#UE4 - Slate Close tooltips when opening modal windows

diff --git a/Engine/Source/Runtime/Slate/Private/SlateApplication.cpp b/Engine/Source/Runtime/Slate/Private/SlateApplication.cpp
--- a/Engine/Source/Runtime/Slate/Private/SlateApplication.cpp
+++ b/Engine/Source/Runtime/Slate/Private/SlateApplication.cpp
@@ -1192,6 +1192,9 @@ void FSlateApplication::AddModalWindow( TSharedRef<SWindow> InSlateWindow, const
 	// Push the active modal window onto the stack.  
 	ActiveModalWindows.AddUnique( InSlateWindow );
 
+	// Close the open tooltip when a new window is open.  Tooltips from non-modal windows can be dangerous and cause rentrancy into code that shouldnt execute in a modal state.
+	CloseToolTip();
+
 	// Set the modal flag on the window
 	InSlateWindow->SetAsModalWindow();
 	
