#ue4
- Prevented FCoreStyle from initializing and failing per file when all of its assets are missing (currently checking for a key asset - the roboto font)
- Saves 1-2 seconds of load time on iPhone4

diff --git a/Engine/Source/Runtime/Slate/Private/CoreStyle.cpp b/Engine/Source/Runtime/Slate/Private/CoreStyle.cpp
--- a/Engine/Source/Runtime/Slate/Private/CoreStyle.cpp
+++ b/Engine/Source/Runtime/Slate/Private/CoreStyle.cpp
@@ -101,6 +101,13 @@ TSharedRef< ISlateStyle > FCoreStyle::Create( const FName& InStyleSetName )
 	TSharedRef< FSlateCoreStyle > Style = MakeShareable( new FSlateCoreStyle(InStyleSetName) );
 	Style->SetContentRoot( FPaths::EngineContentDir() / TEXT("Slate") );
 
+	const FString CanaryPath = Style->RootToContentDir(TEXT("Fonts/Roboto-Regular"), TEXT(".ttf"));
+	if (!FPaths::FileExists(CanaryPath))
+	{
+		UE_LOG(LogSlate, Warning, TEXT("FCoreStyle assets not detected, skipping FCoreStyle initialization"));
+		return Style;
+	}
+
 	// Note, these sizes are in Slate Units.
 	// Slate Units do NOT have to map to pixels.
 	const FVector2D Icon5x16(5.0f, 16.0f);
