Slate: Font: Rendering: Fixed a bug in the slate font renderer that would cause kerning to be applied additionally for ever space directly after a kerned character.

diff --git a/Engine/Source/Runtime/Slate/Private/ElementBatcher.cpp b/Engine/Source/Runtime/Slate/Private/ElementBatcher.cpp
--- a/Engine/Source/Runtime/Slate/Private/ElementBatcher.cpp
+++ b/Engine/Source/Runtime/Slate/Private/ElementBatcher.cpp
@@ -497,6 +497,10 @@ void FSlateElementBatcher::AddTextElement( const FVector2D& Position, const FVec
 			{
 				Kerning = CharacterList.GetKerning( PreviousChar, CurrentChar );
 			}
+			else
+			{
+				Kerning = 0;
+			}
 
 			LineX += Kerning;
 			PreviousChar = CurrentChar;
