Fixed up redundant FText .ToString() usage and some improper FText usage.

diff --git a/Engine/Source/Editor/UnrealEd/Private/PackageRestore.cpp b/Engine/Source/Editor/UnrealEd/Private/PackageRestore.cpp
--- a/Engine/Source/Editor/UnrealEd/Private/PackageRestore.cpp
+++ b/Engine/Source/Editor/UnrealEd/Private/PackageRestore.cpp
@@ -51,7 +51,7 @@ namespace PackageRestore
 		}
 
 		/** @return The tooltip text for this item  */
-		FString GetToolTip() const
+		FText GetToolTip() const
 		{
 			FFormatNamedArguments Args;
 			Args.Add(TEXT("PackageName"), LOCTEXT("PackageName", "Package Name"));
@@ -62,7 +62,7 @@ namespace PackageRestore
 			Args.Add(TEXT("PackageFileStr"), FText::FromString(PackageFilename));
 			Args.Add(TEXT("AutoSaveFileStr"), FText::FromString(AutoSaveFilename));
 
-			return FText::Format(FText::FromString("{PackageName}: {PackageNameStr}\n\n{PackageFile}: {PackageFileStr}\n\n{AutoSaveFile}: {AutoSaveFileStr}"), Args).ToString();
+			return FText::Format(FText::FromString("{PackageName}: {PackageNameStr}\n\n{PackageFile}: {PackageFileStr}\n\n{AutoSaveFile}: {AutoSaveFileStr}"), Args);
 		}
 
 		/** @return Process a request to navigate to the package location */
@@ -128,7 +128,7 @@ namespace PackageRestore
 					.AutoHeight()
 					[
 						SNew(STextBlock)
-						.Text(LOCTEXT("RestoreInfo", "Unreal Editor detected that it did not shut-down cleanly and that the following packages have auto-saves associated with them.\nWould you like to restore from these auto-saves?").ToString())
+						.Text(LOCTEXT("RestoreInfo", "Unreal Editor detected that it did not shut-down cleanly and that the following packages have auto-saves associated with them.\nWould you like to restore from these auto-saves?"))
 						.AutoWrapText(true)
 					]
 					+SVerticalBox::Slot() 
@@ -160,7 +160,7 @@ namespace PackageRestore
 									.FillWidth(1)
 									[
 										SNew(STextBlock) 
-										.Text(LOCTEXT("PackageName", "Package Name").ToString())
+										.Text(LOCTEXT("PackageName", "Package Name"))
 									]
 									+SHorizontalBox::Slot()
 									.Padding(FMargin(4, 0, 0, 0))
@@ -169,7 +169,7 @@ namespace PackageRestore
 									.FillWidth(1)
 									[
 										SNew(STextBlock) 
-										.Text(LOCTEXT("PackageFile", "Package File").ToString())
+										.Text(LOCTEXT("PackageFile", "Package File"))
 									]
 									+SHorizontalBox::Slot()
 									.Padding(FMargin(4, 0, 0, 0))
@@ -178,7 +178,7 @@ namespace PackageRestore
 									.FillWidth(1)
 									[
 										SNew(STextBlock) 
-										.Text(LOCTEXT("AutoSaveFile", "Autosave File").ToString())
+										.Text(LOCTEXT("AutoSaveFile", "Autosave File"))
 									]
 								]
 							]
@@ -203,7 +203,7 @@ namespace PackageRestore
 						.AutoWidth()
 						[
 							SNew(SButton) 
-							.Text(LOCTEXT("RestoreSelectedPackages", "Restore Selected").ToString())
+							.Text(LOCTEXT("RestoreSelectedPackages", "Restore Selected"))
 							.OnClicked(this, &SPackageRestoreDialog::OnRestoreSelectedButtonClicked)
 							.IsEnabled(this, &SPackageRestoreDialog::IsRestoreSelectedButtonEnabled)
 						]
@@ -212,7 +212,7 @@ namespace PackageRestore
 						.AutoWidth()
 						[
 							SNew(SButton) 
-							.Text(LOCTEXT("SkipRestorePackages", "Skip Restore").ToString())
+							.Text(LOCTEXT("SkipRestorePackages", "Skip Restore"))
 							.OnClicked(this, &SPackageRestoreDialog::OnSkipRestoreButtonClicked)
 						]
 					]
diff --git a/Engine/Source/Editor/UnrealEd/Private/Settings/LevelEditorPlaySettingsCustomization.h b/Engine/Source/Editor/UnrealEd/Private/Settings/LevelEditorPlaySettingsCustomization.h
--- a/Engine/Source/Editor/UnrealEd/Private/Settings/LevelEditorPlaySettingsCustomization.h
+++ b/Engine/Source/Editor/UnrealEd/Private/Settings/LevelEditorPlaySettingsCustomization.h
@@ -195,14 +195,14 @@ public:
 						[
 							SNew(STextBlock)
 								.Font(LayoutBuilder->GetDetailFont())
-								.Text(LOCTEXT("CommonResolutionsButtonText", "Common Window Sizes").ToString())
+								.Text(LOCTEXT("CommonResolutionsButtonText", "Common Window Sizes"))
 						]
 						.ContentPadding(FMargin(6.0f, 1.0f))
 						.MenuContent()
 						[
 							MakeCommonResolutionsMenu()
 						]
-						.ToolTipText(LOCTEXT("CommonResolutionsButtonTooltip", "Pick from a list of common screen resolutions").ToString())
+						.ToolTipText(LOCTEXT("CommonResolutionsButtonTooltip", "Pick from a list of common screen resolutions"))
 				]		
 		];
 	}
diff --git a/Engine/Source/Editor/UnrealEd/Public/SSkeletonWidget.h b/Engine/Source/Editor/UnrealEd/Public/SSkeletonWidget.h
--- a/Engine/Source/Editor/UnrealEd/Public/SSkeletonWidget.h
+++ b/Engine/Source/Editor/UnrealEd/Public/SSkeletonWidget.h
@@ -199,7 +199,7 @@ public:
 			+SUniformGridPanel::Slot(0,0)
 			[
 				SNew(SButton)
-				.Text(LOCTEXT("Accept", "Accept").ToString())
+				.Text(LOCTEXT("Accept", "Accept"))
 				.HAlign(HAlign_Center)
 				.ContentPadding(FEditorStyle::GetMargin("StandardDialog.ContentPadding"))
 				.OnClicked_Raw( this, &SSkeletonSelectorWindow::OnAccept )
@@ -207,7 +207,7 @@ public:
 			+SUniformGridPanel::Slot(1,0)
 			[
 				SNew(SButton)
-				.Text(LOCTEXT("Cancel", "Cancel").ToString())
+				.Text(LOCTEXT("Cancel", "Cancel"))
 				.HAlign(HAlign_Center)
 				.ContentPadding(FEditorStyle::GetMargin("StandardDialog.ContentPadding"))
 				.OnClicked_Raw( this, &SSkeletonSelectorWindow::OnCancel )
diff --git a/Engine/Source/Editor/WorldBrowser/Private/StreamingLevels/StreamingLevelCustomization.cpp b/Engine/Source/Editor/WorldBrowser/Private/StreamingLevels/StreamingLevelCustomization.cpp
--- a/Engine/Source/Editor/WorldBrowser/Private/StreamingLevels/StreamingLevelCustomization.cpp
+++ b/Engine/Source/Editor/WorldBrowser/Private/StreamingLevels/StreamingLevelCustomization.cpp
@@ -43,7 +43,7 @@ void FStreamingLevelCustomization::CustomizeDetails(IDetailLayoutBuilder& Detail
 		.NameContent()
 		[
 			SNew(STextBlock)
-				.Text(LOCTEXT("Position", "Position").ToString())
+				.Text(LOCTEXT("Position", "Position"))
 				.Font(IDetailLayoutBuilder::GetDetailFont())
 		]
 		.ValueContent()
@@ -67,7 +67,7 @@ void FStreamingLevelCustomization::CustomizeDetails(IDetailLayoutBuilder& Detail
 		.NameContent()
 		[
 			SNew(STextBlock)
-				.Text(LOCTEXT("Rotation", "Rotation").ToString())
+				.Text(LOCTEXT("Rotation", "Rotation"))
 				.Font(IDetailLayoutBuilder::GetDetailFont())
 		]
 		.ValueContent()
diff --git a/Engine/Source/Runtime/Core/Public/Math/Vector.h b/Engine/Source/Runtime/Core/Public/Math/Vector.h
--- a/Engine/Source/Runtime/Core/Public/Math/Vector.h
+++ b/Engine/Source/Runtime/Core/Public/Math/Vector.h
@@ -637,9 +637,19 @@ public:
 	 */
 	FString ToString() const;
 
+	/**
+	* Get a locale aware textual representation of this vector.
+	*
+	* @return A string describing the vector.
+	*/
+	FText ToText() const;
+
 	/** Get a short textural representation of this vector, for compact readable logging. */
 	FString ToCompactString() const;
 
+	/** Get a short locale aware textural representation of this vector, for compact readable logging. */
+	FText ToCompactText() const;
+
 	/**
 	 * Initialize this Vector based on an FString. The String is expected to contain X=, Y=, Z=.
 	 * The FVector will be bogus when InitFromString returns false.
@@ -1581,6 +1591,68 @@ FORCEINLINE FString FVector::ToString() const
 	return FString::Printf(TEXT("X=%3.3f Y=%3.3f Z=%3.3f"), X, Y, Z);
 }
 
+FORCEINLINE FText FVector::ToText() const
+{
+	FFormatNamedArguments Args;
+	Args.Add(TEXT("X"), X);
+	Args.Add(TEXT("Y"), Y);
+	Args.Add(TEXT("Z"), Z);
+
+	return FText::Format(NSLOCTEXT("Core", "Vector3", "X={X} Y={Y} Z={Z}"), Args);
+}
+
+FORCEINLINE FText FVector::ToCompactText() const
+{
+	if (IsNearlyZero())
+	{
+		return NSLOCTEXT("Core", "Vector3_CompactZeroVector", "V(0)");
+	}
+
+	const bool XIsNotZero = !FMath::IsNearlyZero(X);
+	const bool YIsNotZero = !FMath::IsNearlyZero(Y);
+	const bool ZIsNotZero = !FMath::IsNearlyZero(Z);
+
+	FNumberFormattingOptions FormatRules;
+	FormatRules.MinimumFractionalDigits = 2;
+	FormatRules.MinimumIntegralDigits = 0;
+
+	FFormatNamedArguments Args;
+	Args.Add(TEXT("X"), FText::AsNumber(X, &FormatRules));
+	Args.Add(TEXT("Y"), FText::AsNumber(Y, &FormatRules));
+	Args.Add(TEXT("Z"), FText::AsNumber(Z, &FormatRules));
+
+	if (XIsNotZero && YIsNotZero && ZIsNotZero)
+	{
+		return FText::Format(NSLOCTEXT("Core", "Vector3_CompactXYZ", "V(X={X}, Y={Y}, Z={Z})"), Args);
+	}
+	else if (!XIsNotZero && YIsNotZero && ZIsNotZero)
+	{
+		return FText::Format(NSLOCTEXT("Core", "Vector3_CompactYZ", "V(Y={Y}, Z={Z})"), Args);
+	}
+	else if (XIsNotZero && !YIsNotZero && ZIsNotZero)
+	{
+		return FText::Format(NSLOCTEXT("Core", "Vector3_CompactXZ", "V(X={X}, Z={Z})"), Args);
+	}
+	else if (XIsNotZero && YIsNotZero && !ZIsNotZero)
+	{
+		return FText::Format(NSLOCTEXT("Core", "Vector3_CompactXY", "V(X={X}, Y={Y})"), Args);
+	}
+	else if (!XIsNotZero && !YIsNotZero && ZIsNotZero)
+	{
+		return FText::Format(NSLOCTEXT("Core", "Vector3_CompactZ", "V(Z={Z})"), Args);
+	}
+	else if (XIsNotZero && !YIsNotZero && !ZIsNotZero)
+	{
+		return FText::Format(NSLOCTEXT("Core", "Vector3_CompactX", "V(X={X})"), Args);
+	}
+	else if (!XIsNotZero && YIsNotZero && !ZIsNotZero)
+	{
+		return FText::Format(NSLOCTEXT("Core", "Vector3_CompactY", "V(Y={Y})"), Args);
+	}
+
+	return NSLOCTEXT("Core", "Vector3_CompactZeroVector", "V(0)");
+}
+
 FORCEINLINE FString FVector::ToCompactString() const
 {
 	if( IsNearlyZero() )
diff --git a/Engine/Source/Runtime/Engine/Classes/Debug/GameplayDebuggingHUDComponent.h b/Engine/Source/Runtime/Engine/Classes/Debug/GameplayDebuggingHUDComponent.h
--- a/Engine/Source/Runtime/Engine/Classes/Debug/GameplayDebuggingHUDComponent.h
+++ b/Engine/Source/Runtime/Engine/Classes/Debug/GameplayDebuggingHUDComponent.h
@@ -39,7 +39,8 @@ protected:
 	virtual void DrawNavMeshSnapshot(APlayerController* PC, class UGameplayDebuggingComponent *DebugComponent);
 
 	void PrintAllData();
-	static void CalulateStringSize(const AGameplayDebuggingHUDComponent::FPrintContext& DefaultContext, UFont* Font, const FString& InString, float& OutX, float& OutY );
+	static void CalulateStringSize(const AGameplayDebuggingHUDComponent::FPrintContext& DefaultContext, UFont* Font, const FString& InString, float& OutX, float& OutY);
+	static void CalulateTextSize(const AGameplayDebuggingHUDComponent::FPrintContext& DefaultContext, UFont* Font, const FText& InText, float& OutX, float& OutY);
 	static FVector ProjectLocation(const AGameplayDebuggingHUDComponent::FPrintContext& Context, const FVector& Location);
 	static void DrawItem(const AGameplayDebuggingHUDComponent::FPrintContext& DefaultContext, class FCanvasItem& Item, float X, float Y );
 	static void DrawIcon(const AGameplayDebuggingHUDComponent::FPrintContext& DefaultContext, const FColor& InColor, const FCanvasIcon& Icon, float X, float Y, float Scale = 0.f);
diff --git a/Engine/Source/Runtime/Engine/Classes/Engine/Canvas.h b/Engine/Source/Runtime/Engine/Classes/Engine/Canvas.h
--- a/Engine/Source/Runtime/Engine/Classes/Engine/Canvas.h
+++ b/Engine/Source/Runtime/Engine/Classes/Engine/Canvas.h
@@ -414,6 +414,8 @@ public:
 	 */
 	float DrawText(UFont* InFont, const FString& InText, float X, float Y, float XScale = 1.f, float YScale = 1.f, const FFontRenderInfo& RenderInfo = FFontRenderInfo());
 
+	float DrawText(UFont* InFont, const FText& InText, float X, float Y, float XScale = 1.f, float YScale = 1.f, const FFontRenderInfo& RenderInfo = FFontRenderInfo());
+
 	enum ELastCharacterIndexFormat
 	{
 		// The last whole character before the horizontal offset
diff --git a/Engine/Source/Runtime/Engine/Private/UserInterface/Canvas.cpp b/Engine/Source/Runtime/Engine/Private/UserInterface/Canvas.cpp
--- a/Engine/Source/Runtime/Engine/Private/UserInterface/Canvas.cpp
+++ b/Engine/Source/Runtime/Engine/Private/UserInterface/Canvas.cpp
@@ -1236,25 +1236,31 @@ void VARARGS UCanvas::WrappedStrLenf( UFont* Font, float ScaleX, float ScaleY, i
 	WrappedPrint( false, 0.0f, 0.0f, XL, YL, Font, ScaleX, ScaleY, false, false, Text, Info ); 
 }
 
-float UCanvas::DrawText(UFont* InFont, const FString& InText, float X, float Y, float XScale, float YScale, const FFontRenderInfo& RenderInfo)
+float UCanvas::DrawText(UFont* InFont, const FText& InText, float X, float Y, float XScale, float YScale, const FFontRenderInfo& RenderInfo)
 {
 	ensure(InFont);
-	int32		XL		= 0;
-	int32		YL		= 0; 
+	int32		XL = 0;
+	int32		YL = 0;
 	// need this call in any case to update YL and XL - one of them will be needed anyway
-	WrappedPrint(RenderInfo.bClipText == false, X, Y, XL, YL, InFont, XScale, YScale, bCenterX, bCenterY, *InText, RenderInfo);
+	WrappedPrint(RenderInfo.bClipText == false, X, Y, XL, YL, InFont, XScale, YScale, bCenterX, bCenterY, *InText.ToString(), RenderInfo);
 
 	if (RenderInfo.bClipText)
 	{
-		FCanvasTextItem TextItem( FVector2D( FMath::Trunc(OrgX + X), FMath::Trunc(OrgY + Y) ), FText::FromString( InText ), InFont, DrawColor );
-		TextItem.Scale = FVector2D( XScale, YScale ), 
-		TextItem.BlendMode = SE_BLEND_Translucent;
+		FCanvasTextItem TextItem(FVector2D(FMath::Trunc(OrgX + X), FMath::Trunc(OrgY + Y)), InText, InFont, DrawColor);
+		TextItem.Scale = FVector2D(XScale, YScale),
+			TextItem.BlendMode = SE_BLEND_Translucent;
 		TextItem.FontRenderInfo = RenderInfo;
-		Canvas->DrawItem( TextItem );	
+		Canvas->DrawItem(TextItem);
 	}
+
 	return (float)YL;
 }
 
+float UCanvas::DrawText(UFont* InFont, const FString& InText, float X, float Y, float XScale, float YScale, const FFontRenderInfo& RenderInfo)
+{
+	return DrawText(InFont, FText::FromString(InText), X, Y, XScale, YScale, RenderInfo);
+}
+
 int32 UCanvas::WrappedPrint(bool Draw, float X, float Y, int32& out_XL, int32& out_YL, UFont* Font, float ScaleX, float ScaleY, bool bCenterTextX, bool bCenterTextY, const TCHAR* Text, const FFontRenderInfo& RenderInfo) 
 {
 	if (ClipX < 0 || ClipY < 0)
diff --git a/Engine/Source/Runtime/Slate/Private/SColorPicker.cpp b/Engine/Source/Runtime/Slate/Private/SColorPicker.cpp
--- a/Engine/Source/Runtime/Slate/Private/SColorPicker.cpp
+++ b/Engine/Source/Runtime/Slate/Private/SColorPicker.cpp
@@ -416,15 +416,15 @@ public:
 						[
 							SNew(SImage)
 								.Image(FCoreStyle::Get().GetBrush("ColorPicker.EyeDropper"))
-								.ToolTipText(LOCTEXT("EyeDropperButton_ToolTip", "Activates the eye-dropper for selecting a colored pixel from any window.").ToString())
+								.ToolTipText(LOCTEXT("EyeDropperButton_ToolTip", "Activates the eye-dropper for selecting a colored pixel from any window."))
 								.ColorAndOpacity(this, &SEyeDropperButton::GetDropperImageColor)				
 						]
 
 					+ SOverlay::Slot()
 						[
 							SNew(STextBlock)
-								.Text(LOCTEXT("EscapeCue", "Esc").ToString())
-								.ToolTipText(LOCTEXT("EyeDropperEscapeCue_ToolTip", "Hit Escape key to stop the eye dropper").ToString())
+								.Text(LOCTEXT("EscapeCue", "Esc"))
+								.ToolTipText(LOCTEXT("EyeDropperEscapeCue_ToolTip", "Hit Escape key to stop the eye dropper"))
 								.Visibility(this, &SEyeDropperButton::GetEscapeTextVisibility)
 						]
 				]
@@ -875,7 +875,7 @@ void SColorPicker::GenerateDefaultColorPickerContent( bool bAdvancedSectionExpan
 	ColorThemeComboButton = SNew(SComboButton)
 		.ContentPadding(3.0f)
 		.MenuPlacement(MenuPlacement_ComboBox)
-		.ToolTipText(LOCTEXT("OpenThemeManagerToolTip", "Open Color Theme Manager").ToString());
+		.ToolTipText(LOCTEXT("OpenThemeManagerToolTip", "Open Color Theme Manager"));
 
 	ColorThemeComboButton->SetMenuContent(ColorThemesViewer.ToSharedRef());
 
@@ -913,10 +913,10 @@ void SColorPicker::GenerateDefaultColorPickerContent( bool bAdvancedSectionExpan
 										// color theme bar
 										SAssignNew(CurrentThemeBar, SThemeColorBlocksBar)
 											.ColorTheme(this, &SColorPicker::HandleThemeBarColorTheme)
-											.EmptyText(LOCTEXT("EmptyBarHint", "Drag & drop colors here to save").ToString())
+											.EmptyText(LOCTEXT("EmptyBarHint", "Drag & drop colors here to save"))
 											.HideTrashCallback(this, &SColorPicker::HideSmallTrash)
 											.ShowTrashCallback(this, &SColorPicker::ShowSmallTrash)
-											.ToolTipText(LOCTEXT("CurrentThemeBarToolTip", "Current Color Theme").ToString())
+											.ToolTipText(LOCTEXT("CurrentThemeBarToolTip", "Current Color Theme"))
 											.UseAlpha(SharedThis(this), &SColorPicker::HandleThemeBarUseAlpha)
 											.UseSRGB(SharedThis(this), &SColorPicker::HandleColorPickerUseSRGB)
 											.OnSelectColor(this, &SColorPicker::HandleThemeBarColorSelected)
@@ -928,7 +928,7 @@ void SColorPicker::GenerateDefaultColorPickerContent( bool bAdvancedSectionExpan
 									.VAlign(VAlign_Center)
 									[
 										SNew(STextBlock)
-											.Text(LOCTEXT("EmptyBarHint", "Drag & drop colors here to save").ToString())
+											.Text(LOCTEXT("EmptyBarHint", "Drag & drop colors here to save"))
 											.Visibility(this, &SColorPicker::HandleThemeBarHintVisibility)
 									]
 							]
@@ -949,12 +949,12 @@ void SColorPicker::GenerateDefaultColorPickerContent( bool bAdvancedSectionExpan
 					[
 						// sRGB check box
 						SNew(SCheckBox)
-							.ToolTipText(LOCTEXT("SRGBCheckboxToolTip", "Toggle gamma corrected sRGB previewing").ToString())
+							.ToolTipText(LOCTEXT("SRGBCheckboxToolTip", "Toggle gamma corrected sRGB previewing"))
 							.IsChecked(this, &SColorPicker::HandleSRGBCheckBoxIsChecked)
 							.OnCheckStateChanged(this, &SColorPicker::HandleSRGBCheckBoxCheckStateChanged)
 							[
 								SNew(STextBlock)
-									.Text(LOCTEXT("SRGBCheckboxLabel", "sRGB").ToString())
+									.Text(LOCTEXT("SRGBCheckboxLabel", "sRGB"))
 							]
 					]
 
@@ -1053,7 +1053,7 @@ void SColorPicker::GenerateDefaultColorPickerContent( bool bAdvancedSectionExpan
 											[
 												SNew(SImage)
 													.Image(FCoreStyle::Get().GetBrush("ColorPicker.Mode"))
-													.ToolTipText(LOCTEXT("ColorPickerModeEToolTip", "Toggle between color wheel and color spectrum.").ToString())
+													.ToolTipText(LOCTEXT("ColorPickerModeEToolTip", "Toggle between color wheel and color spectrum."))
 											]									
 									]
 
@@ -1153,7 +1153,7 @@ void SColorPicker::GenerateDefaultColorPickerContent( bool bAdvancedSectionExpan
 									.Padding(0.0f, 8.0f, 0.0f, 0.0f)
 									[
 										SNew(SHorizontalBox)
-											.ToolTipText(LOCTEXT("HexSliderToolTip", "Hexadecimal Value").ToString())
+											.ToolTipText(LOCTEXT("HexSliderToolTip", "Hexadecimal Value"))
 
 										+ SHorizontalBox::Slot()
 											.AutoWidth()
@@ -1195,7 +1195,7 @@ void SColorPicker::GenerateDefaultColorPickerContent( bool bAdvancedSectionExpan
 						SNew(SButton)
 							.ContentPadding( FCoreStyle::Get().GetMargin("StandardDialog.ContentPadding") )
 							.HAlign(HAlign_Center)
-							.Text(LOCTEXT("OKButton", "OK").ToString())
+							.Text(LOCTEXT("OKButton", "OK"))
 							.OnClicked(this, &SColorPicker::HandleOkButtonClicked)
 					]
 
@@ -1205,7 +1205,7 @@ void SColorPicker::GenerateDefaultColorPickerContent( bool bAdvancedSectionExpan
 						SNew(SButton)
 							.ContentPadding( FCoreStyle::Get().GetMargin("StandardDialog.ContentPadding") )
 							.HAlign(HAlign_Center)
-							.Text(LOCTEXT("CancelButton", "Cancel").ToString())
+							.Text(LOCTEXT("CancelButton", "Cancel"))
 							.OnClicked(this, &SColorPicker::HandleCancelButtonClicked)
 					]
 			]
@@ -1442,7 +1442,7 @@ TSharedRef<SWidget> SColorPicker::MakeColorSlider( EColorPickerChannels::Type Ch
 	}
 
 	return SNew(SOverlay)
-		.ToolTipText(SliderTooltip.ToString())
+		.ToolTipText(SliderTooltip)
 
 	+ SOverlay::Slot()
 		.Padding(FMargin(4.0f, 0.0f))
@@ -1565,14 +1565,14 @@ TSharedRef<SWidget> SColorPicker::MakeColorSpinBox( EColorPickerChannels::Type C
 		.VAlign(VAlign_Center)
 		[
 			SNew(STextBlock)
-				.Text(SliderLabel.ToString())
+				.Text(SliderLabel)
 		]
 
 	+ SHorizontalBox::Slot()
 		.FillWidth(1.0f)
 		[
 			SNew(SVerticalBox)
-				.ToolTipText(SliderTooltip.ToString())
+				.ToolTipText(SliderTooltip)
 
 			+ SVerticalBox::Slot()
 				.AutoHeight()
diff --git a/Engine/Source/Runtime/Slate/Private/SColorThemes.cpp b/Engine/Source/Runtime/Slate/Private/SColorThemes.cpp
--- a/Engine/Source/Runtime/Slate/Private/SColorThemes.cpp
+++ b/Engine/Source/Runtime/Slate/Private/SColorThemes.cpp
@@ -324,18 +324,18 @@ const FSlateBrush* SThemeColorBlock::HandleBorderImage() const
 }
 
 #define LOCTEXT_NAMESPACE "ThemeColorBlock"
-FString SThemeColorBlock::GetRedText() const {return ColorPtr.IsValid() ? LOCTEXT("Red", "R").ToString() + FormatToolTipText(ColorPtr.Pin()->HSVToLinearRGB().R) : TEXT("");}
-FString SThemeColorBlock::GetGreenText() const {return ColorPtr.IsValid() ? LOCTEXT("Green", "G").ToString() + FormatToolTipText(ColorPtr.Pin()->HSVToLinearRGB().G) : TEXT("");}
-FString SThemeColorBlock::GetBlueText() const {return ColorPtr.IsValid() ? LOCTEXT("Blue", "B").ToString() + FormatToolTipText(ColorPtr.Pin()->HSVToLinearRGB().B) : TEXT("");}
-FString SThemeColorBlock::GetAlphaText() const {return ColorPtr.IsValid() ? LOCTEXT("Alpha", "A").ToString() + FormatToolTipText(ColorPtr.Pin()->HSVToLinearRGB().A) : TEXT("");}
-FString SThemeColorBlock::GetHueText() const {return ColorPtr.IsValid() ? LOCTEXT("Hue", "H").ToString() + FormatToolTipText(ColorPtr.Pin()->R) : TEXT("");}
-FString SThemeColorBlock::GetSaturationText() const {return ColorPtr.IsValid() ? LOCTEXT("Saturation", "S").ToString() + FormatToolTipText(ColorPtr.Pin()->G) : TEXT("");}
-FString SThemeColorBlock::GetValueText() const {return ColorPtr.IsValid() ? LOCTEXT("Value", "V").ToString() + FormatToolTipText(ColorPtr.Pin()->B) : TEXT("");}
-#undef LOCTEXT_NAMESPACE
-
-FString SThemeColorBlock::FormatToolTipText(float Value) const
+FText SThemeColorBlock::GetRedText() const { return ColorPtr.IsValid() ? FormatToolTipText(LOCTEXT("Red", "R"), ColorPtr.Pin()->HSVToLinearRGB().R) : FText::GetEmpty(); }
+FText SThemeColorBlock::GetGreenText() const { return ColorPtr.IsValid() ? FormatToolTipText(LOCTEXT("Green", "G"), ColorPtr.Pin()->HSVToLinearRGB().G) : FText::GetEmpty(); }
+FText SThemeColorBlock::GetBlueText() const { return ColorPtr.IsValid() ? FormatToolTipText(LOCTEXT("Blue", "B"), ColorPtr.Pin()->HSVToLinearRGB().B) : FText::GetEmpty(); }
+FText SThemeColorBlock::GetAlphaText() const { return ColorPtr.IsValid() ? FormatToolTipText(LOCTEXT("Alpha", "A"), ColorPtr.Pin()->HSVToLinearRGB().A) : FText::GetEmpty(); }
+FText SThemeColorBlock::GetHueText() const { return ColorPtr.IsValid() ? FormatToolTipText(LOCTEXT("Hue", "H"), ColorPtr.Pin()->R) : FText::GetEmpty(); }
+FText SThemeColorBlock::GetSaturationText() const { return ColorPtr.IsValid() ? FormatToolTipText(LOCTEXT("Saturation", "S"), ColorPtr.Pin()->G) : FText::GetEmpty(); }
+FText SThemeColorBlock::GetValueText() const { return ColorPtr.IsValid() ? FormatToolTipText(LOCTEXT("Value", "V"), ColorPtr.Pin()->B) : FText::GetEmpty(); }
+
+FText SThemeColorBlock::FormatToolTipText(const FText& ColorIdentifier, float Value) const
 {
-	FString Result = FString(TEXT(": "));
+	FFormatNamedArguments Args;
+	Args.Add(TEXT("Identifier"), ColorIdentifier);
 
 	if (Value >= 0.f)
 	{
@@ -343,10 +343,20 @@ FString SThemeColorBlock::FormatToolTipText(float Value) const
 		int32 PreRadixDigits = FMath::Max(0, int32(FMath::Loge(Value + KINDA_SMALL_NUMBER) * LogToLog10));
 
 		int32 Precision = FMath::Max(0, 2 - PreRadixDigits);
-		Result += FString::Printf(*FString::Printf(TEXT("%%.%if"), Precision), Value);
+
+		FNumberFormattingOptions FormatRules;
+		FormatRules.MinimumFractionalDigits = Precision;
+
+		Args.Add(TEXT("Value"), FText::AsNumber(Value, &FormatRules));
 	}
-	return Result;
+	else
+	{
+		Args.Add(TEXT("Value"), FText::GetEmpty());
+	}
+
+	return FText::Format(LOCTEXT("ToolTipFormat", "{Identifier}: {Value}"), Args);
 }
+#undef LOCTEXT_NAMESPACE
 
 bool SThemeColorBlock::OnReadIgnoreAlpha() const
 {
@@ -702,22 +712,22 @@ void SColorThemesViewer::Construct(const FArguments& InArgs)
 		SNew(SHorizontalBox)
 		+SHorizontalBox::Slot() .FillWidth(1) .Padding(3)
 		[
-			SNew(SButton) .Text(LOCTEXT("NewButton", "New").ToString()) .HAlign(HAlign_Center)
+			SNew(SButton) .Text(LOCTEXT("NewButton", "New")) .HAlign(HAlign_Center)
 				.OnClicked(this, &SColorThemesViewer::NewColorTheme)
 		]
 		+SHorizontalBox::Slot() .FillWidth(1) .Padding(3)
 		[
-			SNew(SButton) .Text(LOCTEXT("DuplicateButton", "Duplicate").ToString()) .HAlign(HAlign_Center)
+			SNew(SButton) .Text(LOCTEXT("DuplicateButton", "Duplicate")) .HAlign(HAlign_Center)
 				.OnClicked(this, &SColorThemesViewer::DuplicateColorTheme)
 		]
 		+SHorizontalBox::Slot() .FillWidth(1) .Padding(3)
 		[
-			SNew(SButton) .Text(LOCTEXT("RenameButton", "Rename").ToString()) .HAlign(HAlign_Center)
+			SNew(SButton) .Text(LOCTEXT("RenameButton", "Rename")) .HAlign(HAlign_Center)
 				.OnClicked(this, &SColorThemesViewer::MenuToRename)
 		]
 		+SHorizontalBox::Slot() .FillWidth(1) .Padding(3)
 		[
-			SNew(SButton) .Text(LOCTEXT("DeleteButton", "Delete").ToString()) .HAlign(HAlign_Center)
+			SNew(SButton) .Text(LOCTEXT("DeleteButton", "Delete")) .HAlign(HAlign_Center)
 				.OnClicked(this, &SColorThemesViewer::MenuToDelete)
 		];
 
@@ -730,12 +740,12 @@ void SColorThemesViewer::Construct(const FArguments& InArgs)
 		]
 		+SHorizontalBox::Slot().AutoWidth() .Padding(3)
 		[
-			SNew(SButton) .Text(LOCTEXT("AcceptRenameButton", "Accept").ToString()) .HAlign(HAlign_Right)
+			SNew(SButton) .Text(LOCTEXT("AcceptRenameButton", "Accept")) .HAlign(HAlign_Right)
 				.OnClicked(this, &SColorThemesViewer::ChangeThemeName)
 		]
 		+SHorizontalBox::Slot().AutoWidth() .Padding(3)
 		[
-			SNew(SButton) .Text(LOCTEXT("CancelRenameButton", "Cancel").ToString()) .HAlign(HAlign_Right)
+			SNew(SButton) .Text(LOCTEXT("CancelRenameButton", "Cancel")) .HAlign(HAlign_Right)
 				.OnClicked(this, &SColorThemesViewer::MenuToStandard)
 		];
 
@@ -744,16 +754,16 @@ void SColorThemesViewer::Construct(const FArguments& InArgs)
 		SNew(SHorizontalBox)
 		+SHorizontalBox::Slot() .FillWidth(1) .HAlign(HAlign_Left) .Padding(3)
 		[
-			SNew(SButton) .Text(LOCTEXT("ConfirmDeleteButton", "Delete").ToString()) .HAlign(HAlign_Center)
+			SNew(SButton) .Text(LOCTEXT("ConfirmDeleteButton", "Delete")) .HAlign(HAlign_Center)
 				.OnClicked(this, &SColorThemesViewer::DeleteColorTheme)
 		]
 		+SHorizontalBox::Slot() .FillWidth(1) .HAlign(HAlign_Center) .VAlign(VAlign_Center) .Padding(3)
 		[
-			SNew(STextBlock) .Text(LOCTEXT("ConfirmDeletePrompt", "Confirm Delete").ToString()) .Font(SmallLayoutFont)
+			SNew(STextBlock) .Text(LOCTEXT("ConfirmDeletePrompt", "Confirm Delete")) .Font(SmallLayoutFont)
 		]
 		+SHorizontalBox::Slot() .FillWidth(1) .HAlign(HAlign_Right) .Padding(3)
 		[
-			SNew(SButton) .Text(LOCTEXT("CancelDeleteButton", "Cancel").ToString()) .HAlign(HAlign_Center)
+			SNew(SButton) .Text(LOCTEXT("CancelDeleteButton", "Cancel")) .HAlign(HAlign_Center)
 				.OnClicked(this, &SColorThemesViewer::MenuToStandard)
 		];
 
diff --git a/Engine/Source/Runtime/Slate/Public/SColorThemes.h b/Engine/Source/Runtime/Slate/Public/SColorThemes.h
--- a/Engine/Source/Runtime/Slate/Public/SColorThemes.h
+++ b/Engine/Source/Runtime/Slate/Public/SColorThemes.h
@@ -159,18 +159,18 @@ private:
 	FSlateColor HandleBorderColor() const;
 	const FSlateBrush* HandleBorderImage() const;
 
-	FString GetRedText() const;
-	FString GetGreenText() const;
-	FString GetBlueText() const;
-	FString GetAlphaText() const;
-	FString GetHueText() const;
-	FString GetSaturationText() const;
-	FString GetValueText() const;
+	FText GetRedText() const;
+	FText GetGreenText() const;
+	FText GetBlueText() const;
+	FText GetAlphaText() const;
+	FText GetHueText() const;
+	FText GetSaturationText() const;
+	FText GetValueText() const;
 
 	/**
 	 * Function for formatting text for the tooltip which has limited space.
 	 */
-	FString FormatToolTipText(float Value) const;
+	FText FormatToolTipText(const FText& ColorIdentifier, float Value) const;
 	
 	bool OnReadIgnoreAlpha() const;
 	bool OnReadShowBackgroundForAlpha() const;
@@ -228,7 +228,7 @@ public:
 		SLATE_EVENT(FSimpleDelegate, HideTrashCallback)
 		
 		/** Specify what the bar should display when no colors are present */
-		SLATE_ARGUMENT(FString, EmptyText)
+		SLATE_TEXT_ARGUMENT(EmptyText)
 
 		/** Whether to display sRGB color */
 		SLATE_ATTRIBUTE(bool, UseSRGB)
