Added app selection drop down to launcher.

diff --git a/Engine/Source/Runtime/Slate/Public/SComboBox.h b/Engine/Source/Runtime/Slate/Public/SComboBox.h
--- a/Engine/Source/Runtime/Slate/Public/SComboBox.h
+++ b/Engine/Source/Runtime/Slate/Public/SComboBox.h
@@ -82,6 +82,7 @@ public:
 		, _InitiallySelectedItem( NULL )
 		, _Method( SMenuAnchor::CreateNewWindow )
 		, _MaxListHeight(450.0f)
+		, _HasDownArrow( true )
 		{}
 		
 		/** Slot for this button's content (optional) */
@@ -115,6 +116,12 @@ public:
 
 		/** The sound to play when the selection changes (overrides ComboBoxStyle) */
 		SLATE_ARGUMENT( TOptional<FSlateSound>, SelectionChangeSoundOverride )
+
+		/**
+		 * When false, the down arrow is not generated and it is up to the API consumer
+		 * to make their own visual hint that this is a drop down.
+		 */
+		SLATE_ARGUMENT( bool, HasDownArrow )
 				
 	SLATE_END_ARGS()
 
@@ -175,7 +182,7 @@ public:
 			[
 				MenuContent
 			]
-			.HasDownArrow( true )
+			.HasDownArrow( InArgs._HasDownArrow )
 			.ContentPadding( InArgs._ContentPadding )
 			.ForegroundColor( InArgs._ForegroundColor )
 		);
