XboxOne Help and Support menu selections

diff --git a/Engine/Source/Runtime/Core/Public/GenericPlatform/GenericApplication.h b/Engine/Source/Runtime/Core/Public/GenericPlatform/GenericApplication.h
--- a/Engine/Source/Runtime/Core/Public/GenericPlatform/GenericApplication.h
+++ b/Engine/Source/Runtime/Core/Public/GenericPlatform/GenericApplication.h
@@ -288,6 +288,10 @@ public:
     /** Send any analytics captured by the application */
     virtual void SendAnalytics(IAnalyticsProvider* Provider) { }
 
+	virtual bool SupportsSystemHelp() const { return false; }
+
+	virtual void ShowSystemHelp() {}
+
 public:
 
 	const TSharedPtr< ICursor > Cursor;
diff --git a/Engine/Source/Runtime/Slate/Public/SlateApplication.h b/Engine/Source/Runtime/Slate/Public/SlateApplication.h
--- a/Engine/Source/Runtime/Slate/Public/SlateApplication.h
+++ b/Engine/Source/Runtime/Slate/Public/SlateApplication.h
@@ -212,6 +212,11 @@ public:
 	/** @return The force feedback interface for this application */
 	IForceFeedbackSystem *GetForceFeedbackSystem() const { return PlatformApplication->GetForceFeedbackSystem(); }
 
+	/** @return Whether or not the current platform supports system help */
+	bool SupportsSystemHelp() const { return PlatformApplication->SupportsSystemHelp(); }
+
+	void ShowSystemHelp() { PlatformApplication->ShowSystemHelp(); }
+
 	/** @return The text input method interface for this application */
 	ITextInputMethodSystem *GetTextInputMethodSystem() const { return PlatformApplication->GetTextInputMethodSystem(); }
 
