#ttp 331814 - LOC: IME: Some Korean keyboard sequences do not correctly trigger the Slate OnTextChanged event
#summary Fixed OnTextChanged event not being triggered for IME input operations
#proj Editor
#branch UE4

diff --git a/Engine/Source/Runtime/Slate/Private/SEditableText.cpp b/Engine/Source/Runtime/Slate/Private/SEditableText.cpp
--- a/Engine/Source/Runtime/Slate/Private/SEditableText.cpp
+++ b/Engine/Source/Runtime/Slate/Private/SEditableText.cpp
@@ -391,7 +391,13 @@ void SEditableText::FTextInputMethodContext::SetTextInRange(const uint32 BeginIn
 	{
 		const TSharedRef<SEditableText> This = OwningWidget.Pin().ToSharedRef();
 		const FString& OldText = This->EditedText.ToString();
+		
+		// We don't use Start/FinishEditing text here because the whole IME operation handles that.
+		// Also, we don't want to support undo for individual characters added during an IME context
 		This->EditedText = FText::FromString(OldText.Left( BeginIndex ) + InString + OldText.Mid( BeginIndex + Length ));
+		This->SaveText();
+
+		This->OnTextChanged.ExecuteIfBound(This->EditedText);
 	}
 }
 
