#ue4 fix ps4 compile.

diff --git a/Engine/Source/Runtime/Slate/Public/SNumericEntryBox.h b/Engine/Source/Runtime/Slate/Public/SNumericEntryBox.h
--- a/Engine/Source/Runtime/Slate/Public/SNumericEntryBox.h
+++ b/Engine/Source/Runtime/Slate/Public/SNumericEntryBox.h
@@ -27,17 +27,17 @@ public:
 		, _Label()
 		, _LabelVAlign( VAlign_Fill )
 		, _LabelPadding( FMargin(3,0) )
+		, _BorderForegroundColor(FCoreStyle::Get().GetSlateColor("InvertedForeground"))
+		, _BorderBackgroundColor(FLinearColor::White)
 		, _UndeterminedString( TEXT("---") )
 		, _Font( FCoreStyle::Get().GetFontStyle( TEXT("NormalFont") ) )
 		, _AllowSpin(false)
 		, _Delta(0)
 		, _MinValue(TNumericLimits<NumericType>::Lowest())
 		, _MaxValue(TNumericLimits<NumericType>::Max())
-		, _MinSliderValue(0)
+		, _MinSliderValue(0)				
 		, _MaxSliderValue(100)
 		, _SliderExponent(1.f)
-		, _BorderForegroundColor(FCoreStyle::Get().GetSlateColor("InvertedForeground"))
-		, _BorderBackgroundColor(FLinearColor::White)
 	{}		
 
 		/** Style to use for the editable text box within this widget */
