Adding support to the Slate Application to capture "screenshots" of widgets.  You can now ask for an image of a particular widget (or a sub rect inside the widget), the widget must be owned by a Window for this function to succeed.  The function is called TakeScreenshot.

RedrawWindowForScreenshot has been renamed to ForceRedrawWindow.

diff --git a/Engine/Source/Runtime/Core/Public/Math/IntRect.h b/Engine/Source/Runtime/Core/Public/Math/IntRect.h
--- a/Engine/Source/Runtime/Core/Public/Math/IntRect.h
+++ b/Engine/Source/Runtime/Core/Public/Math/IntRect.h
@@ -275,6 +275,12 @@ public:
 	 */
 	int32 Width( ) const;
 
+	/**
+	 * Returns true if the rectangle is 0 x 0.
+	 *
+	 * @return true if the rectangle is 0 x 0.
+	 */
+	bool IsEmpty() const;
 
 public:
 
@@ -535,4 +541,9 @@ FORCEINLINE FString FIntRect::ToString( ) const
 FORCEINLINE int32 FIntRect::Width( ) const
 {
 	return Max.X-Min.X;
-}
\ No newline at end of file
+}
+
+FORCEINLINE bool FIntRect::IsEmpty() const
+{
+	return Width() == 0 && Height() == 0;
+}
diff --git a/Engine/Source/Runtime/Engine/Private/GameViewportClient.cpp b/Engine/Source/Runtime/Engine/Private/GameViewportClient.cpp
--- a/Engine/Source/Runtime/Engine/Private/GameViewportClient.cpp
+++ b/Engine/Source/Runtime/Engine/Private/GameViewportClient.cpp
@@ -583,16 +583,16 @@ void UGameViewportClient::Draw(FViewport* InViewport, FCanvas* SceneCanvas)
 		EngineShowFlags)
 		.SetRealtimeUpdate(true));
 
-    // Allow HMD to modify the view later, just before rendering
-    if (GEngine->HMDDevice.IsValid() && GEngine->IsStereoscopic3D())
-    {
-        ISceneViewExtension* HmdViewExt = GEngine->HMDDevice->GetViewExtension();
-        if (HmdViewExt)
-        {
-            ViewFamily.ViewExtensions.Add(HmdViewExt);
-            HmdViewExt->ModifyShowFlags(ViewFamily.EngineShowFlags);
-        }
-    }
+	// Allow HMD to modify the view later, just before rendering
+	if (GEngine->HMDDevice.IsValid() && GEngine->IsStereoscopic3D())
+	{
+		ISceneViewExtension* HmdViewExt = GEngine->HMDDevice->GetViewExtension();
+		if (HmdViewExt)
+		{
+			ViewFamily.ViewExtensions.Add(HmdViewExt);
+			HmdViewExt->ModifyShowFlags(ViewFamily.EngineShowFlags);
+		}
+	}
 
 
 	ESplitScreenType::Type SplitScreenConfig = GetCurrentSplitscreenConfiguration();
@@ -681,9 +681,9 @@ void UGameViewportClient::Draw(FViewport* InViewport, FCanvas* SceneCanvas)
 							// Update the listener.
 							if (AudioDevice != NULL)
 							{
-		 						FVector Location;
- 								FVector ProjFront;
- 								FVector ProjRight;
+								FVector Location;
+								FVector ProjFront;
+								FVector ProjRight;
 								PlayerController->GetAudioListenerPosition(/*out*/ Location, /*out*/ ProjFront, /*out*/ ProjRight);
 
 								FTransform ListenerTransform(FRotationMatrix::MakeFromXY(ProjFront, ProjRight));
@@ -966,7 +966,7 @@ void UGameViewportClient::ProcessScreenShots(FViewport* InViewport)
 		if (!GIsDumpingMovie && (FScreenshotRequest::ShouldShowUI() && WindowPtr.IsValid()))
 		{
 			TSharedRef<SWindow> WindowRef = WindowPtr.ToSharedRef();
-			FSlateApplication::Get().RedrawWindowForScreenshot( WindowRef );
+			FSlateApplication::Get().ForceRedrawWindow(WindowRef);
 		}
 
 		if (GetViewportScreenShot(InViewport, Bitmap))
diff --git a/Engine/Source/Runtime/Slate/Private/SlateApplication.cpp b/Engine/Source/Runtime/Slate/Private/SlateApplication.cpp
--- a/Engine/Source/Runtime/Slate/Private/SlateApplication.cpp
+++ b/Engine/Source/Runtime/Slate/Private/SlateApplication.cpp
@@ -3036,11 +3036,51 @@ void FSlateApplication::GotoLineInSource(const FString& FileAndLineNumber) const
 	}
 }
 
-void FSlateApplication::RedrawWindowForScreenshot( TSharedRef<SWindow>& InWindowToDraw )
+void FSlateApplication::ForceRedrawWindow(TSharedRef<SWindow>& InWindowToDraw)
 {
 	PrivateDrawWindows( InWindowToDraw );
 }
 
+bool FSlateApplication::TakeScreenshot(TSharedRef<SWidget>& Widget, TArray<FColor>&OutColorData, FIntVector& OutSize)
+{
+	return TakeScreenshot(Widget, FIntRect(), OutColorData, OutSize);
+}
+
+bool FSlateApplication::TakeScreenshot(TSharedRef<SWidget>& Widget, const FIntRect& InnerWidgetArea, TArray<FColor>& OutColorData, FIntVector& OutSize)
+{
+	// We can't screenshot the widget unless there's a valid window handle to draw it in.
+	TSharedPtr<SWindow> WidgetWindow = FSlateApplication::Get().FindWidgetWindow(Widget);
+	if ( !WidgetWindow.IsValid() )
+	{
+		return false;
+	}
+
+	TSharedRef<SWindow> CurrentWindowRef = WidgetWindow.ToSharedRef();
+
+	FWidgetPath WidgetPath;
+	FSlateApplication::Get().GeneratePathToWidgetChecked(Widget, WidgetPath);
+
+	FArrangedWidget ArrangedWidget = WidgetPath.FindArrangedWidget(Widget);
+	FVector2D Position = ArrangedWidget.Geometry.AbsolutePosition;
+	FVector2D Size = ArrangedWidget.Geometry.GetDrawSize();
+	FVector2D WindowPosition = WidgetWindow->GetPositionInScreen();
+
+	FIntRect ScreenshotRect = InnerWidgetArea.IsEmpty() ? FIntRect(0, 0, (int32)Size.X, (int32)Size.Y) : InnerWidgetArea;
+
+	ScreenshotRect.Min.X += ( Position.X - WindowPosition.X );
+	ScreenshotRect.Min.Y += ( Position.Y - WindowPosition.Y );
+	ScreenshotRect.Max.X += ( Position.X - WindowPosition.X );
+	ScreenshotRect.Max.Y += ( Position.Y - WindowPosition.Y );
+
+	Renderer->PrepareToTakeScreenshot(ScreenshotRect, &OutColorData);
+	PrivateDrawWindows(WidgetWindow);
+
+	OutSize.X = ScreenshotRect.Size().X;
+	OutSize.Y = ScreenshotRect.Size().Y;
+
+	return true;
+}
+
 static TSharedPtr<SWindow> FindWindowByPlatformWindow( const TArray< TSharedRef<SWindow> >& WindowsToSearch, const TSharedRef< FGenericWindow >& PlatformWindow )
 {
 	for (int32 WindowIndex=0; WindowIndex < WindowsToSearch.Num(); ++WindowIndex)
diff --git a/Engine/Source/Runtime/Slate/Public/Renderer.h b/Engine/Source/Runtime/Slate/Public/Renderer.h
--- a/Engine/Source/Runtime/Slate/Public/Renderer.h
+++ b/Engine/Source/Runtime/Slate/Public/Renderer.h
@@ -191,6 +191,12 @@ public:
 	 */
 	virtual void InvalidateAllViewports() {}
 
+	/** 
+	 * Prepares the renderer to take a screenshot of the UI.  The Rect is portion of the rendered output
+	 * that will be stored into the TArray of FColors.
+	 */
+	virtual void PrepareToTakeScreenshot(const FIntRect& Rect, TArray<FColor>* OutColorData) {}
+
 protected:
 	TSharedPtr<class FSlateFontCache> FontCache;
 	TSharedPtr<class FSlateFontMeasure> FontMeasure;
diff --git a/Engine/Source/Runtime/Slate/Public/SlateApplication.h b/Engine/Source/Runtime/Slate/Public/SlateApplication.h
--- a/Engine/Source/Runtime/Slate/Public/SlateApplication.h
+++ b/Engine/Source/Runtime/Slate/Public/SlateApplication.h
@@ -713,9 +713,25 @@ public:
 	FPopupSupport& GetPopupSupport() { return PopupSupport; }
 
 	/**
-	 * Redraws a provided window for a screenshot
+	 * Forces the window to redraw immediately.
 	 */
-	void RedrawWindowForScreenshot( TSharedRef<SWindow>& InWindowToDraw );
+	void ForceRedrawWindow( TSharedRef<SWindow>& InWindowToDraw );
+
+	/**
+	 * Takes a screenshot of the widget writing the results into the color buffer provided.  Note that the format is BGRA.
+	 * the size of the resulting image is also output.
+	 * 
+	 * @return true if taking the screenshot was successful.
+	 */
+	bool TakeScreenshot(TSharedRef<SWidget>& Widget, TArray<FColor>&OutColorData, FIntVector& OutSize);
+
+	/**
+	 * Takes a screenshot of the widget writing the results into the color buffer provided, this version allows you to provide 
+	 * an inner area to screenshot.  Note that the format is BGRA.  The size of the resulting image is also output.
+	 *
+	 * @return true if taking the screenshot was successful.
+	 */
+	bool TakeScreenshot(TSharedRef<SWidget>& Widget, const FIntRect& InnerWidgetArea, TArray<FColor>& OutColorData, FIntVector& OutSize);
 
 protected:
 
diff --git a/Engine/Source/Runtime/Slate/Public/WidgetPath.h b/Engine/Source/Runtime/Slate/Public/WidgetPath.h
--- a/Engine/Source/Runtime/Slate/Public/WidgetPath.h
+++ b/Engine/Source/Runtime/Slate/Public/WidgetPath.h
@@ -84,7 +84,7 @@ public:
 		return false;
 	}
 
-	FArrangedWidget FindArrangedWidget( TSharedRef<const SWidget> WidgetToFind ) const;
+	SLATE_API FArrangedWidget FindArrangedWidget( TSharedRef<const SWidget> WidgetToFind ) const;
 	
 	/**
 	 * Get the first (top-most) widget in this path, which is always a window; assumes path is valid
diff --git a/Engine/Source/Runtime/SlateRHIRenderer/Private/SlateRHIRenderer.cpp b/Engine/Source/Runtime/SlateRHIRenderer/Private/SlateRHIRenderer.cpp
--- a/Engine/Source/Runtime/SlateRHIRenderer/Private/SlateRHIRenderer.cpp
+++ b/Engine/Source/Runtime/SlateRHIRenderer/Private/SlateRHIRenderer.cpp
@@ -147,7 +147,8 @@ FSlateRHIRenderer::FSlateRHIRenderer()
 							FPlane(0,	0,	1,  0),
 							FPlane(0,	0,	0,	1));
 
-
+	bTakingAScreenShot = false;
+	OutScreenshotData = NULL;
 }
 
 FSlateRHIRenderer::~FSlateRHIRenderer()
@@ -536,6 +537,15 @@ static void EndDrawingWindows( FSlateDrawBuffer* DrawBuffer, FSlateRHIRenderingP
 }
 
 
+void FSlateRHIRenderer::PrepareToTakeScreenshot(const FIntRect& Rect, TArray<FColor>* OutColorData)
+{
+	check(OutColorData);
+
+	bTakingAScreenShot = true;
+	ScreenshotRect = Rect;
+	OutScreenshotData = OutColorData;
+}
+
 /** 
  * Creates necessary resources to render a window and sends draw commands to the rendering thread
  *
@@ -638,6 +648,23 @@ void FSlateRHIRenderer::DrawWindows_Private( FSlateDrawBuffer& WindowDrawBuffer
 						Params.Renderer->DrawWindow_RenderThread( *Params.ViewportInfo, *Params.WindowElementList, Params.bLockToVsync );
 						Params.MarkWindowAsDrawn.ExecuteIfBound();
 					});
+
+					if ( bTakingAScreenShot )
+					{
+						ENQUEUE_UNIQUE_RENDER_COMMAND_THREEPARAMETER(SlateCaptureScreenshotCommand,
+							FSlateDrawWindowCommandParams, Params, Params,
+							FIntRect, ScreenshotRect, ScreenshotRect,
+							TArray<FColor>*, OutScreenshotData, OutScreenshotData,
+						{
+							FTexture2DRHIRef BackBuffer = RHIGetViewportBackBuffer(Params.ViewportInfo->ViewportRHI);
+							RHIReadSurfaceData(BackBuffer, ScreenshotRect, *OutScreenshotData, FReadSurfaceDataFlags());
+						});
+
+						FlushRenderingCommands();
+
+						bTakingAScreenShot = false;
+						OutScreenshotData = NULL;
+					}
 				}
 			}
 		}
@@ -797,7 +824,7 @@ void FSlateRHIRenderer::CopyWindowsToDrawBuffer(const TArray<FString>& KeypressB
 		FWriteMouseCursorAndKeyPressesContext, Context, WriteMouseCursorAndKeyPressesContext,
 	{
 		RHISetBlendState(TStaticBlendState<CW_RGBA,BO_Add,BF_SourceAlpha,BF_InverseSourceAlpha,BO_Add,BF_Zero,BF_One>::GetRHI());
-        
+		
 		Context.RenderPolicy->UpdateBuffers(*Context.SlateElementList);
 		if( Context.SlateElementList->GetRenderBatches().Num() > 0 )
 		{
diff --git a/Engine/Source/Runtime/SlateRHIRenderer/Private/SlateRHIRenderer.h b/Engine/Source/Runtime/SlateRHIRenderer/Private/SlateRHIRenderer.h
--- a/Engine/Source/Runtime/SlateRHIRenderer/Private/SlateRHIRenderer.h
+++ b/Engine/Source/Runtime/SlateRHIRenderer/Private/SlateRHIRenderer.h
@@ -183,6 +183,8 @@ public:
 	 */
 	virtual void InvalidateAllViewports();
 
+	virtual void PrepareToTakeScreenshot(const FIntRect& Rect, TArray<FColor>* OutColorData) OVERRIDE;
+
 private:
 	/** Loads all known textures from Slate styles */
 	void LoadUsedTextures();
@@ -233,4 +235,8 @@ private:
 
 	/** The resource which allows us to capture the editor to a buffer */
 	FSlateCrashReportResource* CrashTrackerResource;
+
+	bool bTakingAScreenShot;
+	FIntRect ScreenshotRect;
+	TArray<FColor>* OutScreenshotData;
 };
