Fix duplicate line in SRichTextBlock::SetHighlightText

diff --git a/Engine/Source/Runtime/Slate/Private/SRichTextBlock.cpp b/Engine/Source/Runtime/Slate/Private/SRichTextBlock.cpp
--- a/Engine/Source/Runtime/Slate/Private/SRichTextBlock.cpp
+++ b/Engine/Source/Runtime/Slate/Private/SRichTextBlock.cpp
@@ -203,7 +203,6 @@ void SRichTextBlock::SetHighlightText( const FText& InHighlightText )
 	Highlights.Empty();
 	HighlightText = InHighlightText;
 
-	HighlightText = InHighlightText;
 	const FString HighlightTextString = HighlightText.ToString();
 	const int32 HighlightTextLength = HighlightTextString.Len();
 
