- Fix typo in PositionIsInside

diff --git a/Engine/Source/Runtime/Slate/Private/SVirtualJoystick.cpp b/Engine/Source/Runtime/Slate/Private/SVirtualJoystick.cpp
--- a/Engine/Source/Runtime/Slate/Private/SVirtualJoystick.cpp
+++ b/Engine/Source/Runtime/Slate/Private/SVirtualJoystick.cpp
@@ -97,7 +97,7 @@ static bool PositionIsInside(const SVirtualJoystick::FControlInfo& Control, cons
 		Position.X >= Control.Center.X - BoxSize.X * 0.5f &&
 		Position.X <= Control.Center.X + BoxSize.X * 0.5f &&
 		Position.Y >= Control.Center.Y - BoxSize.Y * 0.5f &&
-		Position.Y >= Control.Center.Y - BoxSize.Y * 0.5f;
+		Position.Y <= Control.Center.Y + BoxSize.Y * 0.5f;
 }
 
 int32 SVirtualJoystick::OnPaint( const FGeometry& AllottedGeometry, const FSlateRect& MyClippingRect, FSlateWindowElementList& OutDrawElements, int32 LayerId, const FWidgetStyle& InWidgetStyle, bool bParentEnabled ) const
