#ttp 331158 - UI: Mouse orientation is skewed when rejecting the change from windowed to full screen
#ttp 331513 - Changing from Fullscreen to Windowed causes the cursor to misalign from the game
#proj UE4
#branch UE4
#summary Restoring from fullscreen now sets the window position on the native window
#extra Previously it was just telling Slate what the new window position was, but the native window was moved to 0,0. This caused the window capture rect for Slate to be in a different place to the actual location of the window.
#extra Tested on Windows and OS X.

diff --git a/Engine/Source/Runtime/Slate/Private/SWindow.cpp b/Engine/Source/Runtime/Slate/Private/SWindow.cpp
--- a/Engine/Source/Runtime/Slate/Private/SWindow.cpp
+++ b/Engine/Source/Runtime/Slate/Private/SWindow.cpp
@@ -1991,7 +1991,6 @@ void SWindow::SetWindowMode( EWindowMode::Type NewWindowMode )
 	
 		FSlateApplication::Get().GetRenderer()->UpdateFullscreenState( SharedThis(this), Size.X, Size.Y );
 
-
 		if( TitleArea.IsValid() )
 		{
 			// Collapse the Window title bar when switching to Fullscreen
@@ -2001,7 +2000,7 @@ void SWindow::SetWindowMode( EWindowMode::Type NewWindowMode )
 		if( bWasFullscreen )
 		{
 			// If we left fullscreen, reset the screen position;
-			ScreenPosition = PreFullscreenPosition;
+			MoveWindowTo(PreFullscreenPosition);
 		}
 	}
 
