#slate Added test case for saving/restoring expanded items in a tree view.

diff --git a/Engine/Source/Runtime/Slate/Private/Testing/STableViewTesting.cpp b/Engine/Source/Runtime/Slate/Private/Testing/STableViewTesting.cpp
--- a/Engine/Source/Runtime/Slate/Private/Testing/STableViewTesting.cpp
+++ b/Engine/Source/Runtime/Slate/Private/Testing/STableViewTesting.cpp
@@ -588,7 +588,7 @@ public:
 				[
 					SNew(STextBlock).Text( this, &STableViewTesting::GetNumGeneratedChildren )
 				]
-				+SHorizontalBox::Slot().AutoWidth() .Padding(20,0,0,0)
+				+SHorizontalBox::Slot().AutoWidth() .Padding(20,0,0,0) .VAlign(VAlign_Center)
 				[
 					SNew(STextBlock).Text( LOCTEXT("TotalChildrenLabel", "Total children:") )
 				]
@@ -642,6 +642,33 @@ public:
 					.OnClicked( this, &STableViewTesting::RequestRefresh )
 				]
 			]
+			+SVerticalBox::Slot() .AutoHeight()
+			[
+				SNew(SHorizontalBox)
+				+SHorizontalBox::Slot() .AutoWidth() .VAlign(VAlign_Center)
+				[
+					SNew(STextBlock)
+					.Text( LOCTEXT("ExpansionTestingLabel", "Expansion: ") )
+				]
+				+SHorizontalBox::Slot() .AutoWidth() .VAlign(HAlign_Center)
+				[
+					SNew(SButton)
+					.Text( LOCTEXT("RememberExpansionButton", "Remember") )
+					.OnClicked( this, &STableViewTesting::RememberExpansion )
+				]
+				+SHorizontalBox::Slot() .AutoWidth() .VAlign(VAlign_Center)
+				[
+					SNew(SButton)
+					.Text( LOCTEXT("RestoreExpansionsButton", "Collapse All") )
+					.OnClicked( this, &STableViewTesting::CollapseAll )
+				]
+				+SHorizontalBox::Slot() .AutoWidth() .VAlign(VAlign_Center)
+				[
+					SNew(SButton)
+					.Text( LOCTEXT("RestoreExpansionsButton", "Restore") )
+					.OnClicked( this, &STableViewTesting::RestoreExpansion )
+				]
+			]
 			+SVerticalBox::Slot()
 			.FillHeight(1)
 			[
@@ -835,6 +862,33 @@ protected:
 		];
 	}
 
+	FReply RememberExpansion()
+	{
+		StoredExpandedItems.Empty();
+		TreeBeingTested->GetExpandedItems(StoredExpandedItems);
+		return FReply::Handled();
+	}
+	TSet< TSharedPtr<FTestData> > StoredExpandedItems;
+
+	FReply CollapseAll()
+	{
+		TreeBeingTested->ClearExpandedItems();
+		TreeBeingTested->RequestTreeRefresh();
+		return FReply::Handled();
+	}
+
+	FReply RestoreExpansion()
+	{
+		for ( const TSharedPtr<FTestData>& Item : StoredExpandedItems )
+		{
+			TreeBeingTested->SetItemExpansion( Item, true );
+		}
+		TreeBeingTested->RequestTreeRefresh();
+		return FReply::Handled();
+	}
+
+
+
 	/** @returns How many widgets the TreeView generated to represent the data items*/
 	FText GetNumGeneratedChildren() const
 	{
