compile fix

diff --git a/Engine/Source/Runtime/Slate/Private/Testing/STableViewTesting.cpp b/Engine/Source/Runtime/Slate/Private/Testing/STableViewTesting.cpp
--- a/Engine/Source/Runtime/Slate/Private/Testing/STableViewTesting.cpp
+++ b/Engine/Source/Runtime/Slate/Private/Testing/STableViewTesting.cpp
@@ -650,7 +650,7 @@ public:
 					SNew(STextBlock)
 					.Text( LOCTEXT("ExpansionTestingLabel", "Expansion: ") )
 				]
-				+SHorizontalBox::Slot() .AutoWidth() .VAlign(HAlign_Center)
+				+SHorizontalBox::Slot() .AutoWidth() .VAlign(VAlign_Center)
 				[
 					SNew(SButton)
 					.Text( LOCTEXT("RememberExpansionButton", "Remember") )
