UMG UE4

* Adding metadata to several slate enums to improve their readability in the details panel.

diff --git a/Engine/Source/Runtime/Slate/Classes/SlateTypes.h b/Engine/Source/Runtime/Slate/Classes/SlateTypes.h
--- a/Engine/Source/Runtime/Slate/Classes/SlateTypes.h
+++ b/Engine/Source/Runtime/Slate/Classes/SlateTypes.h
@@ -7,39 +7,39 @@
 UENUM()
 enum EHorizontalAlignment
 {
-	HAlign_Fill,
-	HAlign_Left,
-	HAlign_Center,
-	HAlign_Right,
+	HAlign_Fill UMETA(DisplayName="Fill"),
+	HAlign_Left UMETA(DisplayName="Left"),
+	HAlign_Center UMETA(DisplayName="Center"),
+	HAlign_Right UMETA(DisplayName="Right"),
 };
 
 UENUM()
 enum EVerticalAlignment
 {
-	VAlign_Fill,
-	VAlign_Top,
-	VAlign_Center,
-	VAlign_Bottom
+	VAlign_Fill UMETA(DisplayName="Fill"),
+	VAlign_Top UMETA(DisplayName="Top"),
+	VAlign_Center UMETA(DisplayName="Center"),
+	VAlign_Bottom UMETA(DisplayName="Bottom"),
 };
 
 UENUM()
 enum EOrientation
 {
-	Orient_Horizontal,
-	Orient_Vertical
+	Orient_Horizontal UMETA(DisplayName="Horizontal"),
+	Orient_Vertical UMETA(DisplayName="Vertical"),
 };
 
 UENUM()
 enum EMenuPlacement
 {
 	/** Place the menu immediately below the anchor */
-	MenuPlacement_BelowAnchor,
+	MenuPlacement_BelowAnchor UMETA(DisplayName="Below Anchor"),
 	/** Place the menu immediately below the anchor and match is width to the anchor's content */
-	MenuPlacement_ComboBox,
+	MenuPlacement_ComboBox UMETA(DisplayName="ComboBox"),
 	/** Place the menu to the right of the anchor */
-	MenuPlacement_MenuRight,
+	MenuPlacement_MenuRight UMETA(DisplayName="Menu Right"),
 	/** Place the menu immediately above the anchor, not transition effect */
-	MenuPlacement_AboveAnchor,
+	MenuPlacement_AboveAnchor UMETA(DisplayName="AboveAnchor"),
 };
 
 /**
