Made the AddXExtension and RemoveExtension functions take the same forward declared type

TTP# 332269 - SLATE: FExtensionBase and derivatives are defined in MultiBoxExtender.cpp

This allows us to return a dumb handle without having to expose the actual type declarations, as now that everything is using the same type,  no type conversion needs to happen.

diff --git a/Engine/Source/Runtime/Slate/Private/MultiBox/MultiBoxExtender.cpp b/Engine/Source/Runtime/Slate/Private/MultiBox/MultiBoxExtender.cpp
--- a/Engine/Source/Runtime/Slate/Private/MultiBox/MultiBoxExtender.cpp
+++ b/Engine/Source/Runtime/Slate/Private/MultiBox/MultiBoxExtender.cpp
@@ -97,7 +97,7 @@ public:
 
 
 
-TSharedRef< const FMenuBarExtension > FExtender::AddMenuBarExtension( FName ExtensionHook, EExtensionHook::Position HookPosition, const TSharedPtr< FUICommandList >& CommandList, const FMenuBarExtensionDelegate& MenuBarExtensionDelegate )
+TSharedRef< const FExtensionBase > FExtender::AddMenuBarExtension( FName ExtensionHook, EExtensionHook::Position HookPosition, const TSharedPtr< FUICommandList >& CommandList, const FMenuBarExtensionDelegate& MenuBarExtensionDelegate )
 {
 	TSharedRef< FMenuBarExtension > MenuBarExtension( new FMenuBarExtension );
 	MenuBarExtension->Hook = ExtensionHook;
@@ -110,7 +110,7 @@ TSharedRef< const FMenuBarExtension > FExtender::AddMenuBarExtension( FName Exte
 }
 
 
-TSharedRef< const FMenuExtension > FExtender::AddMenuExtension( FName ExtensionHook, EExtensionHook::Position HookPosition, const TSharedPtr< FUICommandList >& CommandList, const FMenuExtensionDelegate& MenuExtensionDelegate )
+TSharedRef< const FExtensionBase > FExtender::AddMenuExtension( FName ExtensionHook, EExtensionHook::Position HookPosition, const TSharedPtr< FUICommandList >& CommandList, const FMenuExtensionDelegate& MenuExtensionDelegate )
 {
 	TSharedRef< FMenuExtension > MenuExtension( new FMenuExtension );
 	MenuExtension->Hook = ExtensionHook;
@@ -123,7 +123,7 @@ TSharedRef< const FMenuExtension > FExtender::AddMenuExtension( FName ExtensionH
 }
 
 
-TSharedRef< const FToolBarExtension > FExtender::AddToolBarExtension( FName ExtensionHook, EExtensionHook::Position HookPosition, const TSharedPtr< FUICommandList >& CommandList, const FToolBarExtensionDelegate& ToolBarExtensionDelegate )
+TSharedRef< const FExtensionBase > FExtender::AddToolBarExtension( FName ExtensionHook, EExtensionHook::Position HookPosition, const TSharedPtr< FUICommandList >& CommandList, const FToolBarExtensionDelegate& ToolBarExtensionDelegate )
 {
 	TSharedRef< FToolBarExtension > ToolBarExtension( new FToolBarExtension );
 	ToolBarExtension->Hook = ExtensionHook;
diff --git a/Engine/Source/Runtime/Slate/Public/MultiBox/MultiBoxExtender.h b/Engine/Source/Runtime/Slate/Public/MultiBox/MultiBoxExtender.h
--- a/Engine/Source/Runtime/Slate/Public/MultiBox/MultiBoxExtender.h
+++ b/Engine/Source/Runtime/Slate/Public/MultiBox/MultiBoxExtender.h
@@ -29,6 +29,10 @@ namespace EExtensionHook
 }
 
 
+/** Forward declare FExtensionBase. This type is used as a dumb handle to allow you to remove a previously added extension; its implementation is deliberately hidden */
+class FExtensionBase;
+
+
 class FExtender
 {
 public:
@@ -42,7 +46,7 @@ public:
 	 *
 	 * @return	Pointer to the new extension object.  You can use this later to remove the extension.
 	 */
-	SLATE_API TSharedRef< const class FMenuBarExtension > AddMenuBarExtension( FName ExtensionHook, EExtensionHook::Position HookPosition, const TSharedPtr< FUICommandList >& CommandList, const FMenuBarExtensionDelegate& MenuBarExtensionDelegate );
+	SLATE_API TSharedRef< const FExtensionBase > AddMenuBarExtension( FName ExtensionHook, EExtensionHook::Position HookPosition, const TSharedPtr< FUICommandList >& CommandList, const FMenuBarExtensionDelegate& MenuBarExtensionDelegate );
 
 
 	/**
@@ -55,7 +59,7 @@ public:
 	 *
 	 * @return	Pointer to the new extension object.  You can use this later to remove the extension.
 	 */
-	SLATE_API TSharedRef< const class FMenuExtension > AddMenuExtension( FName ExtensionHook, EExtensionHook::Position HookPosition, const TSharedPtr< FUICommandList >& CommandList, const FMenuExtensionDelegate& MenuExtensionDelegate );
+	SLATE_API TSharedRef< const FExtensionBase > AddMenuExtension( FName ExtensionHook, EExtensionHook::Position HookPosition, const TSharedPtr< FUICommandList >& CommandList, const FMenuExtensionDelegate& MenuExtensionDelegate );
 
 	
 	/**
@@ -68,7 +72,7 @@ public:
 	 *
 	 * @return	Pointer to the new extension object.  You can use this later to remove the extension.
 	 */
-	SLATE_API TSharedRef< const class FToolBarExtension > AddToolBarExtension( FName ExtensionHook, EExtensionHook::Position HookPosition, const TSharedPtr< FUICommandList >& CommandList, const FToolBarExtensionDelegate& ToolBarExtensionDelegate );
+	SLATE_API TSharedRef< const FExtensionBase > AddToolBarExtension( FName ExtensionHook, EExtensionHook::Position HookPosition, const TSharedPtr< FUICommandList >& CommandList, const FToolBarExtensionDelegate& ToolBarExtensionDelegate );
 
 
 	/**
@@ -76,7 +80,7 @@ public:
 	 *
 	 * @param	Extension	The extension to remove
 	 */
-	SLATE_API void RemoveExtension( const TSharedRef< const class FExtensionBase >& Extension );
+	SLATE_API void RemoveExtension( const TSharedRef< const FExtensionBase >& Extension );
 
 	/**
 	 * Applies any menu bar extensions at the specified extension point
@@ -119,5 +123,5 @@ public:
 private:
 
 	/** List of extensions to our menu.  The order of the list must be maintained. */
-	TArray< TSharedPtr< const class FExtensionBase > > Extensions;
+	TArray< TSharedPtr< const FExtensionBase > > Extensions;
 };
