Slate element batcher no longer depends on core style.

diff --git a/Engine/Source/Runtime/Slate/Private/ElementBatcher.cpp b/Engine/Source/Runtime/Slate/Private/ElementBatcher.cpp
--- a/Engine/Source/Runtime/Slate/Private/ElementBatcher.cpp
+++ b/Engine/Source/Runtime/Slate/Private/ElementBatcher.cpp
@@ -7,6 +7,7 @@
 #include "RenderingPolicy.h"
 #include "FontCache.h"
 
+static const FSlateImageBrush SplineFilterTable( FPaths::EngineContentDir() / TEXT("Slate/SplineFilterTable.png"), FVector2D(32,1) );
 
 /**
  * Computes the element tint color based in the user specified color and brush being used.
@@ -48,8 +49,7 @@ static FVector2D RotatePoint( FVector2D InPoint, const FVector2D& AboutPoint, fl
 
 
 FSlateElementBatcher::FSlateElementBatcher( TSharedRef<FSlateRenderingPolicy> InRenderingPolicy )
-	: SplineBrush( FCoreStyle::Get().GetBrush("SplineFilterTable") )
-	, RenderingPolicy( InRenderingPolicy )
+	: RenderingPolicy( InRenderingPolicy )
 {
 
 }
@@ -699,7 +699,7 @@ void FSlateElementBatcher::AddSplineElement( const FVector2D& Position, float Sc
 	const float HalfThickness = LineThickness * .5f + Radius;
 
 	// Currently splines are not atlased because they are tiled.  So we just assume the texture proxy holds the actual texture
-	FSlateShaderResourceProxy* ResourceProxy = RenderingPolicy->GetTextureResource( *SplineBrush );
+	FSlateShaderResourceProxy* ResourceProxy = RenderingPolicy->GetTextureResource( SplineFilterTable );
 	check(ResourceProxy && ResourceProxy->Resource);
 
 	// Find a batch for the element
@@ -818,7 +818,7 @@ void FSlateElementBatcher::AddLineElement( const FVector2D& Position, const FVec
 	if( InPayload.bAntialias )
 	{
 		// Currently splines are not atlased because they are tiled.  So we just assume the texture proxy holds the actual texture
-		FSlateShaderResourceProxy* ResourceProxy = RenderingPolicy->GetTextureResource( *SplineBrush );
+		FSlateShaderResourceProxy* ResourceProxy = RenderingPolicy->GetTextureResource( SplineFilterTable );
 		check(ResourceProxy && ResourceProxy->Resource);
 
 		// The radius to use when checking the distance of pixels to the actual line.  Arbitrary value based on what looks the best
diff --git a/Engine/Source/Runtime/Slate/Public/ElementBatcher.h b/Engine/Source/Runtime/Slate/Public/ElementBatcher.h
--- a/Engine/Source/Runtime/Slate/Public/ElementBatcher.h
+++ b/Engine/Source/Runtime/Slate/Public/ElementBatcher.h
@@ -1,8 +1,6 @@
 // Copyright 1998-2014 Epic Games, Inc. All Rights Reserved.
 
-#ifndef __SLATEELEMENTBATCHER_H__
-#define __SLATEELEMENTBATCHER_H__
-
+#pragma once
 
 class FSlateShaderResource;
 
@@ -182,7 +180,6 @@ private:
 	void AddIndices( TArray<SlateIndex>& OutIndices, FSlateElementBatch& ElementBatch, const TArray<SlateIndex>& IndexBatch );
 
 private:
-	const FSlateBrush *SplineBrush;
 	/** Element batch maps sorted by layer */
 	TMap< uint32, TSet<FSlateElementBatch> > LayerToElementBatches;
 	/** Array of vertex lists that are currently free (have no elements in them) */
@@ -207,5 +204,3 @@ private:
 	uint32 TotalIndexMemory;
 	uint32 RequiredIndexMemory;
 };
-
-#endif //__SLATEELEMENTBATCHER_H__
\ No newline at end of file
