Slate: Refactored core Slate implementation into SlateCore module in preparation for UMG.

diff --git a/Engine/Source/Runtime/Slate/Private/Framework/Application/MenuStack.cpp b/Engine/Source/Runtime/Slate/Private/Framework/Application/MenuStack.cpp
new file mode 100644
--- /dev/null
+++ b/Engine/Source/Runtime/Slate/Private/Framework/Application/MenuStack.cpp
@@ -0,0 +1,364 @@
+// Copyright 1998-2014 Epic Games, Inc. All Rights Reserved.
+
+/*=============================================================================
+	MenuStack.cpp: Implements the FMenuStack class.
+=============================================================================*/
+
+#include "Slate.h"
+
+
+namespace FMenuStackDefs
+{
+	/** Maximum size of menus as a fraction of the work area height */
+	const float MaxMenuScreenHeightFraction = 0.8f;
+}
+
+
+namespace
+{
+	/** Widget that wraps any popup menu created through FMenuStack::PushMenu to provide default key handling */
+	DECLARE_DELEGATE_RetVal_OneParam(FReply, FOnKeyDown, FKey)
+	class SMenuContentWrapper : public SCompoundWidget
+	{
+	public:
+		SLATE_BEGIN_ARGS(SMenuContentWrapper) 
+				: _MenuContent()
+				, _OnKeyDown()
+			{}
+
+			SLATE_DEFAULT_SLOT(FArguments, MenuContent)
+			SLATE_EVENT(FOnKeyDown, OnKeyDown)
+		SLATE_END_ARGS()
+
+		/** Construct this widget */
+		void Construct(const FArguments& InArgs)
+		{
+			OnKeyDownDelegate = InArgs._OnKeyDown;
+			ChildSlot
+			[
+				InArgs._MenuContent.Widget
+			];
+		}
+
+	private:
+
+		/** This widget must support keyboard focus */
+		virtual bool SupportsKeyboardFocus() const OVERRIDE
+		{
+			return true;
+		}
+
+		virtual FReply OnKeyDown( const FGeometry& MyGeometry, const FKeyboardEvent& InKeyboardEvent ) OVERRIDE
+		{
+			if (OnKeyDownDelegate.IsBound())
+			{
+				return OnKeyDownDelegate.Execute(InKeyboardEvent.GetKey());
+			}
+
+			return FReply::Unhandled();
+		}
+
+		/** Delegate to forward keys down events on the menu */
+		FOnKeyDown OnKeyDownDelegate;
+	};
+
+
+	/** Global handler used to handle key presses on popup menus */
+	FReply OnMenuKeyDown(const FKey Key)
+	{
+		if (Key == EKeys::Escape)
+		{
+			FSlateApplication::Get().DismissAllMenus();
+			return FReply::Handled();
+		}
+
+		return FReply::Unhandled();
+	}
+
+}	// anon namespace
+
+
+TSharedRef<SWindow> FMenuStack::PushMenu( const TSharedRef<SWindow>& ParentWindow, const TSharedRef<SWidget>& InContent, const FVector2D& SummonLocation, const FPopupTransitionEffect& TransitionEffect, const bool bFocusImmediately, const bool bShouldAutoSize, const FVector2D& WindowSize, const FVector2D& SummonLocationSize)
+{
+	// Only enable window position/size transitions if we're running at a decent frame rate
+	const bool bAllowAnimations = FSlateApplication::Get().AreMenuAnimationsEnabled() && FSlateApplication::Get().IsRunningAtTargetFrameRate();
+
+	// Calc the max height available on screen for the menu
+	FVector2D ExpectedSize = WindowSize;
+	FSlateRect AnchorRect(SummonLocation, SummonLocation + SummonLocationSize);
+	FSlateRect WorkArea = FSlateApplication::Get().GetWorkArea( AnchorRect );
+	float MaxHeight = FMenuStackDefs::MaxMenuScreenHeightFraction * WorkArea.GetSize().Y;
+	ExpectedSize.Y = FMath::Min(ExpectedSize.Y, MaxHeight);
+
+	// Wrap menu content in a box that limits its maximum height
+	TSharedRef<SWidget> WrappedContent = SNew(SVerticalBox)
+		+SVerticalBox::Slot()
+		.MaxHeight(MaxHeight)
+		[
+			SNew(SMenuContentWrapper)
+			.OnKeyDown_Static(&OnMenuKeyDown)
+			.MenuContent()
+			[
+				InContent
+			]
+		];
+
+	// Adjust the position of popup windows so they do not go out of the visible area of the monitor(s)	
+	if( bShouldAutoSize )
+	{
+		// @todo: Doesn't take into account potential window border size
+		WrappedContent->SlatePrepass();
+		ExpectedSize = WrappedContent->GetDesiredSize();
+	}
+	EOrientation Orientation = (TransitionEffect.SlideDirection == FPopupTransitionEffect::SubMenu) ? Orient_Horizontal : Orient_Vertical;
+	FSlateRect Anchor(SummonLocation.X, SummonLocation.Y, SummonLocation.X + SummonLocationSize.X, SummonLocation.Y + SummonLocationSize.Y);
+	FVector2D AdjustedSummonLocation = FSlateApplication::Get().CalculatePopupWindowPosition( Anchor, ExpectedSize, Orientation );
+
+	const float AnimationDuration = 0.15f;
+
+	// Start the pop-up menu at an offset location, then animate it to its target location over time
+	FVector2D AnimationStartLocation = FVector2D( AdjustedSummonLocation.X, AdjustedSummonLocation.Y );
+	if( bAllowAnimations )
+	{
+		const bool bSummonRight = AdjustedSummonLocation.X >= SummonLocation.X;
+		const bool bSummonBelow = AdjustedSummonLocation.Y >= SummonLocation.Y;
+		const int32 SummonDirectionX = bSummonRight ? 1 : -1;
+		const int32 SummonDirectionY = bSummonBelow ? 1 : -1;
+
+		switch( TransitionEffect.SlideDirection )
+		{
+			case FPopupTransitionEffect::None:
+				// No sliding
+				break;
+
+			case FPopupTransitionEffect::ComboButton:
+				AnimationStartLocation.Y = FMath::Max( AnimationStartLocation.Y + 30.0f * SummonDirectionY, 0.0f );
+				break;
+
+			case FPopupTransitionEffect::TopMenu:
+				AnimationStartLocation.Y = FMath::Max( AnimationStartLocation.Y + 60.0f * SummonDirectionY, 0.0f );
+				break;
+
+			case FPopupTransitionEffect::SubMenu:
+				AnimationStartLocation.X += 60.0f * SummonDirectionX;
+				break;
+
+			case FPopupTransitionEffect::TypeInPopup:
+				AnimationStartLocation.Y = FMath::Max( AnimationStartLocation.Y + 30.0f * SummonDirectionY, 0.0f );
+				break;
+
+			case FPopupTransitionEffect::ContextMenu:
+				AnimationStartLocation.X += 30.0f * SummonDirectionX;
+				AnimationStartLocation.Y += 50.0f * SummonDirectionY;
+				break;
+		}
+	}
+
+	// Start pop-up windows out transparent, then fade them in over time
+	const bool bUseTransparency = bAllowAnimations ? true : false;
+	const float InitialWindowOpacity = bAllowAnimations ? 0.0f : 1.0f;
+	const float TargetWindowOpacity = 1.0f;
+
+	//Release the mouse so that context can be properly restored upon closing NewMenuWindow.  See CL 1411833 before changing this.
+	if ( bFocusImmediately ) 
+	{ 
+		FSlateApplication::Get().ReleaseMouseCapture(); 
+	}
+
+	// Create a new window for the menu
+	TSharedRef<SWindow> NewMenuWindow = SNew(SWindow)
+		.IsPopupWindow( true )
+		.SizingRule( bShouldAutoSize ? ESizingRule::Autosized : ESizingRule::UserSized )
+		.ScreenPosition( AnimationStartLocation )
+		.AutoCenter( EAutoCenter::None )
+		.ClientSize( ExpectedSize )
+		.InitialOpacity(InitialWindowOpacity)
+		.SupportsTransparency( bUseTransparency )
+		.FocusWhenFirstShown( bFocusImmediately )
+		.ActivateWhenFirstShown( true )
+		[
+			WrappedContent
+		];
+
+	if (bFocusImmediately)
+	{
+		// Focus the content rather than just the window
+		NewMenuWindow->SetWidgetToFocusOnActivate(InContent);
+	}
+
+	// Find the location where the parent of the new menu resides.  The new menu's location in the stack is the level below the parent
+	int32 LocIndex = FindLocationInStack( ParentWindow ) + 1;
+	if( !Stack.IsValidIndex( LocIndex ) )
+	{
+		// If the new level doesn't exist yet create it now
+		Stack.Add( FMenuWindowList() );
+		// This is to ensure that the parent window is always one level above the new window being created
+		check( Stack.IsValidIndex( LocIndex ) );
+	}
+
+	{
+		// Make a copy of the windows array in case it changes size while we are operating on it
+		FMenuWindowList WindowsInStackCopy = Stack[LocIndex];
+
+		// Remove all other windows on the same level from the stack as they are no longer valid
+		// This will also remove the children of any of those windows
+		for( int32 WindowIndex = 0; WindowIndex < WindowsInStackCopy.Num(); ++WindowIndex )
+		{
+			WindowsInStackCopy[WindowIndex]->RequestDestroyWindow();
+		}
+	}
+
+	// Destroying windows above may cause the stack to get cleared or modified when activating another window.
+	if( !Stack.IsValidIndex( LocIndex ) )
+	{
+		// Re-find the location in the stack to make sure we add to the correct window list
+		LocIndex = FindLocationInStack( ParentWindow ) + 1;
+		if( !Stack.IsValidIndex( LocIndex ) )
+		{
+			// If the new level doesn't exist yet create it now
+			Stack.Add( FMenuWindowList() );
+			// This is to ensure that the parent window is always one level above the new window being created
+			check( Stack.IsValidIndex( LocIndex ) );
+		}
+	}
+
+	// Add the new window to menu window array at the required stack index
+	Stack[LocIndex].Add( NewMenuWindow );
+
+	FSlateApplication::Get().AddWindowAsNativeChild( NewMenuWindow, ParentWindow, true );
+
+	// Kick off the intro animation!
+	if( bAllowAnimations )
+	{
+		FCurveSequence Sequence;
+		Sequence.AddCurve( 0, AnimationDuration, ECurveEaseFunction::CubicOut );
+		NewMenuWindow->MorphToPosition( Sequence, TargetWindowOpacity, AdjustedSummonLocation );
+	}
+
+	// When a new menu is pushed, if we are not already in responsive mode for Slate UI, enter it now
+	// to ensure the menu is responsive in low FPS situations
+	if( !ThrottleHandle.IsValid() )
+	{
+		ThrottleHandle = FSlateThrottleManager::Get().EnterResponsiveMode();
+	}
+	return NewMenuWindow;
+}
+
+
+void FMenuStack::Dismiss( int32 LastStackIndex )
+{
+	// Dismiss the stack in reverse order so we destroy children before parents (causes focusing issues if done the other way around)
+	for( int32 StackIndex = Stack.Num()-1; StackIndex >= LastStackIndex; --StackIndex )
+	{
+		const TArray< TSharedPtr< SWindow> >& MenuWindows = TArray< TSharedPtr<SWindow> >(Stack[StackIndex]);
+		for( int32 MenuIndex = 0; MenuIndex < MenuWindows.Num(); ++MenuIndex )
+		{
+			MenuWindows[MenuIndex]->RequestDestroyWindow();
+		}
+	}
+
+	if (LastStackIndex == 0)
+	{
+		// Empty everything if this is a full dismissal
+		Stack.Empty();
+
+		// Leave responsive mode once the last menu closes
+		if (ThrottleHandle.IsValid())
+		{
+			FSlateThrottleManager::Get().LeaveResponsiveMode(ThrottleHandle);
+		}
+	}
+}
+
+
+void FMenuStack::RemoveWindow( TSharedRef<SWindow> WindowToRemove )
+{
+	// A window was requested to be destroyed, so make sure it's not in the pop-up menu stack to avoid it
+	// becoming a parent to a freshly-created window!
+	int32 Location = FindLocationInStack( WindowToRemove );
+
+	if( Location != INDEX_NONE )
+	{
+		if( Stack.IsValidIndex( Location+1 ) )
+		{
+			// Remove all levels below the window being removed
+			Stack.RemoveAt( Location+1, Stack.Num()-(Location+1) );
+		}
+
+		// Remove the actual window
+		Stack[Location].Remove( WindowToRemove );
+
+		// Remove list if empty
+		if (Stack[Location].Num() == 0)
+		{
+			Stack.RemoveAt(Location);
+		}
+	}
+
+	// Leave responsive mode once the last menu closes
+	if (Stack.Num() == 0 && ThrottleHandle.IsValid())
+	{
+		FSlateThrottleManager::Get().LeaveResponsiveMode( ThrottleHandle );
+	}
+}
+
+
+void FMenuStack::OnWindowActivated( TSharedRef<SWindow> ActivatedWindow )
+{
+	if (Stack.Num() > 0)
+	{
+		// See if the activated window is in the menu stack, if it is not then the current menu should be dismissed
+		int32 Location = FindLocationInStack(ActivatedWindow);
+
+		if (Location == INDEX_NONE)
+		{
+			Dismiss();
+		}
+		else
+		{
+			// If we clicked a window in the stack, we'll let the menu itself figure out what to do.
+		}
+	}
+}
+
+
+int32 FMenuStack::FindLocationInStack( TSharedPtr<SWindow> WindowToFind ) const
+{
+	// Search all levels of the stack
+	for( int32 StackIndex = 0; StackIndex < Stack.Num(); ++StackIndex )
+	{
+		// Search all windows in each level of the stack
+		const TArray<TSharedPtr<SWindow>>& MenuWindows = TArray<TSharedPtr<SWindow>>(Stack[StackIndex]);
+
+		for( int32 MenuIndex = 0; MenuIndex < MenuWindows.Num(); ++MenuIndex )
+		{
+			if( MenuWindows[MenuIndex] == WindowToFind )
+			{
+				return StackIndex;
+			}
+		}
+	}
+
+	// The window was not found
+	return INDEX_NONE;
+}
+
+
+bool FMenuStack::HasOpenSubMenus( const TSharedRef<SWindow>& Window ) const
+{
+	int32 LocationInStack = FindLocationInStack( Window );
+
+	return LocationInStack != INDEX_NONE ? Stack.Num() > LocationInStack + 1 : false;
+}
+
+
+int32 FMenuStack::GetNumStackLevels() const
+{
+	return Stack.Num();
+}
+
+
+FMenuWindowList& FMenuStack::GetWindowsAtStackLevel( const int32 StackLevelIndex )
+{
+	return Stack[StackLevelIndex];
+}
