Slate: Refactored core Slate implementation into SlateCore module in preparation for UMG.

diff --git a/Engine/Source/Runtime/Slate/Private/Framework/Application/SlateApplication.cpp b/Engine/Source/Runtime/Slate/Private/Framework/Application/SlateApplication.cpp
new file mode 100644
--- /dev/null
+++ b/Engine/Source/Runtime/Slate/Private/Framework/Application/SlateApplication.cpp
@@ -0,0 +1,4677 @@
+// Copyright 1998-2014 Epic Games, Inc. All Rights Reserved.
+
+#include "Slate.h"
+#include "SWindowTitleBar.h"
+
+// Statics
+//// @todo checkin: Should this go in an almost empty Events.cpp?
+FTouchKeySet FPointerEvent::StandardTouchKeySet(EKeys::LeftMouseButton);
+FTouchKeySet FPointerEvent::EmptyTouchKeySet(EKeys::Invalid);
+
+
+DECLARE_CYCLE_STAT( TEXT("Message Tick Time"), STAT_SlateMessageTick, STATGROUP_Slate );
+DECLARE_CYCLE_STAT( TEXT("Update Tooltip Time"), STAT_SlateUpdateTooltip, STATGROUP_Slate );
+DECLARE_CYCLE_STAT( TEXT("Tick Window And Children Time"), STAT_SlateTickWindowAndChildren, STATGROUP_Slate );
+DECLARE_CYCLE_STAT( TEXT("FindPathToWidget"), STAT_FindPathToWidget, STATGROUP_Slate );
+
+// Slate Event Logging is enabled to allow crash log dumping
+#define LOG_SLATE_EVENTS 0
+
+#if LOG_SLATE_EVENTS
+	#define LOG_EVENT_CONTENT( EventType, AdditionalContent, WidgetOrReply ) LogSlateEvent(EventLogger, EventType, AdditionalContent, WidgetOrReply);
+	
+	#define LOG_EVENT( EventType, WidgetOrReply ) LOG_EVENT_CONTENT( EventType, FString(), WidgetOrReply )
+	static void LogSlateEvent( const TSharedPtr<IEventLogger>& EventLogger, EEventLog::Type Event, const FString& AdditionalContent, const TSharedPtr<SWidget>& HandlerWidget )
+	{
+		if (EventLogger.IsValid())
+		{
+			EventLogger->Log( Event, AdditionalContent, HandlerWidget );
+		}
+	}
+
+	static void LogSlateEvent( const TSharedPtr<IEventLogger>& EventLogger, EEventLog::Type Event, const FString& AdditionalContent, const FReply& InReply )
+	{
+		if ( EventLogger.IsValid() && InReply.IsEventHandled() )
+		{
+			EventLogger->Log( Event, AdditionalContent, InReply.GetHandler() );
+		}
+	}
+#else
+	#define LOG_EVENT_CONTENT( EventType, AdditionalContent, WidgetOrReply )
+	
+	#define LOG_EVENT( Event, WidgetOrReply ) CheckReplyCorrectness(WidgetOrReply);
+	static void CheckReplyCorrectness(const TSharedPtr<SWidget>& HandlerWidget)
+	{
+	}
+	static void CheckReplyCorrectness(const FReply& InReply)
+	{
+		check( !InReply.IsEventHandled() || InReply.GetHandler().IsValid() );
+	}
+#endif
+
+namespace SlateDefs
+{
+	// How far tool tips should be offset from the mouse cursor position, in pixels
+	static const FVector2D ToolTipOffsetFromMouse( 12.0f, 8.0f );
+
+	// How far tool tips should be pushed out from a force field border, in pixels
+	static const FVector2D ToolTipOffsetFromForceField( 4.0f, 3.0f );
+}
+
+/** True if we should allow throttling based on mouse movement activity.  int32 instead of bool only for console variable system. */
+TAutoConsoleVariable<int32> ThrottleWhenMouseIsMoving( 
+	TEXT( "Slate.ThrottleWhenMouseIsMoving" ),
+	false,
+	TEXT( "Whether to attempt to increase UI responsiveness based on mouse cursor movement." ) );
+
+/** Minimum sustained average frame rate required before we consider the editor to be "responsive" for a smooth UI experience */
+TAutoConsoleVariable<int32> TargetFrameRateForResponsiveness(
+	TEXT( "Slate.TargetFrameRateForResponsiveness" ),
+	35,	// Frames per second
+	TEXT( "Minimum sustained average frame rate required before we consider the editor to be \"responsive\" for a smooth UI experience" ) );
+
+//////////////////////////////////////////////////////////////////////////
+bool FSlateApplication::MouseCaptorHelper::IsValid() const
+{
+	return MouseCaptorWeakPath.IsValid();
+}
+
+TSharedPtr< SWidget > FSlateApplication::MouseCaptorHelper::ToSharedWidget() const
+{
+	// If the path is valid then get the last widget, this is the current mouse captor
+	TSharedPtr< SWidget > SharedWidgetPtr;
+	if ( MouseCaptorWeakPath.IsValid() )
+	{
+		TWeakPtr< SWidget > WeakWidgetPtr = MouseCaptorWeakPath.GetLastWidget();
+		SharedWidgetPtr = WeakWidgetPtr.Pin();
+	}
+
+	return SharedWidgetPtr;
+}
+
+TSharedPtr< SWidget > FSlateApplication::MouseCaptorHelper::ToSharedWindow()
+{
+	// if the path is valid then we can get the window the current mouse captor belongs to
+	FWidgetPath MouseCaptorPath = ToWidgetPath();
+	if ( MouseCaptorPath.IsValid() )
+	{
+		return MouseCaptorPath.GetWindow();
+	}
+
+	return TSharedPtr< SWidget >();
+}
+
+void FSlateApplication::MouseCaptorHelper::SetMouseCaptor( const FWidgetPath& EventPath, TSharedPtr< SWidget > Widget )
+{
+	// Caller is trying to set a new mouse captor, so invalidate the current one - when the function finishes
+	// it still may not have a valid captor widget, this is ok
+	Invalidate();
+
+	if ( Widget.IsValid() )
+	{
+		TSharedRef< SWidget > WidgetRef = Widget.ToSharedRef();
+		MouseCaptorWeakPath = EventPath.GetPathDownTo( WidgetRef );
+
+		if ( !MouseCaptorWeakPath.IsValid() )
+		{
+			// If the target widget wasn't found on the event path then start the search from the root
+			FWidgetPath NewMouseCaptorPath = EventPath.GetPathDownTo( EventPath.Widgets(0).Widget );
+			NewMouseCaptorPath.ExtendPathTo( FWidgetMatcher( WidgetRef ) );
+			MouseCaptorWeakPath = NewMouseCaptorPath;
+		}
+	}
+}
+
+void FSlateApplication::MouseCaptorHelper::Invalidate()
+{
+	InformCurrentCaptorOfCaptureLoss();
+	MouseCaptorWeakPath = FWeakWidgetPath();
+}
+
+FWidgetPath FSlateApplication::MouseCaptorHelper::ToWidgetPath( FWeakWidgetPath::EInterruptedPathHandling::Type InterruptedPathHandling )
+{
+	FWidgetPath WidgetPath;
+	if ( MouseCaptorWeakPath.IsValid() )
+	{
+		if ( MouseCaptorWeakPath.ToWidgetPath( WidgetPath, InterruptedPathHandling ) == FWeakWidgetPath::EPathResolutionResult::Truncated )
+		{
+			// If the path was truncated then it means this widget is no longer part of the active set,
+			// so we make sure to invalidate its capture
+			Invalidate();
+		}
+	}
+
+	return WidgetPath;
+}
+
+FWeakWidgetPath FSlateApplication::MouseCaptorHelper::ToWeakPath() const
+{
+	return MouseCaptorWeakPath;
+}
+
+void FSlateApplication::MouseCaptorHelper::InformCurrentCaptorOfCaptureLoss() const
+{
+	// if we have a path to a widget then it is the current mouse captor and needs to know it has lost capture
+	if ( MouseCaptorWeakPath.IsValid() )
+	{
+		TWeakPtr< SWidget > WeakWidgetPtr = MouseCaptorWeakPath.GetLastWidget();
+		TSharedPtr< SWidget > SharedWidgetPtr = WeakWidgetPtr.Pin();
+		if ( SharedWidgetPtr.IsValid() )
+		{
+			SharedWidgetPtr->OnMouseCaptureLost();
+		}
+	}
+}
+
+//////////////////////////////////////////////////////////////////////////
+void FPopupSupport::RegisterClickNotification( const TSharedRef<SWidget>& NotifyWhenClickedOutsideMe, const FOnClickedOutside& InNotification )
+{
+	// If the subscriber or a zone object is descroyed, the subscription is
+	// no longer active. Clean it up here so that consumers of this API have an
+	// easy time with resource management.
+	struct { void operator()( TArray<FClickSubscriber>& Notifications ) {
+		for ( int32 SubscriberIndex=0; SubscriberIndex < Notifications.Num(); )
+		{
+			if ( !Notifications[SubscriberIndex].ShouldKeep() )
+			{
+				Notifications.RemoveAtSwap(SubscriberIndex);
+			}
+			else
+			{
+				SubscriberIndex++;
+			}		
+		}
+	}} ClearOutStaleNotifications;
+	
+	ClearOutStaleNotifications( ClickZoneNotifications );
+
+	// Add a new notification.
+	ClickZoneNotifications.Add( FClickSubscriber( NotifyWhenClickedOutsideMe, InNotification ) );
+}
+
+void FPopupSupport::UnregisterClickNotification( const FOnClickedOutside& InNotification )
+{
+	for (int32 SubscriptionIndex=0; SubscriptionIndex < ClickZoneNotifications.Num();)
+	{
+		if (ClickZoneNotifications[SubscriptionIndex].Notification == InNotification)
+		{
+			ClickZoneNotifications.RemoveAtSwap(SubscriptionIndex);
+		}
+		else
+		{
+			SubscriptionIndex++;
+		}
+	}	
+}
+
+void FPopupSupport::SendNotifications( const FWidgetPath& WidgetsUnderCursor )
+{
+	struct FArrangedWidgetMatcher
+	{
+		FArrangedWidgetMatcher( const TSharedRef<SWidget>& InWidgetToMatch )
+		: WidgetToMatch( InWidgetToMatch )
+		{}
+
+		bool Matches( const FArrangedWidget& Candidate ) const
+		{
+			return WidgetToMatch == Candidate.Widget;
+		}
+
+		const TSharedRef<SWidget>& WidgetToMatch;
+	};
+
+	// For each subscription, if the widget in question is not being clicked, send the notification.
+	// i.e. Notifications are saying "some widget outside you was clicked".
+	for (int32 SubscriberIndex=0; SubscriberIndex < ClickZoneNotifications.Num(); ++SubscriberIndex)
+	{
+		FClickSubscriber& Subscriber = ClickZoneNotifications[SubscriberIndex];
+		if (Subscriber.DetectClicksOutsideMe.IsValid())
+		{
+			// Did we click outside the region in this subscription? If so send the notification.
+			FArrangedWidgetMatcher Matcher(Subscriber.DetectClicksOutsideMe.Pin().ToSharedRef());
+			const bool bClickedOutsideOfWidget = WidgetsUnderCursor.Widgets.GetInternalArray().FindMatch( Matcher ) == INDEX_NONE;
+			if ( bClickedOutsideOfWidget )
+			{
+				Subscriber.Notification.ExecuteIfBound();
+			}
+		}
+	}
+}
+
+
+/**
+ * Given a widget-geometry pair and the mouse cursor location in absolute coordinates populate OutWidgetsUnderCursor with Candidate's children if they are
+ * under the cursor. If the Candidate is under the cursor return true, otherwise return false.
+ *
+ * @param Candidate                 The widget geometry pair being tested against the cursor location.
+ * @param InAbsoluteCursorLocation  The cursor location in absolute space (absolute space should match that of WidgetGeometry) 
+ * @param OutWidgetsUnderCursor     Widgets to be populated by the Candidate's children if the candidate is under the cursor.
+ * @param bIgnoreEnabledStatus		Allows hit tests to run on disabled widgets. 
+ *
+ * @returns true if the candidate is under the cursor. false otherwise.
+ */
+static bool LocateWidgetsUnderCursor_Helper( FArrangedWidget& Candidate, FVector2D InAbsoluteCursorLocation, FArrangedChildren& OutWidgetsUnderCursor, bool bIgnoreEnabledStatus )
+{
+	const bool bCandidateUnderCursor = 
+		// Candidate is physically under the cursor
+		Candidate.Geometry.IsUnderLocation( InAbsoluteCursorLocation ) && 
+		// Candidate actually considers itself hit by this test
+		(Candidate.Widget->OnHitTest( Candidate.Geometry, InAbsoluteCursorLocation ));
+
+	bool bHitAnyWidget = false;
+	if( bCandidateUnderCursor )
+	{
+		// The candidate widget is under the mouse
+		OutWidgetsUnderCursor.AddWidget(Candidate);
+
+		// Check to see if we were asked to still allow children to be hit test visible
+		bool bHitChildWidget = false;
+		if( ( Candidate.Widget->GetVisibility().AreChildrenHitTestVisible() ) != 0 )
+		{
+			FArrangedChildren ArrangedChildren(OutWidgetsUnderCursor.GetFilter());
+			Candidate.Widget->ArrangeChildren(Candidate.Geometry, ArrangedChildren);
+
+			// A widget's children are implicitly Z-ordered from first to last
+			for (int32 ChildIndex = ArrangedChildren.Num()-1; !bHitChildWidget && ChildIndex >= 0; --ChildIndex)
+			{
+				FArrangedWidget& SomeChild = ArrangedChildren(ChildIndex);
+				bHitChildWidget = (SomeChild.Widget->IsEnabled() || bIgnoreEnabledStatus) && LocateWidgetsUnderCursor_Helper( SomeChild, InAbsoluteCursorLocation, OutWidgetsUnderCursor, bIgnoreEnabledStatus );
+			}
+		}
+
+		// If we hit a child widget or we hit our candidate widget then we'll append our widgets
+		const bool bHitCandidateWidget = Candidate.Widget->GetVisibility().IsHitTestVisible();
+		bHitAnyWidget = bHitChildWidget || bHitCandidateWidget;
+		if( !bHitAnyWidget )
+		{
+			// No child widgets were hit, and even though the cursor was over our candidate widget, the candidate
+			// widget was not hit-testable, so we won't report it
+			check( OutWidgetsUnderCursor.Last() == Candidate );
+			OutWidgetsUnderCursor.Remove( OutWidgetsUnderCursor.Num() - 1 );
+		}
+	}
+
+	return bHitAnyWidget;
+}
+
+
+void FSlateApplication::ArrangeWindowToFront( TArray< TSharedRef<SWindow> >& Windows, const TSharedRef<SWindow>& WindowToBringToFront )
+{
+	Windows.Remove( WindowToBringToFront );
+
+	if ( Windows.Num() == 0 || WindowToBringToFront->IsTopmostWindow() )
+	{
+		Windows.Add( WindowToBringToFront );
+	}
+	else
+	{
+		bool PerformedInsert = false;
+		for (int Index = Windows.Num() - 1; Index >= 0 ; Index--)
+		{
+			if ( !Windows[Index]->IsTopmostWindow() )
+			{
+				Windows.Insert( WindowToBringToFront, Index + 1 );
+				PerformedInsert = true;
+				break;
+			}
+		}
+
+		if ( !PerformedInsert )
+		{
+			Windows.Insert( WindowToBringToFront, 0 );
+		}
+	}
+}
+
+/**
+ * Make BringMeToFront first among its peers.
+ * i.e. make it the last window in its parent's list of child windows.
+ *
+ * @return The top-most window whose children were re-arranged
+ */
+static TSharedRef<SWindow> BringToFrontInParent( const TSharedRef<SWindow>& BringMeToFront )
+{
+	const TSharedPtr<SWindow> ParentWindow = BringMeToFront->GetParentWindow();
+
+	if ( ParentWindow.IsValid() )
+	{
+		FSlateApplication::ArrangeWindowToFront( ParentWindow->GetChildWindows(), BringMeToFront );
+
+		return BringToFrontInParent( ParentWindow.ToSharedRef() );
+	}
+		
+	// This window has no parent, so it must be the top-most window.
+	return BringMeToFront;
+}
+
+/**
+ * Put 'BringMeToFront' at the font of the list of 'WindowsToReorder'. The top-most (front-most) window is last in Z-order and therefore
+ * is added to the end of the list.
+ *
+ * @param WindowsToReorder  An ordered list of windows.
+ * @param BingMeToFront     Window to be brough to front.
+ */
+static void BringWindowToFront( TArray< TSharedRef<SWindow> >& WindowsToReorder, const TSharedRef<SWindow>& BringMeToFront )
+{
+	const TSharedRef<SWindow> TopLeveWindowToReorder = BringToFrontInParent( BringMeToFront );
+
+	FSlateApplication::ArrangeWindowToFront( WindowsToReorder, TopLeveWindowToReorder );
+}
+
+/** @return true if the list of WindowsToSearch or any of their children contains the WindowToFind */
+static bool ContainsWindow( const TArray< TSharedRef<SWindow> >& WindowsToSearch, const TSharedRef<SWindow>& WindowToFind  )
+{
+	if ( WindowsToSearch.Contains(WindowToFind) )
+	{
+		return true;
+	}
+	else
+	{
+		for (int32 WindowIndex=0; WindowIndex < WindowsToSearch.Num(); ++WindowIndex)
+		{
+			if ( ContainsWindow( WindowsToSearch[WindowIndex]->GetChildWindows(), WindowToFind ) )
+			{
+				return true;
+			}
+		}
+	}
+
+	return false;
+}
+
+/** Remove the WindowToRemove (and its children) from TheList of windows or from their children. */
+static void RemoveWindowFromList( TArray< TSharedRef<SWindow> >& TheList, const TSharedRef<SWindow>& WindowToRemove )
+{
+	int32 NumRemoved = TheList.Remove( WindowToRemove );
+	if (NumRemoved == 0)
+	{
+		for ( int32 ChildIndex=0; ChildIndex < TheList.Num(); ++ChildIndex )
+		{
+			RemoveWindowFromList( TheList[ChildIndex]->GetChildWindows(), WindowToRemove ) ;
+		}	
+	}
+}
+
+void FSlateApplication::Create()
+{
+	EKeys::Initialize();
+
+	FCoreStyle::ResetToDefault();
+
+	CurrentApplication = MakeShareable( new FSlateApplication() );
+	CurrentBaseApplication = CurrentApplication;
+
+	PlatformApplication = MakeShareable( FPlatformMisc::CreateApplication() );
+	PlatformApplication->SetMessageHandler( CurrentApplication.ToSharedRef() );
+}
+
+TSharedPtr<FSlateApplication> FSlateApplication::CurrentApplication = NULL;
+TSharedPtr< GenericApplication > FSlateApplication::PlatformApplication = NULL;
+
+FSlateApplication::FSlateApplication()
+	: bAppIsActive(true)
+	, bSlateWindowActive(true)
+	, Scale( 1.0f )
+	, LastUserInteractionTimeForThrottling( 0.0 )
+	, SlateSoundDevice( MakeShareable(new FNullSlateSoundDevice()) )
+	, CurrentTime( FPlatformTime::Seconds() )
+	, LastTickTime( 0.0 )
+	, AverageDeltaTime( 1.0f / 30.0f )	// Prime the running average with a typical frame rate so it doesn't have to spin up from zero
+	, AverageDeltaTimeForResponsiveness( 1.0f / 30.0f )
+	, OnExitRequested()
+	, EventLogger( TSharedPtr<IEventLogger>() )
+	, NumExternalModalWindowsActive( 0 )
+	, bAllowToolTips( true )
+	, ToolTipDelay( 0.15f )
+	, ToolTipFadeInDuration( 0.1f )
+	, ToolTipSummonTime( 0.0 )
+	, DesiredToolTipLocation( FVector2D::ZeroVector )
+	, ToolTipOffsetDirection( EToolTipOffsetDirection::Undetermined )
+	, bRequestLeaveDebugMode( false )
+	, bLeaveDebugForSingleStep( false )
+	, CVarAllowToolTips(
+		TEXT( "Slate.AllowToolTips" ),
+		bAllowToolTips,
+		TEXT( "Whether to allow tool-tips to spawn at all." ) )
+	, CVarToolTipDelay(
+		TEXT( "Slate.ToolTipDelay" ),
+		ToolTipDelay,
+		TEXT( "Delay in seconds before a tool-tip is displayed near the mouse cursor when hovering over widgets that supply tool-tip data." ) )
+	, CVarToolTipFadeInDuration(
+		TEXT( "Slate.ToolTipFadeInDuration" ),
+		ToolTipFadeInDuration,
+		TEXT( "How long it takes for a tool-tip to fade in, in seconds." ) )
+	, bIsExternalUIOpened( false )
+	, SlateTextField( NULL )
+	, bIsFakingTouch(FParse::Param(FCommandLine::Get(), TEXT("simmobile")) || FParse::Param(FCommandLine::Get(), TEXT("faketouches")))
+	, bIsGameFakingTouch( false )
+	, bIsFakingTouched( false )
+	, bTouchFallbackToMouse( true )
+	, bMenuAnimationsEnabled( true )
+	, AppIcon( FCoreStyle::Get().GetBrush("DefaultAppIcon") )
+{
+#if WITH_UNREAL_DEVELOPER_TOOLS
+	FModuleManager::Get().LoadModule(TEXT("Settings"));
+#endif	
+
+	if (GConfig)
+	{	
+		GConfig->GetBool(TEXT("MobileSlateUI"),TEXT("bTouchFallbackToMouse"),bTouchFallbackToMouse,GEngineIni);
+	}
+
+	// causes InputCore to initialize, even if statically linked
+	FInputCoreModule& InputCore = FModuleManager::LoadModuleChecked<FInputCoreModule>(TEXT("InputCore"));
+
+	FGenericCommands::Register();
+
+	NormalExecutionGetter.BindRaw( this, &FSlateApplication::IsNormalExecution );
+
+#if !(UE_BUILD_SHIPPING || UE_BUILD_TEST) && !PLATFORM_HTML5 && !PLATFORM_HTML5_WIN32  
+	bool bEnableRemoteSlateServer = false;
+	if (GConfig && GConfig->GetBool(TEXT("MobileSupport"), TEXT("EnableUDKRemote"), bEnableRemoteSlateServer, GEngineIni) && bEnableRemoteSlateServer)
+	{
+		// @todo gmp: temporarily disabled. remote server is being refactored into plug-in.
+//		RemoteSlateServer = MakeShareable(new FRemoteSlateServer);
+	}
+#endif
+}
+
+FSlateApplication::~FSlateApplication()
+{
+	if (SlateTextField != NULL)
+	{
+		delete SlateTextField;
+		SlateTextField = NULL;
+	}
+}
+
+const FStyleNode* FSlateApplication::GetRootStyle() const
+{
+	return RootStyleNode;
+}
+
+void FSlateApplication::InitializeRenderer( TSharedRef<FSlateRenderer> InRenderer )
+{
+	Renderer = InRenderer;
+	Renderer->Initialize();
+}
+
+void FSlateApplication::InitializeSound( const TSharedRef<ISlateSoundDevice>& InSlateSoundDevice )
+{
+	SlateSoundDevice = InSlateSoundDevice;
+}
+
+void FSlateApplication::DestroyRenderer()
+{
+	if( Renderer.IsValid() )
+	{
+		Renderer->Destroy();
+	}
+}
+
+/**
+ * Called when the user closes the outermost frame (ie quitting the app). Uses standard UE4 global variable
+ * so normal UE4 applications work as expected
+ */
+static void OnRequestExit()
+{
+	GIsRequestingExit = true;
+}
+
+void FSlateApplication::PlaySound( const FSlateSound& SoundToPlay, int32 UserIndex ) const
+{
+	SlateSoundDevice->PlaySound(SoundToPlay, UserIndex);
+}
+
+float FSlateApplication::GetSoundDuration(const FSlateSound& Sound) const
+	{
+	return SlateSoundDevice->GetSoundDuration(Sound);
+}
+
+FVector2D FSlateApplication::GetCursorPos() const
+{
+	if ( PlatformApplication->Cursor.IsValid() )
+	{
+		return PlatformApplication->Cursor->GetPosition();
+	}
+
+	return FVector2D( 0, 0 );
+}
+
+void FSlateApplication::SetCursorPos( const FVector2D& MouseCoordinate )
+{
+	if ( PlatformApplication->Cursor.IsValid() )
+	{
+		return PlatformApplication->Cursor->SetPosition( MouseCoordinate.X, MouseCoordinate.Y );
+	}
+}
+
+
+FWidgetPath FSlateApplication::LocateWindowUnderMouse( FVector2D ScreenspaceMouseCoordinate, const TArray< TSharedRef< SWindow > >& Windows, bool bIgnoreEnabledStatus )
+{
+	bool bPrevWindowWasModal = false;
+	FArrangedChildren OutWidgetPath(EVisibility::Visible);
+
+	for( int32 WindowIndex=Windows.Num()-1; WindowIndex >= 0 && OutWidgetPath.Num() == 0; --WindowIndex )
+	{
+		const TSharedRef<SWindow>& Window = Windows[ WindowIndex ];
+
+		// Hittest the window's children first.
+		FWidgetPath ResultingPath = LocateWindowUnderMouse( ScreenspaceMouseCoordinate, Window->GetChildWindows(), bIgnoreEnabledStatus );
+		if ( ResultingPath.IsValid() )
+		{
+			return ResultingPath;
+		}
+
+		// If none of the children were hit, hittest the parent.
+
+		// Only accept input if the current window accepts input and the current window is not under a modal window or an interactive tooltip
+		const TSharedPtr<IToolTip> ActiveToolTipPtr = ActiveToolTip.Pin();
+		const TSharedPtr< SWindow > ToolTipWindowPtr = ToolTipWindow.Pin();
+		const bool AcceptsInput = Window->AcceptsInput() || ( Window == ToolTipWindowPtr && ActiveToolTipPtr.IsValid() && ActiveToolTipPtr->IsInteractive() );
+
+		if ( Window->IsVisible() && AcceptsInput && Window->IsScreenspaceMouseWithin(ScreenspaceMouseCoordinate) && !bPrevWindowWasModal )
+		{
+			FArrangedWidget WindowWidgetGeometry( Window, Window->GetWindowGeometryInScreen() );
+			if ( Window->IsEnabled() && LocateWidgetsUnderCursor_Helper(WindowWidgetGeometry, ScreenspaceMouseCoordinate, OutWidgetPath, bIgnoreEnabledStatus) )		
+			{			
+				return FWidgetPath( Windows[WindowIndex], OutWidgetPath );
+			}
+		}
+	}
+
+	return FWidgetPath();
+}
+
+/** 
+ * Ticks a single slate window
+ *
+ * @param WindowToTick	The window to tick
+ */
+void FSlateApplication::TickWindowAndChildren( TSharedRef<SWindow> WindowToTick )
+{
+	if ( WindowToTick->IsVisible() )
+	{
+		// Switch to the appropriate world for ticking
+		FScopedSwitchWorldHack SwitchWorld( WindowToTick );
+
+		// Measure all the widgets before we tick, and update their DesiredSize.  This is
+		// needed so that Tick() can call ArrangeChildren(), then pass valid widget metrics into
+		// the Tick() function.
+		
+		{
+			SCOPE_CYCLE_COUNTER( STAT_SlateCacheDesiredSize )
+			WindowToTick->SlatePrepass();
+		}
+
+		if (WindowToTick->IsAutosized())
+		{
+			WindowToTick->ReshapeWindow(WindowToTick->GetPositionInScreen(), WindowToTick->GetDesiredSize());
+		}
+
+		{
+			SCOPE_CYCLE_COUNTER( STAT_SlateTickWidgets )
+			// Tick this window and all of the widgets in this window
+			WindowToTick->TickWidgetsRecursively( WindowToTick->GetWindowGeometryInScreen(), GetCurrentTime(), GetDeltaTime() );
+		}
+
+		// Tick all of this window's child windows.
+		const TArray< TSharedRef<SWindow> >& WindowChildren = WindowToTick->GetChildWindows();
+		for ( int32 ChildIndex=0; ChildIndex < WindowChildren.Num(); ++ChildIndex )
+		{
+			TickWindowAndChildren( WindowChildren[ChildIndex] );
+		}
+	}
+}
+
+void FSlateApplication::DrawWindows()
+{
+	PrivateDrawWindows();
+}
+
+struct FDrawWindowArgs
+{
+	FDrawWindowArgs( FSlateDrawBuffer& InDrawBuffer, const FWidgetPath& InFocusedPath, const FWidgetPath& InWidgetsUnderCursor )
+		: OutDrawBuffer( InDrawBuffer )
+		, FocusedPath( InFocusedPath )
+		, WidgetsUnderCursor( InWidgetsUnderCursor )
+	{}
+
+	TArray<FGenericWindow*, TInlineAllocator<10> > OutDrawnWindows;
+	FSlateDrawBuffer& OutDrawBuffer;
+	const FWidgetPath& FocusedPath;
+	const FWidgetPath& WidgetsUnderCursor;
+};
+
+void FSlateApplication::DrawWindowAndChildren( const TSharedRef<SWindow>& WindowToDraw, FDrawWindowArgs& DrawWindowArgs )
+{
+	// Only draw visible windows
+	if( WindowToDraw->IsVisible() )
+	{
+		// Switch to the appropriate world for drawing
+		FScopedSwitchWorldHack SwitchWorld( WindowToDraw );
+
+		FSlateWindowElementList& WindowElementList = DrawWindowArgs.OutDrawBuffer.AddWindowElementList( WindowToDraw );
+
+		// Drawing is done in window space, so null out the positions and keep the size.
+		FGeometry WindowGeometry = WindowToDraw->GetWindowGeometryInWindow();
+		int32 MaxLayerId = 0;
+		{
+			SCOPE_CYCLE_COUNTER( STAT_SlateOnPaint );
+			MaxLayerId = WindowToDraw->OnPaint( WindowGeometry, WindowToDraw->GetClippingRectangleInWindow(), WindowElementList, 0, FWidgetStyle(), WindowToDraw->IsEnabled() );
+		}
+
+		if (DrawWindowArgs.FocusedPath.IsValid() && DrawWindowArgs.FocusedPath.GetWindow() == WindowToDraw)
+		{
+			MaxLayerId = DrawKeyboardFocus(DrawWindowArgs.FocusedPath, WindowElementList, MaxLayerId);
+		}
+
+		// The widget reflector may want to paint some additional stuff as part of the Widget introspection that it performs.
+		// For example: it may draw layout rectangles for hovered widgets.
+		const bool bVisualizeLayoutUnderCursor = DrawWindowArgs.WidgetsUnderCursor.IsValid();
+		const bool bCapturingFromThisWindow = bVisualizeLayoutUnderCursor && DrawWindowArgs.WidgetsUnderCursor.TopLevelWindow == WindowToDraw;
+		TSharedPtr<IWidgetReflector> WidgetReflector = WidgetReflectorPtr.Pin();
+		if ( bCapturingFromThisWindow || (WidgetReflector.IsValid() && WidgetReflector->ReflectorNeedsToDrawIn(WindowToDraw)) )
+		{
+			MaxLayerId = WidgetReflector->Visualize( DrawWindowArgs.WidgetsUnderCursor, WindowElementList, MaxLayerId );
+		}
+
+		// Keep track of windows that we're actually going to be presenting, so we can mark
+		// them as 'drawn' afterwards.
+		FGenericWindow* NativeWindow = WindowToDraw->GetNativeWindow().Get();
+		DrawWindowArgs.OutDrawnWindows.Add( NativeWindow );
+
+		// Draw the child windows
+		const TArray< TSharedRef<SWindow> >& WindowChildren = WindowToDraw->GetChildWindows();
+		for (int32 ChildIndex=0; ChildIndex < WindowChildren.Num(); ++ChildIndex)
+		{
+			DrawWindowAndChildren( WindowChildren[ChildIndex], DrawWindowArgs );
+		}
+	}
+}
+
+static void DrawWindowPrepass( TSharedRef<SWindow> WindowToDraw )
+{
+	FScopedSwitchWorldHack SwitchWorld( WindowToDraw );
+
+	{
+		SCOPE_CYCLE_COUNTER(STAT_SlateCacheDesiredSize);
+		WindowToDraw->SlatePrepass();
+	}
+
+	if (WindowToDraw->IsAutosized())
+	{
+		WindowToDraw->ReshapeWindow(WindowToDraw->GetPositionInScreen(), WindowToDraw->GetDesiredSize());
+	}
+}
+
+void FSlateApplication::DrawPrepass( TSharedPtr<SWindow> DrawOnlyThisWindow )
+{
+	SCOPE_CYCLE_COUNTER(STAT_SlateDrawWindowTime);
+
+	TSharedPtr<SWindow> ActiveModalWindow = GetActiveModalWindow();
+
+	if (ActiveModalWindow.IsValid())
+	{
+		DrawWindowPrepass( ActiveModalWindow.ToSharedRef() );
+
+		for (TArray< TSharedRef<SWindow> >::TConstIterator CurrentWindowIt(SlateWindows); CurrentWindowIt; ++CurrentWindowIt)
+		{
+			const TSharedRef<SWindow>& CurrentWindow = *CurrentWindowIt;
+			if (CurrentWindow->IsTopmostWindow())
+			{
+				DrawWindowPrepass( CurrentWindow );
+			}
+		}
+
+		TArray< TSharedRef<SWindow> > NotificationWindows;
+		FSlateNotificationManager::Get().GetWindows(NotificationWindows);
+		for (auto CurrentWindowIt(NotificationWindows.CreateIterator()); CurrentWindowIt; ++CurrentWindowIt)
+		{
+			DrawWindowPrepass(*CurrentWindowIt );
+		}
+	}
+	else if (DrawOnlyThisWindow.IsValid())
+	{
+		DrawWindowPrepass( DrawOnlyThisWindow.ToSharedRef() );
+	}
+	else
+	{
+		// Draw all windows
+		for (TArray< TSharedRef<SWindow> >::TConstIterator CurrentWindowIt(SlateWindows); CurrentWindowIt; ++CurrentWindowIt)
+		{
+			TSharedRef<SWindow> CurrentWindow = *CurrentWindowIt;
+			if (CurrentWindow->IsVisible())
+			{
+				DrawWindowPrepass(CurrentWindow );
+			}
+		}
+	}
+}
+
+void FSlateApplication::PrivateDrawWindows( TSharedPtr<SWindow> DrawOnlyThisWindow )
+{
+	check(Renderer.IsValid());
+
+	// Is user expecting visual feedback from the Widget Reflector?
+	const bool bVisualizeLayoutUnderCursor = WidgetReflectorPtr.IsValid() && WidgetReflectorPtr.Pin()->IsVisualizingLayoutUnderCursor();
+
+	FWidgetPath WidgetsUnderCursor;
+	if (bVisualizeLayoutUnderCursor)
+	{
+		WidgetsUnderCursor = LocateWindowUnderMouse( GetCursorPos(), SlateWindows );
+	}
+
+	FWidgetPath FocusPath = FocusedWidgetPath.ToWidgetPath();
+
+	DrawPrepass( DrawOnlyThisWindow );
+
+	FDrawWindowArgs DrawWindowArgs( Renderer->GetDrawBuffer(), FocusPath, WidgetsUnderCursor );
+
+	{
+		SCOPE_CYCLE_COUNTER( STAT_SlateDrawWindowTime );
+
+		TSharedPtr<SWindow> ActiveModalWindow = GetActiveModalWindow(); 
+
+		if (ActiveModalWindow.IsValid())
+		{
+			DrawWindowAndChildren( ActiveModalWindow.ToSharedRef(), DrawWindowArgs );
+
+			for( TArray< TSharedRef<SWindow> >::TConstIterator CurrentWindowIt( SlateWindows ); CurrentWindowIt; ++CurrentWindowIt )
+			{
+				const TSharedRef<SWindow>& CurrentWindow = *CurrentWindowIt;
+				if ( CurrentWindow->IsTopmostWindow() )
+				{
+					DrawWindowAndChildren(CurrentWindow, DrawWindowArgs);
+				}
+			}
+
+			TArray< TSharedRef<SWindow> > NotificationWindows;
+			FSlateNotificationManager::Get().GetWindows(NotificationWindows);
+			for( auto CurrentWindowIt( NotificationWindows.CreateIterator() ); CurrentWindowIt; ++CurrentWindowIt )
+			{
+				DrawWindowAndChildren(*CurrentWindowIt, DrawWindowArgs);
+			}	
+		}
+		else if( DrawOnlyThisWindow.IsValid() )
+		{
+			DrawWindowAndChildren( DrawOnlyThisWindow.ToSharedRef(), DrawWindowArgs );
+		}
+		else
+		{
+			// Draw all windows
+			for( TArray< TSharedRef<SWindow> >::TConstIterator CurrentWindowIt( SlateWindows ); CurrentWindowIt; ++CurrentWindowIt )
+			{
+				TSharedRef<SWindow> CurrentWindow = *CurrentWindowIt;
+				if ( CurrentWindow->IsVisible() )
+				{
+					DrawWindowAndChildren( CurrentWindow, DrawWindowArgs );
+				}
+			}
+		}
+	}
+
+	Renderer->DrawWindows( DrawWindowArgs.OutDrawBuffer );
+}
+
+void FSlateApplication::PollGameDeviceState()
+{
+	if( ActiveModalWindows.Num() == 0 && !GIntraFrameDebuggingGameThread )
+	{
+		// Don't poll when a modal window open or intra frame debugging is happening
+		PlatformApplication->PollGameDeviceState( GetDeltaTime() );
+	}
+}
+
+/**
+ * Ticks this application
+ */
+void FSlateApplication::Tick()
+{
+	SCOPE_CYCLE_COUNTER( STAT_SlateTickTime );
+	
+	{
+		const float DeltaTime = GetDeltaTime();
+
+		SCOPE_CYCLE_COUNTER( STAT_SlateMessageTick );
+
+		// We need to pump messages here so that slate can receive input.  
+		if( (ActiveModalWindows.Num() > 0) || GIntraFrameDebuggingGameThread )
+		{
+			// We only need to pump messages for slate when a modal window or blocking mode is active is up because normally message pumping is handled in FEngineLoop::Tick
+			PlatformApplication->PumpMessages( DeltaTime );
+
+			if (FCoreDelegates::StarvedGameLoop.IsBound())
+			{
+				FCoreDelegates::StarvedGameLoop.Execute();
+			}
+		}
+
+		PlatformApplication->Tick( DeltaTime );
+
+		PlatformApplication->ProcessDeferredEvents( DeltaTime );
+	}
+
+	// When Slate captures the mouse, it is up to us to set the cursor 
+	// because the OS assumes that we own the mouse.
+	if (MouseCaptor.IsValid())
+	{
+		QueryCursor();
+	}
+
+	{
+		SCOPE_CYCLE_COUNTER( STAT_SlateUpdateTooltip );
+
+		// Update tool tip, if we have one
+		const bool AllowSpawningOfToolTips = false;
+		UpdateToolTip( AllowSpawningOfToolTips );
+	}
+
+
+	// Advance time
+	LastTickTime = CurrentTime;
+	CurrentTime = FPlatformTime::Seconds();
+
+	// Update average time between ticks.  This is used to monitor how responsive the application "feels".
+	// Note that we calculate this before we apply the max quantum clamping below, because we want to store
+	// the actual frame rate, even if it is very low.
+	{
+		// Scalar percent of new delta time that contributes to running average.  Use a lower value to add more smoothing
+		// to the average frame rate.  A value of 1.0 will disable smoothing.
+		const float RunningAverageScale = 0.1f;
+
+		AverageDeltaTime = AverageDeltaTime * ( 1.0f - RunningAverageScale ) + GetDeltaTime() * RunningAverageScale;
+
+		// Don't update average delta time if we're in an exceptional situation, such as when throttling mode
+		// is active, because the measured tick time will not be representative of the application's performance.
+		// In these cases, the cached average delta time from before the throttle activated will be used until
+		// throttling has finished.
+		if( FSlateThrottleManager::Get().IsAllowingExpensiveTasks() )
+		{
+			// Clamp to avoid including huge hitchy frames in our average
+			const float ClampedDeltaTime = FMath::Clamp( GetDeltaTime(), 0.0f, 1.0f );
+			AverageDeltaTimeForResponsiveness = AverageDeltaTimeForResponsiveness * ( 1.0f - RunningAverageScale ) + ClampedDeltaTime * RunningAverageScale;
+		}
+	}
+
+	// Handle large quantums
+	const double MaxQuantumBeforeClamp = 1.0 / 8.0;		// 8 FPS
+	if( GetDeltaTime() > MaxQuantumBeforeClamp )
+	{
+		LastTickTime = CurrentTime - MaxQuantumBeforeClamp;
+	}
+
+	// Force a mouse move event to make sure all widgets know whether there is a mouse cursor hovering over them
+	SynthesizeMouseMove();
+
+	// Update auto-throttling based on elapsed time since user interaction
+	ThrottleApplicationBasedOnMouseMovement();
+
+	TSharedPtr<SWindow> ActiveModalWindow = GetActiveModalWindow();
+
+	{
+		SCOPE_CYCLE_COUNTER( STAT_SlateTickWindowAndChildren );
+
+		if ( ActiveModalWindow.IsValid() )
+		{
+			// There is a modal window, and we just need to tick it.
+			TickWindowAndChildren( ActiveModalWindow.ToSharedRef() );
+			// And also tick any top-level windows.
+			for( TArray< TSharedRef<SWindow> >::TIterator CurrentWindowIt( SlateWindows ); CurrentWindowIt; ++CurrentWindowIt )
+			{
+				TSharedRef<SWindow>& CurrentWindow = *CurrentWindowIt;
+				if (CurrentWindow->IsTopmostWindow())
+				{
+					TickWindowAndChildren(CurrentWindow);
+				}
+			}
+			// also tick the notification manager's windows
+			TArray< TSharedRef<SWindow> > NotificationWindows;
+			FSlateNotificationManager::Get().GetWindows(NotificationWindows);
+			for( auto CurrentWindowIt( NotificationWindows.CreateIterator() ); CurrentWindowIt; ++CurrentWindowIt )
+			{
+				TickWindowAndChildren(*CurrentWindowIt);
+			}		
+		}
+		else
+		{
+			// No modal window; tick all slate windows.
+			for( TArray< TSharedRef<SWindow> >::TIterator CurrentWindowIt( SlateWindows ); CurrentWindowIt; ++CurrentWindowIt )
+			{
+				TSharedRef<SWindow>& CurrentWindow = *CurrentWindowIt;
+				TickWindowAndChildren( CurrentWindow );
+			}
+		}
+	}
+
+	// Update any notifications - this needs to be done after windows have updated themselves 
+	// (so they know their size)
+	FSlateNotificationManager::Get().Tick();
+
+	// Draw all windows
+	DrawWindows();
+
+/* @todo gmp: temporarily disabled. remote server is being refactored into plug-in.
+#if !(UE_BUILD_SHIPPING || UE_BUILD_TEST)
+	if (RemoteSlateServer.IsValid())
+	{
+		RemoteSlateServer->Tick();
+	}
+#endif*/
+}
+
+
+void FSlateApplication::PumpMessages()
+{
+	PlatformApplication->PumpMessages( GetDeltaTime() );
+}
+
+
+void FSlateApplication::ThrottleApplicationBasedOnMouseMovement()
+{
+	bool bShouldThrottle = false;
+	if( ThrottleWhenMouseIsMoving.GetValueOnGameThread() )	// Interpreted as bool here
+	{
+		// We only want to engage the throttle for a short amount of time after the mouse stops moving
+		const float TimeToThrottleAfterMouseStops = 0.1f;
+
+		// After a key or mouse button is pressed, we'll leave the throttle disengaged for awhile so the
+		// user can use the keys to navigate in a viewport, for example.
+		const float MinTimeSinceButtonPressToThrottle = 1.0f;
+
+		// Use a small movement threshold to avoid engaging the throttle when the user bumps the mouse
+		const float MinMouseMovePixelsBeforeThrottle = 2.0f;
+
+		const FVector2D& CursorPos = GetCursorPos();
+		static FVector2D LastCursorPos = GetCursorPos();
+		static double LastMouseMoveTime = FPlatformTime::Seconds();
+		static bool bIsMouseMoving = false;
+		if( CursorPos != LastCursorPos )
+		{
+			// Did the cursor move far enough that we care?
+			if( bIsMouseMoving || ( CursorPos - LastCursorPos ).SizeSquared() >= MinMouseMovePixelsBeforeThrottle * MinMouseMovePixelsBeforeThrottle )
+			{
+				bIsMouseMoving = true;
+				LastMouseMoveTime = this->GetCurrentTime();
+				LastCursorPos = CursorPos;
+			}
+		}
+
+		const float TimeSinceLastUserInteraction = CurrentTime - LastUserInteractionTimeForThrottling;
+		const float TimeSinceLastMouseMove = CurrentTime - LastMouseMoveTime;
+		if( TimeSinceLastMouseMove < TimeToThrottleAfterMouseStops )
+		{
+			// Only throttle if a Slate window is currently active.  If a Wx window (such as Matinee) is
+			// being used, we don't want to throttle
+			if( this->GetActiveTopLevelWindow().IsValid() )
+			{
+				// Only throttle if the user hasn't pressed a button in awhile
+				if( TimeSinceLastUserInteraction > MinTimeSinceButtonPressToThrottle )
+				{
+					// If a widget has the mouse captured, then we won't bother throttling
+					if( !MouseCaptor.IsValid() )
+					{
+						// If there is no Slate window under the mouse, then we won't engage throttling
+						if( LocateWindowUnderMouse( GetCursorPos(), GetInteractiveTopLevelWindows() ).IsValid() )
+						{
+							bShouldThrottle = true;
+						}
+					}
+				}
+			}
+		}
+		else
+		{
+			// Mouse hasn't moved in a bit, so reset our movement state
+			bIsMouseMoving = false;
+			LastCursorPos = CursorPos;
+		}
+	}
+
+	if( bShouldThrottle )
+	{
+		if( !UserInteractionResponsivnessThrottle.IsValid() )
+		{
+			// Engage throttling
+			UserInteractionResponsivnessThrottle = FSlateThrottleManager::Get().EnterResponsiveMode();
+		}
+	}
+	else
+	{
+		if( UserInteractionResponsivnessThrottle.IsValid() )
+		{
+			// Disengage throttling
+			FSlateThrottleManager::Get().LeaveResponsiveMode( UserInteractionResponsivnessThrottle );
+		}
+	}
+}
+
+
+TSharedRef<SWindow> FSlateApplication::AddWindow( TSharedRef<SWindow> InSlateWindow, const bool bShowImmediately )
+{
+	// Add the Slate window to the Slate application's top-level window array.  Note that neither the Slate window
+	// or the native window are ready to be used yet, however we need to make sure they're in the Slate window
+	// array so that we can properly respond to OS window messages as soon as they're sent.  For example, a window
+	// activation message may be sent by the OS as soon as the window is shown (in the Init function), and if we
+	// don't add the Slate window to our window list, we wouldn't be able to route that message to the window.
+
+	ArrangeWindowToFront( SlateWindows, InSlateWindow );
+	TSharedRef<FGenericWindow> NewWindow = MakeWindow( InSlateWindow, bShowImmediately );
+
+	if( bShowImmediately )
+	{
+		InSlateWindow->ShowWindow();
+
+		//@todo Slate: Potentially dangerous and annoying if all slate windows that are created steal focus.
+		if( InSlateWindow->SupportsKeyboardFocus() && InSlateWindow->IsFocusedInitially() )
+		{
+			InSlateWindow->GetNativeWindow()->SetWindowFocus();
+		}
+	}
+
+	return InSlateWindow;
+}
+
+TSharedRef< FGenericWindow > FSlateApplication::MakeWindow( TSharedRef<SWindow> InSlateWindow, const bool bShowImmediately )
+{
+	TSharedPtr<FGenericWindow> NativeParent = NULL;
+	TSharedPtr<SWindow> ParentWindow = InSlateWindow->GetParentWindow();
+	if ( ParentWindow.IsValid() )
+	{
+		NativeParent = ParentWindow->GetNativeWindow();
+	}
+
+	TSharedRef< FGenericWindowDefinition > Definition = MakeShareable( new FGenericWindowDefinition() );
+
+	const FVector2D Size = InSlateWindow->GetInitialDesiredSizeInScreen();
+	Definition->WidthDesiredOnScreen = Size.X;
+	Definition->HeightDesiredOnScreen = Size.Y;
+
+	const FVector2D Position = InSlateWindow->GetInitialDesiredPositionInScreen();
+	Definition->XDesiredPositionOnScreen = Position.X;
+	Definition->YDesiredPositionOnScreen = Position.Y;
+
+	Definition->HasOSWindowBorder = InSlateWindow->HasOSWindowBorder();
+	Definition->SupportsTransparency = InSlateWindow->SupportsTransparency();
+	Definition->AppearsInTaskbar = InSlateWindow->AppearsInTaskbar();
+	Definition->IsTopmostWindow = InSlateWindow->IsTopmostWindow();
+	Definition->AcceptsInput = InSlateWindow->AcceptsInput();
+	Definition->ActivateWhenFirstShown = InSlateWindow->ActivateWhenFirstShown();
+
+	Definition->SupportsMinimize = InSlateWindow->HasMinimizeBox();
+	Definition->SupportsMaximize = InSlateWindow->HasMaximizeBox();
+
+	Definition->IsRegularWindow = InSlateWindow->IsRegularWindow();
+	Definition->HasSizingFrame = InSlateWindow->HasSizingFrame();
+	Definition->SizeWillChangeOften = InSlateWindow->SizeWillChangeOften();
+	Definition->ExpectedMaxWidth = InSlateWindow->GetExpectedMaxWidth();
+	Definition->ExpectedMaxHeight = InSlateWindow->GetExpectedMaxHeight();
+
+	Definition->Title = InSlateWindow->GetTitle().ToString();
+	Definition->Opacity = InSlateWindow->GetOpacity();
+	Definition->CornerRadius = InSlateWindow->GetCornerRadius();
+
+	TSharedRef< FGenericWindow > NewWindow = PlatformApplication->MakeWindow();
+
+	InSlateWindow->SetNativeWindow( NewWindow );
+
+	InSlateWindow->SetCachedScreenPosition( Position );
+	InSlateWindow->SetCachedSize( Size );
+
+	PlatformApplication->InitializeWindow( NewWindow, Definition, NativeParent, bShowImmediately );
+
+	return NewWindow;
+}
+
+bool FSlateApplication::CanAddModalWindow() const
+{
+	// A modal window cannot be opened until the renderer has been created.
+	return CanDisplayWindows();
+}
+
+bool FSlateApplication::CanDisplayWindows() const
+{
+	// The renderer must be created and global shaders be available
+	return Renderer.IsValid() && Renderer->AreShadersInitialized();
+}
+
+/**
+ * Adds a modal window to the application.  
+ * In most cases, this function does not return until the modal window is closed (the only exception is a modal window for slow tasks)  
+ *
+ * @param InSlateWindow  A SlateWindow to which to add a native window.
+ */
+void FSlateApplication::AddModalWindow( TSharedRef<SWindow> InSlateWindow, const TSharedPtr<const SWidget> InParentWidget, bool bSlowTaskWindow )
+{
+	if( !CanAddModalWindow() )
+	{
+		// Bail out.  The incoming window will never be added, and no native window will be created.
+		return;
+	}
+
+	// Push the active modal window onto the stack.  
+	ActiveModalWindows.AddUnique( InSlateWindow );
+
+	// Close the open tooltip when a new window is open.  Tooltips from non-modal windows can be dangerous and cause rentrancy into code that shouldnt execute in a modal state.
+	CloseToolTip();
+
+	// Set the modal flag on the window
+	InSlateWindow->SetAsModalWindow();
+	
+	// Make sure we aren't in the middle of using a slate draw buffer
+	Renderer->FlushCommands();
+
+	// In slow task windows, depending on the frequency with which the window is updated, it could be quite some time 
+	// before the window is ticked (and drawn) so we hide the window by default and the slow task window will show it when needed
+	const bool bShowWindow = !bSlowTaskWindow;
+
+	// Create the new window
+	// Note: generally a modal window should not be added without a parent but 
+	// due to this being called from wxWidget editors, this is not always possible
+	if( InParentWidget.IsValid() )
+	{
+		// Find the window of the parent widget
+		FWidgetPath WidgetPath;
+		GeneratePathToWidgetChecked( InParentWidget.ToSharedRef(), WidgetPath );
+		AddWindowAsNativeChild( InSlateWindow, WidgetPath.GetWindow(), bShowWindow );
+	}
+	else
+	{
+		AddWindow( InSlateWindow, bShowWindow );
+	}
+
+	if ( ActiveModalWindows.Num() == 1 )
+	{
+		// Signal that a slate modal window has opened so external windows may be disabled as well
+		ModalWindowStackStartedDelegate.ExecuteIfBound();
+	}
+
+	// Release mouse capture here in case the new modal window has been added in one of the mouse button
+	// event callbacks. Otherwise it will be unresponsive until the next mouse up event.
+	ReleaseMouseCapture();
+
+	// Clear the cached pressed mouse buttons, in case a new modal window has been added between the mouse down and mouse up of another window.
+	PressedMouseButtons.Empty();
+
+	// Also force the platform capture off as the call to ReleaseMouseCapture() above still relies on mouse up messages to clear the capture
+	PlatformApplication->SetCapture( NULL );
+
+	// Disable high precision mouse mode when a modal window is added.  On some OS'es even when a window is diabled, raw input is sent to it.
+	PlatformApplication->SetHighPrecisionMouseMode( false, NULL );
+
+	// Block on all modal windows unless its a slow task.  In that case the game thread is allowed to run.
+	if( !bSlowTaskWindow )
+	{
+		// Show the cursor if it was previously hidden so users can interact with the window
+		if ( PlatformApplication->Cursor.IsValid() )
+		{
+			PlatformApplication->Cursor->Show( true );
+		}
+
+		// Tick slate from here in the event that we should not return until the modal window is closed.
+		while( InSlateWindow == GetActiveModalWindow() )
+		{
+			// Tick and render Slate
+			Tick();
+
+			// Synchronize the game thread and the render thread so that the render thread doesn't get too far behind.
+			Renderer->Sync();
+		}
+	}
+}
+
+void FSlateApplication::SetModalWindowStackStartedDelegate(FModalWindowStackStarted StackStartedDelegate)
+{
+	ModalWindowStackStartedDelegate = StackStartedDelegate;
+}
+
+void FSlateApplication::SetModalWindowStackEndedDelegate(FModalWindowStackEnded StackEndedDelegate)
+{
+	ModalWindowStackEndedDelegate = StackEndedDelegate;
+}
+
+TSharedRef<SWindow> FSlateApplication::AddWindowAsNativeChild( TSharedRef<SWindow> InSlateWindow, TSharedRef<SWindow> InParentWindow, const bool bShowImmediately )
+{
+	// Parent window must already have been added
+	checkSlow( ContainsWindow( SlateWindows, InParentWindow ) );
+
+	// Add the Slate window to the Slate application's top-level window array.  Note that neither the Slate window
+	// or the native window are ready to be used yet, however we need to make sure they're in the Slate window
+	// array so that we can properly respond to OS window messages as soon as they're sent.  For example, a window
+	// activation message may be sent by the OS as soon as the window is shown (in the Init function), and if we
+	// don't add the Slate window to our window list, we wouldn't be able to route that message to the window.
+	InParentWindow->AddChildWindow( InSlateWindow );
+	TSharedRef<FGenericWindow> NewWindow = MakeWindow( InSlateWindow, bShowImmediately );
+
+	if( bShowImmediately )
+	{
+		InSlateWindow->ShowWindow();
+
+		//@todo Slate: Potentially dangerous and annoying if all slate windows that are created steal focus.
+		if( InSlateWindow->SupportsKeyboardFocus() && InSlateWindow->IsFocusedInitially() )
+		{
+			InSlateWindow->GetNativeWindow()->SetWindowFocus();
+		}
+	}
+
+	return InSlateWindow;
+}
+
+
+TSharedRef<SWindow> FSlateApplication::PushMenu( const TSharedRef<SWidget>& InParentContent, const TSharedRef<SWidget>& InContent, const FVector2D& SummonLocation, const FPopupTransitionEffect& TransitionEffect, const bool bFocusImmediately, const bool bShouldAutoSize, const FVector2D& WindowSize, const FVector2D& SummonLocationSize )
+{
+	FWidgetPath WidgetPath;
+	GeneratePathToWidgetChecked( InParentContent, WidgetPath );
+
+#if !(UE_BUILD_SHIPPING && WITH_EDITOR)
+	// The would-be parent of the new menu being pushed is about to be destroyed.  Any children added to an about to be destroyed window will also be destroyed
+	if (IsWindowInDestroyQueue( WidgetPath.GetWindow() ))
+	{
+		UE_LOG(LogSlate, Warning, TEXT("FSlateApplication::PushMenu() called when parent window queued for destroy. New menu will be destroyed."));
+	}
+#endif
+
+	return MenuStack.PushMenu( WidgetPath.GetWindow(), InContent, SummonLocation, TransitionEffect, bFocusImmediately, bShouldAutoSize, WindowSize, SummonLocationSize );
+}
+
+bool FSlateApplication::HasOpenSubMenus( TSharedRef<SWindow> Window ) const
+{
+	return MenuStack.HasOpenSubMenus(Window);
+}
+
+bool FSlateApplication::AnyMenusVisible() const
+{
+	return MenuStack.GetNumStackLevels() > 0;
+}
+
+void FSlateApplication::DismissAllMenus()
+{
+	MenuStack.Dismiss();
+}
+
+void FSlateApplication::DismissMenu( TSharedRef<SWindow> MenuWindowToDismiss )
+{
+	int32 Location = MenuStack.FindLocationInStack( MenuWindowToDismiss );
+	// Dismiss everything starting at the window to dismiss
+	MenuStack.Dismiss( Location );
+}
+
+int32 FSlateApplication::GetLocationInMenuStack( TSharedRef<SWindow> WindowToFind ) const
+{
+	return MenuStack.FindLocationInStack( WindowToFind );
+}
+
+/**
+ * Destroying windows has implications on some OSs ( e.g. destroying Win32 HWNDs can cause events to be lost ).
+ * Slate strictly controls when windows are destroyed. 
+ *
+ * @param WindowToDestroy  Window to queue for destruction.  This will also queue children of this window for destruction
+ */
+void FSlateApplication::RequestDestroyWindow( TSharedRef<SWindow> InWindowToDestroy )
+{
+	struct local
+	{
+		static void Helper( const TSharedRef<SWindow> WindowToDestroy, TArray< TSharedRef<SWindow> >& OutWindowDestroyQueue)
+		{
+			/** @return the list of this window's child windows */
+			TArray< TSharedRef<SWindow> >& ChildWindows = WindowToDestroy->GetChildWindows();
+
+			// Children need to be destroyed first. 
+			if( ChildWindows.Num() > 0 )
+			{
+				for( int32 ChildIndex = 0; ChildIndex < ChildWindows.Num(); ++ChildIndex )
+				{	
+					// Recursively request that the window is destroyed which will also queue any children of children etc...
+					Helper( ChildWindows[ ChildIndex ], OutWindowDestroyQueue );
+				}
+			}
+
+			OutWindowDestroyQueue.AddUnique( WindowToDestroy );
+		}
+	};
+
+	local::Helper( InWindowToDestroy, WindowDestroyQueue );
+
+	DestroyWindowsImmediately();
+}
+
+void FSlateApplication::DestroyWindowImmediately( TSharedRef<SWindow> WindowToDestroy ) 
+{
+	// Request that the window and its children are destroyed
+	RequestDestroyWindow( WindowToDestroy );
+
+	DestroyWindowsImmediately();
+}
+
+/**
+ * Disable Slate components when an external, non-slate, modal window is brought up.  In the case of multiple
+ * external modal windows, we will only increment our tracking counter.
+ */
+void FSlateApplication::ExternalModalStart()
+{
+	if( NumExternalModalWindowsActive++ == 0 )
+	{
+		// Close all open menus.
+		DismissAllMenus();
+
+		// Close tool-tips
+		CloseToolTip();
+
+		// Tick and render Slate so that it can destroy any menu windows if necessary before we disable.
+		Tick();
+		Renderer->Sync();
+
+		if( ActiveModalWindows.Num() > 0 )
+		{
+			// There are still modal windows so only enable the new active modal window.
+			GetActiveModalWindow()->EnableWindow( false );
+		}
+		else
+		{
+			// We are creating a modal window so all other windows need to be disabled.
+			for( TArray< TSharedRef<SWindow> >::TIterator CurrentWindowIt( SlateWindows ); CurrentWindowIt; ++CurrentWindowIt )
+			{
+				TSharedRef<SWindow> CurrentWindow = ( *CurrentWindowIt );
+				CurrentWindow->EnableWindow( false );
+			}
+		}
+	}
+}
+
+/**
+ * Re-enable disabled Slate components when a non-slate modal window is dismissed.  Slate components
+ * will only be re-enabled when all tracked external modal windows have been dismissed.
+ */
+void FSlateApplication::ExternalModalStop()
+{
+	check(NumExternalModalWindowsActive > 0);
+	if( --NumExternalModalWindowsActive == 0 )
+	{
+		if( ActiveModalWindows.Num() > 0 )
+		{
+			// There are still modal windows so only enable the new active modal window.
+			GetActiveModalWindow()->EnableWindow( true );
+		}
+		else
+		{
+			// We are creating a modal window so all other windows need to be disabled.
+			for( TArray< TSharedRef<SWindow> >::TIterator CurrentWindowIt( SlateWindows ); CurrentWindowIt; ++CurrentWindowIt )
+			{
+				TSharedRef<SWindow> CurrentWindow = ( *CurrentWindowIt );
+				CurrentWindow->EnableWindow( true );
+			}
+		}
+	}
+}
+
+void FSlateApplication::InvalidateAllViewports()
+{
+	Renderer->InvalidateAllViewports();
+}
+
+
+void FSlateApplication::RegisterGameViewport( TSharedRef<SViewport> InViewport )
+{
+	GameViewportWidget = InViewport;
+
+	FWidgetPath PathToViewport;
+	// If we cannot find the window it could have been destroyed.
+	if( FindPathToWidget( SlateWindows, InViewport, PathToViewport, EVisibility::All ) )
+	{
+		FReply Reply = FReply::Handled().SetKeyboardFocus( InViewport, EKeyboardFocusCause::SetDirectly );
+	
+		// Set keyboard focus on the actual OS window for the top level Slate window in the viewport path
+		// This is needed because some OS messages are only sent to the window with keyboard focus
+		// Slate will translate the message and send it to the actual widget with focus.
+		// Without this we don't get WM_KEYDOWN or WM_CHAR messages in play in viewport sessions.
+		PathToViewport.GetWindow()->GetNativeWindow()->SetWindowFocus();
+
+		ProcessReply( PathToViewport, Reply, NULL, NULL );
+	}
+}
+
+
+void FSlateApplication::UnregisterGameViewport()
+{
+	ResetToDefaultInputSettings();
+	GameViewportWidget.Reset();
+}
+
+TSharedPtr<SViewport> FSlateApplication::GetGameViewport() const
+{
+	return GameViewportWidget.Pin();
+}
+
+void FSlateApplication::SetFocusToGameViewport()
+{
+	TSharedPtr< SViewport > CurrentGameViewportWidget = GameViewportWidget.Pin();
+
+	if ( CurrentGameViewportWidget.IsValid() )
+	{
+		SetKeyboardFocus( CurrentGameViewportWidget );
+	}
+}
+
+void FSlateApplication::SetJoystickCaptorToGameViewport()
+{
+	TSharedPtr< SViewport > CurrentGameViewportWidget = GameViewportWidget.Pin();
+
+	if ( CurrentGameViewportWidget.IsValid() )
+	{
+		FWidgetPath PathToWidget;
+		FindPathToWidget( SlateWindows, CurrentGameViewportWidget.ToSharedRef(), /*OUT*/ PathToWidget );
+
+		FReply Temp = FReply::Handled().CaptureJoystick(CurrentGameViewportWidget.ToSharedRef());
+
+		ProcessReply(PathToWidget, Temp, NULL, NULL);
+	}
+}
+
+void FSlateApplication::SetKeyboardFocus( const TSharedPtr< SWidget >& OptionalWidgetToFocus, EKeyboardFocusCause::Type ReasonFocusIsChanging )
+{
+	if (OptionalWidgetToFocus.IsValid())
+	{
+		FWidgetPath PathToWidget;
+		FindPathToWidget( SlateWindows, OptionalWidgetToFocus.ToSharedRef(), /*OUT*/ PathToWidget );
+
+		FReply Reply = FReply::Handled();
+		Reply.SetKeyboardFocus( OptionalWidgetToFocus.ToSharedRef(), EKeyboardFocusCause::SetDirectly );
+
+		ProcessReply( PathToWidget, Reply, NULL, NULL );	
+	}
+	else
+	{
+		ClearKeyboardFocus(EKeyboardFocusCause::SetDirectly);
+	}
+}
+
+void FSlateApplication::ResetToDefaultInputSettings()
+{
+	if (MouseCaptor.IsValid())
+	{
+		FWidgetPath MouseCaptorPath = MouseCaptor.ToWidgetPath();
+		ProcessReply( MouseCaptorPath, FReply::Handled().ReleaseMouseCapture(), NULL, NULL );
+	}
+
+	for (int32 UserIndex = 0; UserIndex < ARRAY_COUNT(JoystickCaptorWeakPaths); UserIndex++)
+	{
+		if (JoystickCaptorWeakPaths[UserIndex].IsValid())
+		{
+			FWidgetPath JoystickCaptorPath = JoystickCaptorWeakPaths[UserIndex].ToWidgetPath();
+			ProcessReply( JoystickCaptorPath, FReply::Handled().ReleaseJoystickCapture(), NULL, NULL, UserIndex );
+		}
+	}
+
+	ProcessReply( FWidgetPath(), FReply::Handled().ReleaseMouseLock(), NULL, NULL );
+	if ( PlatformApplication->Cursor.IsValid() )
+	{
+		PlatformApplication->Cursor->SetType(EMouseCursor::Default);
+	}
+}
+
+
+void* FSlateApplication::GetMouseCaptureWindow( void ) const
+{
+	return PlatformApplication->GetCapture();
+}
+
+
+void FSlateApplication::ReleaseMouseCapture()
+{
+	MouseCaptor.Invalidate();
+}
+
+
+TSharedPtr< SWidget > FSlateApplication::GetJoystickCaptor(uint32 UserIndex) const
+{
+	return JoystickCaptorWeakPaths[UserIndex].IsValid() ? JoystickCaptorWeakPaths[UserIndex].GetLastWidget().Pin() : TSharedPtr<SWidget>();
+}
+
+
+void FSlateApplication::ReleaseJoystickCapture(uint32 UserIndex)
+{
+	JoystickCaptorWeakPaths[UserIndex] = FWeakWidgetPath();
+}
+
+
+TSharedPtr<SWindow> FSlateApplication::GetActiveTopLevelWindow() const
+{
+	return ActiveTopLevelWindow.Pin();
+}
+
+
+TSharedPtr<SWindow> FSlateApplication::GetActiveModalWindow() const
+{
+	return (ActiveModalWindows.Num() > 0) ? ActiveModalWindows.Last() : NULL;
+}
+
+
+bool FSlateApplication::SetKeyboardFocus( const FWidgetPath& InFocusPath, const EKeyboardFocusCause::Type InCause )
+{
+	if (!InFocusPath.IsValid())
+	{
+		return false;
+	}
+
+	TSharedPtr<IWidgetReflector> WidgetReflector = WidgetReflectorPtr.Pin();
+	const bool bReflectorShowingFocus = WidgetReflector.IsValid() && WidgetReflector->IsShowingFocus();
+
+	bool bFocusTransferComplete = false;
+
+	for (int32 WidgetIndex = InFocusPath.Widgets.Num()-1; !bFocusTransferComplete  && WidgetIndex >= 0; --WidgetIndex)
+	{
+		const FArrangedWidget& WidgetToFocus = InFocusPath.Widgets( WidgetIndex );
+		
+		// Does this widget support keyboard focus?  If so, then we'll go ahead and set it!
+		if( WidgetToFocus.Widget->SupportsKeyboardFocus() )
+		{
+			// Has focus actually changed?
+			TSharedPtr< SWidget > OldFocusedWidget( GetKeyboardFocusedWidget() );
+			
+			// Is the focus actually changing?
+			if( WidgetToFocus.Widget != OldFocusedWidget )
+			{
+				FWidgetPath NewFocusPath = InFocusPath.GetPathDownTo(WidgetToFocus.Widget);
+				{
+					// Notify all affected widgets about the change in focus.
+					TArray< TSharedRef<SWidget> > NotifyUsAboutFocusChange;
+
+					// Notify widgets in the old focus path
+					{
+						for (int32 ChildIndex=0; ChildIndex < FocusedWidgetPath.Widgets.Num(); ++ChildIndex)
+						{
+							TSharedPtr<SWidget> SomeWidget = FocusedWidgetPath.Widgets[ ChildIndex ].Pin();
+							if (SomeWidget.IsValid())
+							{
+								NotifyUsAboutFocusChange.Add( SomeWidget.ToSharedRef() );
+							}
+						}
+
+						FScopedSwitchWorldHack SwitchWorld( FocusedWidgetPath.Window.Pin() );
+						
+						for( int32 NotifyWidgetIndex=0; NotifyWidgetIndex < NotifyUsAboutFocusChange.Num(); ++NotifyWidgetIndex )
+						{
+							NotifyUsAboutFocusChange[NotifyWidgetIndex]->OnKeyboardFocusChanging( FocusedWidgetPath, NewFocusPath );
+						}
+						
+					}
+	
+					// Empty the array for new widgets
+					NotifyUsAboutFocusChange.Empty();
+
+					// Notify widgets in the new focus path
+					{
+						for (int32 ChildIndex=0; ChildIndex < NewFocusPath.Widgets.Num(); ++ChildIndex)
+						{
+							NotifyUsAboutFocusChange.AddUnique( NewFocusPath.Widgets(ChildIndex).Widget );
+						}
+
+						FScopedSwitchWorldHack SwitchWorld( NewFocusPath );
+					
+						for( int32 NotifyWidgetIndex=0; NotifyWidgetIndex < NotifyUsAboutFocusChange.Num(); ++NotifyWidgetIndex )
+						{
+							NotifyUsAboutFocusChange[NotifyWidgetIndex]->OnKeyboardFocusChanging( FocusedWidgetPath, NewFocusPath );
+						}
+					}
+				}
+				
+				if( OldFocusedWidget.IsValid() )
+				{
+					// Switch worlds for widgets in the old path
+					FScopedSwitchWorldHack SwitchWorld( FocusedWidgetPath.Window.Pin() );
+
+					// Let previously-focused widget know that it's losing focus
+					OldFocusedWidget->OnKeyboardFocusLost( FKeyboardFocusEvent( InCause ) );
+				}
+
+				// Store a weak widget path to the widget that's taking focus
+				FocusedWidgetPath = FWeakWidgetPath( NewFocusPath );
+
+				if (bReflectorShowingFocus)
+				{
+					WidgetReflector->SetWidgetsToVisualize(NewFocusPath);
+				}
+
+				FocusCause = InCause;
+
+				// Let the new widget know that it's received keyboard focus
+				{
+					// Switch worlds for widgets in the new path
+					FScopedSwitchWorldHack SwitchWorld( NewFocusPath );
+
+					FReply Reply = WidgetToFocus.Widget->OnKeyboardFocusReceived( WidgetToFocus.Geometry, FKeyboardFocusEvent( InCause ) );
+					if (Reply.IsEventHandled())
+					{
+						ProcessReply( InFocusPath, Reply, NULL, NULL );
+					}
+				}
+
+				TSharedPtr<SWindow> FocusedWindow = FocusedWidgetPath.Window.Pin();
+				if ( FocusedWindow.IsValid() && FocusedWindow != WidgetToFocus.Widget )
+				{
+					FocusedWindow->SetWidgetToFocusOnActivate( WidgetToFocus.Widget );
+				}
+			}
+		
+			// We are about to successfully transfer focus.
+			bFocusTransferComplete = true;
+		}
+	}
+
+	return bFocusTransferComplete;
+}
+
+
+FModifierKeysState FSlateApplication::GetModifierKeys() const
+{
+	return PlatformApplication->GetModifierKeys();
+}
+
+
+void FSlateApplication::ClearKeyboardFocus( const EKeyboardFocusCause::Type InCause )
+{
+	// Let previously-focused widget know that it's losing focus
+	{
+		TSharedPtr< SWidget > OldFocusedWidget( GetKeyboardFocusedWidget() );
+		if( OldFocusedWidget.IsValid() )
+		{
+			if ( FocusedWidgetPath.Window.IsValid() )
+			{
+				// Switch worlds for widgets in the current path
+				FScopedSwitchWorldHack SwitchWorld( FocusedWidgetPath.Window.Pin().ToSharedRef() );
+
+				OldFocusedWidget->OnKeyboardFocusLost( FKeyboardFocusEvent( InCause ) );
+			}
+			else
+			{
+				OldFocusedWidget->OnKeyboardFocusLost( FKeyboardFocusEvent( InCause ) );
+			}
+		}
+	}
+
+	TSharedPtr<IWidgetReflector> WidgetReflector = WidgetReflectorPtr.Pin();
+	const bool bReflectorShowingFocus = WidgetReflector.IsValid() && WidgetReflector->IsShowingFocus();
+
+	if (bReflectorShowingFocus)
+	{
+		WidgetReflector->SetWidgetsToVisualize(FWidgetPath());
+	}
+
+	FocusedWidgetPath = FWeakWidgetPath();
+}
+
+
+void FSlateApplication::OnSetCursorMessage()
+{
+	QueryCursor();
+}
+
+void FSlateApplication::OnDragLeaveMessage()
+{
+	DragDropContent.Reset();
+}
+
+bool FSlateApplication::OnDropMessage( TSharedRef<SWindow> WindowDroppedInto, FDragDropEvent& DragDropEvent )
+{
+	FWidgetPath WidgetsUnderCursor = LocateWindowUnderMouse( DragDropEvent.GetScreenSpacePosition(), GetInteractiveTopLevelWindows() );
+	DragDropEvent.SetEventPath( WidgetsUnderCursor );
+
+	// Switch worlds for widgets in the current path
+	FScopedSwitchWorldHack SwitchWorld( WidgetsUnderCursor );
+
+	FReply Reply = FReply::Unhandled();
+	for ( int32 WidgetIndex = WidgetsUnderCursor.Widgets.Num()-1; !Reply.IsEventHandled() && WidgetIndex >=0; --WidgetIndex )
+	{
+		const FArrangedWidget& CurWidgetGeometry = WidgetsUnderCursor.Widgets(WidgetIndex);
+		Reply = CurWidgetGeometry.Widget->OnDrop( CurWidgetGeometry.Geometry, DragDropEvent ).SetHandler(CurWidgetGeometry.Widget);
+	}
+
+	LOG_EVENT( EEventLog::DropMessage, Reply )
+
+	return Reply.IsEventHandled();
+}
+
+void FSlateApplication::OnShutdown()
+{
+	// Clean up our tooltip window
+	TSharedPtr< SWindow > PinnedToolTipWindow( ToolTipWindow.Pin() );
+	if( PinnedToolTipWindow.IsValid() )
+	{
+		PinnedToolTipWindow->RequestDestroyWindow();
+		ToolTipWindow.Reset();
+	}
+
+	for( int32 WindowIndex = 0; WindowIndex < SlateWindows.Num(); ++WindowIndex )
+	{
+		// Destroy all top level windows.  This will also request that all children of each window be destroyed
+		RequestDestroyWindow( SlateWindows[WindowIndex] );
+	}
+
+	DestroyWindowsImmediately();
+}
+
+
+void FSlateApplication::DestroyWindowsImmediately()
+{
+	// Destroy any windows that were queued for deletion.
+
+	// Thomas.Sarkanen: I've changed this from a for() to a while() loop so that it is now valid to call RequestDestroyWindow()
+	// in the callstack of another call to RequestDestroyWindow(). Previously this would cause a stack overflow, as the
+	// WindowDestroyQueue would be continually added to each time the for() loop ran.
+	while ( WindowDestroyQueue.Num() > 0 )
+	{
+		TSharedRef<SWindow> CurrentWindow = WindowDestroyQueue[0];
+		WindowDestroyQueue.Remove(CurrentWindow);
+		if( ActiveModalWindows.Num() > 0 && ActiveModalWindows.Contains( CurrentWindow ) )
+		{
+			ActiveModalWindows.Remove( CurrentWindow );
+
+			if( ActiveModalWindows.Num() > 0 )
+			{
+				// There are still modal windows so only enable the new active modal window.
+				GetActiveModalWindow()->EnableWindow( true );
+			}
+			else
+			{
+				//  There are no modal windows so renable all slate windows
+				for ( TArray< TSharedRef<SWindow> >::TConstIterator SlateWindowIter( SlateWindows ); SlateWindowIter; ++SlateWindowIter )
+				{
+					// All other windows need to be re-enabled BEFORE a modal window is destroyed or focus will not be set correctly
+					(*SlateWindowIter)->EnableWindow( true );
+				}
+
+				// Signal that all slate modal windows are closed
+				ModalWindowStackEndedDelegate.ExecuteIfBound();
+			}
+		}
+
+		// Any window being destroyed should be removed from the menu stack if its in it
+		MenuStack.RemoveWindow( CurrentWindow );
+
+		// Perform actual cleanup of the window
+		PrivateDestroyWindow( CurrentWindow );
+	}
+
+	WindowDestroyQueue.Empty();
+}
+
+
+void FSlateApplication::SetExitRequestedHandler( const FSimpleDelegate& OnExitRequestedHandler )
+{
+	OnExitRequested = OnExitRequestedHandler;
+}
+
+
+bool FSlateApplication::FindPathToWidget( const TArray< TSharedRef<SWindow> > WindowsToSearch,  TSharedRef< const SWidget > InWidget, FWidgetPath& OutWidgetPath, EVisibility VisibilityFilter )
+{
+	SCOPE_CYCLE_COUNTER( STAT_FindPathToWidget);
+
+	// Iterate over our top level windows
+	bool bFoundWidget = false;
+	for( int32 WindowIndex = 0; !bFoundWidget && WindowIndex < WindowsToSearch.Num(); ++WindowIndex )
+	{
+		// Make a widget path that contains just the top-level window
+		TSharedRef< SWindow > CurWindow = WindowsToSearch[ WindowIndex ];
+		
+		FArrangedChildren JustWindow(VisibilityFilter);
+		{
+			JustWindow.AddWidget( FArrangedWidget(CurWindow, CurWindow->GetWindowGeometryInScreen()) );
+		}
+		
+		FWidgetPath PathToWidget( CurWindow, JustWindow );
+		
+		// Attempt to extend it to the desired child widget; essentially a full-window search for 'InWidget
+		if ( CurWindow == InWidget || PathToWidget.ExtendPathTo( FWidgetMatcher(InWidget), VisibilityFilter ) )
+		{
+			OutWidgetPath = PathToWidget;
+			bFoundWidget = true;
+		}
+
+		if ( !bFoundWidget )
+		{
+			// Search this window's children
+			bFoundWidget = FindPathToWidget(CurWindow->GetChildWindows(), InWidget, OutWidgetPath, VisibilityFilter);
+		}
+	}
+
+	return bFoundWidget;
+}
+
+
+bool FSlateApplication::GeneratePathToWidgetUnchecked( TSharedRef< const SWidget > InWidget, FWidgetPath& OutWidgetPath, EVisibility VisibilityFilter ) const
+{
+	return FindPathToWidget( SlateWindows, InWidget, OutWidgetPath, VisibilityFilter );
+}
+
+
+void FSlateApplication::GeneratePathToWidgetChecked( TSharedRef< const SWidget > InWidget, FWidgetPath& OutWidgetPath, EVisibility VisibilityFilter ) const
+{
+	const bool bWasFound = FindPathToWidget( SlateWindows, InWidget, OutWidgetPath, VisibilityFilter );
+	check( bWasFound );
+}
+
+
+TSharedPtr<SWindow> FSlateApplication::FindWidgetWindow( TSharedRef< const SWidget > InWidget ) const
+{
+	FWidgetPath WidgetPath;
+	return FindWidgetWindow( InWidget, WidgetPath );
+}
+
+
+TSharedPtr<SWindow> FSlateApplication::FindWidgetWindow( TSharedRef< const SWidget > InWidget, FWidgetPath& OutWidgetPath ) const
+{
+	// If the user wants a widget path back populate it instead
+	const bool bWasFound = FindPathToWidget( SlateWindows, InWidget, OutWidgetPath, EVisibility::All );
+	if( bWasFound )
+	{
+		return OutWidgetPath.TopLevelWindow;
+	}
+	return NULL;
+}
+
+
+void FSlateApplication::ProcessReply( const FWidgetPath& CurrentEventPath, const FReply TheReply, const FWidgetPath* WidgetsUnderMouse, const FPointerEvent* InMouseEvent, uint32 UserIndex )
+{
+	const TSharedPtr<FDragDropOperation> ReplyDragDropContent = TheReply.GetDragDropContent();
+	const bool bStartingDragDrop = ReplyDragDropContent.IsValid();
+
+	// Release mouse capture if requested or if we are starting a drag and drop.
+	// Make sure to only clobber WidgetsUnderCursor if we actually had a mouse capture.
+	if ( MouseCaptor.IsValid() && (TheReply.ShouldReleaseMouse() || bStartingDragDrop) )
+	{
+		WidgetsUnderCursorLastEvent = MouseCaptor.ToWeakPath();
+		MouseCaptor.Invalidate();
+	}
+
+
+	if ( TheReply.ShouldReleaseJoystick() )
+	{
+		if (JoystickCaptorWeakPaths[UserIndex].IsValid())
+		{
+			WidgetsUnderCursorLastEvent = JoystickCaptorWeakPaths[UserIndex];
+		}
+
+		if (TheReply.AffectsAllJoysticks())
+		{
+			for (int32 SlateUserIndex = 0; SlateUserIndex < SlateApplicationDefs::MaxUsers; ++SlateUserIndex)
+			{
+				JoystickCaptorWeakPaths[SlateUserIndex] = FWeakWidgetPath();
+			}
+		}
+		else
+		{
+			JoystickCaptorWeakPaths[UserIndex] = FWeakWidgetPath();
+		}
+	}
+
+	if (TheReply.ShouldEndDragDrop())
+	{
+		DragDropContent.Reset();
+	}
+
+	if ( bStartingDragDrop )
+	{
+		checkf( !this->DragDropContent.IsValid(), TEXT("Drag and Drop already in progress!") );
+		check( true == TheReply.IsEventHandled() );
+		check( WidgetsUnderMouse != NULL );
+		check( InMouseEvent != NULL );
+		DragDropContent = ReplyDragDropContent;
+
+		// We have entered drag and drop mode.
+		// Pretend that the mouse left all the previously hovered widgets, and a drag entered them.
+		for (int32 WidgetIndex=0; WidgetIndex < WidgetsUnderMouse->Widgets.Num(); ++WidgetIndex)
+		{
+			const FArrangedWidget& SomeWidget = WidgetsUnderMouse->Widgets(WidgetIndex);
+			SomeWidget.Widget->OnMouseLeave( *InMouseEvent );
+		}
+
+		FDragDropEvent DragDropEvent( *InMouseEvent, ReplyDragDropContent );
+		for (int32 WidgetIndex=0; WidgetIndex < WidgetsUnderMouse->Widgets.Num(); ++WidgetIndex)
+		{
+			const FArrangedWidget& SomeWidget = WidgetsUnderMouse->Widgets(WidgetIndex);
+			SomeWidget.Widget->OnDragEnter( SomeWidget.Geometry, DragDropEvent );
+		}
+	}
+	
+	TSharedPtr<SWidget> RequestedMouseCaptor = TheReply.GetMouseCaptor();
+	// Do not capture the mouse if we are also starting a drag and drop.
+	if ( RequestedMouseCaptor.IsValid() && !bStartingDragDrop )
+	{
+		MouseCaptor.SetMouseCaptor( CurrentEventPath, RequestedMouseCaptor );
+	}
+	
+	if( CurrentEventPath.IsValid() && ( TheReply.ShouldReleaseMouse() || RequestedMouseCaptor.IsValid() ) )
+	{
+		// If the mouse is being captured or released, toggle high precision raw input if requested by the reply.
+		// Raw input is only used with mouse capture
+		const TSharedRef< SWindow> Window = CurrentEventPath.GetWindow();
+
+		if ( TheReply.ShouldUseHighPrecisionMouse() )
+		{
+			PlatformApplication->SetCapture( Window->GetNativeWindow() );
+			PlatformApplication->SetHighPrecisionMouseMode( true, Window->GetNativeWindow() );
+		}
+		else if ( PlatformApplication->IsUsingHighPrecisionMouseMode() )
+		{
+			PlatformApplication->SetHighPrecisionMouseMode( false, NULL );
+			PlatformApplication->SetCapture( NULL );
+		}
+	}
+
+	TSharedPtr<SWidget> RequestedJoystickCaptor = TheReply.GetJoystickCaptor();
+	if( CurrentEventPath.IsValid() &&  RequestedJoystickCaptor.IsValid() )
+	{
+		FWidgetPath NewJoystickCaptorPath = CurrentEventPath.GetPathDownTo( RequestedJoystickCaptor.ToSharedRef() );
+
+		if ( !NewJoystickCaptorPath.IsValid() )
+		{
+			// The requested mouse captor was not in the event path.
+			// We will attempt to find it in this window; if we don't find it, then give up.
+			NewJoystickCaptorPath = CurrentEventPath.GetPathDownTo( CurrentEventPath.Widgets(0).Widget );
+			NewJoystickCaptorPath.ExtendPathTo( FWidgetMatcher(RequestedJoystickCaptor.ToSharedRef()) );
+		}
+
+		if (TheReply.AffectsAllJoysticks())
+		{
+			for (int32 SlateUserIndex = 0; SlateUserIndex < SlateApplicationDefs::MaxUsers; ++SlateUserIndex)
+			{
+				JoystickCaptorWeakPaths[SlateUserIndex] = NewJoystickCaptorPath;
+			}
+		}
+		else
+		{
+			JoystickCaptorWeakPaths[UserIndex] = NewJoystickCaptorPath;
+		}
+	}
+
+	TOptional<FIntPoint> RequestedMousePos = TheReply.GetRequestedMousePos();
+	if( RequestedMousePos.IsSet() )
+	{
+		const FVector2D Position = RequestedMousePos.GetValue();
+		LastCursorPosition = Position;
+		SetCursorPos( Position );
+	}
+
+	if( TheReply.GetMouseLockWidget().IsValid() )
+	{
+		// The reply requested mouse lock so tell the native application to lock the mouse to the widget receiving the event
+		LockCursor( TheReply.GetMouseLockWidget() );
+	}
+	else if( TheReply.ShouldReleaseMouseLock() )
+	{
+		// Unlock the mouse
+		LockCursor( NULL );
+	}
+
+	if ( TheReply.GetDetectDragRequest().IsValid() )
+	{
+		DragDetector.DetectDragForWidget = WidgetsUnderMouse->GetPathDownTo( TheReply.GetDetectDragRequest().ToSharedRef() );
+		DragDetector.DetectDragButton = TheReply.GetDetectDragRequestButton();
+		DragDetector.DetectDragStartLocation = InMouseEvent->GetScreenSpacePosition();
+	}
+
+	TSharedPtr<SWidget> RequestedKeyboardFocusRecepient = TheReply.GetFocusRecepient();
+	if ( CurrentEventPath.IsValid() && RequestedKeyboardFocusRecepient.IsValid() )
+	{
+		// The widget to focus is probably in the path of this event (likely the handler or handler's parent).
+		FWidgetPath NewFocusedWidgetPath = CurrentEventPath.GetPathDownTo( RequestedKeyboardFocusRecepient.ToSharedRef() );
+		if ( !NewFocusedWidgetPath.IsValid() )
+		{
+			// The widget we want to focus is not in the event processing path.
+			// Search all the widgets for it.
+			GeneratePathToWidgetUnchecked( RequestedKeyboardFocusRecepient.ToSharedRef(), NewFocusedWidgetPath );
+		}		
+		
+		SetKeyboardFocus( NewFocusedWidgetPath, TheReply.GetFocusCause() );
+	}
+}
+
+void FSlateApplication::LockCursor( const TSharedPtr<SWidget>& Widget )
+{
+	if ( PlatformApplication->Cursor.IsValid() )
+	{
+		if( Widget.IsValid() )
+		{
+			// Get a path to this widget so we know the position and size of its geometry
+			FWidgetPath WidgetPath;
+			GeneratePathToWidgetChecked( Widget.ToSharedRef(), WidgetPath );
+
+			// The last widget in the path should be the widget we are locking the cursor to
+			FArrangedWidget& WidgetGeom = WidgetPath.Widgets( WidgetPath.Widgets.Num() - 1 );
+
+			TSharedRef<SWindow> Window = WidgetPath.GetWindow();
+			// Do not attempt to lock the cursor to the window if its not in the foreground.  It would cause annoying side effects
+			if( Window->GetNativeWindow()->IsForegroundWindow() )
+			{
+				check( WidgetGeom.Widget == Widget );
+
+				FSlateRect SlateClipRect = WidgetGeom.Geometry.GetClippingRect();
+
+				// Generate a screen space clip rect based on the widgets geometry
+
+				// Note: We round the upper left coordinate of the clip rect so we guarantee the rect is inside the geometry of the widget.  If we truncated when there is a half pixel we would cause the clip
+				// rect to be half a pixel larger than the geometry and cause the mouse to go outside of the geometry.
+				RECT ClipRect;
+				ClipRect.left = FMath::Round( SlateClipRect.Left );
+				ClipRect.top = FMath::Round( SlateClipRect.Top );
+				ClipRect.right = FMath::Trunc( SlateClipRect.Right );
+				ClipRect.bottom = FMath::Trunc( SlateClipRect.Bottom );
+
+				// Lock the mouse to the widget
+				PlatformApplication->Cursor->Lock( &ClipRect );
+			}
+		}
+		else
+		{
+			// Unlock  the mouse
+			PlatformApplication->Cursor->Lock( NULL );
+		}
+	}
+}
+
+void FSlateApplication::QueryCursor()
+{
+	if ( PlatformApplication->Cursor.IsValid() )
+	{
+		// drag-drop overrides cursor
+		FCursorReply CursorResult = FCursorReply::Unhandled();
+
+		if(IsDragDropping())
+		{
+			CursorResult = DragDropContent->OnCursorQuery();
+			if (CursorResult.IsEventHandled())
+			{
+				// Query was handled, so we should set the cursor.
+				PlatformApplication->Cursor->SetType( CursorResult.GetCursor() );
+			}
+		}
+		
+		if(!CursorResult.IsEventHandled())
+		{
+			FWidgetPath WidgetsToQueryForCursor;
+			const TSharedPtr<SWindow> ActiveModalWindow = GetActiveModalWindow();
+
+			// Query widgets with mouse capture for the cursor
+			if (MouseCaptor.IsValid())
+			{
+				FWidgetPath MouseCaptorPath = MouseCaptor.ToWidgetPath();
+				TSharedRef< SWindow > CaptureWindow = MouseCaptorPath.GetWindow();
+			
+				// Never query the mouse captor path if it is outside an active modal window
+				if (!ActiveModalWindow.IsValid() || (CaptureWindow == ActiveModalWindow || CaptureWindow->IsDescendantOf(ActiveModalWindow)))
+				{
+					WidgetsToQueryForCursor = MouseCaptorPath;
+				}
+			}
+			else
+			{
+				WidgetsToQueryForCursor = LocateWindowUnderMouse( GetCursorPos(), GetInteractiveTopLevelWindows() );
+			}
+
+			if (WidgetsToQueryForCursor.IsValid())
+			{
+				// Switch worlds for widgets in the current path
+				FScopedSwitchWorldHack SwitchWorld( WidgetsToQueryForCursor );
+
+				const FVector2D CurrentCursorPosition = GetCursorPos();
+				const FPointerEvent CursorEvent(
+					CurrentCursorPosition,
+					LastCursorPosition,
+					CurrentCursorPosition - LastCursorPosition,
+					PressedMouseButtons,
+					PlatformApplication->GetModifierKeys()
+				);
+
+				CursorResult = FCursorReply::Unhandled();
+				for( int32 WidgetIndex = WidgetsToQueryForCursor.Widgets.Num() - 1; !CursorResult.IsEventHandled() && WidgetIndex >= 0; --WidgetIndex )
+				{
+					const FArrangedWidget& WidgetToQuery = WidgetsToQueryForCursor.Widgets( WidgetIndex );
+					CursorResult = WidgetToQuery.Widget->OnCursorQuery( WidgetToQuery.Geometry, CursorEvent );
+				}
+
+				if (CursorResult.IsEventHandled())
+				{
+					// Query was handled, so we should set the cursor.
+					PlatformApplication->Cursor->SetType( CursorResult.GetCursor() );
+				}
+				else if (WidgetsToQueryForCursor.IsValid())
+				{
+					// Query was NOT handled, and we are still over a slate window.
+					PlatformApplication->Cursor->SetType(EMouseCursor::Default);
+				}
+			}
+			else
+			{
+				// Set the default cursor when there isn't an active window under the cursor and the mouse isn't captured
+				PlatformApplication->Cursor->SetType(EMouseCursor::Default);
+			}
+		}
+	}
+}
+
+void FSlateApplication::SpawnToolTip( const TSharedRef<IToolTip>& InToolTip, const FVector2D& InSpawnLocation )
+{
+	// Close existing tool tip, if we have one
+	CloseToolTip();
+
+	// Spawn the new tool tip
+	{
+		TSharedPtr< SWindow > NewToolTipWindow( ToolTipWindow.Pin() );
+		if( !NewToolTipWindow.IsValid() )
+		{
+			// Create the tool tip window
+			NewToolTipWindow = SWindow::MakeToolTipWindow();
+
+			// Don't show the window yet.  We'll set it up with some content first!
+			const bool bShowImmediately = false;
+			AddWindow( NewToolTipWindow.ToSharedRef(), bShowImmediately );
+		}
+
+		NewToolTipWindow->SetContent
+		(
+			SNew(SWeakWidget)
+			.PossiblyNullContent(InToolTip->AsWidget())
+		);
+
+		// Move the window again to recalculate popup window position if necessary (tool tip may spawn outside of the monitors work area)
+		// and in that case we need to adjust it
+		DesiredToolTipLocation = InSpawnLocation;
+		{
+			// Make sure the desired size is valid
+			NewToolTipWindow->SlatePrepass();
+
+			FSlateRect Anchor(DesiredToolTipLocation.X, DesiredToolTipLocation.Y, DesiredToolTipLocation.X, DesiredToolTipLocation.Y);
+			DesiredToolTipLocation = CalculatePopupWindowPosition( Anchor, NewToolTipWindow->GetDesiredSize() );
+
+			// MoveWindowTo will adjust the window's position, if needed
+			NewToolTipWindow->MoveWindowTo( DesiredToolTipLocation );
+		}
+
+
+		// Show the window
+		NewToolTipWindow->ShowWindow();
+
+		// Keep a weak reference to the tool tip window
+		ToolTipWindow = NewToolTipWindow;
+
+		// Keep track of when this tool tip was spawned
+		ToolTipSummonTime = FPlatformTime::Seconds();
+	}
+}
+
+void FSlateApplication::CloseToolTip()
+{
+	TSharedPtr< SWindow > PinnedToolTipWindow( ToolTipWindow.Pin() );
+	if( PinnedToolTipWindow.IsValid() && PinnedToolTipWindow->IsVisible() )
+	{
+		// Notify the source widget that it's tooltip is closing
+		if (SWidget* SourceWidget = ActiveToolTipWidgetSource.Pin().Get())
+		{
+			SourceWidget->OnToolTipClosing();
+		}
+
+		// Hide the tool tip window.  We don't destroy the window, because we want to reuse it for future tool tips.
+		PinnedToolTipWindow->HideWindow();
+
+		ActiveToolTip.Reset();
+		ActiveToolTipWidgetSource.Reset();
+	}
+	ToolTipOffsetDirection = EToolTipOffsetDirection::Undetermined;
+}
+
+void FSlateApplication::UpdateToolTip( bool AllowSpawningOfNewToolTips )
+{
+	const bool bCheckForToolTipChanges =
+		bAllowToolTips &&					// Tool-tips must be enabled
+		!IsUsingHighPrecisionMouseMovment() && // If we are using HighPrecision movement then we can't rely on the OS cursor to be accurate
+		!IsDragDropping();					// We must not currently be in the middle of a drag-drop action
+	
+	// We still want to show tooltips for widgets that are disabled
+	const bool bIgnoreEnabledStatus = true;
+
+	FWidgetPath WidgetsToQueryForToolTip;
+	// We don't show any tooltips when drag and dropping or when another app is active
+	if (bCheckForToolTipChanges)
+	{
+		// Ask each widget under the Mouse if they have a tool tip to show.
+		FWidgetPath WidgetsUnderMouse = LocateWindowUnderMouse( GetCursorPos(), GetInteractiveTopLevelWindows(), bIgnoreEnabledStatus );
+		// Don't attempt to show tooltips inside an existing tooltip
+		if (!WidgetsUnderMouse.IsValid() || WidgetsUnderMouse.GetWindow() != ToolTipWindow.Pin())
+		{
+			WidgetsToQueryForToolTip = WidgetsUnderMouse;
+		}
+	}
+
+	bool bHaveForceFieldRect = false;
+	FSlateRect ForceFieldRect;
+
+	TSharedPtr<IToolTip> NewToolTip;
+	TSharedPtr<SWidget> WidgetProvidingNewToolTip;
+	for ( int32 WidgetIndex=WidgetsToQueryForToolTip.Widgets.Num()-1; WidgetIndex >= 0; --WidgetIndex )
+	{
+		FArrangedWidget* CurWidgetGeometry = &WidgetsToQueryForToolTip.Widgets(WidgetIndex);
+		const TSharedRef<SWidget>& CurWidget = CurWidgetGeometry->Widget;
+
+		if( !NewToolTip.IsValid() )
+		{
+			TSharedPtr< IToolTip > WidgetToolTip = CurWidget->GetToolTip();
+
+			// Make sure the tool-tip currently is displaying something before spawning it.
+			if( WidgetToolTip.IsValid() && !WidgetToolTip->IsEmpty() )
+			{
+				WidgetProvidingNewToolTip = CurWidget;
+				NewToolTip = WidgetToolTip;
+			}
+		}
+
+		// Keep track of the root most widget with a tool-tip force field enabled
+		if( CurWidget->HasToolTipForceField() )
+		{
+			if( !bHaveForceFieldRect )
+			{
+				bHaveForceFieldRect = true;
+				ForceFieldRect = CurWidgetGeometry->Geometry.GetClippingRect();				
+			}
+			else
+			{
+				// Grow the rect to encompass this geometry.  Usually, the parent's rect should always be inclusive
+				// of it's child though.  Just is kind of just being paranoid.
+				ForceFieldRect = ForceFieldRect.Expand( CurWidgetGeometry->Geometry.GetClippingRect() );
+			}
+		}
+	}
+
+	// Did the tool tip change from last time?
+	const bool bToolTipChanged = (NewToolTip != ActiveToolTip.Pin());
+	
+	// Any widgets that wish to handle visualizing the tooltip get a change here.
+	TSharedPtr<SWidget> NewTooltipVisualizer;
+	if (bToolTipChanged)
+	{
+		// Remove existing tooltip if there is one.
+		if (TooltipVisualizerPtr.IsValid())
+		{
+			TooltipVisualizerPtr.Pin()->OnVisualizeTooltip( NULL );
+		}
+
+		bool bOnVisualizeTooltipHandled = false;
+		// Some widgets might want to provide an alternative Tooltip Handler.
+		for ( int32 WidgetIndex=WidgetsToQueryForToolTip.Widgets.Num()-1; !bOnVisualizeTooltipHandled && WidgetIndex >= 0; --WidgetIndex )
+		{
+			const FArrangedWidget& CurWidgetGeometry = WidgetsToQueryForToolTip.Widgets(WidgetIndex);
+			bOnVisualizeTooltipHandled = CurWidgetGeometry.Widget->OnVisualizeTooltip( NewToolTip.IsValid() ? NewToolTip->AsWidget() : TSharedPtr<SWidget>() );
+			if (bOnVisualizeTooltipHandled)
+			{
+				// Someone is taking care of visualizing this tooltip
+				NewTooltipVisualizer = CurWidgetGeometry.Widget;
+			}
+		}
+	}
+
+
+	// If a widget under the cursor has a tool-tip forcefield active, then go through any menus
+	// in the menu stack that are above that widget's window, and make sure those windows also
+	// prevent the tool-tip from encroaching.  This prevents tool-tips from drawing over sub-menus
+	// spawned from menu items in a different window, for example.
+	if( bHaveForceFieldRect && WidgetsToQueryForToolTip.IsValid() )
+	{
+		const int32 MenuStackLevel = MenuStack.FindLocationInStack( WidgetsToQueryForToolTip.GetWindow() );
+
+		// Also check widgets in pop-up menus owned by this window
+		for( int32 CurStackLevel = MenuStackLevel + 1; CurStackLevel < MenuStack.GetNumStackLevels(); ++CurStackLevel )
+		{
+			FMenuWindowList& Windows = MenuStack.GetWindowsAtStackLevel( CurStackLevel );
+
+			for( FMenuWindowList::TConstIterator WindowIt( Windows ); WindowIt; ++WindowIt )
+			{
+				TSharedPtr< SWindow > CurWindow = *WindowIt;
+				if( CurWindow.IsValid() )
+				{
+					const FGeometry& WindowGeometry = CurWindow->GetWindowGeometryInScreen();
+					ForceFieldRect = ForceFieldRect.Expand( WindowGeometry.GetClippingRect() );
+				}
+			}
+		}
+	}
+
+	{
+		TSharedPtr<IToolTip> ActiveToolTipPtr = ActiveToolTip.Pin();
+		if ( ( ActiveToolTipPtr.IsValid() && !ActiveToolTipPtr->IsInteractive() ) || ( NewToolTip.IsValid() && NewToolTip != ActiveToolTip.Pin() ) )
+		{
+			// Keep track of where we want tool tips to be positioned
+			DesiredToolTipLocation = LastCursorPosition + SlateDefs::ToolTipOffsetFromMouse;
+		}
+	}
+
+	TSharedPtr< SWindow > ToolTipWindowPtr = ToolTipWindow.Pin();
+	if ( ToolTipWindowPtr.IsValid() )
+	{
+		FSlateRect Anchor(DesiredToolTipLocation.X, DesiredToolTipLocation.Y, DesiredToolTipLocation.X, DesiredToolTipLocation.Y);
+		DesiredToolTipLocation = CalculatePopupWindowPosition( Anchor, ToolTipWindowPtr->GetDesiredSize() );
+	}
+
+	// Repel tool-tip from a force field, if necessary
+	if( bHaveForceFieldRect )
+	{
+		FVector2D ToolTipShift;
+		ToolTipShift.X = ( ForceFieldRect.Right + SlateDefs::ToolTipOffsetFromForceField.X ) - DesiredToolTipLocation.X;
+		ToolTipShift.Y = ( ForceFieldRect.Bottom + SlateDefs::ToolTipOffsetFromForceField.Y ) - DesiredToolTipLocation.Y;
+
+		// Make sure the tool-tip needs to be offset
+		if( ToolTipShift.X > 0.0f && ToolTipShift.Y > 0.0f )
+		{
+			// Find the best edge to move the tool-tip towards
+			if( ToolTipOffsetDirection == EToolTipOffsetDirection::Right ||
+				( ToolTipOffsetDirection == EToolTipOffsetDirection::Undetermined && ToolTipShift.X < ToolTipShift.Y ) )
+			{
+				// Move right
+				DesiredToolTipLocation.X += ToolTipShift.X;
+				ToolTipOffsetDirection = EToolTipOffsetDirection::Right;
+			}
+			else
+			{
+				// Move down
+				DesiredToolTipLocation.Y += ToolTipShift.Y;
+				ToolTipOffsetDirection = EToolTipOffsetDirection::Down;
+			}
+		}
+	}
+
+	// The tool tip changed...
+	if ( bToolTipChanged )
+	{
+		// Close any existing tooltips; Unless the current tooltip is interactive and we don't have a valid tooltip to replace it
+		TSharedPtr<IToolTip> ActiveToolTipPtr = ActiveToolTip.Pin();
+		if ( NewToolTip.IsValid() || ( ActiveToolTipPtr.IsValid() && !ActiveToolTipPtr->IsInteractive() ) )
+		{
+			CloseToolTip();
+			
+			if (NewTooltipVisualizer.IsValid())
+			{
+				TooltipVisualizerPtr = NewTooltipVisualizer;
+			}
+			else if( bAllowToolTips && AllowSpawningOfNewToolTips )
+			{
+				// Spawn a new one if we have it
+				if( NewToolTip.IsValid() )
+				{
+					SpawnToolTip( NewToolTip.ToSharedRef(), DesiredToolTipLocation );
+				}
+			}
+			else
+			{
+				NewToolTip = NULL;
+			}
+
+			ActiveToolTip = NewToolTip;
+			ActiveToolTipWidgetSource = WidgetProvidingNewToolTip;
+		}
+	}
+
+	// Do we have a tool tip window?
+	if( ToolTipWindow.IsValid() )
+	{
+		// Only enable tool-tip transitions if we're running at a decent frame rate
+		const bool bAllowInstantToolTips = false;
+		const bool bAllowAnimations = !bAllowInstantToolTips && FSlateApplication::Get().IsRunningAtTargetFrameRate();
+
+		// How long since the tool tip was summoned?
+		const float TimeSinceSummon = FPlatformTime::Seconds() - ToolTipDelay - ToolTipSummonTime;
+		const float ToolTipOpacity = bAllowInstantToolTips ? 1.0f : FMath::Clamp< float >( TimeSinceSummon / ToolTipFadeInDuration, 0.0f, 1.0f );
+
+		// Update window opacity
+		TSharedRef< SWindow > PinnedToolTipWindow( ToolTipWindow.Pin().ToSharedRef() );
+		PinnedToolTipWindow->SetOpacity( ToolTipOpacity );
+
+		// How far tool tips should slide
+		const FVector2D SlideDistance( 30.0f, 5.0f );
+
+		// Apply steep inbound curve to the movement, so it looks like it quickly decelerating
+		const float SlideProgress = bAllowAnimations ? FMath::Pow( 1.0f - ToolTipOpacity, 3.0f ) : 0.0f;
+
+		FVector2D WindowLocation = DesiredToolTipLocation + SlideProgress * SlideDistance;
+		if( WindowLocation != PinnedToolTipWindow->GetPositionInScreen() )
+		{
+			// Avoid the edges of the desktop
+			FSlateRect Anchor(WindowLocation.X, WindowLocation.Y, WindowLocation.X, WindowLocation.Y);
+			WindowLocation = CalculatePopupWindowPosition( Anchor, PinnedToolTipWindow->GetDesiredSize() );
+
+			// Update the tool tip window positioning
+			PinnedToolTipWindow->MoveWindowTo( WindowLocation );
+		}
+	}
+}
+
+int32 FSlateApplication::DrawKeyboardFocus( const FWidgetPath& FocusPath, FSlateWindowElementList& WindowElementList, int32 InLayerId ) const
+{
+	if (FocusCause == EKeyboardFocusCause::Keyboard)
+	{
+		// Widgets where being focused matters draw themselves differently when focused.
+		// When the user navigates keyboard focus, we draw keyboard focus for everything, 
+		// so the user can see what they are doing.
+		const FArrangedWidget& FocusedWidgetGeomertry = FocusPath.Widgets.Last();
+
+		FSlateDrawElement::MakeBox(
+			WindowElementList,
+			InLayerId++,
+			FPaintGeometry(
+			FocusedWidgetGeomertry.Geometry.AbsolutePosition - FocusPath.GetWindow()->GetPositionInScreen(),
+			FocusedWidgetGeomertry.Geometry.Size * FocusedWidgetGeomertry.Geometry.Scale,
+			FocusedWidgetGeomertry.Geometry.Scale ),
+			FCoreStyle::Get().GetBrush("FocusRectangle"),
+			FocusPath.GetWindow()->GetClippingRectangleInWindow(),
+			ESlateDrawEffect::None,
+			FColor(255,255,255,128)
+		);
+	}
+
+	return InLayerId;
+}
+
+TArray< TSharedRef<SWindow> > FSlateApplication::GetInteractiveTopLevelWindows()
+{
+	if (ActiveModalWindows.Num() > 0)
+	{
+		// If we have modal windows, only the topmost modal window and its children are interactive.
+		TArray< TSharedRef<SWindow>, TInlineAllocator<1> > OutWindows;
+		OutWindows.Add( ActiveModalWindows.Last().ToSharedRef() );
+		return TArray< TSharedRef<SWindow> >(OutWindows);
+	}
+	else
+	{
+		// No modal windows? All windows are interactive.
+		return SlateWindows;
+	}
+}
+
+void FSlateApplication::GetAllVisibleWindowsOrdered(TArray< TSharedRef<SWindow> >& OutWindows)
+{
+	for( TArray< TSharedRef<SWindow> >::TConstIterator CurrentWindowIt( SlateWindows ); CurrentWindowIt; ++CurrentWindowIt )
+	{
+		TSharedRef<SWindow> CurrentWindow = *CurrentWindowIt;
+		if ( CurrentWindow->IsVisible() )
+		{
+			GetAllVisibleChildWindows(OutWindows, CurrentWindow);
+		}
+	}
+}
+
+void FSlateApplication::GetAllVisibleChildWindows(TArray< TSharedRef<SWindow> >& OutWindows, TSharedRef<SWindow> CurrentWindow)
+{
+	if ( CurrentWindow->IsVisible() )
+	{
+		OutWindows.Add(CurrentWindow);
+
+		const TArray< TSharedRef<SWindow> >& WindowChildren = CurrentWindow->GetChildWindows();
+		for (int32 ChildIndex=0; ChildIndex < WindowChildren.Num(); ++ChildIndex)
+		{
+			GetAllVisibleChildWindows( OutWindows, WindowChildren[ChildIndex] );
+		}
+	}
+}
+
+bool FSlateApplication::IsDragDropping() const
+{
+	return DragDropContent.IsValid();
+}
+
+TSharedPtr<FDragDropOperation> FSlateApplication::GetDragDroppingContent() const
+{
+	return DragDropContent;
+}
+
+void FSlateApplication::EndDragDrop()
+{
+	DragDropContent.Reset();
+}
+
+void FSlateApplication::EnterDebuggingMode()
+{
+	bRequestLeaveDebugMode = false;
+
+	// Note it is ok to hold a reference here as the game viewport should not be destroyed while in debugging mode
+	TSharedPtr<SViewport> PreviousGameViewport;
+
+	// Disable any game viewports while we are in debug mode so that mouse capture is released and the cursor is visible
+	if (GameViewportWidget.IsValid())
+	{
+		PreviousGameViewport = GameViewportWidget.Pin();
+		UnregisterGameViewport();
+	}
+
+	Renderer->FlushCommands();
+	
+	// We are about to start an in stack tick. Make sure the rendering thread isn't already behind
+	Renderer->Sync();
+
+#if WITH_EDITORONLY_DATA
+	// Flag that we're about to enter the first frame of intra-frame debugging.
+	GFirstFrameIntraFrameDebugging = true;
+#endif	//WITH_EDITORONLY_DATA
+
+	// Tick slate from here in the event that we should not return until the modal window is closed.
+	while (!bRequestLeaveDebugMode)
+	{
+		// Tick and render Slate
+		Tick();
+
+		// Synchronize the game thread and the render thread so that the render thread doesn't get too far behind.
+		Renderer->Sync();
+
+#if WITH_EDITORONLY_DATA
+		// We are done with the first frame
+		GFirstFrameIntraFrameDebugging = false;
+
+		// If we are requesting leaving debugging mode, leave it now.
+		GIntraFrameDebuggingGameThread = !bRequestLeaveDebugMode;
+#endif	//WITH_EDITORONLY_DATA
+	}
+
+	bRequestLeaveDebugMode = false;
+	
+	if ( PreviousGameViewport.IsValid() )
+	{
+		check(!GameViewportWidget.IsValid());
+
+		// When in single step mode, register the game viewport so we can unregister it later
+		// but do not do any of the other stuff like locking or capturing the mouse.
+		if( bLeaveDebugForSingleStep )
+		{
+			GameViewportWidget = PreviousGameViewport;
+		}
+		else
+		{
+			// If we had a game viewport before debugging, re-register it now to capture the mouse and lock the cursor
+			RegisterGameViewport( PreviousGameViewport.ToSharedRef() );
+		}
+	}
+
+	bLeaveDebugForSingleStep = false;
+}
+
+void FSlateApplication::LeaveDebuggingMode(  bool bLeavingForSingleStep )
+{
+	bRequestLeaveDebugMode = true;
+	bLeaveDebugForSingleStep = bLeavingForSingleStep;
+}
+
+bool FSlateApplication::IsWindowInDestroyQueue(TSharedRef<SWindow> Window) const
+{
+	return WindowDestroyQueue.Contains(Window);
+}
+
+void FSlateApplication::SynthesizeMouseMove()
+{
+	FPointerEvent MouseEvent(
+		LastCursorPosition, 
+		LastCursorPosition, 
+		PressedMouseButtons, 
+		EKeys::Invalid, 
+		0, 
+		PlatformApplication->GetModifierKeys()
+	);
+
+	ProcessMouseMoveMessage(MouseEvent);
+}
+
+void FSlateApplication::OnLogSlateEvent(EEventLog::Type Event, const FString& AdditionalContent)
+{
+#if LOG_SLATE_EVENTS
+	if (EventLogger.IsValid())
+	{
+		LOG_EVENT_CONTENT(Event, AdditionalContent, NULL);
+	}
+#endif
+}
+
+void FSlateApplication::OnLogSlateEvent(EEventLog::Type Event, const FText& AdditionalContent )
+{
+#if LOG_SLATE_EVENTS
+	if (EventLogger.IsValid())
+	{
+		LOG_EVENT_CONTENT(Event, AdditionalContent.ToString(), NULL);
+	}
+#endif
+};
+
+void FSlateApplication::SetSlateUILogger(TSharedPtr<IEventLogger> InEventLogger)
+{
+#if LOG_SLATE_EVENTS
+	EventLogger = InEventLogger;
+#endif
+}
+
+void FSlateApplication::SetUnhandledKeyDownEventHandler( const FOnKeyboardEvent& NewHandler )
+{
+	UnhandledKeyDownEventHandler = NewHandler;
+}
+
+FVector2D FSlateApplication::CalculatePopupWindowPosition( const FSlateRect& InAnchor, const FVector2D& InSize, const EOrientation Orientation ) const
+{
+	// Do nothing if this window has no size
+	if (InSize == FVector2D::ZeroVector)
+	{
+		return FVector2D(InAnchor.Left, InAnchor.Top);
+	}
+
+	FVector2D CalculatedPopUpWindowPosition( 0, 0 );
+
+	FPlatformRect AnchorRect;
+	AnchorRect.Left = InAnchor.Left;
+	AnchorRect.Top = InAnchor.Top;
+	AnchorRect.Right = InAnchor.Right;
+	AnchorRect.Bottom = InAnchor.Bottom;
+
+	EPopUpOrientation::Type PopUpOrientation = EPopUpOrientation::Horizontal;
+
+	if ( Orientation == EOrientation::Orient_Vertical )
+	{
+		PopUpOrientation =  EPopUpOrientation::Vertical;
+	}
+
+	if ( PlatformApplication->TryCalculatePopupWindowPosition( AnchorRect, InSize, PopUpOrientation, /*OUT*/&CalculatedPopUpWindowPosition ) )
+	{
+		return CalculatedPopUpWindowPosition;
+	}
+	else
+	{
+		// Calculate the rectangle around our work area
+		// Use our own rect.  This window as probably doesn't have a size or position yet.
+		// Use a size of 1 to get the closest monitor to the start point
+		AnchorRect.Left = InAnchor.Left + 1;
+		AnchorRect.Top = InAnchor.Top + 1;
+		const FPlatformRect PlatformWorkArea = PlatformApplication->GetWorkArea( AnchorRect );
+
+		FGeometry WorkArea(
+			FVector2D(PlatformWorkArea.Left, PlatformWorkArea.Top),
+			FVector2D::ZeroVector,
+			FVector2D(PlatformWorkArea.Right - PlatformWorkArea.Left, PlatformWorkArea.Bottom - PlatformWorkArea.Top),
+			1 );
+
+		FSlateRect WorkAreaRect( WorkArea.Position.X, WorkArea.Position.Y, WorkArea.Position.X+WorkArea.Size.X, WorkArea.Position.Y+WorkArea.Size.Y );
+
+		// In the direction we are opening, see if there is enough room. If there is not, flip the opening direction along the same axis.
+		FVector2D NewPosition = FVector2D::ZeroVector;
+		if ( Orientation == Orient_Horizontal )
+		{
+			const bool bFitsRight = InAnchor.Right + InSize.X < WorkAreaRect.Right;
+			const bool bFitsLeft = InAnchor.Left - InSize.X >= WorkAreaRect.Left;
+
+			if ( bFitsRight || !bFitsLeft )
+			{
+				// The menu fits to the right of the anchor or it does not fit to the left, display to the right
+				NewPosition = FVector2D(InAnchor.Right, InAnchor.Top);
+			}
+			else
+			{
+				// The menu does not fit to the right of the anchor but it does fit to the left, display to the left
+				NewPosition = FVector2D(InAnchor.Left - InSize.X, InAnchor.Top);
+			}
+		}
+		else
+		{
+			const bool bFitsDown = InAnchor.Bottom + InSize.Y < WorkAreaRect.Bottom;
+			const bool bFitsUp = InAnchor.Top - InSize.Y >= WorkAreaRect.Top;
+
+			if ( bFitsDown || !bFitsUp )
+			{
+				// The menu fits below the anchor or it does not fit above, display below
+				NewPosition = FVector2D(InAnchor.Left, InAnchor.Bottom);
+			}
+			else
+			{
+				// The menu does not fit below the anchor but it does fit above, display above
+				NewPosition = FVector2D(InAnchor.Left, InAnchor.Top - InSize.Y);
+			}
+
+			if ( !bFitsDown && !bFitsUp )
+			{
+				NewPosition.X = InAnchor.Right;
+			}
+		}
+
+		// Adjust the position of popup windows so they do not go out of the visible area of the monitor(s)
+		// This can happen along the opposite axis that we are opening with
+		// Assumes this window has a valid size
+		// Adjust any menus that my not fit on the screen where they are opened
+		FVector2D StartPos = NewPosition;
+		FVector2D EndPos = NewPosition+InSize;
+		FVector2D Adjust = FVector2D::ZeroVector;
+		if (StartPos.X < WorkAreaRect.Left)
+		{
+			// Window is clipped by the left side of the work area
+			Adjust.X = WorkAreaRect.Left - StartPos.X;
+		}
+
+		if (StartPos.Y < WorkAreaRect.Top)
+		{
+			// Window is clipped by the top of the work area
+			Adjust.Y = WorkAreaRect.Top - StartPos.Y;
+		}
+
+		if (EndPos.X > WorkAreaRect.Right)
+		{
+			// Window is clipped by the right side of the work area
+			Adjust.X = WorkAreaRect.Right - EndPos.X;
+		}
+
+		if (EndPos.Y > WorkAreaRect.Bottom)
+		{
+			// Window is clipped by the bottom of the work area
+			Adjust.Y = WorkAreaRect.Bottom - EndPos.Y;
+		}
+
+		NewPosition += Adjust;
+
+		return NewPosition;
+	}
+}
+
+bool FSlateApplication::IsRunningAtTargetFrameRate() const
+{
+	const float MinimumDeltaTime = 1.0f / TargetFrameRateForResponsiveness.GetValueOnGameThread();
+	return ( AverageDeltaTimeForResponsiveness <= MinimumDeltaTime ) || !IsNormalExecution();
+}
+
+
+bool FSlateApplication::AreMenuAnimationsEnabled() const
+{
+	return bMenuAnimationsEnabled;
+}
+
+
+void FSlateApplication::EnableMenuAnimations( const bool bEnableAnimations )
+{
+	bMenuAnimationsEnabled = bEnableAnimations;
+}
+
+
+void FSlateApplication::SetAppIcon(const FSlateBrush* const InAppIcon)
+{
+	check(InAppIcon);
+	AppIcon = InAppIcon;
+}
+
+
+const FSlateBrush* FSlateApplication::GetAppIcon() const
+{
+	return AppIcon;
+}
+
+
+void FSlateApplication::ShowKeyboard( bool bShow, TSharedPtr<SVirtualKeyboardEntry> TextEntryWidget )
+{
+	if(SlateTextField == NULL)
+	{
+		SlateTextField = new FPlatformTextField();
+	}
+
+	SlateTextField->ShowKeyboard(bShow, TextEntryWidget);
+}
+
+FSlateRect FSlateApplication::GetPreferredWorkArea() const
+{
+	// First see if we have a focused widget
+	if( FocusedWidgetPath.IsValid() && FocusedWidgetPath.Window.IsValid() )
+	{
+		const FVector2D WindowPos = FocusedWidgetPath.Window.Pin()->GetPositionInScreen();
+		const FVector2D WindowSize = FocusedWidgetPath.Window.Pin()->GetSizeInScreen();
+		return GetWorkArea( FSlateRect( WindowPos.X, WindowPos.Y, WindowPos.X + WindowSize.X, WindowPos.Y + WindowSize.Y ) );
+	}
+	else
+	{
+		// no focus widget, so use mouse position if there are windows present in the work area
+		const FVector2D CursorPos = GetCursorPos();
+		const FSlateRect WorkArea = GetWorkArea( FSlateRect( CursorPos.X, CursorPos.Y, CursorPos.X + 1.0f, CursorPos.Y + 1.0f ) );
+
+		struct Local
+		{
+			static bool CheckWorkAreaForWindows_Helper( const TArray< TSharedRef<SWindow> >& WindowsToSearch, const FSlateRect& InWorkAreaRect )
+			{
+				for( TArray< TSharedRef<SWindow> >::TConstIterator CurrentWindowIt( WindowsToSearch ); CurrentWindowIt; ++CurrentWindowIt )
+				{
+					const TSharedRef<SWindow>& CurrentWindow = *CurrentWindowIt;
+					const FVector2D Position = CurrentWindow->GetPositionInScreen();
+					const FVector2D Size = CurrentWindow->GetSizeInScreen();
+					const FSlateRect WindowRect( Position.X, Position.Y, Size.X, Size.Y );
+
+					if( FSlateRect::DoRectanglesIntersect( InWorkAreaRect, WindowRect ) || CheckWorkAreaForWindows_Helper( CurrentWindow->GetChildWindows(), InWorkAreaRect ) )
+					{
+						return true;
+					}
+				}
+
+				return false;
+			}
+		};
+
+		if( Local::CheckWorkAreaForWindows_Helper( SlateWindows, WorkArea ) )
+		{
+			return WorkArea;
+		}
+
+		// If we can't find a window where the cursor is at, try finding a main window.
+		TSharedPtr<SWindow> ActiveTop = GetActiveTopLevelWindow();
+		if ( ActiveTop.IsValid() )
+		{
+			// Use the current top level windows rect
+			return GetWorkArea( ActiveTop->GetRectInScreen() );
+		}
+		
+		// If we can't find a top level window check for an active modal window
+		TSharedPtr<SWindow> ActiveModal = GetActiveModalWindow();
+		if ( ActiveModal.IsValid() )
+		{
+			// Use the current active modal windows rect
+			return GetWorkArea( ActiveModal->GetRectInScreen() );
+		}
+
+		// no windows on work area - default to primary display
+		FDisplayMetrics DisplayMetrics;
+		GetDisplayMetrics( DisplayMetrics );
+		const FPlatformRect& DisplayRect = DisplayMetrics.PrimaryDisplayWorkAreaRect;
+		return FSlateRect( (float)DisplayRect.Left, (float)DisplayRect.Top, (float)DisplayRect.Right, (float)DisplayRect.Bottom );
+	}
+}
+
+FSlateRect FSlateApplication::GetWorkArea( const FSlateRect& InRect ) const
+{
+	FPlatformRect InPlatformRect;
+	InPlatformRect.Left = FMath::Trunc(InRect.Left);
+	InPlatformRect.Top = FMath::Trunc(InRect.Top);
+	InPlatformRect.Right = FMath::Trunc(InRect.Right);
+	InPlatformRect.Bottom = FMath::Trunc(InRect.Bottom);
+
+	const FPlatformRect OutPlatformRect = PlatformApplication->GetWorkArea( InPlatformRect );
+	return FSlateRect( OutPlatformRect.Left, OutPlatformRect.Top, OutPlatformRect.Right, OutPlatformRect.Bottom );
+}
+
+bool FSlateApplication::SupportsSourceAccess() const
+{
+	if(QuerySourceCodeAccessDelegate.IsBound())
+	{
+		return QuerySourceCodeAccessDelegate.Execute();
+	}
+	return false;
+}
+
+void FSlateApplication::GotoLineInSource(const FString& FileName, int32 LineNumber) const
+{
+	if ( SupportsSourceAccess() )
+	{
+		if(SourceCodeAccessDelegate.IsBound())
+		{
+			SourceCodeAccessDelegate.Execute(FileName, LineNumber, 0);
+		}
+	}
+}
+
+void FSlateApplication::ForceRedrawWindow(TSharedRef<SWindow>& InWindowToDraw)
+{
+	PrivateDrawWindows( InWindowToDraw );
+}
+
+bool FSlateApplication::TakeScreenshot(TSharedRef<SWidget>& Widget, TArray<FColor>&OutColorData, FIntVector& OutSize)
+{
+	return TakeScreenshot(Widget, FIntRect(), OutColorData, OutSize);
+}
+
+bool FSlateApplication::TakeScreenshot(TSharedRef<SWidget>& Widget, const FIntRect& InnerWidgetArea, TArray<FColor>& OutColorData, FIntVector& OutSize)
+{
+	// We can't screenshot the widget unless there's a valid window handle to draw it in.
+	TSharedPtr<SWindow> WidgetWindow = FSlateApplication::Get().FindWidgetWindow(Widget);
+	if ( !WidgetWindow.IsValid() )
+	{
+		return false;
+	}
+
+	TSharedRef<SWindow> CurrentWindowRef = WidgetWindow.ToSharedRef();
+
+	FWidgetPath WidgetPath;
+	FSlateApplication::Get().GeneratePathToWidgetChecked(Widget, WidgetPath);
+
+	FArrangedWidget ArrangedWidget = WidgetPath.FindArrangedWidget(Widget);
+	FVector2D Position = ArrangedWidget.Geometry.AbsolutePosition;
+	FVector2D Size = ArrangedWidget.Geometry.GetDrawSize();
+	FVector2D WindowPosition = WidgetWindow->GetPositionInScreen();
+
+	FIntRect ScreenshotRect = InnerWidgetArea.IsEmpty() ? FIntRect(0, 0, (int32)Size.X, (int32)Size.Y) : InnerWidgetArea;
+
+	ScreenshotRect.Min.X += ( Position.X - WindowPosition.X );
+	ScreenshotRect.Min.Y += ( Position.Y - WindowPosition.Y );
+	ScreenshotRect.Max.X += ( Position.X - WindowPosition.X );
+	ScreenshotRect.Max.Y += ( Position.Y - WindowPosition.Y );
+
+	Renderer->PrepareToTakeScreenshot(ScreenshotRect, &OutColorData);
+	PrivateDrawWindows(WidgetWindow);
+
+	OutSize.X = ScreenshotRect.Size().X;
+	OutSize.Y = ScreenshotRect.Size().Y;
+
+	return true;
+}
+
+static TSharedPtr<SWindow> FindWindowByPlatformWindow( const TArray< TSharedRef<SWindow> >& WindowsToSearch, const TSharedRef< FGenericWindow >& PlatformWindow )
+{
+	for (int32 WindowIndex=0; WindowIndex < WindowsToSearch.Num(); ++WindowIndex)
+	{
+		TSharedRef<SWindow> SomeWindow = WindowsToSearch[WindowIndex];
+		TSharedRef<FGenericWindow> SomeNativeWindow = StaticCastSharedRef<FGenericWindow>( SomeWindow->GetNativeWindow().ToSharedRef() );
+		if ( SomeNativeWindow == PlatformWindow )
+		{
+			return SomeWindow;
+		}
+		else
+		{
+			// Search child windows
+			TSharedPtr<SWindow> FoundChildWindow = FindWindowByPlatformWindow( SomeWindow->GetChildWindows(), PlatformWindow );
+			if (FoundChildWindow.IsValid())
+			{
+				return FoundChildWindow;
+			}
+		}
+	}
+
+	return TSharedPtr<SWindow>( NULL );
+}
+
+
+/* FSlateApplicationBase interface
+ *****************************************************************************/
+
+FVector2D FSlateApplication::GetCursorSize( ) const
+{
+	if ( PlatformApplication->Cursor.IsValid() )
+	{
+		int32 X;
+		int32 Y;
+		PlatformApplication->Cursor->GetSize( X, Y );
+		return FVector2D( X, Y );
+	}
+
+	return FVector2D( 1.0f, 1.0f );
+}
+
+
+TSharedPtr< SWidget > FSlateApplication::GetKeyboardFocusedWidget() const
+{
+	if( FocusedWidgetPath.IsValid() )
+	{
+		return FocusedWidgetPath.GetLastWidget().Pin();
+	}
+
+	return TSharedPtr< SWidget >();
+}
+
+
+TSharedPtr< SWidget > FSlateApplication::GetMouseCaptor() const
+{
+	return MouseCaptor.ToSharedWidget();
+}
+
+
+bool FSlateApplication::HasFocusedDescendants( const TSharedRef< const SWidget >& Widget ) const
+{
+	return FocusedWidgetPath.IsValid() && FocusedWidgetPath.GetLastWidget().Pin() != Widget && FocusedWidgetPath.ContainsWidget( Widget );
+}
+
+
+TSharedRef<SWidget> FSlateApplication::MakeImage( const TAttribute<const FSlateBrush*>& Image, const TAttribute<FSlateColor>& Color, const TAttribute<EVisibility>& Visibility ) const
+{
+	return SNew(SImage)
+		.ColorAndOpacity(Color)
+		.Image(Image)
+		.Visibility(Visibility);
+}
+
+
+TSharedRef<SWidget> FSlateApplication::MakeWindowTitleBar( const TSharedRef<SWindow>& Window, const TSharedPtr<SWidget>& CenterContent, TSharedPtr<IWindowTitleBar>& OutTitleBar ) const
+{
+	TSharedRef<SWindowTitleBar> TitleBar = SNew(SWindowTitleBar, Window, CenterContent)
+		.Visibility(EVisibility::SelfHitTestInvisible);
+
+	OutTitleBar = TitleBar;
+
+	return TitleBar;
+}
+
+
+TSharedRef<IToolTip> FSlateApplication::MakeToolTip( const TAttribute<FString>& ToolTipString )
+{
+	return SNew(SToolTip)
+		.Text(ToolTipString);
+}
+
+
+TSharedRef<IToolTip> FSlateApplication::MakeToolTip( const FText& ToolTipText )
+{
+	return SNew(SToolTip)
+		.Text(ToolTipText);
+}
+
+
+/* FGenericApplicationMessageHandler interface
+ *****************************************************************************/
+
+bool FSlateApplication::ShouldProcessUserInputMessages( const TSharedPtr< FGenericWindow >& PlatformWindow ) const
+{
+	TSharedPtr< SWindow > Window;
+	if ( PlatformWindow.IsValid() )
+	{
+		Window = FindWindowByPlatformWindow( SlateWindows, PlatformWindow.ToSharedRef() );
+	}
+
+	if ( ActiveModalWindows.Num() == 0 || 
+		( Window.IsValid() &&
+			( Window->IsDescendantOf( GetActiveModalWindow() ) || ActiveModalWindows.Contains( Window ) ) ) )
+	{
+		return true;
+	}
+	return false;
+}
+
+bool FSlateApplication::OnKeyChar( const TCHAR Character, const bool IsRepeat )
+{
+	FCharacterEvent CharacterEvent( Character, PlatformApplication->GetModifierKeys(), IsRepeat );
+	return ProcessKeyCharMessage( CharacterEvent );
+}
+
+bool FSlateApplication::ProcessKeyCharMessage( FCharacterEvent& InCharacterEvent )
+{
+	FReply Reply = FReply::Unhandled();
+
+	const int32 EventCount = InCharacterEvent.IsRepeat() ? SlateApplicationDefs::NumRepeatsPerActualRepeat : 1;
+	for( int32 CurEventIndex = 0; CurEventIndex < EventCount; ++CurEventIndex )
+	{
+		// NOTE: We intentionally don't reset LastUserInteractionTimeForThrottling here so that the UI can be responsive while typing
+
+		// Bubble the keyboard event
+		FWidgetPath EventPath = FocusedWidgetPath.ToWidgetPath();
+		InCharacterEvent.SetEventPath(EventPath);
+	
+		// Switch worlds for widgets in the current path
+		FScopedSwitchWorldHack SwitchWorld( EventPath );
+		TSharedPtr<SWidget> WidgetToLog;
+
+		Reply = FReply::Unhandled();
+		// Send out mouse enter events.
+		InCharacterEvent.SetEventPath( EventPath );
+		for( int32 WidgetIndex = EventPath.Widgets.Num()-1; !Reply.IsEventHandled() && WidgetIndex >= 0; --WidgetIndex )
+		{
+			FArrangedWidget& SomeWidgetGettingEvent = EventPath.Widgets( WidgetIndex );
+			{
+				// Widget newly under cursor, so send a MouseEnter.
+				Reply = SomeWidgetGettingEvent.Widget->OnKeyChar( SomeWidgetGettingEvent.Geometry, InCharacterEvent ).SetHandler( SomeWidgetGettingEvent.Widget );				
+				ProcessReply(EventPath, Reply, NULL, NULL);
+
+				WidgetToLog = SomeWidgetGettingEvent.Widget;
+			}
+		}
+
+		LOG_EVENT_CONTENT( EEventLog::KeyChar, FString::Printf(TEXT("%c"), InCharacterEvent.GetCharacter()), Reply );
+	}
+
+	return Reply.IsEventHandled();
+}
+
+bool FSlateApplication::OnKeyDown( const int32 KeyCode, const uint32 CharacterCode, const bool IsRepeat ) 
+{
+	FKey const Key = FInputKeyManager::Get().GetKeyFromCodes( KeyCode, CharacterCode );
+	FKeyboardEvent KeyboardEvent( Key, PlatformApplication->GetModifierKeys(), IsRepeat, CharacterCode );
+
+	return ProcessKeyDownMessage( KeyboardEvent );
+}
+
+bool FSlateApplication::ProcessKeyDownMessage( FKeyboardEvent& InKeyboardEvent )
+{
+	FReply Reply = FReply::Unhandled();
+
+	if (IsDragDropping() && InKeyboardEvent.GetKey() == EKeys::Escape)
+	{
+		// Pressing ESC while drag and dropping terminates the drag drop.
+		DragDropContent.Reset();
+		Reply = FReply::Handled();
+	}
+	else
+	{
+		const int32 EventCount = InKeyboardEvent.IsRepeat() ? SlateApplicationDefs::NumRepeatsPerActualRepeat : 1;
+		for( int32 CurEventIndex = 0; CurEventIndex < EventCount; ++CurEventIndex )
+		{
+			LastUserInteractionTimeForThrottling = this->GetCurrentTime();
+
+			// If we are inspecting, pressing ESC exits inspection mode.
+			if ( InKeyboardEvent.GetKey() == EKeys::Escape )
+			{
+				TSharedPtr<IWidgetReflector> WidgetReflector = WidgetReflectorPtr.Pin();
+				const bool bIsWidgetReflectorPicking = WidgetReflector.IsValid() && WidgetReflector->IsInPickingMode();
+				if ( bIsWidgetReflectorPicking )
+				{
+					if ( WidgetReflector.IsValid() )
+					{
+						WidgetReflector->OnWidgetPicked();
+					}
+				}
+			}
+
+			// Bubble the keyboard event
+			FWidgetPath EventPath = FocusedWidgetPath.ToWidgetPath();
+			InKeyboardEvent.SetEventPath(EventPath);
+
+			// Switch worlds for widgets in the current path
+			FScopedSwitchWorldHack SwitchWorld( EventPath );
+
+			TSharedPtr<SWidget> WidgetToLog;
+
+			Reply = FReply::Unhandled();
+
+			// Tunnel the keyboard event
+			for( int32 WidgetIndex = 0; !Reply.IsEventHandled() && WidgetIndex < EventPath.Widgets.Num(); ++WidgetIndex )
+			{
+				FArrangedWidget& SomeWidgetGettingEvent = EventPath.Widgets( WidgetIndex );
+				{
+					Reply = SomeWidgetGettingEvent.Widget->OnPreviewKeyDown( SomeWidgetGettingEvent.Geometry, InKeyboardEvent ).SetHandler(SomeWidgetGettingEvent.Widget);
+					ProcessReply(EventPath, Reply, NULL, NULL);
+
+					WidgetToLog = SomeWidgetGettingEvent.Widget;
+				}
+			}
+
+			// Send out key down events.
+			for( int32 WidgetIndex = EventPath.Widgets.Num()-1; !Reply.IsEventHandled() && WidgetIndex >= 0; --WidgetIndex )
+			{
+				FArrangedWidget& SomeWidgetGettingEvent = EventPath.Widgets( WidgetIndex );
+				{
+					Reply = SomeWidgetGettingEvent.Widget->OnKeyDown( SomeWidgetGettingEvent.Geometry, InKeyboardEvent ).SetHandler(SomeWidgetGettingEvent.Widget);
+					ProcessReply(EventPath, Reply, NULL, NULL);
+
+					WidgetToLog = SomeWidgetGettingEvent.Widget;
+				}
+			}
+
+			LOG_EVENT_CONTENT( EEventLog::KeyDown, GetKeyName(InKeyboardEvent.GetKey()).ToString(), Reply );
+
+			// If the keyboard event was not processed by any widget...
+			if ( !Reply.IsEventHandled() )
+			{
+				// If the key was Tab, interpret as an attempt to move focus.
+				if ( InKeyboardEvent.GetKey() == EKeys::Tab )
+				{
+					if (FocusedWidgetPath.IsValid())
+					{
+						EFocusMoveDirection::Type MoveDirection = ( InKeyboardEvent.IsShiftDown() )
+							? EFocusMoveDirection::Previous
+							: EFocusMoveDirection::Next;
+						SetKeyboardFocus( FocusedWidgetPath.ToNextFocusedPath( MoveDirection ), EKeyboardFocusCause::Keyboard );
+					}
+				}
+				else if ( UnhandledKeyDownEventHandler.IsBound() )
+				{
+					// Nothing else handled this event, give external code a chance to handle it.
+					Reply = UnhandledKeyDownEventHandler.Execute( InKeyboardEvent );
+				}
+			}
+		}
+	}
+
+	return Reply.IsEventHandled();
+}
+
+bool FSlateApplication::OnKeyUp( const int32 KeyCode, const uint32 CharacterCode, const bool IsRepeat )
+{
+	FKey const Key = FInputKeyManager::Get().GetKeyFromCodes( KeyCode, CharacterCode );
+	FKeyboardEvent KeyboardEvent( Key, PlatformApplication->GetModifierKeys(), IsRepeat, CharacterCode );
+
+	return ProcessKeyUpMessage( KeyboardEvent );
+}
+
+bool FSlateApplication::ProcessKeyUpMessage( FKeyboardEvent& InKeyboardEvent )
+{
+	FReply Reply = FReply::Unhandled();
+
+	const int32 EventCount = InKeyboardEvent.IsRepeat() ? SlateApplicationDefs::NumRepeatsPerActualRepeat : 1;
+	for( int32 CurEventIndex = 0; CurEventIndex < EventCount; ++CurEventIndex )
+	{
+		LastUserInteractionTimeForThrottling = this->GetCurrentTime();
+
+		// Bubble the keyboard event
+		FWidgetPath EventPath = FocusedWidgetPath.ToWidgetPath();
+		InKeyboardEvent.SetEventPath(EventPath);
+
+		// Switch worlds for widgets in the current path
+		FScopedSwitchWorldHack SwitchWorld( EventPath );
+
+		TSharedPtr<SWidget> WidgetToLog;
+
+		Reply = FReply::Unhandled();
+		// Send out mouse enter events.
+		for( int32 WidgetIndex = EventPath.Widgets.Num()-1; !Reply.IsEventHandled() && WidgetIndex >= 0; --WidgetIndex )
+		{
+			FArrangedWidget& SomeWidgetGettingEvent = EventPath.Widgets( WidgetIndex );
+			{
+				// Widget newly under cursor, so send a MouseEnter.
+				Reply = SomeWidgetGettingEvent.Widget->OnKeyUp( SomeWidgetGettingEvent.Geometry, InKeyboardEvent ).SetHandler( SomeWidgetGettingEvent.Widget );				
+				ProcessReply(EventPath, Reply, NULL, NULL);
+
+				WidgetToLog = SomeWidgetGettingEvent.Widget;
+			}
+		}
+
+		LOG_EVENT_CONTENT( EEventLog::KeyUp, GetKeyName(InKeyboardEvent.GetKey()).ToString(), Reply );
+	}
+
+	return Reply.IsEventHandled();
+}
+
+FKey TranslateMouseButtonToKey( const EMouseButtons::Type Button )
+{
+	FKey Key = EKeys::Invalid;
+
+	switch( Button )
+	{
+	case EMouseButtons::Left:
+		Key = EKeys::LeftMouseButton;
+		break;
+	case EMouseButtons::Middle:
+		Key = EKeys::MiddleMouseButton;
+		break;
+	case EMouseButtons::Right:
+		Key = EKeys::RightMouseButton;
+		break;
+	case EMouseButtons::Thumb01:
+		Key = EKeys::ThumbMouseButton;
+		break;
+	case EMouseButtons::Thumb02:
+		Key = EKeys::ThumbMouseButton2;
+		break;
+	}
+
+	return Key;
+}
+
+void FSlateApplication::SetGameIsFakingTouchEvents(const bool bIsFaking)
+{
+	if (bIsFakingTouched && !bIsFaking && bIsGameFakingTouch && !bIsFakingTouch)
+	{
+		OnTouchEnded(PlatformApplication->Cursor->GetPosition(), 0, 0);
+	}
+	bIsGameFakingTouch = bIsFaking;
+}
+
+bool FSlateApplication::IsFakingTouchEvents() const
+{
+	return bIsFakingTouch || bIsGameFakingTouch;
+}
+
+
+bool FSlateApplication::OnMouseDown( const TSharedPtr< FGenericWindow >& PlatformWindow, const EMouseButtons::Type Button )
+{
+	// convert to touch event if we are faking it	
+	if (bIsFakingTouch || bIsGameFakingTouch)
+	{
+		bIsFakingTouched = true;
+		return OnTouchStarted( PlatformWindow, PlatformApplication->Cursor->GetPosition(), 0, 0 );
+	}
+
+	FKey Key = TranslateMouseButtonToKey( Button );
+
+	FPointerEvent MouseEvent(
+		PlatformApplication->Cursor->GetPosition(), 	
+		LastCursorPosition,
+		PressedMouseButtons,
+		Key,
+		0,
+		PlatformApplication->GetModifierKeys()
+		);
+
+	return ProcessMouseButtonDownMessage( PlatformWindow, MouseEvent );
+}
+
+bool FSlateApplication::ProcessMouseButtonDownMessage( const TSharedPtr< FGenericWindow >& PlatformWindow, FPointerEvent& MouseEvent )
+{
+	LastUserInteractionTimeForThrottling = this->GetCurrentTime();
+
+	PlatformApplication->SetCapture( PlatformWindow );
+	PressedMouseButtons.Add( MouseEvent.GetEffectingButton() );
+
+	bool bInGame = false;
+
+	// Only process mouse down messages if we are not drag/dropping
+	if ( !IsDragDropping() )
+	{
+		FReply Reply = FReply::Unhandled();
+		if (MouseCaptor.IsValid())
+		{
+			FWidgetPath MouseCaptorPath = MouseCaptor.ToWidgetPath();
+			FArrangedWidget& MouseCaptorWidget = MouseCaptorPath.Widgets.Last();
+
+			// Switch worlds widgets in the current path
+			FScopedSwitchWorldHack SwitchWorld( MouseCaptorPath );
+			bInGame = FApp::IsGame();
+
+			MouseEvent.SetEventPath( MouseCaptorPath );
+			
+			Reply = MouseCaptorWidget.Widget->OnPreviewMouseButtonDown( MouseCaptorWidget.Geometry, MouseEvent ).SetHandler(MouseCaptorWidget.Widget);
+			ProcessReply(MouseCaptorPath, Reply, &MouseCaptorPath, &MouseEvent);
+			
+			if (!Reply.IsEventHandled())
+			{
+				if (MouseEvent.IsTouchEvent())
+				{
+					Reply = MouseCaptorWidget.Widget->OnTouchStarted( MouseCaptorWidget.Geometry, MouseEvent ).SetHandler( MouseCaptorWidget.Widget );
+				}
+				if (!MouseEvent.IsTouchEvent() || (!Reply.IsEventHandled() && bTouchFallbackToMouse))
+				{
+					Reply = MouseCaptorWidget.Widget->OnMouseButtonDown( MouseCaptorWidget.Geometry, MouseEvent ).SetHandler( MouseCaptorWidget.Widget );
+				}
+				ProcessReply(MouseCaptorPath, Reply, &MouseCaptorPath, &MouseEvent);
+			}
+			LOG_EVENT( EEventLog::MouseButtonDown, Reply );
+		}
+		else
+		{
+			FWidgetPath WidgetsUnderCursor = LocateWindowUnderMouse( MouseEvent.GetScreenSpacePosition(), GetInteractiveTopLevelWindows() );
+			MouseEvent.SetEventPath(WidgetsUnderCursor);
+
+			PopupSupport.SendNotifications( WidgetsUnderCursor );
+
+			// Switch worlds widgets in the current path
+			FScopedSwitchWorldHack SwitchWorld( WidgetsUnderCursor );
+			bInGame = FApp::IsGame();
+
+			TSharedPtr<SWidget> WidgetToLog;
+
+			const TSharedPtr<SWidget> PreviouslyFocusedWidget = GetKeyboardFocusedWidget();
+
+			for( int32 WidgetIndex = 0; !Reply.IsEventHandled() && WidgetIndex < WidgetsUnderCursor.Widgets.Num() ; ++WidgetIndex )
+			{
+				FArrangedWidget& CurWidget = WidgetsUnderCursor.Widgets(WidgetIndex);
+
+				Reply = CurWidget.Widget->OnPreviewMouseButtonDown(CurWidget.Geometry, MouseEvent).SetHandler( CurWidget.Widget );
+				ProcessReply( WidgetsUnderCursor, Reply, &WidgetsUnderCursor, &MouseEvent);
+			}
+
+			for( int32 WidgetIndex = WidgetsUnderCursor.Widgets.Num()-1; !Reply.IsEventHandled() && WidgetIndex >= 0 ; --WidgetIndex )
+			{
+				FArrangedWidget& CurWidget = WidgetsUnderCursor.Widgets(WidgetIndex);				
+
+				if (!Reply.IsEventHandled())
+				{
+					if (MouseEvent.IsTouchEvent())
+					{
+						Reply = CurWidget.Widget->OnTouchStarted( CurWidget.Geometry, MouseEvent ).SetHandler( CurWidget.Widget );
+					}
+					if (!MouseEvent.IsTouchEvent() || (!Reply.IsEventHandled() && bTouchFallbackToMouse))
+					{
+						Reply = CurWidget.Widget->OnMouseButtonDown( CurWidget.Geometry, MouseEvent ).SetHandler( CurWidget.Widget );
+					}
+					ProcessReply( WidgetsUnderCursor, Reply, &WidgetsUnderCursor, &MouseEvent);
+				}
+
+				WidgetToLog = CurWidget.Widget;
+			}
+			LOG_EVENT( EEventLog::MouseButtonDown, Reply );
+
+			// If none of the widgets requested keyboard focus to be set (or set the keyboard focus explicitly), set it to the leaf-most widget under the mouse.
+			const bool bFocusChangedByEventHandler = PreviouslyFocusedWidget != GetKeyboardFocusedWidget();
+			if (!Reply.GetFocusRecepient().IsValid() && !bFocusChangedByEventHandler)
+			{
+				// The event handler for OnMouseButton down may have altered the widget hierarchy.
+				// Refresh the previously-cached widget path.
+				WidgetsUnderCursor = LocateWindowUnderMouse( MouseEvent.GetScreenSpacePosition(), GetInteractiveTopLevelWindows() );
+
+				bool bFocusCandidateFound = false;
+				for( int32 WidgetIndex = WidgetsUnderCursor.Widgets.Num()-1; !bFocusCandidateFound && WidgetIndex >= 0 ; --WidgetIndex )
+				{
+					FArrangedWidget& CurWidget = WidgetsUnderCursor.Widgets(WidgetIndex);
+					if (CurWidget.Widget->SupportsKeyboardFocus())
+					{
+						bFocusCandidateFound = true;
+						FWidgetPath NewFocusedWidgetPath = WidgetsUnderCursor.GetPathDownTo( CurWidget.Widget );
+						SetKeyboardFocus( NewFocusedWidgetPath, EKeyboardFocusCause::Mouse );
+					}
+				}
+			}
+		}
+
+		// See if expensive tasks should be throttled.  By default on mouse down expensive tasks are throttled
+		// to ensure Slate responsiveness in low FPS situations
+		if (Reply.IsEventHandled() && !bInGame && !MouseEvent.IsTouchEvent())
+		{
+			// Enter responsive mode if throttling should occur and its not already happening
+			if( Reply.ShouldThrottle() && !MouseButtonDownResponsivnessThrottle.IsValid() )
+			{
+				MouseButtonDownResponsivnessThrottle = FSlateThrottleManager::Get().EnterResponsiveMode();
+			}
+			else if( !Reply.ShouldThrottle() && MouseButtonDownResponsivnessThrottle.IsValid() )
+			{
+				// Leave responsive mode if a widget chose not to throttle
+				FSlateThrottleManager::Get().LeaveResponsiveMode( MouseButtonDownResponsivnessThrottle );
+			}
+		}
+	}
+
+	LastCursorPosition = MouseEvent.GetScreenSpacePosition();
+
+	return true;
+}
+
+bool FSlateApplication::OnMouseDoubleClick( const TSharedPtr< FGenericWindow >& PlatformWindow, const EMouseButtons::Type Button )
+{
+	if (bIsFakingTouch || bIsGameFakingTouch)
+	{
+		bIsFakingTouched = true;
+		return OnTouchStarted(PlatformWindow, PlatformApplication->Cursor->GetPosition(), 0, 0);
+	}
+
+	FKey Key = TranslateMouseButtonToKey( Button );
+
+	FPointerEvent MouseEvent(
+		GetCursorPos(), 	
+		LastCursorPosition,
+		PressedMouseButtons,
+		Key,
+		0,
+		PlatformApplication->GetModifierKeys()
+		);
+
+	return ProcessMouseButtonDoubleClickMessage( PlatformWindow, MouseEvent );
+}
+
+bool FSlateApplication::ProcessMouseButtonDoubleClickMessage( const TSharedPtr< FGenericWindow >& PlatformWindow, FPointerEvent& InMouseEvent )
+{
+	LastUserInteractionTimeForThrottling = this->GetCurrentTime();
+
+	PlatformApplication->SetCapture( PlatformWindow );
+	PressedMouseButtons.Add( InMouseEvent.GetEffectingButton() );
+
+	FWidgetPath WidgetsUnderCursor = LocateWindowUnderMouse( InMouseEvent.GetScreenSpacePosition(), GetInteractiveTopLevelWindows() );
+	InMouseEvent.SetEventPath(WidgetsUnderCursor);
+
+	// Switch worlds widgets in the current path
+	FScopedSwitchWorldHack SwitchWorld( WidgetsUnderCursor );
+
+	TSharedPtr<SWidget> WidgetToLog;
+
+	FReply Reply = FReply::Unhandled();
+	for( int32 WidgetIndex = WidgetsUnderCursor.Widgets.Num()-1; !Reply.IsEventHandled() && WidgetIndex >= 0; --WidgetIndex )
+	{
+		FArrangedWidget& CurWidget = WidgetsUnderCursor.Widgets(WidgetIndex);
+		Reply = CurWidget.Widget->OnMouseButtonDoubleClick( CurWidget.Geometry, InMouseEvent ).SetHandler( CurWidget.Widget );
+		ProcessReply(WidgetsUnderCursor, Reply, &WidgetsUnderCursor, &InMouseEvent);
+
+		WidgetToLog = CurWidget.Widget;
+	}
+
+	LOG_EVENT( EEventLog::MouseButtonDoubleClick, Reply );
+
+	LastCursorPosition = InMouseEvent.GetScreenSpacePosition();
+
+	return Reply.IsEventHandled();
+}
+
+bool FSlateApplication::OnMouseUp( const EMouseButtons::Type Button )
+{
+	// convert to touch event if we are faking it	
+	if (bIsFakingTouch || bIsGameFakingTouch)
+	{
+		bIsFakingTouched = false;
+		return OnTouchEnded(PlatformApplication->Cursor->GetPosition(), 0, 0);
+	}
+
+	FKey Key = TranslateMouseButtonToKey( Button );
+
+	FPointerEvent MouseEvent(
+		GetCursorPos(), 	
+		LastCursorPosition,
+		PressedMouseButtons,
+		Key,
+		0,
+		PlatformApplication->GetModifierKeys()
+		);
+
+	return ProcessMouseButtonUpMessage( MouseEvent );
+}
+
+bool FSlateApplication::ProcessMouseButtonUpMessage( FPointerEvent& MouseEvent )
+{
+	LastUserInteractionTimeForThrottling = this->GetCurrentTime();
+	PressedMouseButtons.Remove( MouseEvent.GetEffectingButton() );
+
+	if (DragDetector.DetectDragForWidget.IsValid() && MouseEvent.GetEffectingButton() == DragDetector.DetectDragButton )
+	{
+		// The user has release the button that was supposed to start the drag; stop detecting it.
+		DragDetector = FDragDetector();
+	}
+
+	if (MouseCaptor.IsValid())
+	{
+		FWidgetPath MouseCaptorPath = MouseCaptor.ToWidgetPath();
+		if ( ensureMsg(MouseCaptorPath.Widgets.Num() > 0, TEXT("A window had a widget with mouse capture. That entire window has been dismissed before the mouse up could be processed.")) )
+		{
+			FArrangedWidget& MouseCaptorWidget = MouseCaptorPath.Widgets.Last();
+			MouseEvent.SetEventPath(MouseCaptorPath);
+
+			// Switch worlds widgets in the current path
+			FScopedSwitchWorldHack SwitchWorld( MouseCaptorPath );
+
+			FReply Reply = FReply::Unhandled();
+			if (MouseEvent.IsTouchEvent())
+			{
+				Reply = MouseCaptorWidget.Widget->OnTouchEnded( MouseCaptorWidget.Geometry, MouseEvent ).SetHandler( MouseCaptorWidget.Widget );
+			}
+			if (!MouseEvent.IsTouchEvent() || (!Reply.IsEventHandled() && bTouchFallbackToMouse))
+			{
+				Reply = MouseCaptorWidget.Widget->OnMouseButtonUp( MouseCaptorWidget.Geometry, MouseEvent ).SetHandler( MouseCaptorWidget.Widget );
+			}
+			ProcessReply(MouseCaptorPath, Reply, &MouseCaptorPath, &MouseEvent);
+			LOG_EVENT( EEventLog::MouseButtonUp, Reply );
+		}
+	}
+	else
+	{
+		FWidgetPath WidgetsUnderCursor = LocateWindowUnderMouse( MouseEvent.GetScreenSpacePosition(), GetInteractiveTopLevelWindows() );
+		MouseEvent.SetEventPath(WidgetsUnderCursor);
+
+		// If we are doing a drag and drop, we will send this event instead.
+		FDragDropEvent DragDropEvent( MouseEvent, DragDropContent );
+
+		// Cache the drag drop content and reset the pointer in case OnMouseButtonUpMessage re-enters as a result of OnDrop
+		const bool bIsDragDropping = IsDragDropping();
+		TSharedPtr< FDragDropOperation > LocalDragDropContent = DragDropContent;
+		DragDropContent.Reset();
+
+		// Switch worlds widgets in the current path
+		FScopedSwitchWorldHack SwitchWorld( WidgetsUnderCursor );
+
+		FReply Reply = FReply::Unhandled();
+		for( int32 WidgetIndex = WidgetsUnderCursor.Widgets.Num()-1; !Reply.IsEventHandled() && WidgetIndex >= 0; --WidgetIndex )
+		{
+			FArrangedWidget& CurWidget = WidgetsUnderCursor.Widgets(WidgetIndex);
+			if (MouseEvent.IsTouchEvent())
+			{
+				Reply = CurWidget.Widget->OnTouchEnded( CurWidget.Geometry, MouseEvent ).SetHandler( CurWidget.Widget );
+			}
+			if (!MouseEvent.IsTouchEvent() || (!Reply.IsEventHandled() && bTouchFallbackToMouse))
+			{
+				Reply = (bIsDragDropping)
+					? CurWidget.Widget->OnDrop( CurWidget.Geometry, DragDropEvent ).SetHandler( CurWidget.Widget )
+					: CurWidget.Widget->OnMouseButtonUp( CurWidget.Geometry, MouseEvent ).SetHandler( CurWidget.Widget );
+			}
+
+			ProcessReply(WidgetsUnderCursor, Reply, &WidgetsUnderCursor, &MouseEvent);
+		}
+
+		LOG_EVENT( bIsDragDropping ? EEventLog::DragDrop : EEventLog::MouseButtonUp, Reply );
+
+		// If we were dragging, notify the content
+		if ( bIsDragDropping )
+		{
+			LocalDragDropContent->OnDrop( Reply.IsEventHandled(), MouseEvent );
+		}
+	}
+
+	// If in responsive mode throttle, leave it on mouse up.
+	if( MouseButtonDownResponsivnessThrottle.IsValid() )
+	{
+		FSlateThrottleManager::Get().LeaveResponsiveMode( MouseButtonDownResponsivnessThrottle );
+	}
+
+	if ( PressedMouseButtons.Num() == 0 )
+	{
+		// Release Capture
+		PlatformApplication->SetCapture( NULL );
+	}
+
+	return true;
+}
+
+bool FSlateApplication::OnMouseWheel( const float Delta )
+{
+	const FVector2D CurrentCursorPosition = GetCursorPos();
+
+	FPointerEvent MouseWheelEvent(
+		CurrentCursorPosition, 	
+		CurrentCursorPosition,
+		PressedMouseButtons,
+		EKeys::MouseWheelSpin,
+		Delta,
+		PlatformApplication->GetModifierKeys()
+		);
+
+	return ProcessMouseWheelOrGestureMessage( MouseWheelEvent, NULL );
+}
+
+bool FSlateApplication::ProcessMouseWheelOrGestureMessage( FPointerEvent& InWheelEvent, const FPointerEvent* InGestureEvent )
+{
+	const bool bShouldProcessEvent = InWheelEvent.GetWheelDelta() != 0 ||
+		(InGestureEvent!=NULL && InGestureEvent->GetGestureDelta()!=FVector2D::ZeroVector);
+	
+	if ( !bShouldProcessEvent )
+	{
+		return false;
+	}
+
+	// NOTE: We intentionally don't reset LastUserInteractionTimeForThrottling here so that the UI can be responsive while scrolling
+
+	FWidgetPath EventPath = ( MouseCaptor.IsValid() )
+		? MouseCaptor.ToWidgetPath()
+		: LocateWindowUnderMouse( InWheelEvent.GetScreenSpacePosition(), GetInteractiveTopLevelWindows() );
+
+	InWheelEvent.SetEventPath( EventPath );
+
+	// Switch worlds widgets in the current path
+	FScopedSwitchWorldHack SwitchWorld( EventPath );
+
+	TSharedPtr<SWidget> WidgetToLog;
+
+	FReply Reply = FReply::Unhandled();
+	for( int32 WidgetIndex = EventPath.Widgets.Num()-1; !Reply.IsEventHandled() && WidgetIndex >= 0; --WidgetIndex )
+	{
+		FArrangedWidget& CurWidget = EventPath.Widgets(WidgetIndex);
+		// Gesture event gets first shot, if slate doesn't respond to it, we'll try the wheel event.
+		if( InGestureEvent!=NULL )
+		{
+			Reply = CurWidget.Widget->OnTouchGesture( CurWidget.Geometry, *InGestureEvent ).SetHandler( CurWidget.Widget );
+			ProcessReply(EventPath, Reply, &EventPath, InGestureEvent);
+		}
+		
+		// Send the mouse wheel event if we haven't already handled the gesture version of this event.
+		if( !Reply.IsEventHandled() )
+		{
+			Reply = CurWidget.Widget->OnMouseWheel( CurWidget.Geometry, InWheelEvent ).SetHandler( CurWidget.Widget );
+			ProcessReply(EventPath, Reply, &EventPath, &InWheelEvent);
+		}
+
+		WidgetToLog = CurWidget.Widget;
+	}
+
+	LOG_EVENT( bIsGestureEvent ? EEventLog::TouchGesture : EEventLog::MouseWheel, Reply );
+
+	return Reply.IsEventHandled();
+}
+
+bool FSlateApplication::OnMouseMove()
+{
+	// convert to touch event if we are faking it	
+	if (bIsFakingTouched)
+	{
+		return OnTouchMoved(PlatformApplication->Cursor->GetPosition(), 0, 0);
+	}
+	else if (!bIsGameFakingTouch && bIsFakingTouch)
+	{
+		return false;
+	}
+
+	bool Result = true;
+	const FVector2D CurrentCursorPosition = GetCursorPos();
+
+	if ( LastCursorPosition != CurrentCursorPosition )
+	{
+		FPointerEvent MouseEvent(
+			CurrentCursorPosition,
+			LastCursorPosition,
+			PressedMouseButtons,
+			EKeys::Invalid,
+			0,
+			PlatformApplication->GetModifierKeys()
+			);
+
+		Result = ProcessMouseMoveMessage( MouseEvent );
+	}
+
+	return Result;
+}
+
+bool FSlateApplication::OnRawMouseMove( const int32 X, const int32 Y )
+{
+	if ( X != 0 || Y != 0 )
+	{
+		FPointerEvent MouseEvent(
+			GetCursorPos(), 	
+			LastCursorPosition,
+			FVector2D( X, Y ), 
+			PressedMouseButtons,
+			PlatformApplication->GetModifierKeys()
+		);
+
+		ProcessMouseMoveMessage(MouseEvent);
+	}
+	
+	return true;
+}
+
+bool FSlateApplication::ProcessMouseMoveMessage( FPointerEvent& MouseEvent )
+{
+	if ( !MouseEvent.GetCursorDelta().IsZero() )
+	{
+		// Detecting a mouse move of zero delta is our way of filtering out synthesized move events
+		const bool AllowSpawningOfToolTips = true;
+		UpdateToolTip( AllowSpawningOfToolTips );
+	}
+
+	FWidgetPath WidgetsUnderCursor = LocateWindowUnderMouse( MouseEvent.GetScreenSpacePosition(), GetInteractiveTopLevelWindows() );
+	bool bHandled = false;
+
+	FWeakWidgetPath LastWidgetsUnderCursor;
+
+	// User asked us to detect a drag.
+	bool bDragDetected = false;
+	if( DragDetector.DetectDragForWidget.IsValid() )
+	{	
+		const FVector2D DragDelta = (DragDetector.DetectDragStartLocation - MouseEvent.GetScreenSpacePosition());
+		bDragDetected = ( DragDelta.Size() > SlateDragStartDistance );
+		if (bDragDetected)
+		{
+			FWidgetPath DragDetectPath = DragDetector.DetectDragForWidget.ToWidgetPath();
+			if( DragDetectPath.IsValid() && DragDetector.DetectDragForWidget.GetLastWidget().IsValid() )
+			{
+				FArrangedWidget DetectDragForMe = DragDetectPath.FindArrangedWidget(DragDetector.DetectDragForWidget.GetLastWidget().Pin().ToSharedRef());
+
+				// A drag has been triggered. The cursor exited some widgets as a result.
+				// This assignment ensures that we will send OnLeave notifications to those widgets.
+				LastWidgetsUnderCursor = DragDetector.DetectDragForWidget;
+
+				// We're finished with the drag detect.
+				DragDetector = FDragDetector();
+
+				// Send an OnDragDetected to the widget that requested drag-detection.
+				MouseEvent.SetEventPath( DragDetectPath );
+
+				// Switch worlds widgets in the current path
+				FScopedSwitchWorldHack SwitchWorld( DragDetectPath );
+
+				FReply Reply = DetectDragForMe.Widget->OnDragDetected(DetectDragForMe.Geometry, MouseEvent).SetHandler(DetectDragForMe.Widget);
+				ProcessReply( DragDetectPath, Reply, &DragDetectPath, &MouseEvent );
+				LOG_EVENT( EEventLog::DragDetected, Reply );
+			}
+			else
+			{
+				bDragDetected = false;
+			}
+		}		
+	}
+
+
+	if (bDragDetected)
+	{
+		// When a drag was detected, we pretend that the widgets under the mouse last time around.
+		// We have set LastWidgetsUnderCursor accordingly when the drag was detected above.
+	}
+	else
+	{
+		// No Drag Detection
+		LastWidgetsUnderCursor = WidgetsUnderCursorLastEvent;
+	}
+
+	// In the case of drag leave/enter events there is no path to speak of
+	MouseEvent.SetEventPath( FWidgetPath() );
+
+	// Send out mouse leave events
+	// If we are doing a drag and drop, we will send this event instead.
+	{
+		FDragDropEvent DragDropEvent( MouseEvent, DragDropContent );
+		// Switch worlds widgets in the current path
+		FScopedSwitchWorldHack SwitchWorld( LastWidgetsUnderCursor.Window.Pin() );
+
+		for ( int32 WidgetIndex = LastWidgetsUnderCursor.Widgets.Num()-1; WidgetIndex >=0; --WidgetIndex )
+		{
+			// Guards for cases where WidgetIndex can become invalid due to MouseMove being re-entrant.
+			while ( WidgetIndex >= LastWidgetsUnderCursor.Widgets.Num() )
+			{
+				WidgetIndex--;
+			}
+
+			if ( WidgetIndex >= 0 )
+			{
+				const TSharedPtr<SWidget>& SomeWidgetPreviouslyUnderCursor = LastWidgetsUnderCursor.Widgets[WidgetIndex].Pin();
+				if( SomeWidgetPreviouslyUnderCursor.IsValid() )
+				{
+					if ( !WidgetsUnderCursor.ContainsWidget(SomeWidgetPreviouslyUnderCursor.ToSharedRef()) )
+					{
+						// Widget is no longer under cursor, so send a MouseLeave.
+						if ( IsDragDropping() )
+						{
+							SomeWidgetPreviouslyUnderCursor->OnDragLeave( DragDropEvent );
+							LOG_EVENT( EEventLog::DragLeave, SomeWidgetPreviouslyUnderCursor );
+
+							// Reset the cursor override
+							DragDropEvent.GetOperation()->SetCursorOverride( TOptional<EMouseCursor::Type>() );
+						}
+						else
+						{
+							SomeWidgetPreviouslyUnderCursor->OnMouseLeave( MouseEvent );
+							LOG_EVENT( EEventLog::MouseLeave, SomeWidgetPreviouslyUnderCursor );
+						}			
+					}
+				}
+			}
+		}
+	}
+
+
+	FWidgetPath MouseCaptorPath;
+	if (MouseCaptor.IsValid())
+	{
+		MouseCaptorPath = MouseCaptor.ToWidgetPath(FWeakWidgetPath::EInterruptedPathHandling::ReturnInvalid);
+	}
+
+	if (MouseCaptorPath.IsValid())
+	{
+		FArrangedWidget& MouseCaptorWidget = MouseCaptorPath.Widgets.Last();
+		MouseEvent.SetEventPath(MouseCaptorPath);
+
+		// Switch worlds widgets in the current path
+		FScopedSwitchWorldHack SwitchWorld( MouseCaptorPath );
+
+		FReply Reply = FReply::Unhandled();
+		if (MouseEvent.IsTouchEvent())
+		{
+			Reply = MouseCaptorWidget.Widget->OnTouchMoved( MouseCaptorWidget.Geometry, MouseEvent ).SetHandler( MouseCaptorWidget.Widget );
+		}
+		if (!MouseEvent.IsTouchEvent() || (!Reply.IsEventHandled() && bTouchFallbackToMouse))
+		{
+			Reply = MouseCaptorWidget.Widget->OnMouseMove( MouseCaptorWidget.Geometry, MouseEvent ).SetHandler( MouseCaptorWidget.Widget );
+		}
+		ProcessReply(MouseCaptorPath, Reply, &MouseCaptorPath, &MouseEvent);
+		bHandled = Reply.IsEventHandled();
+	}
+	else
+	{	
+		FReply Reply = FReply::Unhandled();
+
+		MouseEvent.SetEventPath(WidgetsUnderCursor);
+		FDragDropEvent DragDropEvent( MouseEvent, DragDropContent );
+		// Switch worlds widgets in the current path
+		FScopedSwitchWorldHack SwitchWorld( WidgetsUnderCursor );
+
+		// Send out mouse enter events.
+		for( int32 WidgetIndex = WidgetsUnderCursor.Widgets.Num()-1; WidgetIndex >= 0; --WidgetIndex )
+		{
+			FArrangedWidget& SomeWidgetUnderCursor = WidgetsUnderCursor.Widgets( WidgetIndex );
+			if ( !LastWidgetsUnderCursor.ContainsWidget(SomeWidgetUnderCursor.Widget) )
+			{
+				// Widget newly under cursor, so send a MouseEnter.
+				if ( IsDragDropping() )
+				{
+					// Doing a drag and drop; send a DragDropEvent
+					SomeWidgetUnderCursor.Widget->OnDragEnter( SomeWidgetUnderCursor.Geometry, DragDropEvent );
+					LOG_EVENT( EEventLog::DragEnter, SomeWidgetUnderCursor.Widget );
+				}
+				else
+				{
+					// Not drag dropping; send regular mouse event
+					SomeWidgetUnderCursor.Widget->OnMouseEnter( SomeWidgetUnderCursor.Geometry, MouseEvent );
+					LOG_EVENT( EEventLog::MouseEnter, SomeWidgetUnderCursor.Widget );
+				}
+			}
+		}
+
+		// Bubble the MouseMove event.
+		for( int32 WidgetIndex = WidgetsUnderCursor.Widgets.Num()-1; !Reply.IsEventHandled() && WidgetIndex >= 0; --WidgetIndex )
+		{
+			FArrangedWidget& CurWidget = WidgetsUnderCursor.Widgets(WidgetIndex);
+
+			if (MouseEvent.IsTouchEvent())
+			{
+				Reply = CurWidget.Widget->OnTouchMoved( CurWidget.Geometry, MouseEvent ).SetHandler( CurWidget.Widget );
+			}
+			if (!Reply.IsEventHandled())
+			{
+				Reply = (IsDragDropping())
+					? CurWidget.Widget->OnDragOver( CurWidget.Geometry, DragDropEvent ).SetHandler( CurWidget.Widget )
+					: CurWidget.Widget->OnMouseMove( CurWidget.Geometry, MouseEvent ).SetHandler( CurWidget.Widget );
+			}
+
+			ProcessReply(WidgetsUnderCursor, Reply, &WidgetsUnderCursor, &MouseEvent);
+		}
+
+		LOG_EVENT( IsDragDropping() ? EEventLog::DragOver : EEventLog::MouseMove, Reply )
+
+			bHandled = Reply.IsEventHandled();
+	}
+
+	// Give the current drag drop operation a chance to do something
+	// custom (e.g. update the Drag/Drop preview based on content)
+	if (IsDragDropping())
+	{
+		FDragDropEvent DragDropEvent( MouseEvent, DragDropContent );
+		FScopedSwitchWorldHack SwitchWorld( WidgetsUnderCursor );
+		DragDropContent->OnDragged( DragDropEvent );
+
+		// check the drag-drop operation for a cursor switch (on Windows, the OS thinks the mouse is
+		// captured so we wont get QueryCursor calls for drag/drops internal to the Slate application)
+		FCursorReply CursorResult = DragDropContent->OnCursorQuery();
+		if (CursorResult.IsEventHandled())
+		{
+			// Query was handled, so we should set the cursor.
+			PlatformApplication->Cursor->SetType( CursorResult.GetCursor() );
+		}
+		else
+		{
+			// reset the cursor to default for drag-drops
+			PlatformApplication->Cursor->SetType( EMouseCursor::Default );
+		}
+	}
+
+	WidgetsUnderCursorLastEvent = FWeakWidgetPath( WidgetsUnderCursor );
+
+	LastCursorPosition = MouseEvent.GetScreenSpacePosition();
+
+	return bHandled;
+}
+
+bool FSlateApplication::OnCursorSet()
+{
+	QueryCursor();
+	return true;
+}
+
+FKey TranslateControllerButtonToKey( EControllerButtons::Type Button )
+{
+	FKey Key = EKeys::Invalid;
+
+	switch ( Button )
+	{
+	case EControllerButtons::LeftAnalogY: Key = EKeys::Gamepad_LeftY; break;
+	case EControllerButtons::LeftAnalogX: Key = EKeys::Gamepad_LeftX; break;
+
+	case EControllerButtons::RightAnalogY: Key = EKeys::Gamepad_RightY; break;
+	case EControllerButtons::RightAnalogX: Key = EKeys::Gamepad_RightX; break;
+
+	case EControllerButtons::LeftTriggerAnalog: Key = EKeys::Gamepad_LeftTriggerAxis; break;
+	case EControllerButtons::RightTriggerAnalog: Key = EKeys::Gamepad_RightTriggerAxis; break;
+
+	case EControllerButtons::FaceButtonBottom: Key = EKeys::Gamepad_FaceButton_Bottom; break;
+	case EControllerButtons::FaceButtonRight: Key = EKeys::Gamepad_FaceButton_Right; break;
+	case EControllerButtons::FaceButtonLeft: Key = EKeys::Gamepad_FaceButton_Left; break;
+	case EControllerButtons::FaceButtonTop: Key = EKeys::Gamepad_FaceButton_Top; break;
+
+	case EControllerButtons::LeftShoulder: Key = EKeys::Gamepad_LeftShoulder; break;
+	case EControllerButtons::RightShoulder: Key = EKeys::Gamepad_RightShoulder; break;
+	case EControllerButtons::SpecialLeft: Key = EKeys::Gamepad_Special_Left; break;
+	case EControllerButtons::SpecialRight: Key = EKeys::Gamepad_Special_Right; break;
+	case EControllerButtons::LeftThumb: Key = EKeys::Gamepad_LeftThumbstick; break;
+	case EControllerButtons::RightThumb: Key = EKeys::Gamepad_RightThumbstick; break;
+	case EControllerButtons::LeftTriggerThreshold: Key = EKeys::Gamepad_LeftTrigger; break;
+	case EControllerButtons::RightTriggerThreshold: Key = EKeys::Gamepad_RightTrigger; break;
+
+	case EControllerButtons::DPadUp: Key = EKeys::Gamepad_DPad_Up; break;
+	case EControllerButtons::DPadDown: Key = EKeys::Gamepad_DPad_Down; break;
+	case EControllerButtons::DPadLeft: Key = EKeys::Gamepad_DPad_Left; break;
+	case EControllerButtons::DPadRight: Key = EKeys::Gamepad_DPad_Right; break;
+
+	case EControllerButtons::LeftStickUp: Key = EKeys::Gamepad_LeftStick_Up; break;
+	case EControllerButtons::LeftStickDown: Key = EKeys::Gamepad_LeftStick_Down; break;
+	case EControllerButtons::LeftStickLeft: Key = EKeys::Gamepad_LeftStick_Left; break;
+	case EControllerButtons::LeftStickRight: Key = EKeys::Gamepad_LeftStick_Right; break;
+
+	case EControllerButtons::RightStickUp: Key = EKeys::Gamepad_RightStick_Up; break;
+	case EControllerButtons::RightStickDown: Key = EKeys::Gamepad_RightStick_Down; break;
+	case EControllerButtons::RightStickLeft: Key = EKeys::Gamepad_RightStick_Left; break;
+	case EControllerButtons::RightStickRight: Key = EKeys::Gamepad_RightStick_Right; break;
+
+	case EControllerButtons::GlobalMenu: Key = EKeys::Global_Menu; break;
+	case EControllerButtons::GlobalView: Key = EKeys::Global_View; break;
+	case EControllerButtons::GlobalPause: Key = EKeys::Global_Pause; break;
+	case EControllerButtons::GlobalPlay: Key = EKeys::Global_Play; break;
+	case EControllerButtons::GlobalBack: Key = EKeys::Global_Back; break;
+
+	case EControllerButtons::Invalid: Key = EKeys::Invalid; break;
+	}
+
+	return Key;
+}
+
+#define CALL_WIDGET_FUNCTION(Event, Function) \
+	if (Event.GetUserIndex() < ARRAY_COUNT(JoystickCaptorWeakPaths) && JoystickCaptorWeakPaths[Event.GetUserIndex()].IsValid()) \
+{ \
+	/* Get the joystick capture target for this user */ \
+	FWidgetPath PathToWidget = JoystickCaptorWeakPaths[Event.GetUserIndex()].ToWidgetPath(); \
+	FArrangedWidget& ArrangedWidget = PathToWidget.Widgets.Last(); \
+	\
+	/* Switch worlds for widgets in the current path */ \
+	FScopedSwitchWorldHack SwitchWorld( PathToWidget ); \
+	\
+	/* Send the message to the widget */ \
+	FReply Reply = ArrangedWidget.Widget->Function( ArrangedWidget.Geometry, Event ).SetHandler( ArrangedWidget.Widget ); \
+	ProcessReply( PathToWidget, Reply, NULL, NULL, Event.GetUserIndex() );\
+}
+
+bool FSlateApplication::OnControllerAnalog( FKey Button, int32 ControllerId, float AnalogValue )
+{
+	if ( GetJoystickCaptor( ControllerId ).IsValid() )
+	{
+		FControllerEvent ControllerEvent( 
+			Button, 
+			ControllerId,
+			AnalogValue,
+			false
+		);
+
+		OnControllerAnalogValueChangedMessage( ControllerEvent );
+	}
+
+	return true;
+}
+
+bool FSlateApplication::OnControllerAnalog( EControllerButtons::Type Button, int32 ControllerId, float AnalogValue )
+{
+	return OnControllerAnalog( TranslateControllerButtonToKey( Button ), ControllerId, AnalogValue );
+}
+
+void FSlateApplication::OnControllerAnalogValueChangedMessage( FControllerEvent& ControllerEvent )
+{
+	CALL_WIDGET_FUNCTION(ControllerEvent, OnControllerAnalogValueChanged);
+}
+
+bool FSlateApplication::OnControllerButtonPressed( FKey Button, int32 ControllerId, bool IsRepeat )
+{
+	if ( GetJoystickCaptor( ControllerId ).IsValid() )
+	{
+		FControllerEvent ControllerEvent( 
+			Button, 
+			ControllerId,
+			1.0f,
+			IsRepeat
+		);
+
+		OnControllerButtonPressedMessage( ControllerEvent );
+	}
+
+	return true;
+}
+
+bool FSlateApplication::OnControllerButtonPressed( EControllerButtons::Type Button, int32 ControllerId, bool IsRepeat )
+{
+	return OnControllerButtonPressed( TranslateControllerButtonToKey( Button ), ControllerId, IsRepeat );
+}
+
+void FSlateApplication::OnControllerButtonPressedMessage( FControllerEvent& ControllerEvent )
+{
+	CALL_WIDGET_FUNCTION(ControllerEvent, OnControllerButtonPressed);
+}
+
+bool FSlateApplication::OnControllerButtonReleased( FKey Button, int32 ControllerId, bool IsRepeat )
+{
+	if ( GetJoystickCaptor( ControllerId ).IsValid() )
+	{
+		FControllerEvent ControllerEvent( 
+			Button, 
+			ControllerId,
+			1.0f,
+			IsRepeat
+		);
+
+		OnControllerButtonReleasedMessage( ControllerEvent );
+	}
+
+	return false;
+}
+
+bool FSlateApplication::OnControllerButtonReleased( EControllerButtons::Type Button, int32 ControllerId, bool IsRepeat )
+{
+	return OnControllerButtonReleased( TranslateControllerButtonToKey( Button ), ControllerId, IsRepeat );
+}
+
+void FSlateApplication::OnControllerButtonReleasedMessage( FControllerEvent& ControllerEvent )
+{
+	CALL_WIDGET_FUNCTION(ControllerEvent, OnControllerButtonReleased);
+
+}
+
+bool FSlateApplication::OnTouchGesture( EGestureEvent::Type GestureType, const FVector2D &Delta, const float MouseWheelDelta )
+{
+	const FVector2D CurrentCursorPosition = GetCursorPos();
+	
+	FPointerEvent GestureEvent(
+		CurrentCursorPosition,
+		CurrentCursorPosition,
+		PressedMouseButtons,
+		PlatformApplication->GetModifierKeys(),
+		GestureType,
+		Delta
+	);
+	
+	FPointerEvent MouseWheelEvent(
+		CurrentCursorPosition,
+		CurrentCursorPosition,
+		PressedMouseButtons,
+		EKeys::MouseWheelSpin,
+		MouseWheelDelta,
+		PlatformApplication->GetModifierKeys()
+	);
+	
+	return ProcessMouseWheelOrGestureMessage( MouseWheelEvent, &GestureEvent );
+}
+
+bool FSlateApplication::OnTouchStarted( const TSharedPtr< FGenericWindow >& PlatformWindow, const FVector2D& Location, int32 TouchIndex, int32 ControllerId )
+{
+	FPointerEvent PointerEvent(
+		ControllerId,
+		TouchIndex,
+		Location,
+		Location,
+		true);
+	OnTouchStartedMessage( PlatformWindow, PointerEvent );
+
+	return true;
+}
+
+void FSlateApplication::OnTouchStartedMessage( const TSharedPtr< FGenericWindow >& PlatformWindow, FPointerEvent& PointerEvent )
+{
+	ProcessMouseButtonDownMessage( PlatformWindow, PointerEvent );
+}
+
+bool FSlateApplication::OnTouchMoved( const FVector2D& Location, int32 TouchIndex, int32 ControllerId )
+{
+	FPointerEvent PointerEvent(
+		ControllerId,
+		TouchIndex,
+		Location,
+		Location,
+		true);
+	OnTouchMovedMessage(PointerEvent);
+
+	return true;
+}
+
+void FSlateApplication::OnTouchMovedMessage( FPointerEvent& PointerEvent )
+{
+	ProcessMouseMoveMessage(PointerEvent);
+}
+
+bool FSlateApplication::OnTouchEnded( const FVector2D& Location, int32 TouchIndex, int32 ControllerId )
+{
+	FPointerEvent PointerEvent(
+		ControllerId,
+		TouchIndex,
+		Location,
+		Location,
+		true);
+	OnTouchEndedMessage(PointerEvent);
+
+	return true;
+}
+
+void FSlateApplication::OnTouchEndedMessage( FPointerEvent& PointerEvent )
+{
+	ProcessMouseButtonUpMessage(PointerEvent);
+}
+
+bool FSlateApplication::OnMotionDetected(const FVector& Tilt, const FVector& RotationRate, const FVector& Gravity, const FVector& Acceleration, int32 ControllerId)
+{
+	FMotionEvent MotionEvent( 
+		ControllerId,
+		Tilt,
+		RotationRate,
+		Gravity,
+		Acceleration
+		);
+	OnMotionDetectedMessage(MotionEvent);
+
+	return true;
+}
+
+void FSlateApplication::OnMotionDetectedMessage( FMotionEvent& MotionEvent )
+{
+	CALL_WIDGET_FUNCTION(MotionEvent, OnMotionDetected);
+}
+
+bool FSlateApplication::OnSizeChanged( const TSharedRef< FGenericWindow >& PlatformWindow, const int32 Width, const int32 Height, bool bWasMinimized )
+{
+	TSharedPtr< SWindow > Window = FindWindowByPlatformWindow( SlateWindows, PlatformWindow );
+
+	if ( Window.IsValid() )
+	{
+		Window->SetCachedSize( FVector2D( Width, Height ) );
+
+		Renderer->RequestResize( Window, Width, Height );
+
+		if ( !bWasMinimized && Window->IsRegularWindow() && !Window->HasOSWindowBorder() && Window->IsVisible() )
+		{
+			PrivateDrawWindows( Window );
+		}
+
+		if( !bWasMinimized && Window->IsVisible() && Window->IsRegularWindow() && Window->IsAutosized() )
+		{
+			// Reduces flickering due to one frame lag when windows are resized automatically
+			Renderer->FlushCommands();
+		}
+
+		// Inform the notification manager we have activated a window - it may want to force notifications 
+		// back to the front of the z-order
+		FSlateNotificationManager::Get().ForceNotificationsInFront( Window.ToSharedRef() );
+	}
+
+	return true;
+}
+
+void FSlateApplication::OnOSPaint( const TSharedRef< FGenericWindow >& PlatformWindow )
+{
+	TSharedPtr< SWindow > Window = FindWindowByPlatformWindow( SlateWindows, PlatformWindow );
+	PrivateDrawWindows( Window );
+	Renderer->FlushCommands();
+}
+
+void FSlateApplication::OnResizingWindow( const TSharedRef< FGenericWindow >& PlatformWindow )
+{
+	// Flush the rendering command queue to ensure that there aren't pending viewport draw commands for the old viewport size.
+	Renderer->FlushCommands();
+}
+
+bool FSlateApplication::BeginReshapingWindow( const TSharedRef< FGenericWindow >& PlatformWindow )
+{
+	if(!IsExternalUIOpened())
+	{
+		if (!ThrottleHandle.IsValid())
+		{
+			ThrottleHandle = FSlateThrottleManager::Get().EnterResponsiveMode();
+		}
+
+		return true;
+	}
+
+	return false;
+}
+
+void FSlateApplication::FinishedReshapingWindow( const TSharedRef< FGenericWindow >& PlatformWindow )
+{
+	if (ThrottleHandle.IsValid())
+	{
+		FSlateThrottleManager::Get().LeaveResponsiveMode(ThrottleHandle);
+	}
+}
+
+void FSlateApplication::OnMovedWindow( const TSharedRef< FGenericWindow >& PlatformWindow, const int32 X, const int32 Y )
+{
+	TSharedPtr< SWindow > Window = FindWindowByPlatformWindow( SlateWindows, PlatformWindow );
+
+	if ( Window.IsValid() )
+	{
+		Window->SetCachedScreenPosition( FVector2D( X, Y ) );
+	}
+}
+
+FWindowActivateEvent::EActivationType TranslationWindowActivationMessage( const EWindowActivation::Type ActivationType )
+{
+	FWindowActivateEvent::EActivationType Result = FWindowActivateEvent::EA_Activate;
+
+	switch( ActivationType )
+	{
+	case EWindowActivation::Activate:
+		Result = FWindowActivateEvent::EA_Activate;
+		break;
+	case EWindowActivation::ActivateByMouse:
+		Result = FWindowActivateEvent::EA_ActivateByMouse;
+		break;
+	case EWindowActivation::Deactivate:
+		Result = FWindowActivateEvent::EA_Deactivate;
+		break;
+	default:
+		check( false );
+	}
+
+	return Result;
+}
+
+bool FSlateApplication::OnWindowActivationChanged( const TSharedRef< FGenericWindow >& PlatformWindow, const EWindowActivation::Type ActivationType )
+{
+	TSharedPtr< SWindow > Window = FindWindowByPlatformWindow( SlateWindows, PlatformWindow );
+
+	if ( !Window.IsValid() )
+	{
+		return false;
+	}
+
+	FWindowActivateEvent::EActivationType TranslatedActivationType = TranslationWindowActivationMessage( ActivationType );
+	FWindowActivateEvent WindowActivateEvent( TranslatedActivationType, Window.ToSharedRef() );
+
+	return ProcessWindowActivatedMessage( WindowActivateEvent );
+}
+
+bool FSlateApplication::ProcessWindowActivatedMessage( const FWindowActivateEvent& ActivateEvent )
+{
+	TSharedPtr<SWindow> ActiveModalWindow = GetActiveModalWindow();
+
+	if ( ActivateEvent.GetActivationType() != FWindowActivateEvent::EA_Deactivate )
+	{
+		// Do not process activation messages unless we have no modal windows or the current window is modal
+		if( !ActiveModalWindow.IsValid() || ActivateEvent.GetAffectedWindow() == ActiveModalWindow || ActivateEvent.GetAffectedWindow()->IsDescendantOf(ActiveModalWindow) )
+		{
+			// Window being ACTIVATED
+
+			BringWindowToFront( SlateWindows, ActivateEvent.GetAffectedWindow() );
+
+			{
+				// Switch worlds widgets in the current path
+				FScopedSwitchWorldHack SwitchWorld( ActivateEvent.GetAffectedWindow() );
+				ActivateEvent.GetAffectedWindow()->OnIsActiveChanged( ActivateEvent );
+			}
+
+			if ( ActivateEvent.GetAffectedWindow()->IsRegularWindow() )
+			{
+				ActiveTopLevelWindow = ActivateEvent.GetAffectedWindow();
+			}
+
+			// A Slate window was activated
+			bSlateWindowActive = true;
+
+			if ( ActivateEvent.GetAffectedWindow()->IsFocusedInitially() && ActivateEvent.GetAffectedWindow()->SupportsKeyboardFocus() )
+			{
+				// Set keyboard focus on the window being activated.
+				{
+					FWidgetPath PathToWindowBeingActivated;
+					GeneratePathToWidgetChecked( ActivateEvent.GetAffectedWindow(), PathToWindowBeingActivated );
+
+					if( ActivateEvent.GetActivationType() == FWindowActivateEvent::EA_ActivateByMouse )
+					{
+						SetKeyboardFocus(PathToWindowBeingActivated, EKeyboardFocusCause::Mouse);
+					}
+					else
+					{
+						SetKeyboardFocus(PathToWindowBeingActivated, EKeyboardFocusCause::WindowActivate);
+					}
+				}
+
+			}
+
+			{
+				FScopedSwitchWorldHack SwitchWorld( ActivateEvent.GetAffectedWindow() );
+				// let the menu stack know of new window being activated.  We may need to close menus as a result
+				MenuStack.OnWindowActivated( ActivateEvent.GetAffectedWindow() );
+			}
+
+			// Inform the notification manager we have activated a window - it may want to force notifications 
+			// back to the front of the z-order
+			FSlateNotificationManager::Get().ForceNotificationsInFront( ActivateEvent.GetAffectedWindow() );
+		}
+		else
+		{
+			// An attempt is being made to activate another window when a modal window is running
+			ActiveModalWindow->BringToFront();
+			ActiveModalWindow->FlashWindow();
+		}
+	}
+	else
+	{
+		// Window being DEACTIVATED
+
+		// If our currently-active top level window was deactivated, take note of that
+		if ( ActivateEvent.GetAffectedWindow()->IsRegularWindow() &&
+			ActivateEvent.GetAffectedWindow() == ActiveTopLevelWindow.Pin() )
+		{
+			ActiveTopLevelWindow.Reset();
+		}
+
+		// A Slate window was deactivated.  Currently there is no active Slate window
+		bSlateWindowActive = false;
+
+		// Switch worlds for the activated window
+		FScopedSwitchWorldHack SwitchWorld( ActivateEvent.GetAffectedWindow() );
+		ActivateEvent.GetAffectedWindow()->OnIsActiveChanged( ActivateEvent );
+
+		// A window was deactivated; mouse capture should be cleared
+		ResetToDefaultInputSettings();
+	}
+
+	return true;
+}
+
+bool FSlateApplication::OnApplicationActivationChanged( const bool IsActive )
+{
+	ProcessApplicationActivationMessage( IsActive );
+	return true;
+}
+
+void FSlateApplication::ProcessApplicationActivationMessage( bool InAppActivated )
+{
+	const bool UserSwitchedAway = bAppIsActive && !InAppActivated;
+
+	bAppIsActive = InAppActivated;
+
+
+	// If the user switched to a different application then we should dismiss our pop-ups.  In the case
+	// where a user clicked on a different Slate window, OnWindowActivatedMessage() will be call MenuStack.OnWindowActivated()
+	// to destroy any windows in our stack that are no longer appropriate to be displayed.
+	if( UserSwitchedAway )
+	{
+		// Close pop-up menus
+		DismissAllMenus();
+
+		// Close tool-tips
+		CloseToolTip();
+
+		// No slate window is active when our entire app becomes inactive
+		bSlateWindowActive = false;
+
+		// Clear keyboard focus when the app is deactivated
+		ClearKeyboardFocus(EKeyboardFocusCause::OtherWidgetLostFocus);
+
+		// If we have a slate-only drag-drop occurring, stop the drag drop.
+		if ( IsDragDropping() && !DragDropContent->IsExternalOperation() )
+		{
+			DragDropContent.Reset();
+		}
+	}
+}
+
+EWindowZone::Type FSlateApplication::GetWindowZoneForPoint( const TSharedRef< FGenericWindow >& PlatformWindow, const int32 X, const int32 Y )
+{
+	TSharedPtr< SWindow > Window = FindWindowByPlatformWindow( SlateWindows, PlatformWindow );
+
+	if ( Window.IsValid() )
+	{
+		return Window->GetCurrentWindowZone( FVector2D( X, Y ) );
+	}
+
+	return EWindowZone::NotInWindow;
+}
+
+
+void FSlateApplication::PrivateDestroyWindow( const TSharedRef<SWindow>& DestroyedWindow )
+{
+	// Notify the window that it is going to be destroyed.  The window must be completely intact when this is called 
+	// because delegates are allowed to leave Slate here
+	DestroyedWindow->NotifyWindowBeingDestroyed();
+
+	// Release rendering resources.  
+	// This MUST be done before destroying the native window as the native window is required to be valid before releasing rendering resources with some API's
+	Renderer->OnWindowDestroyed( DestroyedWindow );
+
+	// Destroy the native window
+	DestroyedWindow->DestroyWindowImmediately();
+
+	// Remove the window and all its children from the Slate window list
+	RemoveWindowFromList( SlateWindows, DestroyedWindow );
+
+	// Shutdown the application if there are no more windows
+	{
+		bool bAnyRegularWindows = false;
+		for( auto WindowIter( SlateWindows.CreateConstIterator() ); WindowIter; ++WindowIter )
+		{
+			auto Window = *WindowIter;
+			if( Window->IsRegularWindow() )
+			{
+				bAnyRegularWindows = true;
+				break;
+			}
+		}
+
+		if (!bAnyRegularWindows)
+		{
+			OnExitRequested.ExecuteIfBound();
+		}
+	}
+}
+
+void FSlateApplication::OnWindowClose( const TSharedRef< FGenericWindow >& PlatformWindow )
+{
+	TSharedPtr< SWindow > Window = FindWindowByPlatformWindow( SlateWindows, PlatformWindow );
+
+	if ( Window.IsValid() )
+	{
+		Window->RequestDestroyWindow();
+	}
+}
+
+EDropEffect::Type FSlateApplication::OnDragEnterText( const TSharedRef< FGenericWindow >& Window, const FString& Text )
+{
+	const TSharedPtr< FExternalDragOperation > DragDropOperation = FExternalDragOperation::NewText( Text );
+	const TSharedPtr< SWindow > EffectingWindow = FindWindowByPlatformWindow( SlateWindows, Window );
+
+	EDropEffect::Type Result = EDropEffect::None;
+	if ( DragDropOperation.IsValid() && EffectingWindow.IsValid() )
+	{
+		Result = OnDragEnter( EffectingWindow.ToSharedRef(), DragDropOperation.ToSharedRef() );
+	}
+
+	return Result;
+}
+
+EDropEffect::Type FSlateApplication::OnDragEnterFiles( const TSharedRef< FGenericWindow >& Window, const TArray< FString >& Files )
+{
+	const TSharedPtr< FExternalDragOperation > DragDropOperation = FExternalDragOperation::NewFiles( Files );
+	const TSharedPtr< SWindow > EffectingWindow = FindWindowByPlatformWindow( SlateWindows, Window );
+
+	EDropEffect::Type Result = EDropEffect::None;
+	if ( DragDropOperation.IsValid() && EffectingWindow.IsValid() )
+	{
+		Result = OnDragEnter( EffectingWindow.ToSharedRef(), DragDropOperation.ToSharedRef() );
+	}
+
+	return Result;
+}
+
+EDropEffect::Type FSlateApplication::OnDragEnter( const TSharedRef< SWindow >& Window, const TSharedRef<FExternalDragOperation>& DragDropOperation )
+{
+	// We are encountering a new drag and drop operation.
+	// Assume we cannot handle it.
+	DragIsHandled = false;
+
+	const FVector2D CurrentCursorPosition = GetCursorPos();
+
+	// Tell slate to enter drag and drop mode.
+	// Make a faux mouse event for slate, so we can initiate a drag and drop.
+	FDragDropEvent DragDropEvent(
+		FPointerEvent(
+		CurrentCursorPosition,
+		LastCursorPosition,
+		PressedMouseButtons,
+		EKeys::Invalid,
+		0,
+		PlatformApplication->GetModifierKeys() ),
+		DragDropOperation
+	);
+
+	LastCursorPosition = CurrentCursorPosition;
+
+	ProcessDragEnterMessage( Window, DragDropEvent );
+	return EDropEffect::None;
+}
+
+bool FSlateApplication::ProcessDragEnterMessage( TSharedRef<SWindow> WindowEntered, FDragDropEvent& DragDropEvent )
+{
+	FWidgetPath WidgetsUnderCursor = LocateWindowUnderMouse( DragDropEvent.GetScreenSpacePosition(), GetInteractiveTopLevelWindows() );
+	DragDropEvent.SetEventPath( WidgetsUnderCursor );
+
+	// Switch worlds for widgets in the current path
+	FScopedSwitchWorldHack SwitchWorld( WidgetsUnderCursor );
+
+	FReply TriggerDragDropReply = FReply::Handled().BeginDragDrop( DragDropEvent.GetOperation().ToSharedRef() );
+	ProcessReply( WidgetsUnderCursor, TriggerDragDropReply, &WidgetsUnderCursor, &DragDropEvent );
+
+	return true;
+}
+
+EDropEffect::Type FSlateApplication::OnDragOver( const TSharedPtr< FGenericWindow >& Window )
+{
+	EDropEffect::Type Result = EDropEffect::None;
+
+	if ( IsDragDropping() )
+	{
+		bool MouseMoveHandled = true;
+		FVector2D CursorMovementDelta( 0, 0 );
+		const FVector2D CurrentCursorPosition = GetCursorPos();
+
+		if ( LastCursorPosition != CurrentCursorPosition )
+		{
+			FPointerEvent MouseEvent(
+				CurrentCursorPosition,
+				LastCursorPosition,
+				PressedMouseButtons,
+				EKeys::Invalid,
+				0,
+				PlatformApplication->GetModifierKeys()
+			);
+
+			MouseMoveHandled = ProcessMouseMoveMessage( MouseEvent );
+			CursorMovementDelta = MouseEvent.GetCursorDelta();
+		}
+
+		// Slate is now in DragAndDrop mode. It is tracking the payload.
+		// We just need to convey mouse movement.
+		if ( CursorMovementDelta.SizeSquared() > 0 )
+		{
+			DragIsHandled = MouseMoveHandled;
+		}
+
+		if ( DragIsHandled )
+		{
+			Result = EDropEffect::Copy;
+		}
+	}
+
+	return Result;
+}
+
+void FSlateApplication::OnDragLeave( const TSharedPtr< FGenericWindow >& Window )
+{
+	DragDropContent.Reset();
+}
+
+EDropEffect::Type FSlateApplication::OnDragDrop( const TSharedPtr< FGenericWindow >& Window )
+{
+	EDropEffect::Type Result = EDropEffect::None;
+
+	if ( IsDragDropping() )
+	{
+		FPointerEvent MouseEvent(
+			GetCursorPos(), 	
+			LastCursorPosition,
+			PressedMouseButtons,
+			EKeys::LeftMouseButton,
+			0,
+			PlatformApplication->GetModifierKeys()
+			);
+
+		// User dropped into a Slate window. Slate is already in drag and drop mode.
+		// It knows what to do based on a mouse up.
+		if ( ProcessMouseButtonUpMessage( MouseEvent ) )
+{
+			Result = EDropEffect::Copy;
+		}
+	}
+
+	return Result;
+}
+
+bool FSlateApplication::OnWindowAction( const TSharedRef< FGenericWindow >& PlatformWindow, const EWindowAction::Type InActionType)
+{
+	return !IsExternalUIOpened();
+}
+
+
+/* 
+ *****************************************************************************/
+
+void FSlateApplication::InitializeAsStandaloneApplication( const TSharedRef<FSlateRenderer>& PlatformRenderer)
+{
+	// create the platform slate application (what FSlateApplication::Get() returns)
+	FSlateApplication::Create();
+
+	// initialize renderer
+	FSlateApplication::Get().InitializeRenderer(PlatformRenderer);
+
+	// set the normal UE4 GIsRequestingExit when outer frame is closed
+	FSlateApplication::Get().SetExitRequestedHandler(FSimpleDelegate::CreateStatic(&OnRequestExit));
+}
