Slate: Refactored core Slate implementation into SlateCore module in preparation for UMG.

diff --git a/Engine/Source/Runtime/Slate/Private/Framework/Application/SlateIcon.cpp b/Engine/Source/Runtime/Slate/Private/Framework/Application/SlateIcon.cpp
new file mode 100644
--- /dev/null
+++ b/Engine/Source/Runtime/Slate/Private/Framework/Application/SlateIcon.cpp
@@ -0,0 +1,69 @@
+// Copyright 1998-2014 Epic Games, Inc. All Rights Reserved.
+
+/*=============================================================================
+	SlateIcon.cpp: Implements the FSlateIcon structure.
+=============================================================================*/
+
+#include "Slate.h"
+
+
+/* FSlateIcon structors
+ *****************************************************************************/
+
+FSlateIcon::FSlateIcon( )
+	: StyleSetName(NAME_None)
+	, StyleName(NAME_None)
+	, SmallStyleName(NAME_None)
+	, bIsSet(false)
+{ }
+
+
+FSlateIcon::FSlateIcon( const FName& InStyleSetName, const FName& InStyleName )
+	: StyleSetName(InStyleSetName)
+	, StyleName(InStyleName)
+	, SmallStyleName(ISlateStyle::Join(InStyleName, ".Small"))
+	, bIsSet(true)
+{ }
+
+
+FSlateIcon::FSlateIcon( const FName& InStyleSetName, const FName& InStyleName, const FName& InSmallStyleName )
+	: StyleSetName(InStyleSetName)
+	, StyleName(InStyleName)
+	, SmallStyleName(InSmallStyleName)
+	, bIsSet(true)
+{ }
+
+
+/* FSlateIcon interface
+ *****************************************************************************/
+
+const FSlateBrush* FSlateIcon::GetSmallIcon( ) const
+{
+	const ISlateStyle* StyleSet = GetStyleSet();
+	
+	if (StyleSet)
+	{
+		return StyleSet->GetOptionalBrush(SmallStyleName);
+	}
+
+	return FStyleDefaults::GetNoBrush();
+}
+
+
+const ISlateStyle* FSlateIcon::GetStyleSet( ) const
+{
+	return FSlateStyleRegistry::FindSlateStyle(StyleSetName);
+}
+
+
+const FSlateBrush* FSlateIcon::GetIcon( ) const
+{
+	const ISlateStyle* StyleSet = GetStyleSet();
+	
+	if(StyleSet)
+	{
+		return StyleSet->GetOptionalBrush(StyleName);
+	}
+
+	return FStyleDefaults::GetNoBrush();
+}
diff --git a/Engine/Source/Runtime/Slate/Private/Framework/Commands/Commands.cpp b/Engine/Source/Runtime/Slate/Private/Framework/Commands/Commands.cpp
new file mode 100644
--- /dev/null
+++ b/Engine/Source/Runtime/Slate/Private/Framework/Commands/Commands.cpp
@@ -0,0 +1,24 @@
+// Copyright 1998-2014 Epic Games, Inc. All Rights Reserved.
+
+#include "Slate.h"
+
+
+#define LOC_DEFINE_REGION
+
+void UI_COMMAND_Function( FBindingContext* This, TSharedPtr< FUICommandInfo >& OutCommand, const TCHAR* OutCommandName, const TCHAR* OutCommandNameUnderscoreTooltip, const ANSICHAR* DotOutCommandName, const TCHAR* FriendlyName, const TCHAR* InDescription, const EUserInterfaceActionType::Type CommandType, const FInputGesture& InDefaultGesture )
+{
+	static const FString UICommandsStr(TEXT("UICommands"));
+
+	FUICommandInfo::MakeCommandInfo(
+		This->AsShared(),
+		OutCommand,
+		OutCommandName,
+		FInternationalization::ForUseOnlyByLocMacroAndGraphNodeTextLiterals_CreateText( FriendlyName, *UICommandsStr, OutCommandName ),
+		FInternationalization::ForUseOnlyByLocMacroAndGraphNodeTextLiterals_CreateText( InDescription, *UICommandsStr, OutCommandNameUnderscoreTooltip ),
+		FSlateIcon( This->GetStyleSetName(), ISlateStyle::Join( This->GetContextName(), DotOutCommandName ) ),
+		CommandType,
+		InDefaultGesture
+	);
+}
+
+#undef LOC_DEFINE_REGION
\ No newline at end of file
diff --git a/Engine/Source/Runtime/Slate/Private/Framework/Commands/InputBindingManager.cpp b/Engine/Source/Runtime/Slate/Private/Framework/Commands/InputBindingManager.cpp
new file mode 100644
--- /dev/null
+++ b/Engine/Source/Runtime/Slate/Private/Framework/Commands/InputBindingManager.cpp
@@ -0,0 +1,462 @@
+// Copyright 1998-2014 Epic Games, Inc. All Rights Reserved.
+
+/*=============================================================================
+	InputBindingManager.cpp: Implements the FInputBindingManager class.
+=============================================================================*/
+
+#include "Slate.h"
+#include "RemoteConfigIni.h"
+
+
+/* FUserDefinedGestures helper class
+ *****************************************************************************/
+
+class FUserDefinedGestures
+{
+public:
+	void LoadGestures();
+	void SaveGestures();
+	bool GetUserDefinedGesture( const FName BindingContext, const FName CommandName, FInputGesture& OutUserDefinedGesture );
+	void SetUserDefinedGesture( const FUICommandInfo& CommandInfo );
+	/** Remove all user defined gestures */
+	void RemoveAll();
+private:
+	TSharedPtr<FJsonObject> Gestures;
+};
+
+void FUserDefinedGestures::LoadGestures()
+{
+	if( !Gestures.IsValid() )
+	{
+		FString Content;
+		GConfig->GetString(TEXT("UserDefinedGestures"), TEXT("Content"), Content, GEditorKeyBindingsIni);
+
+		TSharedRef< TJsonReader<> > Reader = TJsonReaderFactory<>::Create( FRemoteConfig::ReplaceIniSpecialCharWithChar(Content).ReplaceEscapedCharWithChar() );
+		bool bResult = FJsonSerializer::Deserialize( Reader, Gestures );
+
+		if (!Gestures.IsValid())
+		{
+			// Gestures have not been loaded from the ini file, try reading them from the txt file now
+			FArchive* Ar = IFileManager::Get().CreateFileReader( *( FPaths::GameSavedDir() / TEXT("Preferences/EditorKeyBindings.txt") ) );
+			if( Ar )
+			{
+				TSharedRef< TJsonReader<ANSICHAR> > TextReader = TJsonReaderFactory<ANSICHAR>::Create( Ar );
+				bResult = FJsonSerializer::Deserialize( TextReader, Gestures );
+				delete Ar;
+			}
+		}
+		
+		if( !Gestures.IsValid() )
+		{
+			Gestures = MakeShareable( new FJsonObject );
+		}
+	}
+}
+
+void FUserDefinedGestures::SaveGestures()
+{
+	if( Gestures.IsValid() ) 
+	{
+		FString Content;
+		TSharedRef< TJsonWriter<> > Writer = TJsonWriterFactory<>::Create( &Content );
+		FJsonSerializer::Serialize( Gestures.ToSharedRef(), Writer );
+
+		GConfig->SetString(TEXT("UserDefinedGestures"), TEXT("Content"), *FRemoteConfig::ReplaceIniCharWithSpecialChar(Content).ReplaceCharWithEscapedChar(), GEditorKeyBindingsIni);
+	}
+}
+
+bool FUserDefinedGestures::GetUserDefinedGesture( const FName BindingContext, const FName CommandName, FInputGesture& OutUserDefinedGesture )
+{
+	bool bResult = false;
+
+	if( Gestures.IsValid() )
+	{
+		const TSharedPtr<FJsonValue> ContextObj = Gestures->Values.FindRef( BindingContext.ToString() );
+		if( ContextObj.IsValid() )
+		{
+			const TSharedPtr<FJsonValue> CommandObj = ContextObj->AsObject()->Values.FindRef( CommandName.ToString() );
+			if( CommandObj.IsValid() )
+			{
+				TSharedPtr<FJsonObject> GestureObj = CommandObj->AsObject();
+
+				const TSharedPtr<FJsonValue> KeyObj = GestureObj->Values.FindRef( TEXT("Key") );
+				const TSharedPtr<FJsonValue> CtrlObj = GestureObj->Values.FindRef( TEXT("Control") );
+				const TSharedPtr<FJsonValue> AltObj = GestureObj->Values.FindRef( TEXT("Alt") );
+				const TSharedPtr<FJsonValue> ShiftObj = GestureObj->Values.FindRef( TEXT("Shift") );
+
+				OutUserDefinedGesture.Key = *KeyObj->AsString();
+				OutUserDefinedGesture.bCtrl = CtrlObj->AsBool();
+				OutUserDefinedGesture.bAlt = AltObj->AsBool();
+				OutUserDefinedGesture.bShift = ShiftObj->AsBool();
+
+				bResult = true;
+			}
+		}
+	}
+
+	return bResult;
+}
+
+void FUserDefinedGestures::SetUserDefinedGesture( const FUICommandInfo& CommandInfo )
+{
+	// Find or create the command context json value
+	const FName BindingContext = CommandInfo.GetBindingContext();
+	const FName CommandName = CommandInfo.GetCommandName();
+
+	TSharedPtr<FJsonValue> ContextObj = Gestures->Values.FindRef( BindingContext.ToString() );
+	if( !ContextObj.IsValid() )
+	{
+		ContextObj = Gestures->Values.Add( BindingContext.ToString(), MakeShareable( new FJsonValueObject( MakeShareable( new FJsonObject ) ) ) );
+	}
+
+	const TSharedPtr<const FInputGesture> InputGesture = CommandInfo.GetActiveGesture();
+
+	FInputGesture GestureToSave;
+
+	// Save an empty invalid gesture if one was not set
+	// This is an indication that the user doesn't want this bound and not to use the default gesture
+	if( InputGesture.IsValid() )
+	{
+		GestureToSave = *InputGesture;
+	}
+
+	{
+		TSharedPtr<FJsonValueObject> GestureValueObj = MakeShareable( new FJsonValueObject( MakeShareable( new FJsonObject ) ) );
+		TSharedPtr<FJsonObject> GestureObj = GestureValueObj->AsObject();
+
+		// Set the gesture values for the command
+		GestureObj->Values.Add( TEXT("Key"), MakeShareable( new FJsonValueString( GestureToSave.Key.ToString() ) ) );
+		GestureObj->Values.Add( TEXT("Control"),  MakeShareable( new FJsonValueBoolean( GestureToSave.bCtrl ) ) );
+		GestureObj->Values.Add( TEXT("Alt"),  MakeShareable( new FJsonValueBoolean( GestureToSave.bAlt ) ) );
+		GestureObj->Values.Add( TEXT("Shift"),  MakeShareable( new FJsonValueBoolean( GestureToSave.bShift ) ) );
+
+		ContextObj->AsObject()->Values.Add( CommandName.ToString(), GestureValueObj );
+	}
+
+}
+
+void FUserDefinedGestures::RemoveAll()
+{
+	Gestures = MakeShareable(new FJsonObject);
+}
+
+
+/* FInputBindingManager structors
+ *****************************************************************************/
+
+FInputBindingManager& FInputBindingManager::Get()
+{
+	static FInputBindingManager* Instance= NULL;
+	if( Instance == NULL )
+	{
+		Instance = new FInputBindingManager();
+	}
+
+	return *Instance;
+}
+
+
+/* FInputBindingManager interface
+ *****************************************************************************/
+
+/**
+ * Gets the user defined gesture (if any) from the provided command name
+ * 
+ * @param InBindingContext	The context in which the command is active
+ * @param InCommandName		The name of the command to get the gesture from
+ */
+bool FInputBindingManager::GetUserDefinedGesture( const FName InBindingContext, const FName InCommandName, FInputGesture& OutUserDefinedGesture )
+{
+	if( !UserDefinedGestures.IsValid() )
+	{
+		UserDefinedGestures = MakeShareable( new FUserDefinedGestures );
+		UserDefinedGestures->LoadGestures();
+	}
+
+	return UserDefinedGestures->GetUserDefinedGesture( InBindingContext, InCommandName, OutUserDefinedGesture );
+}
+
+void FInputBindingManager::CheckForDuplicateDefaultGestures( const FBindingContext& InBindingContext, TSharedPtr<FUICommandInfo> InCommandInfo ) const
+{
+	const bool bCheckDefault = true;
+	TSharedPtr<FUICommandInfo> ExistingInfo = GetCommandInfoFromInputGesture( InBindingContext.GetContextName(), InCommandInfo->DefaultGesture, bCheckDefault );
+	if( ExistingInfo.IsValid() )
+	{
+		if( ExistingInfo->CommandName != InCommandInfo->CommandName )
+		{
+			// Two different commands with the same name in the same context or parent context
+			UE_LOG(LogSlate, Fatal, TEXT("The command '%s.%s' has the same default gesture as '%s.%s' [%s]"), 
+				*InCommandInfo->BindingContext.ToString(),
+				*InCommandInfo->CommandName.ToString(), 
+				*ExistingInfo->BindingContext.ToString(),
+				*ExistingInfo->CommandName.ToString(), 
+				*InCommandInfo->DefaultGesture.GetInputText().ToString() );
+		}
+	}
+}
+
+
+void FInputBindingManager::NotifyActiveGestureChanged( const FUICommandInfo& CommandInfo )
+{
+	FContextEntry& ContextEntry = ContextMap.FindChecked( CommandInfo.GetBindingContext() );
+
+	// Slow but doesn't happen frequently
+	for( FGestureMap::TIterator It( ContextEntry.GestureToCommandInfoMap ); It; ++It )
+	{
+		// Remove the currently active gesture from the map if one exists
+		if( It.Value() == CommandInfo.GetCommandName() )
+		{
+			It.RemoveCurrent();
+			// There should only be one active gesture
+			break;
+		}
+	}
+
+	if( CommandInfo.GetActiveGesture()->IsValidGesture() )
+	{
+		checkSlow( !ContextEntry.GestureToCommandInfoMap.Contains( *CommandInfo.GetActiveGesture() ) )
+		ContextEntry.GestureToCommandInfoMap.Add( *CommandInfo.GetActiveGesture(), CommandInfo.GetCommandName() );
+	}
+	
+
+	// The user defined gestures should have already been created
+	check( UserDefinedGestures.IsValid() );
+
+	UserDefinedGestures->SetUserDefinedGesture( CommandInfo );
+
+	// Broadcast the gesture event when a new one is added
+	OnUserDefinedGestureChanged.Broadcast(CommandInfo);
+}
+
+void FInputBindingManager::SaveInputBindings()
+{
+	if( UserDefinedGestures.IsValid() )
+	{
+		UserDefinedGestures->SaveGestures();
+	}
+}
+
+void FInputBindingManager::RemoveUserDefinedGestures()
+{
+	if( UserDefinedGestures.IsValid() )
+	{
+		UserDefinedGestures->RemoveAll();
+		UserDefinedGestures->SaveGestures();
+	}
+}
+
+void FInputBindingManager::GetCommandInfosFromContext( const FName InBindingContext, TArray< TSharedPtr<FUICommandInfo> >& OutCommandInfos ) const
+{
+	ContextMap.FindRef( InBindingContext ).CommandInfoMap.GenerateValueArray( OutCommandInfos );
+}
+
+void FInputBindingManager::CreateInputCommand( const TSharedRef<FBindingContext>& InBindingContext, TSharedRef<FUICommandInfo> InCommandInfo )
+{
+	check( InCommandInfo->BindingContext == InBindingContext->GetContextName() );
+
+	// The command name should be valid
+	check( InCommandInfo->CommandName != NAME_None );
+
+	// Should not have already created a gesture for this command
+	check( !InCommandInfo->ActiveGesture->IsValidGesture() );
+	
+	const FName ContextName = InBindingContext->GetContextName();
+
+	FContextEntry& ContextEntry = ContextMap.FindOrAdd( ContextName );
+	
+	// Our parent context must exist.
+	check( InBindingContext->GetContextParent() == NAME_None || ContextMap.Find( InBindingContext->GetContextParent() ) != NULL );
+
+	FCommandInfoMap& CommandInfoMap = ContextEntry.CommandInfoMap;
+
+	if( !ContextEntry.BindingContext.IsValid() )
+	{
+		ContextEntry.BindingContext = InBindingContext;
+	}
+
+	if( InBindingContext->GetContextParent() != NAME_None  )
+	{
+		check( InBindingContext->GetContextName() != InBindingContext->GetContextParent() );
+		// Set a mapping from the parent of the current context to the current context
+		ParentToChildMap.AddUnique( InBindingContext->GetContextParent(), InBindingContext->GetContextName() );
+	}
+
+	if( InCommandInfo->DefaultGesture.IsValidGesture() )
+	{
+		CheckForDuplicateDefaultGestures( *InBindingContext, InCommandInfo );
+	}
+	
+	{
+		TSharedPtr<FUICommandInfo> ExistingInfo = CommandInfoMap.FindRef( InCommandInfo->CommandName );
+		ensureMsgf( !ExistingInfo.IsValid(), TEXT("A command with name %s already exists in context %s"), *InCommandInfo->CommandName.ToString(), *InBindingContext->GetContextName().ToString() );
+	}
+
+	// Add the command info to the list of known infos.  It can only exist once.
+	CommandInfoMap.Add( InCommandInfo->CommandName, InCommandInfo );
+
+	// See if there are user defined gestures for this command
+	FInputGesture UserDefinedGesture;
+	bool bFoundUserDefinedGesture = GetUserDefinedGesture( ContextName, InCommandInfo->CommandName, UserDefinedGesture );
+
+	
+	if( !bFoundUserDefinedGesture && InCommandInfo->DefaultGesture.IsValidGesture() )
+	{
+		// Find any existing command with the same gesture 
+		// This is for inconsistency between default and user defined gesture.  We need to make sure that if default gestures are changed to a gesture that a user set to a different command, that the default gesture doesn't replace
+		// the existing commands gesture. Note: Duplicate default gestures are found above in CheckForDuplicateDefaultGestures
+		FName ExisingCommand = ContextEntry.GestureToCommandInfoMap.FindRef( InCommandInfo->DefaultGesture );
+
+		if( ExisingCommand == NAME_None )
+		{
+			// No existing command has a user defined gesture and no user defined gesture is available for this command 
+			TSharedRef<FInputGesture> NewGesture = MakeShareable( new FInputGesture( InCommandInfo->DefaultGesture ) );
+			InCommandInfo->ActiveGesture = NewGesture;
+		}
+
+	}
+	else if( bFoundUserDefinedGesture )
+	{
+		// Find any existing command with the same gesture 
+		// This is for inconsistency between default and user defined gesture.  We need to make sure that if default gestures are changed to a gesture that a user set to a different command, that the default gesture doesn't replace
+		// the existing commands gesture.
+		FName ExisingCommandName = ContextEntry.GestureToCommandInfoMap.FindRef( UserDefinedGesture );
+
+		if( ExisingCommandName != NAME_None )
+		{
+			// Get the command with using the same gesture
+			TSharedPtr<FUICommandInfo> ExistingInfo = CommandInfoMap.FindRef( ExisingCommandName );
+			if( *ExistingInfo->ActiveGesture != ExistingInfo->DefaultGesture )
+			{
+				// two user defined gestures are the same within a context.  If the keybinding editor was used this wont happen so this must have been directly a modified user setting file
+				UE_LOG(LogSlate, Error, TEXT("Duplicate user defined gestures found: [%s,%s].  Gesture for %s being removed"), *InCommandInfo->GetLabel().ToString(), *ExistingInfo->GetLabel().ToString(), *ExistingInfo->GetLabel().ToString() );
+			}
+			ContextEntry.GestureToCommandInfoMap.Remove( *ExistingInfo->ActiveGesture );
+			// Remove the existing gesture. 
+			ExistingInfo->ActiveGesture = MakeShareable( new FInputGesture() );
+		
+		}
+
+		TSharedRef<FInputGesture> NewGesture = MakeShareable( new FInputGesture( UserDefinedGesture ) );
+		// Set the active gesture on the command info
+		InCommandInfo->ActiveGesture = NewGesture;
+	}
+
+	// If the active gesture is valid, map the gesture to the map for fast lookup when processing bindings
+	if( InCommandInfo->ActiveGesture->IsValidGesture() )
+	{
+		checkSlow( !ContextEntry.GestureToCommandInfoMap.Contains( *InCommandInfo->GetActiveGesture() ) );
+		ContextEntry.GestureToCommandInfoMap.Add( *InCommandInfo->GetActiveGesture(), InCommandInfo->GetCommandName() );
+	}
+}
+
+const TSharedPtr<FUICommandInfo> FInputBindingManager::FindCommandInContext( const FName InBindingContext, const FInputGesture& InGesture, bool bCheckDefault ) const
+{
+	const FContextEntry& ContextEntry = ContextMap.FindRef( InBindingContext );
+	
+	TSharedPtr<FUICommandInfo> FoundCommand = NULL;
+
+	if( bCheckDefault )
+	{
+		const FCommandInfoMap& InfoMap = ContextEntry.CommandInfoMap;
+		for( FCommandInfoMap::TConstIterator It(InfoMap); It && !FoundCommand.IsValid(); ++It )
+		{
+			const FUICommandInfo& CommandInfo = *It.Value();
+			if( CommandInfo.DefaultGesture == InGesture )
+			{
+				FoundCommand = It.Value();
+			}	
+		}
+	}
+	else
+	{
+		// faster lookup for active gestures
+		FName CommandName = ContextEntry.GestureToCommandInfoMap.FindRef( InGesture );
+		if( CommandName != NAME_None )
+		{
+			FoundCommand = ContextEntry.CommandInfoMap.FindChecked( CommandName );
+		}
+	}
+
+	return FoundCommand;
+}
+
+const TSharedPtr<FUICommandInfo> FInputBindingManager::FindCommandInContext( const FName InBindingContext, const FName CommandName ) const 
+{
+	const FContextEntry& ContextEntry = ContextMap.FindRef( InBindingContext );
+	
+	return ContextEntry.CommandInfoMap.FindRef( CommandName );
+}
+
+
+void FInputBindingManager::GetAllChildContexts( const FName InBindingContext, TArray<FName>& AllChildren ) const
+{
+	AllChildren.Add( InBindingContext );
+
+	TArray<FName> TempChildren;
+	ParentToChildMap.MultiFind( InBindingContext, TempChildren );
+	for( int32 ChildIndex = 0; ChildIndex < TempChildren.Num(); ++ChildIndex )
+	{
+		GetAllChildContexts( TempChildren[ChildIndex], AllChildren );
+	}
+}
+
+const TSharedPtr<FUICommandInfo> FInputBindingManager::GetCommandInfoFromInputGesture( const FName InBindingContext, const FInputGesture& InGesture, bool bCheckDefault ) const
+{
+	TSharedPtr<FUICommandInfo> FoundCommand = NULL;
+
+	FName CurrentContext = InBindingContext;
+	while( CurrentContext != NAME_None && !FoundCommand.IsValid() )
+	{
+		const FContextEntry& ContextEntry = ContextMap.FindRef( CurrentContext );
+
+		FoundCommand = FindCommandInContext( CurrentContext, InGesture, bCheckDefault );
+
+		CurrentContext = ContextEntry.BindingContext->GetContextParent();
+	}
+	
+	if( !FoundCommand.IsValid() )
+	{
+		TArray<FName> Children;
+		GetAllChildContexts( InBindingContext, Children );
+
+		for( int32 ContextIndex = 0; ContextIndex < Children.Num() && !FoundCommand.IsValid(); ++ContextIndex )
+		{
+			FoundCommand = FindCommandInContext( Children[ContextIndex], InGesture, bCheckDefault );
+		}
+	}
+	
+
+	return FoundCommand;
+}
+
+/**
+ * Returns a list of all known input contexts
+ *
+ * @param OutInputContexts	The generated list of contexts
+ * @return A list of all known input contexts                   
+ */
+void FInputBindingManager::GetKnownInputContexts( TArray< TSharedPtr<FBindingContext> >& OutInputContexts ) const
+{
+	for( TMap< FName, FContextEntry >::TConstIterator It(ContextMap); It; ++It )
+	{
+		OutInputContexts.Add( It.Value().BindingContext );
+	}
+}
+
+TSharedPtr<FBindingContext> FInputBindingManager::GetContextByName( const FName& InContextName )
+{
+	FContextEntry* FindResult = ContextMap.Find( InContextName );
+	if ( FindResult == NULL )
+	{
+		return TSharedPtr<FBindingContext>();
+	}
+	else
+	{
+		return FindResult->BindingContext;
+	}
+}
+
+void FInputBindingManager::RemoveContextByName( const FName& InContextName )
+{
+	ContextMap.Remove(InContextName);
+}
diff --git a/Engine/Source/Runtime/Slate/Private/Framework/Commands/InputGesture.cpp b/Engine/Source/Runtime/Slate/Private/Framework/Commands/InputGesture.cpp
new file mode 100644
--- /dev/null
+++ b/Engine/Source/Runtime/Slate/Private/Framework/Commands/InputGesture.cpp
@@ -0,0 +1,89 @@
+// Copyright 1998-2014 Epic Games, Inc. All Rights Reserved.
+
+/*=============================================================================
+	InputGesture.cpp: Implements the FInputGesture class.
+=============================================================================*/
+
+#include "Slate.h"
+
+
+#define LOCTEXT_NAMESPACE "FInputGesture"
+
+
+/* FInputGesture interface
+ *****************************************************************************/
+
+/**
+ * Returns the friendly, localized string name of this key binding
+ * @todo Slate: Got to be a better way to do this
+ */
+FText FInputGesture::GetInputText( ) const
+{
+	FText FormatString = LOCTEXT("ShortcutKeyWithNoAdditionalModifierKeys", "{Key}");
+
+	if ( bCtrl && !bAlt && !bShift )
+	{
+#if PLATFORM_MAC
+		FormatString = LOCTEXT("KeyName_Command", "Command+{Key}");
+#else
+		FormatString = LOCTEXT("KeyName_Control", "Ctrl+{Key}");
+#endif
+	}
+	else if ( bCtrl && bAlt && !bShift )
+	{
+#if PLATFORM_MAC
+		FormatString = LOCTEXT("KeyName_Command + KeyName_Alt", "Command+Alt+{Key}");
+#else
+		FormatString = LOCTEXT("KeyName_Control + KeyName_Alt", "Ctrl+Alt+{Key}");
+#endif
+	}
+	else if ( bCtrl && !bAlt && bShift )
+	{
+#if PLATFORM_MAC
+		FormatString = LOCTEXT("KeyName_Command + KeyName_Shift", "Command+Shift+{Key}");
+#else
+		FormatString = LOCTEXT("KeyName_Control + KeyName_Shift", "Ctrl+Shift+{Key}");
+#endif
+	}
+	else if ( bCtrl && bAlt && bShift )
+	{
+#if PLATFORM_MAC
+		FormatString = LOCTEXT("KeyName_Command + KeyName_Alt + KeyName_Shift", "Command+Alt+Shift+{Key}");
+#else
+		FormatString = LOCTEXT("KeyName_Control + KeyName_Alt + KeyName_Shift", "Ctrl+Alt+Shift+{Key}");
+#endif
+	}
+	else if ( !bCtrl && bAlt && !bShift )
+	{
+		FormatString = LOCTEXT("KeyName_Alt", "Alt+{Key}");
+	}
+	else if ( !bCtrl && bAlt && bShift )
+	{
+		FormatString = LOCTEXT("KeyName_Alt + KeyName_Shift", "Alt+Shift+{Key}");
+	}
+	else if ( !bCtrl && !bAlt && bShift )
+	{
+		FormatString = LOCTEXT("KeyName_Shift", "Shift+{Key}");
+	}
+
+	FFormatNamedArguments Args;
+	Args.Add( TEXT("Key"), GetKeyText() );
+
+	return FText::Format( FormatString, Args );
+}
+
+
+FText FInputGesture::GetKeyText( ) const
+{
+	FText OutString; // = KeyGetDisplayName(Key);
+
+	if (Key.IsValid() && !Key.IsModifierKey())
+	{
+		OutString = Key.GetDisplayName();
+	}
+
+	return OutString;
+}
+
+
+#undef LOCTEXT_NAMESPACE
diff --git a/Engine/Source/Runtime/Slate/Private/Framework/Commands/UICommandInfo.cpp b/Engine/Source/Runtime/Slate/Private/Framework/Commands/UICommandInfo.cpp
new file mode 100644
--- /dev/null
+++ b/Engine/Source/Runtime/Slate/Private/Framework/Commands/UICommandInfo.cpp
@@ -0,0 +1,128 @@
+// Copyright 1998-2014 Epic Games, Inc. All Rights Reserved.
+
+/*=============================================================================
+	.cpp: Implements the  class.
+=============================================================================*/
+
+#include "Slate.h"
+
+
+FSimpleMulticastDelegate FBindingContext::CommandsChanged;
+
+
+FUICommandInfoDecl FBindingContext::NewCommand( const FName InCommandName, const FText& InCommandLabel, const FText& InCommandDesc )
+{
+	return FUICommandInfoDecl( this->AsShared(), InCommandName, InCommandLabel, InCommandDesc );
+}
+
+
+
+
+
+FUICommandInfoDecl::FUICommandInfoDecl( const TSharedRef<FBindingContext>& InContext, const FName InCommandName, const FText& InLabel, const FText& InDesc )
+	: Context( InContext )
+{
+	Info = MakeShareable( new FUICommandInfo( InContext->GetContextName() ) );
+	Info->CommandName = InCommandName;
+	Info->Label = InLabel;
+	Info->Description = InDesc;
+}
+
+FUICommandInfoDecl& FUICommandInfoDecl::DefaultGesture( const FInputGesture& InDefaultGesture )
+{
+	Info->DefaultGesture = InDefaultGesture;
+	return *this;
+}
+FUICommandInfoDecl& FUICommandInfoDecl::UserInterfaceType( EUserInterfaceActionType::Type InType )
+{
+	Info->UserInterfaceType = InType;
+	return *this;
+}
+
+FUICommandInfoDecl& FUICommandInfoDecl::Icon( const FSlateIcon& InIcon )
+{
+	Info->Icon = InIcon;
+	return *this;
+}
+
+FUICommandInfoDecl& FUICommandInfoDecl::Description( const FText& InDescription )
+{
+	Info->Description = InDescription;
+	return *this;
+}
+
+FUICommandInfoDecl::operator TSharedPtr<FUICommandInfo>() const
+{
+	FInputBindingManager::Get().CreateInputCommand( Context, Info.ToSharedRef() );
+	return Info;
+}
+
+FUICommandInfoDecl::operator TSharedRef<FUICommandInfo>() const
+{
+	FInputBindingManager::Get().CreateInputCommand( Context, Info.ToSharedRef() );
+	return Info.ToSharedRef();
+}
+
+
+
+const FText FUICommandInfo::GetInputText() const
+{	
+	return ActiveGesture->GetInputText();
+}
+
+
+void FUICommandInfo::MakeCommandInfo( const TSharedRef<class FBindingContext>& InContext, TSharedPtr< FUICommandInfo >& OutCommand, const FName InCommandName, const FText& InCommandLabel, const FText& InCommandDesc, const FSlateIcon& InIcon, const EUserInterfaceActionType::Type InUserInterfaceType, const FInputGesture& InDefaultGesture )
+{
+	ensureMsgf( !InCommandLabel.IsEmpty(), TEXT("Command labels cannot be empty") );
+
+	OutCommand = MakeShareable( new FUICommandInfo( InContext->GetContextName() ) );
+	OutCommand->CommandName = InCommandName;
+	OutCommand->Label = InCommandLabel;
+	OutCommand->Description = InCommandDesc;
+	OutCommand->Icon = InIcon;
+	OutCommand->UserInterfaceType = InUserInterfaceType;
+	OutCommand->DefaultGesture = InDefaultGesture;
+	FInputBindingManager::Get().CreateInputCommand( InContext, OutCommand.ToSharedRef() );
+}
+
+void FUICommandInfo::SetActiveGesture( const FInputGesture& NewGesture )
+{
+	ActiveGesture->Set( NewGesture );
+
+	// Set the user defined gesture for this command so it can be saved to disk later
+	FInputBindingManager::Get().NotifyActiveGestureChanged( *this );
+}
+
+void FUICommandInfo::RemoveActiveGesture()
+{
+	// Gesture already exists
+	// Reset the other gesture that has the same binding
+	ActiveGesture = MakeShareable( new FInputGesture() );
+	FInputBindingManager::Get().NotifyActiveGestureChanged( *this );
+}
+
+TSharedRef<SToolTip> FUICommandInfo::MakeTooltip( const TAttribute<FText>& InText, const TAttribute< EVisibility >& InToolTipVisibility ) const
+{
+	return 
+		SNew(SToolTip)
+		.Visibility(InToolTipVisibility.IsBound() ? InToolTipVisibility : EVisibility::Visible)
+		.Content()
+		[
+			SNew(SVerticalBox)
+			+SVerticalBox::Slot()
+			.Padding(0.0f, 0.0f, 0.0f, 4.0f)
+			[
+				SNew(STextBlock)
+				.Text(InText.IsBound() ? InText : GetDescription())
+				.Font(FCoreStyle::Get().GetFontStyle( "ToolTip.Font" ))
+				.ColorAndOpacity( FSlateColor::UseForeground() )
+			]
+			+SVerticalBox::Slot()
+			[
+				SNew(STextBlock)
+				.Text(GetInputText())
+				.Font(FCoreStyle::Get().GetFontStyle( "ToolTip.Font" ))
+				.ColorAndOpacity( FSlateColor::UseSubduedForeground() )
+			]
+		];
+}
\ No newline at end of file
diff --git a/Engine/Source/Runtime/Slate/Private/Framework/Commands/UICommandList.cpp b/Engine/Source/Runtime/Slate/Private/Framework/Commands/UICommandList.cpp
new file mode 100644
--- /dev/null
+++ b/Engine/Source/Runtime/Slate/Private/Framework/Commands/UICommandList.cpp
@@ -0,0 +1,404 @@
+// Copyright 1998-2014 Epic Games, Inc. All Rights Reserved.
+
+/*=============================================================================
+	UICommandList.cpp: Implements the FUICommandList class.
+=============================================================================*/
+
+#include "Slate.h"
+
+
+/* FUICommandList interface
+ *****************************************************************************/
+
+/**
+ * Maps a command info to a series of delegates that are executed by a multibox or mouse/keyboard input
+ *
+ * @param InUICommandInfo	The command info to map
+ * @param ExecuteAction		The delegate to call when the command should be executed
+ */
+void FUICommandList::MapAction( const TSharedPtr< const FUICommandInfo > InUICommandInfo, FExecuteAction ExecuteAction )
+{
+	MapAction( InUICommandInfo, ExecuteAction, FCanExecuteAction(), FIsActionChecked(), FIsActionButtonVisible() );
+}
+
+
+/**
+ * Maps a command info to a series of delegates that are executed by a multibox or mouse/keyboard input
+ *
+ * @param InUICommandInfo	The command info to map
+ * @param ExecuteAction		The delegate to call when the command should be executed
+ * @param CanExecuteAction	The delegate to call to see if the command can be executed
+ */
+void FUICommandList::MapAction( const TSharedPtr< const FUICommandInfo > InUICommandInfo, FExecuteAction ExecuteAction, FCanExecuteAction CanExecuteAction )
+{
+	MapAction( InUICommandInfo, ExecuteAction, CanExecuteAction, FIsActionChecked(), FIsActionButtonVisible() );
+}
+
+
+/**
+ * Maps a command info to a series of delegates that are executed by a multibox or mouse/keyboard input
+ *
+ * @param InUICommandInfo	The command info to map
+ * @param ExecuteAction		The delegate to call when the command should be executed
+ * @param CanExecuteAction	The delegate to call to see if the command can be executed
+ * @param IsCheckedDelegate	The delegate to call to see if the command should appear checked when visualized in a multibox
+ */
+void FUICommandList::MapAction( const TSharedPtr< const FUICommandInfo > InUICommandInfo, FExecuteAction ExecuteAction, FCanExecuteAction CanExecuteAction, FIsActionChecked IsCheckedDelegate )
+{
+	MapAction( InUICommandInfo, ExecuteAction, CanExecuteAction, IsCheckedDelegate, FIsActionButtonVisible() );
+}
+
+
+void FUICommandList::MapAction( const TSharedPtr< const FUICommandInfo > InUICommandInfo, FExecuteAction ExecuteAction, FCanExecuteAction CanExecuteAction, FIsActionChecked IsCheckedDelegate, FIsActionButtonVisible IsVisibleDelegate )
+{
+	FUIAction Action;
+	Action.ExecuteAction = ExecuteAction;
+	Action.CanExecuteAction = CanExecuteAction;
+	Action.IsCheckedDelegate = IsCheckedDelegate;
+	Action.IsActionVisibleDelegate = IsVisibleDelegate;
+
+	MapAction( InUICommandInfo, Action );
+
+}
+
+
+void FUICommandList::MapAction( const TSharedPtr< const FUICommandInfo > InUICommandInfo, const FUIAction& InUIAction )
+{
+	check( InUICommandInfo.IsValid() );
+
+	// Check against already mapped actions
+	checkSlow( !UICommandBindingMap.Contains( InUICommandInfo ) );
+
+	ContextsInList.Add( InUICommandInfo->GetBindingContext() );
+	UICommandBindingMap.Add( InUICommandInfo, InUIAction );
+}
+
+
+void FUICommandList::Append( const TSharedRef<FUICommandList>& InCommandsToAppend )
+{
+	check( &InCommandsToAppend.Get() != this );
+
+	// Clear out any invalid parents or children
+	CleanupPointerArray(ParentUICommandLists);
+	CleanupPointerArray(ChildUICommandLists);
+
+	// Add the new parent. Add this to the parent's child list.
+	ParentUICommandLists.AddUnique( InCommandsToAppend );
+	InCommandsToAppend->ChildUICommandLists.AddUnique( this->AsShared() );
+}
+
+
+/**
+ * Executes the action associated with the provided command info
+ * Note: It is assumed at this point that CanExecuteAction was already checked
+ *
+ * @param InUICommandInfo	The command info execute
+ */
+bool FUICommandList::ExecuteAction( const TSharedRef< const FUICommandInfo > InUICommandInfo ) const
+{
+	const FUIAction* Action = GetActionForCommand(InUICommandInfo);
+
+	if( Action )
+	{
+		FSlateApplication::Get().OnLogSlateEvent(EEventLog::UICommand, InUICommandInfo->GetLabel());
+		Action->Execute();
+		return true;
+	}
+
+	return false;
+}
+
+
+/**
+ * Calls the CanExecuteAction associated with the provided command info to see if ExecuteAction can be called
+ *
+ * @param InUICommandInfo	The command info execute
+ */
+bool FUICommandList::CanExecuteAction( const TSharedRef< const FUICommandInfo > InUICommandInfo ) const
+{
+	const FUIAction* Action = GetActionForCommand(InUICommandInfo);
+
+	if( Action )
+	{
+		return Action->CanExecute();
+	}
+
+	// If there is no action then assume its possible to execute (Some menus with only children that do nothing themselves will have no delegates bound )
+	return true;
+}
+
+
+bool FUICommandList::TryExecuteAction( const TSharedRef< const FUICommandInfo > InUICommandInfo ) const
+{
+	if( !CanExecuteAction( InUICommandInfo ) )
+	{
+		return false;
+	}
+
+	return ExecuteAction( InUICommandInfo );
+}
+
+
+/**
+ * Calls the IsVisible delegate associated with the provided command info to see if the command should be visible in a toolbar
+ *
+ * @param InUICommandInfo	The command info execute
+ */
+EVisibility FUICommandList::GetVisibility( const TSharedRef< const FUICommandInfo > InUICommandInfo ) const
+{
+	const FUIAction* Action = GetActionForCommand(InUICommandInfo);
+
+	if (Action)
+	{
+		return Action->IsVisible();
+	}
+
+	// If there is no action then assume it is visible
+	return EVisibility::Visible;
+}
+
+
+/**
+ * Calls the IsChecked delegate to see if the visualization of this command in a multibox should appear checked
+ *
+ * @param InUICommandInfo	The command info execute
+ */
+bool FUICommandList::IsChecked( const TSharedRef< const FUICommandInfo > InUICommandInfo ) const
+{
+	const FUIAction* Action = GetActionForCommand(InUICommandInfo);
+
+	if( Action )
+	{
+		return Action->IsChecked();
+	}
+
+	return false;
+}
+
+
+/**
+ * Processes any user interface actions which are activated by the specified keyboard event
+ *
+ * @param InBindingContext	The context in which the actions are valid
+ * @param InKeyboardEvent	The keyboard event to check
+ *
+ * @return true if an action was processed
+ */
+bool FUICommandList::ProcessCommandBindings( const FKeyboardEvent& InKeyboardEvent ) const 
+{
+	return ConditionalProcessCommandBindings( InKeyboardEvent.GetKey(), InKeyboardEvent.IsControlDown(), InKeyboardEvent.IsAltDown(), InKeyboardEvent.IsShiftDown(), InKeyboardEvent.IsRepeat() );
+}
+
+
+/**
+ * Processes any user interface actions which are activated by the specified mouse event
+ *
+ * @param InBindingContext	The context in which the actions are valid
+ * @param InKeyboardEvent	The mouse event to check
+ *
+ * @return true if an action was processed
+ */
+bool FUICommandList::ProcessCommandBindings( const FPointerEvent& InMouseEvent ) const
+{
+	return ConditionalProcessCommandBindings( InMouseEvent.GetEffectingButton(), InMouseEvent.IsControlDown(), InMouseEvent.IsAltDown(), InMouseEvent.IsShiftDown(), InMouseEvent.IsRepeat() );
+}
+
+
+/* FUICommandList implementation
+ *****************************************************************************/
+
+/**
+ * Helper function to execute an interface action delegate or exec command if valid
+ *
+ * @param InBindingContext The binding context where commands are currently valid
+ * @param Key		The current key that is pressed
+ * @param bCtrl		True if control is pressed
+ * @param bAlt		True if alt is pressed
+ * @param bShift	True if shift is pressed
+ * @param bRepeat	True if command is repeating (held)
+ * @return True if a command was executed, False otherwise
+ */
+bool FUICommandList::ConditionalProcessCommandBindings( const FKey Key, bool bCtrl, bool bAlt, bool bShift, bool bRepeat ) const
+{
+	if ( !bRepeat && !FSlateApplication::Get().IsDragDropping() )
+	{
+		FInputGesture CheckGesture( Key );
+		CheckGesture.bCtrl = bCtrl;
+		CheckGesture.bAlt = bAlt;
+		CheckGesture.bShift = bShift;
+
+		if( CheckGesture.IsValidGesture() )
+		{
+			TSet<FName> AllContextsToCheck;
+			GatherContextsForList(AllContextsToCheck);
+
+			for( TSet<FName>::TConstIterator It(AllContextsToCheck); It; ++It )
+			{
+				FName Context = *It;
+	
+				// Only active gestures process commands
+				const bool bCheckDefault = false;
+
+				// Check to see if there is any command in the context activated by the gesture
+				TSharedPtr<FUICommandInfo> Command = FInputBindingManager::Get().FindCommandInContext( Context, CheckGesture, bCheckDefault );
+
+				if( Command.IsValid() && ensure( *Command->GetActiveGesture() == CheckGesture ) )
+				{
+					// Find the bound action for this command
+					const FUIAction* Action = GetActionForCommand(Command);
+
+					// If there is no Action mapped to this command list, continue to the next context
+					if( Action )
+					{
+						if(Action->CanExecute())
+						{
+							// If the action was found and can be executed, do so now
+							Action->Execute();
+							return true;
+						}
+						else
+						{
+							// An action wasn't bound to the command but a gesture was found or an action was found but cant be executed
+							return false;
+						}
+					}
+				}
+			}
+		}
+	}
+
+	// No action was processed
+	return false;
+}
+
+
+const FUIAction* FUICommandList::GetActionForCommand(TSharedPtr<const FUICommandInfo> Command) const
+{
+	// Make sure the command is valid
+	if ( !ensure(Command.IsValid()) )
+	{
+		return NULL;
+	}
+
+	// Check in my own binding map. This should not be prevented by CanProduceActionForCommand.
+	// Any action directly requested from a command list should be returned if it actually exists in the list.
+	const FUIAction* Action = UICommandBindingMap.Find( Command );
+
+	if ( !Action )
+	{
+		// We did not find the action in our own list. Recursively attempt to find the command in children and parents.
+		const bool bIncludeChildren = true;
+		const bool bIncludeParents = true;
+		TSet<TSharedRef<const FUICommandList>> VisitedLists;
+		Action = GetActionForCommandRecursively(Command.ToSharedRef(), bIncludeChildren, bIncludeParents, VisitedLists);
+	}
+	
+	return Action;
+}
+
+
+const FUIAction* FUICommandList::GetActionForCommandRecursively(const TSharedRef<const FUICommandInfo>& Command, bool bIncludeChildren, bool bIncludeParents, TSet<TSharedRef<const FUICommandList>>& InOutVisitedLists) const
+{
+	// Detect cycles in the graph
+	{
+		const TSharedRef<const FUICommandList>& ListAsShared = AsShared();
+		if ( InOutVisitedLists.Contains(ListAsShared) )
+		{
+			// This node was already visited. End recursion.
+			return NULL;
+		}
+
+		InOutVisitedLists.Add( ListAsShared );
+	}
+
+	const FUIAction* Action = NULL;
+
+	// Make sure I am capable of processing this command
+	bool bCapableOfCommand = true;
+	if ( CanProduceActionForCommand.IsBound() )
+	{
+		bCapableOfCommand = CanProduceActionForCommand.Execute(Command);
+	}
+
+	if ( bCapableOfCommand )
+	{
+		// Check in my own binding map
+		Action = UICommandBindingMap.Find( Command );
+	
+		// If the action was not found, check in my children binding maps
+		if ( !Action && bIncludeChildren )
+		{
+			for ( auto ChildIt = ChildUICommandLists.CreateConstIterator(); ChildIt; ++ChildIt )
+			{
+				TWeakPtr<FUICommandList> Child = *ChildIt;
+				if ( Child.IsValid() )
+				{
+					const bool bShouldIncludeChildrenOfChild = true;
+					const bool bShouldIncludeParentsOfChild = false;
+					Action = Child.Pin()->GetActionForCommandRecursively(Command, bShouldIncludeChildrenOfChild, bShouldIncludeParentsOfChild, InOutVisitedLists);
+					if ( Action )
+					{
+						break;
+					}
+				}
+			}
+		}
+	}
+
+	// If the action was not found, check in my parent binding maps
+	if ( !Action && bIncludeParents )
+	{
+		for ( auto ParentIt = ParentUICommandLists.CreateConstIterator(); ParentIt; ++ParentIt )
+		{
+			TWeakPtr<FUICommandList> Parent = *ParentIt;
+			if ( Parent.IsValid() )
+			{
+				const bool bShouldIncludeChildrenOfParent = false;
+				const bool bShouldIncludeParentsOfParent = true;
+				Action = Parent.Pin()->GetActionForCommandRecursively(Command, bShouldIncludeChildrenOfParent, bShouldIncludeParentsOfParent, InOutVisitedLists);
+				if ( Action )
+				{
+					break;
+				}
+			}
+		}
+	}
+
+	return Action;
+}
+
+
+void FUICommandList::GatherContextsForList(TSet<FName>& OutAllContexts) const
+{
+	TSet<TSharedRef<const FUICommandList>> VisitedLists;
+	GatherContextsForListRecursively(OutAllContexts, VisitedLists);
+}
+
+
+void FUICommandList::GatherContextsForListRecursively(TSet<FName>& OutAllContexts, TSet<TSharedRef<const FUICommandList>>& InOutVisitedLists) const
+{
+	// Detect cycles in the graph
+	{
+		const TSharedRef<const FUICommandList>& ListAsShared = AsShared();
+		if ( InOutVisitedLists.Contains(ListAsShared) )
+		{
+			// This node was already visited. End recursion.
+			return;
+		}
+
+		InOutVisitedLists.Add( ListAsShared );
+	}
+
+	// Include all contexts on this list
+	OutAllContexts.Append(ContextsInList);
+
+	// Include all the parent contexts
+	for ( auto ParentIt = ParentUICommandLists.CreateConstIterator(); ParentIt; ++ParentIt )
+	{
+		TWeakPtr<FUICommandList> Parent = *ParentIt;
+		if ( Parent.IsValid() )
+		{
+			Parent.Pin()->GatherContextsForListRecursively(OutAllContexts, InOutVisitedLists);
+		}
+	}
+}
diff --git a/Engine/Source/Runtime/Slate/Private/SImage.cpp b/Engine/Source/Runtime/Slate/Private/SImage.cpp
--- a/Engine/Source/Runtime/Slate/Private/SImage.cpp
+++ b/Engine/Source/Runtime/Slate/Private/SImage.cpp
@@ -1,8 +1,6 @@
 // Copyright 1998-2014 Epic Games, Inc. All Rights Reserved.
 
 #include "Slate.h"
-#include "SImage.h"
-
 
 
 /**
diff --git a/Engine/Source/Runtime/Slate/Private/SMissingWidget.cpp b/Engine/Source/Runtime/Slate/Private/SMissingWidget.cpp
new file mode 100644
--- /dev/null
+++ b/Engine/Source/Runtime/Slate/Private/SMissingWidget.cpp
@@ -0,0 +1,26 @@
+// Copyright 1998-2014 Epic Games, Inc. All Rights Reserved.
+
+/*=============================================================================
+	SMissingWidget.h: Declares the SMissingWidget class.
+=============================================================================*/
+
+#include "Slate.h"
+
+
+#define LOCTEXT_NAMESPACE "SMissingWidget"
+
+
+TSharedRef<class SWidget> SMissingWidget::MakeMissingWidget( )
+{
+	return SNew(SBorder)
+		.HAlign(HAlign_Center)
+		.VAlign(VAlign_Center)
+		[
+			SNew(STextBlock)
+				.Text(LOCTEXT("MissingWidget", "Missing Widget"))
+				.TextStyle(FCoreStyle::Get(), "EmbossedText")
+		];
+}
+
+
+#undef LOCTEXT_NAMESPACE
diff --git a/Engine/Source/Runtime/Slate/Private/Slate.cpp b/Engine/Source/Runtime/Slate/Private/Slate.cpp
--- a/Engine/Source/Runtime/Slate/Private/Slate.cpp
+++ b/Engine/Source/Runtime/Slate/Private/Slate.cpp
@@ -4,25 +4,6 @@
 #include "ModuleManager.h"
 
 
-
-IMPLEMENT_MODULE( FDefaultModuleImpl, Slate );
-
-
-DEFINE_STAT(STAT_SlateBuildLines);
-DEFINE_STAT(STAT_SlateBuildBoxes);
-DEFINE_STAT(STAT_SlateBuildBorders);
-DEFINE_STAT(STAT_SlateBuildText);
-DEFINE_STAT(STAT_SlateBuildGradients);
-DEFINE_STAT(STAT_SlateBuildSplines);
-DEFINE_STAT(STAT_SlateBuildViewports);
-DEFINE_STAT(STAT_SlateBuildElements);
-DEFINE_STAT(STAT_SlateFindBatchTime);
-DEFINE_STAT(STAT_SlateUpdateBufferGTTime);
-DEFINE_STAT(STAT_SlateFontCachingTime);
-DEFINE_STAT(STAT_SlateRenderingGTTime);
-DEFINE_STAT(STAT_SlateMeasureStringTime);
-
-
 DEFINE_STAT(STAT_SlateCacheDesiredSize_STextBlock);
 DEFINE_STAT(STAT_SlateCacheDesiredSize);
 
@@ -41,32 +22,8 @@ DEFINE_STAT(STAT_SlateOnPaint);
 DEFINE_STAT(STAT_SlateOnPaint_SSlider);
 
 DEFINE_STAT(STAT_SlateDrawWindowTime);
-
 DEFINE_STAT(STAT_SlateTickWidgets);
-
 DEFINE_STAT(STAT_SlateTickTime);
 
-DEFINE_STAT(STAT_SlateDrawTime);
-DEFINE_STAT(STAT_SlateUpdateBufferRTTime);
-DEFINE_STAT(STAT_SlateRenderingRTTime);
-
 // Useful for locally pref testing miscellaneous Slate areas.  Usage of this should not be checked in 
 DEFINE_STAT(STAT_SlateMiscTime);
-
-
-DEFINE_STAT(STAT_SlateNumLayers);
-DEFINE_STAT(STAT_SlateNumBatches);
-DEFINE_STAT(STAT_SlateVertexCount);
-
-DEFINE_STAT(STAT_SlateVertexBatchMemory);
-DEFINE_STAT(STAT_SlateIndexBatchMemory);
-DEFINE_STAT(STAT_SlateVertexBufferMemory);
-DEFINE_STAT(STAT_SlateIndexBufferMemory);
-DEFINE_STAT(STAT_SlateTextureAtlasMemory);
-DEFINE_STAT(STAT_SlateFontKerningTableMemory);
-DEFINE_STAT(STAT_SlateFontMeasureCacheMemory);
-DEFINE_STAT(STAT_SlateNumTextureAtlases);
-DEFINE_STAT(STAT_SlateNumFontAtlases);
-DEFINE_STAT(STAT_SlateNumNonAtlasedTextures);
-DEFINE_STAT(STAT_SlateTextureDataMemory);
-DEFINE_STAT(STAT_SlateNumDynamicTextures);
\ No newline at end of file
diff --git a/Engine/Source/Runtime/Slate/Private/SlateModule.cpp b/Engine/Source/Runtime/Slate/Private/SlateModule.cpp
new file mode 100644
--- /dev/null
+++ b/Engine/Source/Runtime/Slate/Private/SlateModule.cpp
@@ -0,0 +1,21 @@
+// Copyright 1998-2014 Epic Games, Inc. All Rights Reserved.
+
+/*=============================================================================
+	SlateModule.cpp: Implements the FSlateModule class.
+=============================================================================*/
+
+#include "Slate.h"
+#include "SlatePrivatePCH.h"
+#include "ModuleManager.h"
+
+
+/**
+ * Implements the Slate module.
+ */
+class FSlateModule
+	: public IModuleInterface
+{
+};
+
+
+IMPLEMENT_MODULE(FSlateModule, Slate);
diff --git a/Engine/Source/Runtime/Slate/Private/SlatePrivatePCH.h b/Engine/Source/Runtime/Slate/Private/SlatePrivatePCH.h
new file mode 100644
--- /dev/null
+++ b/Engine/Source/Runtime/Slate/Private/SlatePrivatePCH.h
@@ -0,0 +1,30 @@
+// Copyright 1998-2014 Epic Games, Inc. All Rights Reserved.
+
+/*=============================================================================
+	SlatePrivatePCH.h: Pre-compiled header file for the Slate module.
+=============================================================================*/
+
+#pragma once
+
+//#include "Slate.h"
+
+
+/* Dependencies
+ *****************************************************************************/
+
+
+
+/* Private includes
+ *****************************************************************************/
+
+
+
+/* Global functions
+ *****************************************************************************/
+
+/**
+ * Creates a widget reflector widget.
+ *
+ * @todo Slate: move this into a plug-in.
+ */
+//TSharedRef<SWidgetReflector> MakeWidgetReflector( );
\ No newline at end of file
diff --git a/Engine/Source/Runtime/Slate/Private/Testing/STestSuite.cpp b/Engine/Source/Runtime/Slate/Private/Testing/STestSuite.cpp
--- a/Engine/Source/Runtime/Slate/Private/Testing/STestSuite.cpp
+++ b/Engine/Source/Runtime/Slate/Private/Testing/STestSuite.cpp
@@ -3366,8 +3366,6 @@ void RestoreSlateTestSuite()
 {
 	FTestStyle::ResetToDefault();
 
-	FSlateApplication::Get().RegisterWidgetReflectorTabSpawner( TestSuiteMenu::DeveloperCategory );
-
 	FGlobalTabmanager::Get()->RegisterTabSpawner("TestSuite1", FOnSpawnTab::CreateStatic( &SpawnTestSuite1 ) );
 	FGlobalTabmanager::Get()->RegisterTabSpawner("TestSuite2", FOnSpawnTab::CreateStatic( &SpawnTestSuite2 ) );
 
diff --git a/Engine/Source/Runtime/Slate/Private/UICommandDragDropOp.cpp b/Engine/Source/Runtime/Slate/Private/UICommandDragDropOp.cpp
new file mode 100644
--- /dev/null
+++ b/Engine/Source/Runtime/Slate/Private/UICommandDragDropOp.cpp
@@ -0,0 +1,54 @@
+// Copyright 1998-2014 Epic Games, Inc. All Rights Reserved.
+
+/*=============================================================================
+	UICommandDragDropOp.cpp: Implements the FUICommandDragDropOp class.
+=============================================================================*/
+
+#include "Slate.h"
+
+
+TSharedRef<FUICommandDragDropOp> FUICommandDragDropOp::New( TSharedRef<const FUICommandInfo> InCommandInfo, FName InOriginMultiBox, TSharedPtr<SWidget> CustomDectorator, FVector2D DecoratorOffset )
+{
+	TSharedRef<FUICommandDragDropOp> Operation = MakeShareable(new FUICommandDragDropOp(InCommandInfo, InOriginMultiBox, CustomDectorator, DecoratorOffset));
+	Operation->Construct();
+
+	return Operation;
+}
+
+
+void FUICommandDragDropOp::OnDragged( const class FDragDropEvent& DragDropEvent )
+{
+	CursorDecoratorWindow->SetOpacity(0.85f);
+	CursorDecoratorWindow->MoveWindowTo(DragDropEvent.GetScreenSpacePosition() + Offset);
+}
+
+
+void FUICommandDragDropOp::OnDrop( bool bDropWasHandled, const FPointerEvent& MouseEvent )
+{
+	FDragDropOperation::OnDrop(bDropWasHandled, MouseEvent);
+	OnDropNotification.ExecuteIfBound();
+}
+
+
+TSharedPtr<SWidget> FUICommandDragDropOp::GetDefaultDecorator() const
+{
+	TSharedPtr<SWidget> Content;
+
+	if (CustomDecorator.IsValid())
+	{
+		Content = CustomDecorator;
+	}
+	else
+	{
+		Content = SNew(STextBlock)
+			.Text(UICommand->GetLabel());
+	}
+
+	// Create hover widget
+	return SNew(SBorder)
+		.BorderImage(FCoreStyle::Get().GetBrush("ToolPanel.GroupBorder"))
+		.Content()
+		[	
+			Content.ToSharedRef()
+		];
+}
diff --git a/Engine/Source/Runtime/Slate/Public/Framework/Application/IPlatformTextField.h b/Engine/Source/Runtime/Slate/Public/Framework/Application/IPlatformTextField.h
new file mode 100644
--- /dev/null
+++ b/Engine/Source/Runtime/Slate/Public/Framework/Application/IPlatformTextField.h
@@ -0,0 +1,16 @@
+// Copyright 1998-2014 Epic Games, Inc. All Rights Reserved.
+
+#pragma once
+
+class SVirtualKeyboardEntry;
+
+class IPlatformTextField
+{
+public:
+	virtual ~IPlatformTextField() {};
+
+	virtual void ShowKeyboard(bool bShow, TSharedPtr<SVirtualKeyboardEntry> TextEntryWidget) = 0;
+
+private:
+
+};
diff --git a/Engine/Source/Runtime/Slate/Public/Framework/Application/MenuStack.h b/Engine/Source/Runtime/Slate/Public/Framework/Application/MenuStack.h
new file mode 100644
--- /dev/null
+++ b/Engine/Source/Runtime/Slate/Public/Framework/Application/MenuStack.h
@@ -0,0 +1,130 @@
+// Copyright 1998-2014 Epic Games, Inc. All Rights Reserved.
+
+/*=============================================================================
+	MenuStack.h: Declares the FMenuStack class.
+=============================================================================*/
+
+#pragma once
+
+
+#if UE_BUILD_DEBUG
+	// Debugging is harder with the inline allocator
+	typedef TArray<TSharedPtr<SWindow>> FMenuWindowList;
+#else
+	typedef TArray<TSharedPtr<SWindow>, TInlineAllocator<4>> FMenuWindowList;
+#endif
+
+
+/** Describes a simple animation for popup window introductions */
+struct FPopupTransitionEffect
+{
+	/** Direction to slide in from */
+	enum ESlideDirection
+	{
+		/** No sliding */
+		None,
+
+		/** Sliding direction for a combo button */
+		ComboButton,
+
+		/** Sliding direction for a top-level pull-down menu or combo box */
+		TopMenu,
+
+		/** Sliding direction for a sub-menu */
+		SubMenu,
+
+		/** Sliding direction for a popup that lets the user type in data */
+		TypeInPopup,
+
+		/** Sliding direction preferred for context menu popups */
+		ContextMenu,
+
+	} SlideDirection;
+
+
+	/** Constructor */
+	FPopupTransitionEffect( const ESlideDirection InitSlideDirection )
+		: SlideDirection(InitSlideDirection)
+	{ }
+};
+
+
+/**
+ * Represents a stack of open menus.  The last item in the stack is the top most menu
+ * Each stack level can have any number of windows opened on that level but they all must be parented to the same window
+ * which must reside in the previous level of the stack
+ */
+class FMenuStack
+{
+public:
+
+	/**
+	 * Pushes a new menu onto the stack.  Automatically removes windows from the stack that are no longer valid
+	 * Invalid windows are those that are not parents of the window menu being pushed and any windows on the same stack level
+	 *
+	 * @param InParentWindow		The parent window of the menu
+	 * @param InContent				The menu's content
+	 * @param SummonLocation		Location where the menu should appear
+	 * @param TransitionEffect		Animation to use when the popup appears
+	 * @param bFocusImmediately		Should the popup steal focus when shown?
+	 * @param bShouldAutoSize		True if the new window should automatically size itself to its content
+	 * @param WindowSize			When bShouldAutoSize=false, this must be set to the size of the window to be created
+	 * @param SummonLocationSize	An optional size around the summon location which describes an area in which the menu may not appear
+	 */
+	TSharedRef<SWindow> PushMenu( const TSharedRef<SWindow>& InParentWindow, const TSharedRef<SWidget>& InContent, const FVector2D& SummonLocation, const FPopupTransitionEffect& TransitionEffect, const bool bFocusImmediately = true, const bool bShouldAutoSize = true, const FVector2D& WindowSize = FVector2D::ZeroVector, const FVector2D& SummonLocationSize = FVector2D::ZeroVector );
+
+	/**
+	 * Dismisses the menu stack up to a certian level (by default removes all levels in the stack)
+	 * Dismisses in reverse order so to keep just the first level you would pass in 1
+	 *
+	 * @param FirstStackIndexToRemove	The first level of the stack to remove
+	 */
+	void Dismiss( int32 FirstStackIndexToRemove = 0 );
+	
+	/**
+	 * Removes a window from the stack and all of its children.  Assumes the window is cleaned up elsewhere
+	 *
+	 * @param WindowToRemove	The window to remove (will remove all children of this window as well)	
+	 */
+	void RemoveWindow( TSharedRef<SWindow> WindowToRemove );
+
+	/**
+	 * Notifies the stack that a new window was activated.  The menu stack will be dismissed if the activated window is not in the menu stack
+	 *
+	 * @param ActivatedWindow	The window that was just activated
+	 */
+	void OnWindowActivated( TSharedRef<SWindow> ActivatedWindow );
+
+	/**
+	 * Finds a window in the stack
+	 *
+	 * @param WindowToFind	The window to look fort
+	 * @param The level of the stack where the window is located or INDEX_NONE if the window could not be found
+	 */
+	int32 FindLocationInStack( TSharedPtr<SWindow> WindowToFind ) const;
+
+	/**
+	 * @return	Returns the size of the menu stack
+	 */
+	int32 GetNumStackLevels( ) const;
+
+	/** @return Returns whether the window has child menus. */
+	bool HasOpenSubMenus( const TSharedRef<SWindow>& Window ) const;
+
+	/**
+	 * Returns all of the windows at the specified stack level index
+	 *
+	 * @param	StackLevelIndex	Index into the menu stack
+	 *
+	 * @return	List of windows at this stack level
+	 */
+	FMenuWindowList& GetWindowsAtStackLevel( const int32 StackLevelIndex );
+
+private:
+
+	/** The stack of windows, each level in the stack has an arbitrary number of siblings.  All siblings must belong to the same parent in the previous level */
+	TArray<FMenuWindowList> Stack;
+	
+	/** Handle to a throttle request made to ensure the menu is responsive in low FPS situations */
+	FThrottleRequest ThrottleHandle;
+};
