Slate: Refactored core Slate implementation into SlateCore module in preparation for UMG.

diff --git a/Engine/Source/Runtime/Slate/Public/Framework/Application/SlateApplication.h b/Engine/Source/Runtime/Slate/Public/Framework/Application/SlateApplication.h
new file mode 100644
--- /dev/null
+++ b/Engine/Source/Runtime/Slate/Public/Framework/Application/SlateApplication.h
@@ -0,0 +1,1423 @@
+// Copyright 1998-2014 Epic Games, Inc. All Rights Reserved.
+
+#pragma once
+
+#include "GenericApplicationMessageHandler.h"
+
+
+class SToolTip;
+class SViewport;
+class SWindow;
+
+
+/** A Delegate for querying whether source code access is possible */
+DECLARE_DELEGATE_RetVal(bool, FQueryAccessSourceCode);
+
+/** Delegates for when modal windows open or close */
+DECLARE_DELEGATE(FModalWindowStackStarted)
+DECLARE_DELEGATE(FModalWindowStackEnded)
+
+extern SLATE_API const FName NAME_UnrealOS;
+
+/**
+ * Design constraints for Slate applications
+ */
+namespace SlateApplicationDefs
+{
+	/** How much to amplify keyboard repeat rate by */
+	static const int32 NumRepeatsPerActualRepeat = 2;
+
+	/** How many users can we support at once? */
+	static const int32 MaxUsers = 8;
+}
+
+/** Allow widgets to find out when someone clicked outside them. Currently needed by MenuAnchros. */
+class FPopupSupport
+{
+	public:
+
+	/**
+	 * Given a WidgetPath that was clicked, send notifications to any subscribers that were not under the mouse.
+	 * i.e. Send the "Someone clicked outside me" notifications.
+	 */
+	void SendNotifications( const FWidgetPath& WidgetsUnderCursor );
+
+	private:
+
+	/** Only menu anchors currently need this service */
+	friend class SMenuAnchor;
+
+	/** A single subscription about clicks happening outside the widget. */
+	struct FClickSubscriber
+	{
+		FClickSubscriber( const TSharedRef<SWidget>& DetectClicksOutsideThisWidget, const FOnClickedOutside& InNotification )
+		: DetectClicksOutsideMe( DetectClicksOutsideThisWidget )
+		, Notification( InNotification )
+		{}
+
+		bool ShouldKeep() const
+		{
+			return DetectClicksOutsideMe.IsValid() && Notification.IsBound();
+		}
+
+		/** If a click occurs outside this widget, we'll send the notification */
+		TWeakPtr<SWidget> DetectClicksOutsideMe;
+		/** Notification to send */
+		FOnClickedOutside Notification;
+	};
+
+	/**
+	 * Register for a notification when the user clicks outside a specific widget.
+	 * 
+	 * @param NotifyWhenClickedOutsideMe    When the user clicks outside this widget, fire a notification.
+	 * @param InNotification                The notification to invoke.
+	 */
+	void RegisterClickNotification( const TSharedRef<SWidget>& NotifyWhenClickedOutsideMe, const FOnClickedOutside& InNotification );
+
+	/**
+	 * NOTE: Only necessary if notification no longer desired.
+	 *       Stale notifications are cleaned up automatically.
+	 *       
+	 * Unregister the notification because it is no longer desired.
+	 */
+	void UnregisterClickNotification( const FOnClickedOutside& InNotification );
+
+	/** List of subscriptions that want to be notified when the user clicks outside a certain widget. */
+	TArray<FClickSubscriber> ClickZoneNotifications;
+
+};
+
+
+class SLATE_API FSlateApplication
+	: public FSlateApplicationBase
+	, public FGenericApplicationMessageHandler
+{
+public:
+
+	/**
+	 * Destructor.
+	 */
+	virtual ~FSlateApplication( );
+
+public:
+
+	/**
+	 * Returns the running average delta time (smoothed over several frames)
+	 *
+	 * @return  The average time delta
+	 */
+	const float GetAverageDeltaTime() const
+	{
+		return AverageDeltaTime;
+	}
+
+	/**
+	 * Returns the real time delta since Slate last ticked widgets
+	 *
+	 * @return  The time delta since last tick
+	 */
+	const float GetDeltaTime() const
+	{
+		return (float)( CurrentTime - LastTickTime );
+	}
+
+	/**
+	 * Returns the running average delta time (smoothed over several frames)
+	 * Unlike GetAverageDeltaTime() it excludes exceptional
+	 * situations, such as when throttling mode is active.
+	 *
+	 * @return  The average time delta
+	 */
+	float GetAverageDeltaTimeForResponsiveness() const
+	{
+		return AverageDeltaTimeForResponsiveness;
+	}
+
+public:
+
+	static void Create();
+
+	static void InitializeAsStandaloneApplication( const TSharedRef< class FSlateRenderer >& PlatformRenderer);
+
+	/**
+	 * Returns true if a Slate application instance is currently initialized and ready
+	 *
+	 * @return  True if Slate application is initialized
+	 */
+	static bool IsInitialized()
+	{
+		return CurrentApplication.IsValid();
+	}
+
+	/**
+	 * Returns the current instance of the application. The application should have been initialized before
+	 * this method is called
+	 *
+	 * @return  Reference to the application
+	 */
+	static FSlateApplication& Get()
+	{
+		check( IsInGameThread() );
+		return *CurrentApplication;
+	}
+
+	static void Shutdown()
+	{
+		if ( FSlateApplication::IsInitialized() )
+		{
+			CurrentApplication->OnShutdown();
+			CurrentApplication->DestroyRenderer();
+			CurrentApplication->Renderer.Reset();
+			PlatformApplication.Reset();
+			CurrentApplication.Reset();
+			CurrentBaseApplication.Reset();
+		}
+	}
+
+	/** @return the global tab manager */
+	static TSharedRef<class FGlobalTabmanager> GetGlobalTabManager();
+
+	/** @return the root style node, which is the entry point to the style graph representing all the current style rules. */
+	const class FStyleNode* GetRootStyle() const;
+
+	/**
+	 * Initializes the renderer responsible for drawing all elements in this application
+	 *
+	 * @param InRenderer The renderer to use.
+	 */
+	virtual void InitializeRenderer( TSharedRef<FSlateRenderer> InRenderer );
+
+	/** Set the slate sound provider that the slate app should use. */
+	virtual void InitializeSound( const TSharedRef<ISlateSoundDevice>& InSlateSoundDevice );
+
+	void DestroyRenderer();
+
+	/** Play SoundToPlay. Interrupt previous sound if one is playing. */
+	void PlaySound( const FSlateSound& SoundToPlay, int32 UserIndex = 0 ) const;
+
+	/** @return The duration of the given sound resource */
+	float GetSoundDuration(const FSlateSound& Sound) const;
+
+	/** @return The force feedback interface for this application */
+	IForceFeedbackSystem *GetForceFeedbackSystem() const { return PlatformApplication->GetForceFeedbackSystem(); }
+
+	/** @return Whether or not the current platform supports system help */
+	bool SupportsSystemHelp() const { return PlatformApplication->SupportsSystemHelp(); }
+
+	void ShowSystemHelp() { PlatformApplication->ShowSystemHelp(); }
+
+	/** @return The text input method interface for this application */
+	ITextInputMethodSystem *GetTextInputMethodSystem() const { return PlatformApplication->GetTextInputMethodSystem(); }
+
+	/** 
+	 * Sets the position of the cursor.
+	 *
+	 * @param MouseCoordinate		The new position.
+	 */
+	void SetCursorPos( const FVector2D& MouseCoordinate );
+
+	/**
+	 * Reorders an array of windows so the specified window is "brought to the front"
+	 */
+	static void ArrangeWindowToFront( TArray< TSharedRef<SWindow> >& Windows, const TSharedRef<SWindow>& WindowToBringToFront );
+
+	/**
+	 * Polls game devices for input
+	 */
+	void PollGameDeviceState();
+
+	/**
+	 * Ticks this application
+	 */
+	void Tick();
+
+	/**
+	 * Pumps OS messages when a modal window or intra-frame debugging session exists
+	 */
+	void PumpMessages();
+
+	/**
+	 * Returns true if this slate application is ready to open modal windows
+	 */
+	bool CanAddModalWindow() const;
+
+	/**
+	 * Returns true if this slate application is ready to display windows.
+	 */
+	bool CanDisplayWindows() const;
+	
+	/**
+	 * Adds a modal window to the application.  
+	 * In most cases, this function does not return until the modal window is closed (the only exception is a modal window for slow tasks)  
+	 *
+	 * @param InSlateWindow		A SlateWindow to which to add a native window.
+	 * @param InParentWindow	The parent of the modal window.  All modal windows must have a parent.
+	 * @param bSlowTaskWindow	true if the window is for a slow task and this function should return before the window is closed
+	 */
+	void AddModalWindow( TSharedRef<SWindow> InSlateWindow, const TSharedPtr<const SWidget> InParentWidget, bool bSlowTaskWindow = false );
+
+	/** Sets the delegate for when a modal window stack begins */
+	void SetModalWindowStackStartedDelegate(FModalWindowStackStarted StackStartedDelegate);
+
+	/** Sets the delegate for when a modal window stack ends */
+	void SetModalWindowStackEndedDelegate(FModalWindowStackEnded StackEndedDelegate);
+
+	/**
+	 * Associates a top level Slate Window with a native window, and "natively" parents that window to the specified Slate window.
+	 * Although the window still a top level window in Slate, it will be considered a child window to the operating system.
+	 *
+	 * @param	InSlateWindow		A Slate window to which to add a native window.
+	 * @param	InParentWindow		Slate window that the window being added should be a native child of
+	 * @param	bShowImmediately	True to show the window.  Pass false if you're going to call ShowWindow() yourself later.
+	 *
+	 * @return a reference to the SWindow that was just added.
+	 */
+	TSharedRef<SWindow> AddWindowAsNativeChild( TSharedRef<SWindow> InSlateWindow, TSharedRef<SWindow> InParentWindow, const bool bShowImmediately = true );
+
+	/**
+	 * Creates a new Menu window and add it to the menu stack.
+	 *
+	 * @param InWindow				The parent of the menu.  If there is already an open menu this parent must exist in the menu stack or the menu stack is dismissed and a new one started
+	 * @param Content				The content to be placed inside the new window
+	 * @param SummonLocation		The location where this menu should be summoned
+	 * @param TransitionEffect		Animation to use when the popup appears
+	 * @param bFocusImmediately		Should the popup steal focus when shown?
+	 * @param bShouldAutoSize		True if the newCalculatePopupWindowPosition window should automatically size itself to its content
+	 * @param WindowSize			When bShouldAutoSize=false, this must be set to the size of the window to be created
+	 * @param SummonLocationSize	An optional rect which describes an area in which the menu may not appear
+	 */
+	TSharedRef<SWindow> PushMenu( const TSharedRef<SWidget>& InParentContent, const TSharedRef<SWidget>& InContent, const FVector2D& SummonLocation, const FPopupTransitionEffect& TransitionEffect, const bool bFocusImmediately = true, const bool bShouldAutoSize = true, const FVector2D& WindowSize = FVector2D::ZeroVector, const FVector2D& SummonLocationSize = FVector2D::ZeroVector );
+
+	/** @return Returns whether the window has child menus. */
+	bool HasOpenSubMenus( TSharedRef<SWindow> Window ) const;
+
+	/** @return	Returns true if there are any pop-up menus summoned */
+	bool AnyMenusVisible() const;
+
+	/**
+	 * Dismisses any open menus
+	 */
+	void DismissAllMenus();
+
+	/**
+	 * Dismisses a menu window and all its children
+	 *
+	 * @param MenuWindowToDismiss	The window to dismiss, any children, grandchildren etc will also be dismissed
+	 */
+	void DismissMenu( TSharedRef<SWindow> MenuWindowToDismiss );
+
+	/**
+	 * Finds the window in the menu stack
+	 *
+	 * @param WindowToFind	The window to look for
+	 * @return The level in the stack  that the window is in or INDEX_NONE if it is not found
+	 */
+	int32 GetLocationInMenuStack( TSharedRef<SWindow> WindowToFind ) const;
+
+	/**
+	 * HACK: Don't use this unless shutting down a game viewport
+	 * Game viewport windows need to be destroyed instantly or else the viewport could tick and access deleted data
+	 *
+	 * @param WindowToDestroy		Window for destruction
+	 */
+	void DestroyWindowImmediately( TSharedRef<SWindow> WindowToDestroy );
+
+	/**
+	 * Disable Slate components when an external, non-slate, modal window is brought up.  In the case of multiple
+	 * external modal windows, we will only increment our tracking counter.
+	 */
+	void ExternalModalStart( );
+
+	/**
+	 * Re-enable disabled Slate components when a non-slate modal window is dismissed.  Slate components
+	 * will only be re-enabled when all tracked external modal windows have been dismissed.
+	 */
+	void ExternalModalStop( );
+
+	/** 
+	 * Removes references to FViewportRHI's.  
+	 * This has to be done explicitly instead of using the FRenderResource mechanism because FViewportRHI's are managed by the game thread.
+	 * This is needed before destroying the RHI device. 
+	 */
+	void InvalidateAllViewports();
+
+	/**
+	 * Registers a game viewport with the Slate application so that specific messages can be routed directly to a viewport
+	 * 
+	 * @param InViewport	The viewport to register.  Note there is currently only one registered viewport
+	 */
+	void RegisterGameViewport( TSharedRef<SViewport> InViewport );
+
+	TSharedPtr<SViewport> GetGameViewport() const;
+
+	/** 
+	 * Unregisters the current game viewport from Slate.  This method sends a final deactivation message to the viewport
+	 * to allow it to do a final cleanup before being closed.
+	 */
+	void UnregisterGameViewport();
+
+	/**
+	 *	Sets the current focus to the SWidget representing the currently active game viewport
+	 */
+	void SetFocusToGameViewport();
+
+	/**
+	 *	Forces the currently active game viewport to capture the joystick
+	 */
+	void SetJoystickCaptorToGameViewport();
+
+	/**
+	 * 	Sets the current focus to the specified SWidget
+	 */
+	void SetKeyboardFocus( const TSharedPtr< SWidget >& OptionalWidgetToFocus, EKeyboardFocusCause::Type ReasonFocusIsChanging = EKeyboardFocusCause::SetDirectly );
+	
+	/**
+	 * Returns the current modifier keys state
+	 *
+	 * @return  State of modifier keys
+	 */
+	FModifierKeysState GetModifierKeys() const;
+
+	/**
+	 *	Restores all input settings to their original values
+	 * 
+	 *  Shows the mouse, clears any locks and captures, turns off high precision input
+	 */
+	void ResetToDefaultInputSettings();
+
+	/**
+	 * Mouse capture
+	 */
+
+	/** returning platform-specific value designating window that captures mouse, or NULL if mouse isn't captured */
+	virtual void* GetMouseCaptureWindow( void ) const;
+
+
+	/** Releases the mouse capture from whatever it currently is on. */
+	void ReleaseMouseCapture();
+
+	/** @return a pointer to the Widget that currently captures the joystick; Empty pointer when the joystick is not captured. */
+	TSharedPtr< SWidget > GetJoystickCaptor(uint32 UserIndex) const;
+
+	/** Releases the joystick capture from whatever it currently is on. */
+	void ReleaseJoystickCapture(uint32 UserIndex);
+
+	/** @return The active modal window or NULL if there is no modal window. */
+	TSharedPtr<SWindow> GetActiveModalWindow() const;
+
+	/**
+	 * Keyboard focus
+	 */
+
+	/**
+	 * Clears keyboard focus, if any widget is currently focused
+	 *
+	 * @param  InCause  The reason that keyboard focus is changing
+	 */
+	void ClearKeyboardFocus( const EKeyboardFocusCause::Type InCause );
+
+	/**
+	 * Assign a delegate to be called when this application is requesting an exit (e.g. when the last window is closed).
+	 *
+	 * @param OnExitRequestHandler  Function to execute when the application wants to quit
+	 */
+	void SetExitRequestedHandler( const FSimpleDelegate& OnExitRequestedHandler );
+		
+	/**
+	 * @todo slate: Remove this method or make it private.
+	 * Searches for the specified widget and generates a full path to it.  Note that this is
+	 * a relatively slow operation!
+	 * 
+	 * @param  InWidget       Widget to generate a path to
+	 * @param  OutWidgetPath  The generated widget path
+	 * @param  VisibilityFilter	Widgets must have this type of visibility to be included the path
+	 *
+	 * @return	True if the widget path was found
+	 */
+	static bool FindPathToWidget( const TArray< TSharedRef<SWindow> > WindowsToSearch, TSharedRef< const SWidget > InWidget, FWidgetPath& OutWidgetPath, EVisibility VisibilityFilter = EVisibility::Visible );
+
+	/**
+	 * @todo slate: Remove this method or make it private.
+	 * Searches for the specified widget and generates a full path to it.  Note that this is
+	 * a relatively slow operation!  It can fail, in which case OutWidgetPath.IsValid() will be false.
+	 * 
+	 * @param  InWidget       Widget to generate a path to
+	 * @param  OutWidgetPath  The generated widget path
+	 * @param  VisibilityFilter	Widgets must have this type of visibility to be included the path
+	 */
+	bool GeneratePathToWidgetUnchecked( TSharedRef< const SWidget > InWidget, FWidgetPath& OutWidgetPath, EVisibility VisibilityFilter = EVisibility::Visible ) const;
+	
+	/**
+	 * @todo slate: Remove this method or make it private.
+	 * Searches for the specified widget and generates a full path to it.  Note that this is
+	 * a relatively slow operation!  Asserts if the widget isn't found.
+	 * 
+	 * @param  InWidget       Widget to generate a path to
+	 * @param  OutWidgetPath  The generated widget path
+	 * @param  VisibilityFilter	Widgets must have this type of visibility to be included the path
+	 */
+	void GeneratePathToWidgetChecked( TSharedRef< const SWidget > InWidget, FWidgetPath& OutWidgetPath, EVisibility VisibilityFilter = EVisibility::Visible ) const;
+	
+	/**
+	 * Finds the window that the provided widget resides in
+	 * 
+	 * @param InWidget		The widget to find the window for
+	 *
+	 * @return The window where the widget resides, or null if the widget wasn't found.  Remember, a widget might not be found simply because its parent decided not to report the widget in ArrangeChildren.
+	 */
+	TSharedPtr<SWindow> FindWidgetWindow( TSharedRef< const SWidget > InWidget ) const;
+
+	/**
+	 * Finds the window that the provided widget resides in
+	 * 
+	 * @param InWidget		The widget to find the window for
+	 * @param OutWidgetPath Full widget path generated 
+	 *
+	 * @return The window where the widget resides, or null if the widget wasn't found.  Remember, a widget might not be found simply because its parent decided not to report the widget in ArrangeChildren.
+	 */
+	TSharedPtr<SWindow> FindWidgetWindow( TSharedRef< const SWidget > InWidget, FWidgetPath& OutWidgetPath) const;
+
+	/**
+	 * @return True if the application is currently routing high precision mouse movement events (OS specific)
+	 * If this value is true, the mouse is captured and hidden by the widget that originally made the request.
+	 */
+	bool IsUsingHighPrecisionMouseMovment() const { return PlatformApplication.IsValid() ? PlatformApplication->IsUsingHighPrecisionMouseMode() : false; }
+	
+	/**
+	 * @return True if the last mouse event was from a trackpad.
+	 */
+	bool IsUsingTrackpad() const { return PlatformApplication.IsValid() ? PlatformApplication->IsUsingTrackpad() : false; }
+
+	/**
+	 * Sets the widget reflector.
+	 *
+	 * @param WidgetReflector The widget reflector to set.
+	 */
+	void SetWidgetReflector( const TSharedRef<IWidgetReflector>& WidgetReflector )
+	{
+		if (SourceCodeAccessDelegate.IsBound())
+		{
+			WidgetReflector->SetSourceAccessDelegate(SourceCodeAccessDelegate);
+		}
+
+		WidgetReflectorPtr = WidgetReflector;
+	}
+
+	/** @param AccessDelegate The delegate to pass along to the widget reflector */
+	void SetWidgetReflectorSourceAccessDelegate(FAccessSourceCode AccessDelegate)
+	{
+		SourceCodeAccessDelegate = AccessDelegate;
+	}
+
+	/** @param QueryAccessDelegate The delegate to pass along to the widget reflector */
+	void SetWidgetReflectorQuerySourceAccessDelegate(FQueryAccessSourceCode QueryAccessDelegate)
+	{
+		QuerySourceCodeAccessDelegate = QueryAccessDelegate;
+	}
+
+	/** @param Scale  Sets the ratio SlateUnit / ScreenPixel */
+	void SetApplicationScale( float InScale ){ Scale = InScale; }
+
+	virtual void GetInitialDisplayMetrics( FDisplayMetrics& OutDisplayMetrics ) const { PlatformApplication->GetInitialDisplayMetrics( OutDisplayMetrics ); }
+
+	virtual const TSharedPtr<GenericApplication> GetPlatformApplication ( ) const { return PlatformApplication; }
+
+	/** Are we drag-dropping right now? */
+	bool IsDragDropping() const;
+
+	/** Get the current drag-dropping content */
+	TSharedPtr<class FDragDropOperation> GetDragDroppingContent() const;
+
+	/** End any in flight drag and drops */
+	void EndDragDrop();
+
+	/**
+	 * Returns the attribute that can be used by widgets to check if the application is in normal execution mode
+	 * Don't hold a reference to this anywhere that can exist when this application closes
+	 */
+	const TAttribute<bool>& GetNormalExecutionAttribute() const { return NormalExecutionGetter; }
+
+	/**
+	 * @return true if not in debugging mode
+	 */
+	bool IsNormalExecution() const { return !GIntraFrameDebuggingGameThread; }
+
+	/**
+	 * @return true if in debugging mode
+	 */
+	bool InKismetDebuggingMode() const { return GIntraFrameDebuggingGameThread; }
+
+	/**
+	 * Enters debugging mode which is a special state that causes the
+	 * Slate application to tick in place which in the middle of a stack frame
+	 */
+	void EnterDebuggingMode();
+
+	/**
+	 * Leaves debugging mode
+	 *
+	 * @param bLeavingDebugForSingleStep	Whether or not we are leaving debug mode due to single stepping
+	 */
+	void LeaveDebuggingMode( bool bLeavingDebugForSingleStep = false );
+	
+	/**
+	 * Calculates the popup window position from the passed in window position and size. 
+	 * Adjusts position for popup windows which are outside of the work area of the monitor where they reside
+	 *
+	 * @param InAnchor	The current(suggseted) window position and size of an area which may not be covered by the popup.
+	 * @param InSize		The size of the window
+	 * @param Orientation	The direction of the popup.
+	 *						If vertical it will attempt to open below the anchor but will open above if there is no room.
+	 *						If horizontal it will attempt to open below the anchor but will open above if there is no room.
+	 * @return The adjusted position
+	 */
+	virtual FVector2D CalculatePopupWindowPosition( const FSlateRect& InAnchor, const FVector2D& InSize, const EOrientation Orientation = Orient_Vertical ) const;
+
+
+	/**
+	 * Is the window in the app's destroy queue? If so it will be destroyed next tick.
+	 *
+	 * @param Window		The window to find in the destroy list
+	 * @return				true if Window is in the destroy queue
+	 */
+	bool IsWindowInDestroyQueue(TSharedRef<SWindow> Window) const;
+
+	/** @return	Returns true if the application's average frame rate is at least as high as our target frame rate that satisfies our requirement for a smooth and responsive UI experience */
+	bool IsRunningAtTargetFrameRate() const;
+
+	/** @return Returns true if transition effects for new menus and windows should be played */
+	bool AreMenuAnimationsEnabled() const;
+
+	/** 
+	 * Sets whether transition effects for new menus and windows should be played.  This can be called at any time.
+	 *
+	 * @param	bEnableAnimations	True if animations should be used, otherwise false.
+	 */
+	void EnableMenuAnimations( const bool bEnableAnimations );
+
+	/** Set the global application icon */
+	void SetAppIcon(const FSlateBrush* const InAppIcon);
+
+	/** Sets the display state of external UI such as Steam. */
+	void ExternalUIChange(bool bIsOpening)
+	{
+		bIsExternalUIOpened = bIsOpening;
+	}
+
+	/**
+	 * Shows or hides an onscreen keyboard
+	 *
+	 * @param bShow	true to show the keyboard, false to hide it
+	 */
+	void ShowKeyboard( bool bShow, TSharedPtr<SVirtualKeyboardEntry> TextEntryWidget = NULL );
+
+	/** Get the work area that has the largest intersection with the specified rectangle */
+	FSlateRect GetWorkArea( const FSlateRect& InRect ) const;
+
+	/**
+	 * Shows or hides an onscreen keyboard
+	 *
+	 * @param bShow	true to show the keyboard, false to hide it
+	 */
+	virtual void NativeApp_ShowKeyboard( bool bShow, FString InitialString = "", int32 SelectionStart = -1, int32 SelectionEnd = -1 )
+	{
+		// empty default functionality
+	}
+
+	/** @return true if the current platform allows source code access */
+	bool SupportsSourceAccess() const;
+
+	/** Opens the current platform's code editing IDE (if necessary) and focuses the specified line in the specified file. Will only work if SupportsSourceAccess() returns true */
+	void GotoLineInSource(const FString& FileName, int32 LineNumber) const;
+
+	/** @return Service that supports popups; helps with auto-hiding of popups */
+	FPopupSupport& GetPopupSupport() { return PopupSupport; }
+
+	/**
+	 * Forces the window to redraw immediately.
+	 */
+	void ForceRedrawWindow( TSharedRef<SWindow>& InWindowToDraw );
+
+	/**
+	 * Takes a screenshot of the widget writing the results into the color buffer provided.  Note that the format is BGRA.
+	 * the size of the resulting image is also output.
+	 * 
+	 * @return true if taking the screenshot was successful.
+	 */
+	bool TakeScreenshot(TSharedRef<SWidget>& Widget, TArray<FColor>&OutColorData, FIntVector& OutSize);
+
+	/**
+	 * Takes a screenshot of the widget writing the results into the color buffer provided, this version allows you to provide 
+	 * an inner area to screenshot.  Note that the format is BGRA.  The size of the resulting image is also output.
+	 *
+	 * @return true if taking the screenshot was successful.
+	 */
+	bool TakeScreenshot(TSharedRef<SWidget>& Widget, const FIntRect& InnerWidgetArea, TArray<FColor>& OutColorData, FIntVector& OutSize);
+
+protected:
+
+	/** 
+	 * Ticks a slate window and all of its children
+	 *
+	 * @param WindowToTick	The window to tick
+	 */
+	void TickWindowAndChildren( TSharedRef<SWindow> WindowToTick );
+
+	/**
+	 * Draws Slate windows.  Should only be called by the application's main loop or renderer.
+	 */
+	void DrawWindows();
+
+	/**
+	 * Draws slate windows, optionally only drawing the passed in window
+	 */
+	void PrivateDrawWindows( TSharedPtr<SWindow> DrawOnlyThisWindow = NULL );
+
+	/**
+	 * Prepass step before drawing windows to compute geometry size and reshape autosized windows
+	 */
+	void DrawPrepass( TSharedPtr<SWindow> DrawOnlyThisWindow );
+
+	/**
+	 * Draws a window and its children
+	 */
+	void DrawWindowAndChildren( const TSharedRef<SWindow>& WindowToDraw, struct FDrawWindowArgs& DrawWindowArgs );
+
+	/**
+	 * Gets all visible child windows of a window.
+	 */
+	void GetAllVisibleChildWindows(TArray< TSharedRef<SWindow> >& OutWindows, TSharedRef<SWindow> CurrentWindow);
+
+
+	/** Engages or disengages application throttling based on user behavior */
+	void ThrottleApplicationBasedOnMouseMovement();
+
+	/**
+	 * Allow the Remote Slate (aka UDKRemote) server class to have time
+	 */
+	void TickRemoteSlateServer();
+
+protected:
+
+	/**
+	 * Called by the native application in response to a mouse move. Routs the event to Slate Widgets.
+	 *
+	 * @param  InMouseEvent  Mouse event
+	 *
+	 * @return  Was this event handled by the Slate application?
+	 */
+	bool ProcessMouseMoveMessage( FPointerEvent& MouseEvent );
+
+	/**
+	 * Called by the native application in response to a mouse button press. Routs the event to Slate Widgets.
+	 *
+	 * @param  InMouseEvent  Mouse event
+	 *
+	 * @return  Was this event handled by the Slate application?
+	 */
+	bool ProcessMouseButtonDownMessage( const TSharedPtr< FGenericWindow >& PlatformWindow, FPointerEvent& MouseEvent );
+
+	/**
+	 * Called by the native application in response to a mouse button release. Routs the event to Slate Widgets.
+	 *
+	 * @param  InMouseEvent  Mouse event
+	 *
+	 * @return  Was this event handled by the Slate application?
+	 */
+	bool ProcessMouseButtonUpMessage( FPointerEvent& MouseEvent );
+
+
+	/**
+	 * Called by the native application in response to a mouse release. Routs the event to Slate Widgets.
+	 *
+	 * @param  InMouseEvent  Mouse event
+	 *
+	 * @return  Was this event handled by the Slate application?
+	 */
+	bool ProcessMouseButtonDoubleClickMessage( const TSharedPtr< FGenericWindow >& PlatformWindow, FPointerEvent& InMouseEvent );
+	
+	/**
+	 * Called by the native application in response to a mouse wheel spin or a touch gesture. Routs the event to Slate Widgets.
+	 *
+	 * @param  InWheelEvent    Mouse wheel event details
+	 * @param  InGestureEvent  Optional gesture event details
+	 *
+	 * @return  Was this event handled by the Slate application?
+	 */
+	bool ProcessMouseWheelOrGestureMessage( FPointerEvent& InWheelEvent, const FPointerEvent* InGestureEvent );
+
+	/**
+	 * Called when a character is entered
+	 *
+	 * @param  InCharacterEvent  Character event
+	 *
+	 * @return  Was this event handled by the Slate application?
+	 */
+	bool ProcessKeyCharMessage( FCharacterEvent& InCharacterEvent );
+
+	/**
+	 * Called when a key is pressed
+	 *
+	 * @param  InKeyboardEvent  Keyboard event
+	 *
+	 * @return  Was this event handled by the Slate application?
+	 */
+	bool ProcessKeyDownMessage( FKeyboardEvent& InKeyboardEvent );
+
+	/**
+	 * Called when a key is released
+	 *
+	 * @param  InKeyboardEvent  Keyboard event
+	 *
+	 * @return  Was this event handled by the Slate application?
+	 */
+	bool ProcessKeyUpMessage( FKeyboardEvent& InKeyboardEvent );
+	
+public:
+
+	/**
+	 * Called by the native application in response to a set cursor event.
+	 */
+	void OnSetCursorMessage();
+
+	/**
+	 * Called when a drag from an external (non-slate) source enters a window
+	 *
+	 * @param WindowEntered  The window that was entered by the drag and drop
+	 * @param DragDropEvent  Describes the mouse state (position, pressed buttons, etc) and associated payload
+	 *
+	 * @return true if the drag enter was handled and can be processed by some widget in this window; false otherwise
+	 */
+	bool ProcessDragEnterMessage( TSharedRef<SWindow> WindowEntered, FDragDropEvent& DragDropEvent );
+
+	/**
+	 * Called when a drag from an external (non-slate) source leaves a window
+	 *
+	 * @param WindowLeft     The window that was left by the drag and drop
+	 * @param DragDropEvent  Describes the mouse state (position, pressed buttons, etc) and associated payload
+	 *
+	 * @return true if the drag enter was handled and can be processed by some widget in this window; false otherwise
+	 */
+	void OnDragLeaveMessage();
+
+	/**
+	 * Called when a drag from an external (non-slate) is dropped onto a window
+	 *
+	 * @param WindowDroppedInto  The window that was dropped into
+	 * @param DragDropEvent  Describes the mouse state (position, pressed buttons, etc) and associated payload
+	 *
+	 * @return true if the drag enter was handled and can be processed by some widget in this window; false otherwise
+	 */
+	bool OnDropMessage( TSharedRef<SWindow> WindowDroppedInto, FDragDropEvent& DragDropEvent );
+
+	/**
+	 * Called when a controller button is found pressed when polling game device state
+	 * 
+	 * @param ControllerEvent	The controller event generated
+	 */
+	void OnControllerButtonPressedMessage( FControllerEvent& ControllerEvent );
+
+	/**
+	 * Called when a controller button is found released when polling game device state
+	 * 
+	 * @param ControllerEvent	The controller event generated
+	 */
+	void OnControllerButtonReleasedMessage( FControllerEvent& ControllerEvent );
+
+	/**
+	 * Called when a controller analog value has changed while polling game device state
+	 * 
+	 * @param ControllerEvent	The controller event generated
+	 */
+	void OnControllerAnalogValueChangedMessage( FControllerEvent& ControllerEvent );
+
+	/**
+	 * Called when a touchpad touch is started (finger down) when polling game device state
+	 * 
+	 * @param ControllerEvent	The touch event generated
+	 */
+	void OnTouchStartedMessage( const TSharedPtr< FGenericWindow >& PlatformWindow, FPointerEvent& InTouchEvent );
+
+	/**
+	 * Called when a touchpad touch is moved  (finger moved) when polling game device state
+	 * 
+	 * @param ControllerEvent	The touch event generated
+	 */
+	void OnTouchMovedMessage( FPointerEvent& InTouchEvent );
+
+	/**
+	 * Called when a touchpad touch is ended (finger lifted) when polling game device state
+	 * 
+	 * @param ControllerEvent	The touch event generated
+	 */
+	void OnTouchEndedMessage( FPointerEvent& InTouchEvent );
+
+	/**
+	 * Called when motion is detected (controller or device) when polling game device state
+	 * 
+	 * @param MotionEvent		The motion event generated
+	 */
+	void OnMotionDetectedMessage( FMotionEvent& InMotionEvent );
+
+public:
+	/**
+	 * Called by the native application in response to an activation or deactivation. 
+	 *
+	 * @param ActivateEvent Information about the window activation/deactivation
+	 *
+	 * @return  Was this event handled by the Slate application?
+	 */
+	bool ProcessWindowActivatedMessage( const FWindowActivateEvent& ActivateEvent );
+	
+	/** Called when the application is activated (i.e. one of its windows becomes active) or deactivated. */
+	void ProcessApplicationActivationMessage( bool InAppActivated );
+
+
+
+	/** Called when the slate application is being shut down. */
+	void OnShutdown();
+
+	/**
+	 * Destroy the native and slate windows in the array provided.
+	 *
+	 * @param WindowsToDestroy   Destroy these windows
+	 */
+	void DestroyWindowsImmediately();
+
+	/**
+	 * Apply any requests form the Reply to the application. E.g. Capture mouse
+	 *
+	 * @param CurrentEventPath   The WidgetPath along which the reply-generated event was routed
+	 * @param TheReply           The reply generated by an event that was being processed.
+	 * @param WidgetsUnderMouse  Optional widgets currently under the mouse; initiating drag and drop needs access to widgets under the mouse.
+	 * @param InMouseEvent       Optional mouse event that caused this action.
+	 * @param UserIndex			 User index that generated the event we are replying to (defaults to 0, at least for now)
+	 */
+	void ProcessReply( const FWidgetPath& CurrentEventPath, const FReply TheReply, const FWidgetPath* WidgetsUnderMouse, const FPointerEvent* InMouseEvent, uint32 UserIndex=0 );
+	
+	void LockCursor( const TSharedPtr<SWidget>& Widget );
+
+	/**
+	 * Bubble a request for which cursor to display for widgets under the mouse or the widget that captured the mouse.
+	 */
+	void QueryCursor();
+	
+	/**
+	 * Spawns a tool tip window.  If an existing tool tip window is open, it will be dismissed first.
+	 *
+	 * @param	InToolTip           Widget to display.
+	 * @param	InSpawnLocation     Screen space location to show the tool tip (window top left)
+	 */
+	void SpawnToolTip( const TSharedRef<IToolTip>& InToolTip, const FVector2D& InSpawnLocation );
+
+	/**
+	 * Closes the open tool-tip, if a tool-tip is open
+	 */
+	void CloseToolTip();
+
+	/**
+	 * Updates tool tip
+	 */
+	void UpdateToolTip( bool AllowSpawningOfNewToolTips );
+
+	/**
+	 * Creates a mouse move event for the last known cursor position.  This should be called every tick to make
+	 * sure that widgets that appear (or vanish from) underneath the cursor have hover state set appropriately.
+	 */
+	virtual void SynthesizeMouseMove();
+
+	/**
+	 * Draw the keyboard focus rectangle over the currently focused widget
+	 *
+	 * @param FocusPath             Path to the focused widget.
+	 * @param WindowElementList     List of draw elements to which we will add our focus rectangle
+	 * @param InLayerId             Layer onto which we should paint the rectangle
+	 *
+	 * @return New maximum layer id
+	 */
+	int32 DrawKeyboardFocus( const FWidgetPath& FocusPath, class FSlateWindowElementList& WindowElementList, int32 InLayerId ) const;
+
+	/** @return an array of top-level windows that can be interacted with. e.g. when a modal window is up, only return the modal window */
+	TArray< TSharedRef<SWindow> > GetInteractiveTopLevelWindows();
+
+	/** Gets all visible slate windows ordered from back to front based on child hierarchies */
+	void GetAllVisibleWindowsOrdered(TArray< TSharedRef<SWindow> >& OutWindows);
+	
+	/** Tell the slate application to log a string to it's ui log */
+	void OnLogSlateEvent(enum EEventLog::Type Event, const FString& AdditionalContent = FString());
+
+	/** Tell the slate application to log a FText to it's ui log */
+	void OnLogSlateEvent(enum EEventLog::Type Event, const FText& AdditionalContent );
+
+	/** Sets the slate event logger */
+	void SetSlateUILogger(TSharedPtr<class IEventLogger> InEventLogger = TSharedPtr<class IEventLogger>());
+
+	/** @return true if mouse events are being turned into touch events, and touch UI should be forced on */
+	bool IsFakingTouchEvents() const;
+
+	void SetGameIsFakingTouchEvents(const bool bIsFaking);
+
+	/** Sets the handler for otherwise unhandled key down events. This is used by the editor to provide a global action list, if the key was not consumed by any widget. */
+	void SetUnhandledKeyDownEventHandler( const FOnKeyboardEvent& NewHandler );
+
+public:
+
+	// Begin FSlateApplicationBase interface
+
+	virtual TSharedRef<SWindow> AddWindow( TSharedRef<SWindow> InSlateWindow, const bool bShowImmediately = true ) OVERRIDE;
+
+	virtual void ArrangeWindowToFrontVirtual( TArray<TSharedRef<SWindow>>& Windows, const TSharedRef<SWindow>& WindowToBringToFront ) OVERRIDE
+	{
+		FSlateApplication::ArrangeWindowToFront(Windows, WindowToBringToFront);
+	}
+
+	virtual bool FindPathToWidgetVirtual( const TArray<TSharedRef<SWindow>> WindowsToSearch, TSharedRef<const SWidget> InWidget, FWidgetPath& OutWidgetPath, EVisibility VisibilityFilter = EVisibility::Visible ) OVERRIDE
+	{
+		return FSlateApplication::FindPathToWidget(WindowsToSearch, InWidget, OutWidgetPath, VisibilityFilter);
+	}
+
+	virtual const double GetCurrentTime( ) const OVERRIDE
+	{
+		return CurrentTime;
+	}
+
+	virtual TSharedPtr<SWindow> GetActiveTopLevelWindow( ) const OVERRIDE;
+
+	virtual const FSlateBrush* GetAppIcon( ) const OVERRIDE;
+
+	virtual float GetApplicationScale( ) const OVERRIDE
+	{
+		return Scale;
+	}
+
+	virtual FVector2D GetCursorPos( ) const OVERRIDE;
+
+	virtual FVector2D GetCursorSize( ) const OVERRIDE;
+
+	virtual void GetDisplayMetrics( FDisplayMetrics& OutDisplayMetrics ) const OVERRIDE
+	{
+		PlatformApplication->GetDisplayMetrics( OutDisplayMetrics );
+	}
+
+	virtual TSharedPtr<SWidget> GetKeyboardFocusedWidget( ) const OVERRIDE;
+
+	virtual TSharedPtr<SWidget> GetMouseCaptor( ) const OVERRIDE;
+
+	virtual FSlateRect GetPreferredWorkArea( ) const OVERRIDE;
+
+	virtual bool HasFocusedDescendants( const TSharedRef<const SWidget>& Widget ) const OVERRIDE;
+
+	virtual bool IsExternalUIOpened() OVERRIDE
+	{
+		return bIsExternalUIOpened;
+	}
+
+	virtual FWidgetPath LocateWindowUnderMouse( FVector2D ScreenspaceMouseCoordinate, const TArray<TSharedRef<SWindow>>& Windows, bool bIgnoreEnabledStatus = false ) OVERRIDE;
+
+	virtual TSharedRef<SWidget> MakeImage( const TAttribute<const FSlateBrush*>& Image, const TAttribute<FSlateColor>& Color, const TAttribute<EVisibility>& Visibility ) const OVERRIDE;
+
+	virtual TSharedRef<SWidget> MakeWindowTitleBar( const TSharedRef<SWindow>& Window, const TSharedPtr<SWidget>& CenterContent, TSharedPtr<IWindowTitleBar>& OutTitleBar ) const OVERRIDE;
+
+	virtual TSharedRef<IToolTip> MakeToolTip( const TAttribute<FString>& ToolTipString ) OVERRIDE;
+
+	virtual TSharedRef<IToolTip> MakeToolTip( const FText& ToolTipText ) OVERRIDE;
+
+	virtual void RequestDestroyWindow( TSharedRef<SWindow> WindowToDestroy ) OVERRIDE;
+
+	virtual bool SetKeyboardFocus( const FWidgetPath& InFocusPath, const EKeyboardFocusCause::Type InCause ) OVERRIDE;
+
+	// End FSlateApplicationBase interface
+
+public:
+
+	// Begin FGenericApplicationMessageHandler interface
+
+	virtual bool ShouldProcessUserInputMessages( const TSharedPtr< FGenericWindow >& PlatformWindow ) const OVERRIDE;
+
+	virtual bool OnKeyChar( const TCHAR Character, const bool IsRepeat ) OVERRIDE;
+
+	virtual bool OnKeyDown( const int32 KeyCode, const uint32 CharacterCode, const bool IsRepeat ) OVERRIDE;
+
+	virtual bool OnKeyUp( const int32 KeyCode, const uint32 CharacterCode, const bool IsRepeat ) OVERRIDE;
+
+	virtual bool OnMouseDown( const TSharedPtr< FGenericWindow >& PlatformWindow, const EMouseButtons::Type Button ) OVERRIDE;
+
+	virtual bool OnMouseUp( const EMouseButtons::Type Button ) OVERRIDE;
+
+	virtual bool OnMouseDoubleClick( const TSharedPtr< FGenericWindow >& PlatformWindow, const EMouseButtons::Type Button ) OVERRIDE;
+
+	virtual bool OnMouseWheel( const float Delta ) OVERRIDE;
+
+	virtual bool OnMouseMove() OVERRIDE;
+
+	virtual bool OnRawMouseMove( const int32 X, const int32 Y ) OVERRIDE;
+
+	virtual bool OnCursorSet() OVERRIDE;
+
+	virtual bool OnControllerAnalog( FKey Button, int32 ControllerId, float AnalogValue );
+
+	virtual bool OnControllerAnalog( EControllerButtons::Type Button, int32 ControllerId, float AnalogValue ) OVERRIDE;
+
+	virtual bool OnControllerButtonPressed( FKey Button, int32 ControllerId, bool IsRepeat );
+
+	virtual bool OnControllerButtonPressed( EControllerButtons::Type Button, int32 ControllerId, bool IsRepeat ) OVERRIDE;
+
+	virtual bool OnControllerButtonReleased( FKey, int32 ControllerId, bool IsRepeat );
+
+	virtual bool OnControllerButtonReleased( EControllerButtons::Type Button, int32 ControllerId, bool IsRepeat ) OVERRIDE;
+
+	virtual bool OnTouchGesture( EGestureEvent::Type GestureType, const FVector2D& Delta, float WheelDelta ) OVERRIDE;
+	
+	virtual bool OnTouchStarted( const TSharedPtr< FGenericWindow >& PlatformWindow, const FVector2D& Location, int32 TouchIndex, int32 ControllerId ) OVERRIDE;
+
+	virtual bool OnTouchMoved( const FVector2D& Location, int32 TouchIndex, int32 ControllerId ) OVERRIDE;
+
+	virtual bool OnTouchEnded( const FVector2D& Location, int32 TouchIndex, int32 ControllerId ) OVERRIDE;
+
+	virtual bool OnMotionDetected(const FVector& Tilt, const FVector& RotationRate, const FVector& Gravity, const FVector& Acceleration, int32 ControllerId) OVERRIDE;
+
+	virtual bool OnSizeChanged( const TSharedRef< FGenericWindow >& PlatformWindow, const int32 Width, const int32 Height, bool bWasMinimized = false ) OVERRIDE;
+
+	virtual void OnOSPaint( const TSharedRef< FGenericWindow >& PlatformWindow ) OVERRIDE;
+
+	virtual void OnResizingWindow( const TSharedRef< FGenericWindow >& PlatformWindow ) OVERRIDE;
+
+	virtual bool BeginReshapingWindow( const TSharedRef< FGenericWindow >& PlatformWindow ) OVERRIDE;
+
+	virtual void FinishedReshapingWindow( const TSharedRef< FGenericWindow >& PlatformWindow ) OVERRIDE;
+
+	virtual void OnMovedWindow( const TSharedRef< FGenericWindow >& PlatformWindow, const int32 X, const int32 Y ) OVERRIDE;
+
+	virtual bool OnWindowActivationChanged( const TSharedRef< FGenericWindow >& PlatformWindow, const EWindowActivation::Type ActivationType ) OVERRIDE;
+
+	virtual bool OnApplicationActivationChanged( const bool IsActive ) OVERRIDE;
+
+	virtual EWindowZone::Type GetWindowZoneForPoint( const TSharedRef< FGenericWindow >& PlatformWindow, const int32 X, const int32 Y ) OVERRIDE;
+
+	virtual void OnWindowClose( const TSharedRef< FGenericWindow >& PlatformWindow ) OVERRIDE;
+
+	virtual EDropEffect::Type OnDragEnterText( const TSharedRef< FGenericWindow >& Window, const FString& Text ) OVERRIDE;
+
+	virtual EDropEffect::Type OnDragEnterFiles( const TSharedRef< FGenericWindow >& Window, const TArray< FString >& Files ) OVERRIDE;
+
+	EDropEffect::Type OnDragEnter( const TSharedRef< SWindow >& Window, const TSharedRef<FExternalDragOperation>& DragDropOperation );
+
+	virtual EDropEffect::Type OnDragOver( const TSharedPtr< FGenericWindow >& Window ) OVERRIDE;
+
+	virtual void OnDragLeave( const TSharedPtr< FGenericWindow >& Window ) OVERRIDE;
+
+	virtual EDropEffect::Type OnDragDrop( const TSharedPtr< FGenericWindow >& Window ) OVERRIDE;
+
+	virtual bool OnWindowAction( const TSharedRef< FGenericWindow >& PlatformWindow, const EWindowAction::Type InActionType ) OVERRIDE;
+
+	// End FGenericApplicationMessageHandler interface
+
+private:
+
+	TSharedRef< FGenericWindow > MakeWindow( TSharedRef<SWindow> InSlateWindow, const bool bShowImmediately );
+
+	/**
+	 * Destroys an SWindow, removing it and all its children from the Slate window list.  Notifies the native window to destroy itself and releases rendering resources
+	 *
+	 * @param DestroyedWindow The window to destroy
+	 */
+	void PrivateDestroyWindow( const TSharedRef<SWindow>& DestroyedWindow );
+
+private:
+
+	// Hidden default constructor.
+	FSlateApplication();
+
+private:
+
+	/** Application singleton */
+	static TSharedPtr< FSlateApplication > CurrentApplication;
+
+	static TSharedPtr< class GenericApplication > PlatformApplication;
+	
+	TSet<FKey> PressedMouseButtons;
+
+	/** true when the slate app is active; i.e. the current foreground window is from our Slate app*/
+	bool bAppIsActive;
+
+	/** true if any slate window is currently active (not just top level windows) */
+	bool bSlateWindowActive;
+
+	/** Application-wide scale for supporting monitors of varying pixel density */
+	float Scale;
+
+	/** All the top-level windows owned by this application; they are tracked here in a platform-agnostic way. */
+	TArray< TSharedRef<SWindow> > SlateWindows;
+
+	/** The currently active slate window that is a top-level window (full fledged window; not a menu or tooltip)*/
+	TWeakPtr<SWindow> ActiveTopLevelWindow;
+	
+	/** List of active modal windows.  The last item in the list is the top-most modal window */
+	TArray< TSharedPtr<SWindow> > ActiveModalWindows;
+
+	/** These windows will be destroyed next tick. */
+	TArray< TSharedRef<SWindow> > WindowDestroyQueue;
+	
+	/** The stack of menus that are open */
+	FMenuStack MenuStack;
+
+	/** A vertical slice through the tree of widgets on screen; it represents widgets that were under the cursor last time an event was processed */
+	FWeakWidgetPath WidgetsUnderCursorLastEvent;
+
+	/**
+	 * A helper class to wrap the weak path functionality. The advantage of using this
+	 * class is that the path can be validated and the current mouse captor (if any) can
+	 * be informed they have lost mouse capture in the case of a number of events, such as
+	 * application losing focus, or the widget being hidden.
+	 */
+	class MouseCaptorHelper
+	{
+	public:
+		/** Returns true if the weak path is valid (i.e. more than zero widgets) */
+		bool IsValid() const;
+
+		/**
+		 * Sets a new mouse captor widget, invalidating the previous one if any and calling
+		 * its OnMouseCaptureLost() handler.
+		 *
+		 * @param EventPath	The path to the event.
+		 * @param Widget	The widget that wants to capture the mouse.
+		 */
+		void SetMouseCaptor( const FWidgetPath& EventPath, TSharedPtr< SWidget > Widget );
+
+		/** Invalidates the current mouse captor. Calls OnMouseCaptureLost() on the current mouse captor if one exists */
+		void Invalidate();
+
+		/**
+		 * Retrieves a resolved FWidgetPath from the current, internal weak path, if possible.
+		 * If the path is invalid or truncated (i.e. the widget is no longer relevant) then the current mouse captor's
+		 * OnMouseCaptureLost() handler is called and the path is invalidated.
+		 *
+		 * @param InterruptedPathHandling	How to handled incomplete paths. "Truncate" will return a partial path, "ReturnInvalid" will return an empty path.
+		 */
+		FWidgetPath ToWidgetPath( FWeakWidgetPath::EInterruptedPathHandling::Type InterruptedPathHandling = FWeakWidgetPath::EInterruptedPathHandling::Truncate );
+
+		/** Retrieves the currently set weak path to the current mouse captor */
+		FWeakWidgetPath ToWeakPath() const;
+
+		/* Walks the weak path and retrieves the widget that is set as the current mouse captor */
+		TSharedPtr< SWidget > ToSharedWidget() const;
+
+		/* Walks the weak path and retrieves the window that the current mouse captor widget belongs to */
+		TSharedPtr< SWidget > ToSharedWindow();
+
+	protected:
+		/** If there is a current mouse captor, will call its OnMouseCaptureLost() handler */
+		void InformCurrentCaptorOfCaptureLoss() const;
+
+		/** The weak pointer path to the current mouse captor */
+		FWeakWidgetPath MouseCaptorWeakPath;
+	};
+	/** The current mouse captor for the application, if any. */
+	MouseCaptorHelper MouseCaptor;
+
+	/** A weak path to the widget currently capturing joystick events, if any. */
+	FWeakWidgetPath JoystickCaptorWeakPaths[SlateApplicationDefs::MaxUsers];
+
+	/**
+	 * Application throttling
+	 */
+
+	/** Holds a current request to ensure Slate is responsive in low FPS situations, based in mouse button pressed state */
+	FThrottleRequest MouseButtonDownResponsivnessThrottle;
+
+	/** Separate throttle handle that engages automatically based on mouse movement and other user behavior */
+	FThrottleRequest UserInteractionResponsivnessThrottle;
+
+	/** The last real time that the user pressed a key or mouse button */
+	double LastUserInteractionTimeForThrottling;
+
+
+	/** Helper for detecting when a drag should begin */
+	struct FDragDetector
+	{
+		FDragDetector()
+		: DetectDragStartLocation( FVector2D::ZeroVector )
+		, DetectDragButton( EKeys::Invalid )
+		{
+		}
+		/** If not null, a widget has request that we detect a drag being triggered in this widget and send an OnDragDetected() event*/
+		FWeakWidgetPath DetectDragForWidget;
+		/** Location from which be begin detecting the drag */
+		FVector2D DetectDragStartLocation;
+		/** Button that must be pressed to trigger the drag */
+		FKey DetectDragButton;
+	} DragDetector;
+
+	/** Support for auto-dismissing popups */
+	FPopupSupport PopupSupport;
+
+	/** Weak path to the widget that currently has keyboard focus, if any */
+	FWeakWidgetPath FocusedWidgetPath;
+	/** The way in which the last focus was set. */
+	EKeyboardFocusCause::Type FocusCause;
+	
+	/** Pointer to the currently registered game viewport widget if any */
+	TWeakPtr<SViewport> GameViewportWidget;
+
+	TSharedPtr<ISlateSoundDevice> SlateSoundDevice;
+
+	/** The current cached absolute real time, right before we tick widgets */
+	double CurrentTime;
+
+	/** Last absolute real time that we ticked */
+	double LastTickTime;
+
+	/** Running average time in seconds between calls to Tick (used for monitoring responsiveness) */
+	float AverageDeltaTime;
+
+	/** Average delta time for application responsiveness tracking.  This is like AverageDeltaTime, but it excludes frame
+	    deltas spent while the the application is in a throttled state */
+	float AverageDeltaTimeForResponsiveness;
+
+	
+	/**
+	 * Provides a platform-agnostic method for requesting that the application exit.
+	 * Implementations should assign a handler that terminates the process when this delegate is invoked.
+	 */
+	FSimpleDelegate OnExitRequested;
+	
+	/** A Widget that introspects the current UI hierarchy */
+	TWeakPtr<IWidgetReflector> WidgetReflectorPtr;
+
+	/** Delegate for accessing source code to pass to any widget inspectors. */
+	FAccessSourceCode SourceCodeAccessDelegate;
+
+	/** Delegate for querying if source code access is available */
+	FQueryAccessSourceCode QuerySourceCodeAccessDelegate;
+
+	/** System for logging all relevant slate events to a log file */
+	TSharedPtr<class IEventLogger> EventLogger;
+
+	/** Allows us to track the number of non-slate modal windows active. */
+	int32 NumExternalModalWindowsActive;
+
+
+	/**
+	 * Tool-tips
+	 */
+
+	/** Window that we'll re-use for spawned tool tips */
+	TWeakPtr< SWindow > ToolTipWindow;
+
+	/** Widget that last visualized the tooltip;  */
+	TWeakPtr< SWidget > TooltipVisualizerPtr;
+
+	/** The tool tip that is currently being displayed; can be invalid (i.e. not showing tooltip) */
+	TWeakPtr<IToolTip> ActiveToolTip;
+
+	/** The widget that sourced the currently active tooltip widget */
+	TWeakPtr<SWidget> ActiveToolTipWidgetSource;
+
+	/** Whether tool-tips are allowed to spawn at all.  Used only for debugging purposes. */
+	int32 bAllowToolTips;
+
+	/** How long before the tool tip should start fading in? */
+	float ToolTipDelay;
+
+	/** Tool-tip fade-in duration */
+	float ToolTipFadeInDuration;
+
+	/** Absolute real time that the current tool-tip window was summoned */
+	double ToolTipSummonTime;
+
+	/** Desired tool tip position in screen space, updated whenever the mouse moves */
+	FVector2D DesiredToolTipLocation;
+
+	/** Direction that tool-tip is being repelled from a force field in.  We cache this to avoid tool-tips
+	    teleporting between different offset directions as the user moves the mouse cursor around. */
+	struct EToolTipOffsetDirection
+	{
+		enum Type
+		{
+			Undetermined, Down, Right
+		};
+	};
+	EToolTipOffsetDirection::Type ToolTipOffsetDirection;
+
+	/** The top of the Style tree. */
+	const class FStyleNode* RootStyleNode;
+
+	/**
+	 * Drag and Drop (DragDrop)
+	 */
+	
+	/** When not null, the content of the current drag drop operation. */
+	TSharedPtr< FDragDropOperation > DragDropContent;
+
+	/** Whether or not we are requesting that we leave debugging mode after the tick is complete */
+	bool bRequestLeaveDebugMode;
+	/** Whether or not we need to leave debug mode for single stepping */
+	bool bLeaveDebugForSingleStep;
+	TAttribute<bool> NormalExecutionGetter;
+
+	/**
+	 * Console objects
+	 */
+
+	/** AllowToolTips console variable **/
+	FAutoConsoleVariableRef CVarAllowToolTips;
+	/** ToolTipDelay console variable **/
+	FAutoConsoleVariableRef CVarToolTipDelay;
+	/** ToolTipFadeInDuration console variable **/
+	FAutoConsoleVariableRef CVarToolTipFadeInDuration;
+
+	/**
+	 * Modal Windows
+	 */
+
+	/** Delegates for when modal windows open or close */
+	FModalWindowStackStarted ModalWindowStackStartedDelegate;
+	FModalWindowStackEnded ModalWindowStackEndedDelegate;
+
+	/** Keeps track of whether or not the UI for services such as Steam is open. */
+	bool bIsExternalUIOpened;
+
+	/**
+	 * Remote Slate server
+	 */
+	TSharedPtr<class FRemoteSlateServer> RemoteSlateServer;
+
+	/** Handle to a throttle request made to ensure the window is responsive in low FPS situations */
+	FThrottleRequest ThrottleHandle;
+
+	/** When an drag and drop is happening, we keep track of whether slate knew what to do with the payload on last mouse move */
+	bool DragIsHandled;
+
+	FVector2D LastCursorPosition;
+
+	/**
+	 * Virtual keyboard text field
+	 */
+	IPlatformTextField* SlateTextField;
+
+	/** For desktop platforms that want to test touch style input, pass -faketouches or -simmobile on the commandline to set this */
+	bool bIsFakingTouch;
+
+	/** For games that want to allow mouse to imitate touch */
+	bool bIsGameFakingTouch;
+
+	/**For desktop platforms that the touch move event be called when this variable is true */
+	bool bIsFakingTouched;
+
+	/** Delegate for when a key down event occurred but was not handled in any other way by ProcessKeyDownMessage */
+	FOnKeyboardEvent UnhandledKeyDownEventHandler;
+
+	/** controls whether unhandled touch events fall back to sending mouse events */
+	bool bTouchFallbackToMouse;
+
+	/**
+	 * Slate look and feel
+	 */
+
+	/** Globally enables or disables transition effects for pop-up menus (menu stacks) */
+	bool bMenuAnimationsEnabled;
+
+	/** The icon to use on application windows */
+	const FSlateBrush *AppIcon;
+};
diff --git a/Engine/Source/Runtime/Slate/Public/Framework/Commands/Commands.h b/Engine/Source/Runtime/Slate/Public/Framework/Commands/Commands.h
new file mode 100644
--- /dev/null
+++ b/Engine/Source/Runtime/Slate/Public/Framework/Commands/Commands.h
@@ -0,0 +1,102 @@
+// Copyright 1998-2014 Epic Games, Inc. All Rights Reserved.
+
+#pragma once
+
+//
+// COMMAND DECLARATION
+//
+
+#define LOC_DEFINE_REGION
+
+SLATE_API void UI_COMMAND_Function( FBindingContext* This, TSharedPtr< FUICommandInfo >& OutCommand, const TCHAR* OutCommandName, const TCHAR* OutCommandNameUnderscoreTooltip, const ANSICHAR* DotOutCommandName, const TCHAR* FriendlyName, const TCHAR* InDescription, const EUserInterfaceActionType::Type CommandType, const FInputGesture& InDefaultGesture );
+
+#define UI_COMMAND( CommandId, FriendlyName, InDescription, CommandType, InDefaultGesture ) \
+	UI_COMMAND_Function( this, CommandId, TEXT(#CommandId), TEXT(#CommandId) TEXT("_ToolTip"), "." #CommandId, TEXT(FriendlyName), TEXT(InDescription), CommandType, InDefaultGesture );
+
+#undef LOC_DEFINE_REGION
+
+/** A base class for a set of commands. Inherit from it to make a set of commands. See MainFrameActions for an example. */
+template<typename CommandContextType>
+class TCommands : public FBindingContext
+{
+public:
+	
+
+	/** Use this method to register commands. Usually done in StartupModule(). */
+	static void Register()
+	{
+		if ( !Instance.IsValid() )
+		{
+			// We store the singleton instances in the FInputBindingManager in order to prevent
+			// different modules from instantiating their own version of TCommands<CommandContextType>.
+			TSharedRef<CommandContextType> NewInstance = MakeShareable( new CommandContextType() );
+
+			TSharedPtr<FBindingContext> ExistingBindingContext = FInputBindingManager::Get().GetContextByName( NewInstance->GetContextName() );
+			if (ExistingBindingContext.IsValid())
+			{
+				// Someone already made this set of commands, and registered it.
+				Instance = StaticCastSharedPtr<CommandContextType>(ExistingBindingContext);
+			}
+			else
+			{
+				// Make a new set of commands and register it.
+				Instance = NewInstance;
+
+				// Registering the first command will add the NewInstance into the Binding Manager, who holds on to it.
+				NewInstance->RegisterCommands();
+
+				// Notify that new commands have been registered
+				CommandsChanged.Broadcast();
+			}
+		}
+	}
+
+	/** Get the singleton instance of this set of commands. */
+	static const CommandContextType& Get()
+	{
+		return *(Instance.Pin());
+	}
+
+	/** Use this method to clean up any resources used by the command set. Usually done in ShutdownModule() */
+	static void Unregister()
+	{
+		// The instance may not be valid if it was never used.
+		if( Instance.IsValid() )
+		{
+			FInputBindingManager::Get().RemoveContextByName(Instance.Pin()->GetContextName());
+			
+			// Notify that new commands have been unregistered
+			CommandsChanged.Broadcast();
+
+			check(!Instance.IsValid());
+		}
+	}
+
+	/** Get the BindingContext for this set of commands. */
+	static const FBindingContext& GetContext()
+	{
+		check( Instance.IsValid() );
+		return *(Instance.Pin());
+	}
+	
+protected:
+	
+	/** Construct a set of commands; call this from your custom commands class. */
+	TCommands( const FName InContextName, const FText& InContextDesc, const FName InContextParent, const FName InStyleSetName )
+		: FBindingContext( InContextName, InContextDesc, InContextParent, InStyleSetName )
+	{
+	}
+	virtual ~TCommands()
+	{
+	}
+
+	/** A static instance of the command set. */
+	static TWeakPtr< CommandContextType > Instance;
+
+	/** Pure virtual to override; describe and instantiate the commands in here by using the UI COMMAND macro. */
+	virtual void RegisterCommands() = 0;
+
+};
+
+template<typename CommandContextType>
+TWeakPtr< CommandContextType > TCommands<CommandContextType>::Instance = NULL;
diff --git a/Engine/Source/Runtime/Slate/Public/Framework/Commands/InputBindingManager.h b/Engine/Source/Runtime/Slate/Public/Framework/Commands/InputBindingManager.h
new file mode 100644
--- /dev/null
+++ b/Engine/Source/Runtime/Slate/Public/Framework/Commands/InputBindingManager.h
@@ -0,0 +1,175 @@
+// Copyright 1998-2014 Epic Games, Inc. All Rights Reserved.
+
+/*=============================================================================
+	InputBindingManager.h: Declares the FInputBindingManager class.
+=============================================================================*/
+
+#pragma once
+
+
+typedef TMap<FName, TSharedPtr<FUICommandInfo>> FCommandInfoMap;
+typedef TMap<FInputGesture, FName> FGestureMap;
+
+
+/** Delegate for alerting subscribers the input manager records a user-defined gesture */
+DECLARE_MULTICAST_DELEGATE_OneParam(FOnUserDefinedGestureChanged, const FUICommandInfo& );
+
+
+/**
+ * Manager responsible for creating and processing input bindings.                  
+ */
+class SLATE_API FInputBindingManager
+{
+public:
+	/**
+	 * @return The instance of this manager                   
+	 */
+	static FInputBindingManager& Get();
+
+	/**
+	 * Returns a list of all known input contexts
+	 *
+	 * @param OutInputContexts	The generated list of contexts
+	 * @return A list of all known input contexts                   
+	 */
+	void GetKnownInputContexts( TArray< TSharedPtr<FBindingContext> >& OutInputContexts ) const;
+
+	/**
+	 * Look up a binding context by name.
+	 */
+	TSharedPtr<FBindingContext> GetContextByName( const FName& InContextName );
+
+	/**
+	 * Remove the context with this name
+	 */
+	void RemoveContextByName( const FName& InContextName );
+
+	/**
+	 * Creates an input command from the specified user interface action
+	 * 
+	 * @param InBindingContext		The context where the command is valid
+	 * @param InUserInterfaceAction	The user interface action to create the input command from
+	 *								Will automatically bind the user interface action to the command specified 
+	 *								by the action so when the command's gesture is pressed, the action is executed
+	 */
+	void CreateInputCommand( const TSharedRef<FBindingContext>& InBindingContext, TSharedRef<FUICommandInfo> InUICommandInfo );
+
+	/**
+	 * Returns a command info that is has the same active gesture as the provided gesture and is in the same binding context or parent context
+	 *
+	 * @param	InBindingContext		The context in which the command is valid
+	 * @param	InGesture				The gesture to match against commands
+	 * @param	bCheckDefault			Whether or not to check the default gesture.  Will check the active gesture if false
+	 * @return	A pointer to the command info which has the InGesture as its active gesture or null if one cannot be found
+	 */
+	const TSharedPtr<FUICommandInfo> GetCommandInfoFromInputGesture( const FName InBindingContext, const FInputGesture& InGesture, bool bCheckDefault ) const;
+
+	/** 
+	 * Finds the command in the provided context which uses the provided input gesture
+	 * 
+	 * @param InBindingContext	The binding context name
+	 * @param InGesture			The gesture to check against when looking for commands
+	 * @param bCheckDefault		Whether or not to check the default gesture of commands instead of active gestures
+	 * @param OutActiveGestureIndex	The index into the commands active gesture array where the passed in gesture was matched.  INDEX_NONE if bCheckDefault is true or nothing was found.
+	 */
+	const TSharedPtr<FUICommandInfo> FindCommandInContext( const FName InBindingContext, const FInputGesture& InGesture, bool bCheckDefault ) const;
+
+	/** 
+	 * Finds the command in the provided context which has the provided name 
+	 * 
+	 * @param InBindingContext	The binding context name
+	 * @param CommandName		The name of the command to find
+	 */
+	const TSharedPtr<FUICommandInfo> FindCommandInContext( const FName InBindingContext, const FName CommandName ) const;
+
+	/**
+	 * Called when the active gesture is changed on a command
+	 *
+	 * @param CommandInfo	The command that had the active gesture changed. 
+	 */
+	virtual void NotifyActiveGestureChanged( const FUICommandInfo& CommandInfo );
+	
+	/**
+	 * Saves the user defined gestures to a json file
+	 */
+	void SaveInputBindings();
+
+	/**
+	 * Removes any user defined gestures
+	 */
+	void RemoveUserDefinedGestures();
+
+	/**
+	 * Returns all known command infos for a given binding context
+	 *
+	 * @param InBindingContext	The binding context to get command infos from
+	 * @param OutCommandInfos	The list of command infos for the binding context
+	 */
+	void GetCommandInfosFromContext( const FName InBindingContext, TArray< TSharedPtr<FUICommandInfo> >& OutCommandInfos ) const;
+
+	/** Registers a delegate to be called when a user-defined gesture is edited */
+	void RegisterUserDefinedGestureChanged(const FOnUserDefinedGestureChanged::FDelegate& Delegate)
+	{
+		OnUserDefinedGestureChanged.Add(Delegate);
+	}
+
+	/** Unregisters a delegate to be called when a user-defined gesture is edited */
+	void UnregisterUserDefinedGestureChanged(const FOnUserDefinedGestureChanged::FDelegate& Delegate)
+	{
+		OnUserDefinedGestureChanged.Remove(Delegate);
+	}
+
+private:
+
+	/**
+	 * Hidden default constructor.
+	 */
+	FInputBindingManager() { }
+
+	/**
+	 * Gets the user defined gesture (if any) from the provided command name
+	 * 
+	 * @param InBindingContext	The context in which the command is active
+	 * @param InCommandName		The name of the command to get the gesture from
+	 */
+	bool GetUserDefinedGesture( const FName InBindingContext, const FName InCommandName, FInputGesture& OutUserDefinedGesture );
+
+	/**
+	 *	Checks a binding context for duplicate gestures 
+	 */
+	void CheckForDuplicateDefaultGestures( const FBindingContext& InBindingContext, TSharedPtr<FUICommandInfo> InCommandInfo ) const;
+
+	/**
+	 * Recursively finds all child contexts of the provided binding context.  
+	 *
+	 * @param InBindingContext	The binding context to search
+	 * @param AllChildren		All the children of InBindingContext. InBindingContext is the first element in AllChildren
+	 */
+	void GetAllChildContexts( const FName InBindingContext, TArray<FName>& AllChildren ) const;
+
+private:
+
+	struct FContextEntry
+	{
+		/** A list of commands associated with the context */
+		FCommandInfoMap CommandInfoMap;
+
+		/** Gesture to command info map */
+		FGestureMap GestureToCommandInfoMap;
+
+		/** The binding context for this entry*/
+		TSharedPtr< FBindingContext > BindingContext;
+	};
+
+	/** A mapping of context name to the associated entry map */
+	TMap< FName, FContextEntry > ContextMap;
+	
+	/** A mapping of contexts to their child contexts */
+	TMultiMap< FName, FName > ParentToChildMap;
+
+	/** User defined gesture overrides for commands */
+	TSharedPtr< class FUserDefinedGestures > UserDefinedGestures;
+
+	/** Delegate called when a user-defined gesture is edited */
+	FOnUserDefinedGestureChanged OnUserDefinedGestureChanged;
+};
diff --git a/Engine/Source/Runtime/Slate/Public/Framework/Commands/InputGesture.h b/Engine/Source/Runtime/Slate/Public/Framework/Commands/InputGesture.h
new file mode 100644
--- /dev/null
+++ b/Engine/Source/Runtime/Slate/Public/Framework/Commands/InputGesture.h
@@ -0,0 +1,197 @@
+// Copyright 1998-2014 Epic Games, Inc. All Rights Reserved.
+
+/*=============================================================================
+	InputGesture.h: Declares the FInputGesture class.
+=============================================================================*/
+
+#pragma once
+
+
+/**
+ * Enumerates available modifier keys for input gestures.
+ */
+namespace EModifierKey
+{
+	typedef uint8 Type;
+
+	/** No key. */
+	const Type None	= 0;
+
+	/** Ctrl key. */
+	const Type Control = 1 << 0;
+
+	/** Alt key. */
+	const Type Alt = 1 << 1;
+
+	/** Shift key. */
+	const Type Shift = 1 << 2;
+};
+
+
+/**
+ * Raw input gesture that defines input that must be valid when           
+ */
+struct SLATE_API FInputGesture
+{
+	/** Key that must be pressed */
+	FKey Key;
+
+	/** True if control must be pressed */
+	uint32 bCtrl:1;
+	
+	/** True if alt must be pressed */
+	uint32 bAlt:1;
+	
+	/** True if shift must be pressed */
+	uint32 bShift:1;
+
+public:
+
+	/**
+	 * Default constructor.
+	 */
+	FInputGesture( )
+	{
+		bCtrl = bAlt = bShift = 0;
+	}
+
+	/**
+	 * Creates and initializes a new instance.
+	 *
+	 * @param InModifierKeys
+	 * @param InKey
+	 */
+	FInputGesture( EModifierKey::Type InModifierKeys, const FKey InKey )
+		: Key(InKey)
+	{
+		bCtrl = (InModifierKeys & EModifierKey::Control) != 0;
+		bAlt = (InModifierKeys & EModifierKey::Alt) != 0;
+		bShift = (InModifierKeys & EModifierKey::Shift) != 0;
+	}
+
+	/**
+	 * Creates and initializes a new instance.
+	 *
+	 * @param InKey
+	 */
+	FInputGesture( const FKey InKey )
+		: Key(InKey)
+	{
+		bCtrl = bAlt = bShift = 0;
+	}
+
+	/**
+	 * Creates and initializes a new instance.
+	 *
+	 * @param InKey
+	 * @param InModifierKeys
+	 */
+	FInputGesture( const FKey InKey, EModifierKey::Type InModifierKeys )
+		: Key(InKey)
+	{
+		bCtrl = (InModifierKeys & EModifierKey::Control) != 0;
+		bAlt = (InModifierKeys & EModifierKey::Alt) != 0;
+		bShift = (InModifierKeys & EModifierKey::Shift) != 0;
+	}
+
+	/**
+	 * Copy constructor.
+	 *
+	 * @param Other
+	 */
+	FInputGesture( const FInputGesture& Other )
+		: Key(Other.Key)
+		, bCtrl(Other.bCtrl)
+		, bAlt(Other.bAlt)
+		, bShift(Other.bShift)
+	{ }
+
+public:
+
+	/**
+	 * Compares this input gesture with another for equality.
+	 *
+	 * @param Other The other gesture to compare with.
+	 *
+	 * @return true if the gestures are equal, false otherwise.
+	 */
+	bool operator!=( const FInputGesture& Other ) const
+	{
+		return (Key != Other.Key) || (bCtrl != Other.bCtrl) || (bAlt != Other.bAlt) || (bShift != Other.bShift);
+	}
+
+	/**
+	 * Compares this input gesture with another for inequality.
+	 *
+	 * @param Other The other gesture to compare with.
+	 *
+	 * @return true if the gestures are not equal, false otherwise.
+	 */
+	bool operator==( const FInputGesture& Other ) const
+	{
+		return (Key == Other.Key) && (bCtrl == Other.bCtrl) && (bAlt == Other.bAlt) && (bShift == Other.bShift);
+	}
+
+public:
+
+	/** 
+	 * Gets a localized string that represents the gesture.
+	 *
+	 * @return A localized string.
+	 */
+	FText GetInputText( ) const;
+	
+	/**
+	 * Gets the key represented as a localized string.
+	 *
+	 * @return A localized string.
+	 */
+	FText GetKeyText( ) const;
+
+	/**
+	 * Checks whether this gesture requires an modifier keys to be pressed.
+	 *
+	 * @return true if modifier keys must be pressed, false otherwise.
+	 */
+	bool HasAnyModifierKeys( ) const
+	{
+		return bCtrl || bAlt || bShift;
+	}
+
+	/**
+	 * Determines if this gesture is valid.  A gesture is valid if it has a non modifier key that must be pressed
+	 * and zero or more modifier keys that must be pressed
+	 *
+	 * @return true if the gesture is valid
+	 */
+	bool IsValidGesture( ) const
+	{
+		return (Key.IsValid() && !Key.IsModifierKey());
+	}
+
+	/**
+	 * Sets this gesture to a new key and modifier state based on the provided template
+	 * Should not be called directly.  Only used by the key binding editor to set user defined keys
+	 */
+	void Set( FInputGesture InTemplate )
+	{
+		Key = InTemplate.Key;
+		bCtrl = InTemplate.bCtrl;
+		bAlt = InTemplate.bAlt;
+		bShift = InTemplate.bShift;
+	}
+
+public:
+
+	/**
+	 * Gets a type hash value for the specified gesture.
+	 *
+	 * @param Gesture The input gesture to get the hash value for.
+	 *
+	 * @return The hash value.
+	 */
+	friend uint32 GetTypeHash( const FInputGesture& Gesture )
+	{
+		return GetTypeHash(Gesture.Key) ^ ((Gesture.bCtrl << 2) | (Gesture.bShift << 1) | Gesture.bAlt);
+	}
+};
diff --git a/Engine/Source/Runtime/Slate/Public/Framework/Commands/UIAction.h b/Engine/Source/Runtime/Slate/Public/Framework/Commands/UIAction.h
new file mode 100644
--- /dev/null
+++ b/Engine/Source/Runtime/Slate/Public/Framework/Commands/UIAction.h
@@ -0,0 +1,192 @@
+// Copyright 1998-2014 Epic Games, Inc. All Rights Reserved.
+
+/*=============================================================================
+	UIAction.h: Declares the FUIAction structure.
+=============================================================================*/
+
+#pragma once
+
+
+/** Defines FExecuteAction delegate interface */
+DECLARE_DELEGATE(FExecuteAction);
+
+/** Defines FCanExecuteAction delegate interface.  Returns true when an action is able to execute. */
+DECLARE_DELEGATE_RetVal(bool, FCanExecuteAction);
+
+/** Defines FIsActionChecked delegate interface.  Returns true if the action is currently toggled on. */
+DECLARE_DELEGATE_RetVal(bool, FIsActionChecked);
+
+/** Defines FIsActionButtonVisible delegate interface.  Returns true when UI buttons associated with the action should be visible. */
+DECLARE_DELEGATE_RetVal(bool, FIsActionButtonVisible);
+
+
+/**
+ * Implements an UI action.
+ */
+struct FUIAction
+{
+	/** Holds a delegate that is executed when this action is activated. */
+	FExecuteAction ExecuteAction;
+
+	/** Holds a delegate that is executed when determining whether this action can execute. */
+	FCanExecuteAction CanExecuteAction;
+
+	/** Holds a delegate that is executed when determining whether this action is checked. */
+	FIsActionChecked IsCheckedDelegate;
+
+	/** Holds a delegate that is executed when determining whether this action is visible. */
+	FIsActionButtonVisible IsActionVisibleDelegate;
+
+public:
+
+	/**
+	 * Default constructor
+	 */
+	FUIAction( ) { }
+
+	/**
+	 * Constructor that takes delegates to initialize the action with
+	 *
+	 * @param	ExecuteAction		The delegate to call when the action should be executed
+	 */
+	FUIAction( FExecuteAction InitExecuteAction )
+		: ExecuteAction(InitExecuteAction)
+		, CanExecuteAction(FCanExecuteAction::CreateRaw(&FSlateApplication::Get(), &FSlateApplication::IsNormalExecution))
+		, IsCheckedDelegate()
+		, IsActionVisibleDelegate()
+	{ }
+
+	/**
+	 * Constructor that takes delegates to initialize the action with
+	 *
+	 * @param	ExecuteAction		The delegate to call when the action should be executed
+	 * @param	CanExecuteAction	The delegate to call to see if the action can be executed
+	 */
+	FUIAction( FExecuteAction InitExecuteAction, FCanExecuteAction InitCanExecuteAction )
+		: ExecuteAction(InitExecuteAction )
+		, CanExecuteAction(InitCanExecuteAction)
+		, IsCheckedDelegate()
+		, IsActionVisibleDelegate()
+	{ }
+
+	/**
+	 * Constructor that takes delegates to initialize the action with
+	 *
+	 * @param	ExecuteAction		The delegate to call when the action should be executed
+	 * @param	CanExecuteAction	The delegate to call to see if the action can be executed
+	 * @param	IsCheckedDelegate	The delegate to call to see if the action should appear checked when visualized
+	 */
+	FUIAction( FExecuteAction InitExecuteAction, FCanExecuteAction InitCanExecuteAction, FIsActionChecked InitIsCheckedDelegate )
+		: ExecuteAction(InitExecuteAction)
+		, CanExecuteAction(InitCanExecuteAction)
+		, IsCheckedDelegate(InitIsCheckedDelegate)
+		, IsActionVisibleDelegate()
+	{ }
+
+	/**
+	 * Constructor that takes delegates to initialize the action with
+	 *
+	 * @param	ExecuteAction			The delegate to call when the action should be executed
+	 * @param	CanExecuteAction		The delegate to call to see if the action can be executed
+	 * @param	IsCheckedDelegate		The delegate to call to see if the action should appear checked when visualized
+	 * @param	IsActionVisibleDelegate	The delegate to call to see if the action should be visible
+	 */
+	FUIAction( FExecuteAction InitExecuteAction, FCanExecuteAction InitCanExecuteAction, FIsActionChecked InitIsCheckedDelegate, FIsActionButtonVisible InitIsActionVisibleDelegate )
+		: ExecuteAction(InitExecuteAction)
+		, CanExecuteAction(InitCanExecuteAction)
+		, IsCheckedDelegate(InitIsCheckedDelegate)
+		, IsActionVisibleDelegate(InitIsActionVisibleDelegate)
+	{ }
+
+public:
+
+	/**
+	 * Compares this UI action to another for equality.
+	 *
+	 * @param Other - The other action to compare to.
+	 *
+	 * @return true if both actions are equal, false otherwise.
+	 */
+	bool operator==( const FUIAction& Other ) const
+	{
+		return
+			(ExecuteAction == Other.ExecuteAction) &&
+			(CanExecuteAction == Other.CanExecuteAction) &&
+			(IsCheckedDelegate == Other.IsCheckedDelegate) &&
+			(IsActionVisibleDelegate == Other.IsActionVisibleDelegate);
+	}
+
+public:
+
+	/**
+	 * Checks to see if its currently safe to execute this action.
+	 *
+	 * @return	True if this action can be executed and we should reflect that state visually in the UI
+	 */
+	bool CanExecute( ) const
+	{
+		// Fire the 'can execute' delegate if we have one, otherwise always return true
+		return CanExecuteAction.IsBound() ? CanExecuteAction.Execute() : true;
+	}
+
+	/**
+	 * Executes this action
+	 */
+	bool Execute( ) const
+	{
+		// It's up to the programmer to ensure that the action is still valid by the time the user clicks on the
+		// button.  Otherwise the user won't know why the action didn't take place!
+		if( CanExecute() )
+		{
+			return ExecuteAction.ExecuteIfBound();
+		}
+
+		return false;
+	}
+
+	/**
+	 * Queries the checked state for this action.  This is only valid for actions that are toggleable!
+	 *
+	 * @return	Checked state for this action
+	 */
+	bool IsChecked( ) const
+	{
+		if (IsCheckedDelegate.IsBound())
+		{
+			const bool bIsChecked = IsCheckedDelegate.Execute();
+
+			if (bIsChecked)
+			{
+				return true;
+			}
+		}
+
+		return false;
+	}
+
+	/**
+	 * Checks whether this action's execution delegate is bound.
+	 *
+	 * @return true if the delegate is bound, false otherwise.
+	 */
+	bool IsBound( ) const
+	{
+		return ExecuteAction.IsBound();
+	}
+
+	/**
+	 * Queries the visibility for this action.
+	 *
+	 * @return	Visibility state for this action
+	 */
+	EVisibility IsVisible( ) const
+	{
+		if (IsActionVisibleDelegate.IsBound())
+		{
+			const bool bIsVisible = IsActionVisibleDelegate.Execute();
+			return bIsVisible ? EVisibility::Visible : EVisibility::Collapsed;
+		}
+
+		return EVisibility::Visible;
+	}
+};
diff --git a/Engine/Source/Runtime/Slate/Public/Framework/Commands/UICommandDragDropOp.h b/Engine/Source/Runtime/Slate/Public/Framework/Commands/UICommandDragDropOp.h
new file mode 100644
--- /dev/null
+++ b/Engine/Source/Runtime/Slate/Public/Framework/Commands/UICommandDragDropOp.h
@@ -0,0 +1,55 @@
+// Copyright 1998-2014 Epic Games, Inc. All Rights Reserved.
+
+/*=============================================================================
+	UICommandDragDropOp.h: Declares the FUICommandDragDropOp class.
+=============================================================================*/
+
+#pragma once
+
+
+/**
+ * A drag drop operation for UI Commands
+ */
+class SLATE_API FUICommandDragDropOp
+	: public FDragDropOperation
+{
+public:
+
+	DRAG_DROP_OPERATOR_TYPE(FUICommandDragDropOp, FDragDropOperation)
+
+	static TSharedRef<FUICommandDragDropOp> New( TSharedRef<const FUICommandInfo> InCommandInfo, FName InOriginMultiBox, TSharedPtr<SWidget> CustomDectorator, FVector2D DecoratorOffset );
+
+	FUICommandDragDropOp( TSharedRef<const FUICommandInfo> InUICommand, FName InOriginMultiBox, TSharedPtr<SWidget> InCustomDecorator, FVector2D DecoratorOffset )
+		: UICommand( InUICommand )
+		, OriginMultiBox( InOriginMultiBox )
+		, CustomDecorator( InCustomDecorator )
+		, Offset( DecoratorOffset )
+	{}
+
+	/**
+	 * Sets a delegate that will be called when the command is dropped
+	 */
+	void SetOnDropNotification( FSimpleDelegate InOnDropNotification ) { OnDropNotification = InOnDropNotification; }
+
+	/** FDragDropOperation interface */
+	virtual TSharedPtr<SWidget> GetDefaultDecorator() const OVERRIDE;
+	virtual void OnDragged( const class FDragDropEvent& DragDropEvent ) OVERRIDE;
+	virtual void OnDrop( bool bDropWasHandled, const FPointerEvent& MouseEvent ) OVERRIDE;
+
+public:
+
+	/** UI command being dragged */
+	TSharedPtr< const FUICommandInfo > UICommand;
+
+	/** Multibox the UI command was dragged from if any*/
+	FName OriginMultiBox;
+
+	/** Custom decorator to display */
+	TSharedPtr<SWidget> CustomDecorator;
+
+	/** Offset from the cursor where the decorator should be displayed */
+	FVector2D Offset;
+
+	/** Delegate called when the command is dropped */
+	FSimpleDelegate OnDropNotification;
+};
diff --git a/Engine/Source/Runtime/Slate/Public/Framework/Commands/UICommandInfo.h b/Engine/Source/Runtime/Slate/Public/Framework/Commands/UICommandInfo.h
new file mode 100644
--- /dev/null
+++ b/Engine/Source/Runtime/Slate/Public/Framework/Commands/UICommandInfo.h
@@ -0,0 +1,232 @@
+// Copyright 1998-2014 Epic Games, Inc. All Rights Reserved.
+
+/*=============================================================================
+	UICommandInfo.h: Declares the UICommandInfo class.
+=============================================================================*/
+
+#pragma once
+
+
+/** Types of user interfaces that can be associated with a user interface action */
+namespace EUserInterfaceActionType
+{
+	enum Type
+	{
+		/** Momentary buttons or menu items.  These support enable state, and execute a delegate when clicked. */
+		Button,
+
+		/** Toggleable buttons or menu items that store on/off state.  These support enable state, and execute a delegate when toggled. */
+		ToggleButton,
+		
+		/** Radio buttons are similar to toggle buttons in that they are for menu items that store on/off state.  However they should be used to indicate that menu items in a group can only be in one state */
+		RadioButton,
+
+		/** Similar to Button but will display a readonly checkbox next to the item. */
+		Check
+	};
+};
+
+
+class FUICommandInfo;
+
+
+/**
+ *
+ */
+class SLATE_API FUICommandInfoDecl
+{
+	friend class FBindingContext;
+
+public:
+
+	FUICommandInfoDecl& DefaultGesture( const FInputGesture& InDefaultGesture );
+	FUICommandInfoDecl& UserInterfaceType( EUserInterfaceActionType::Type InType );
+	FUICommandInfoDecl& Icon( const FSlateIcon& InIcon );
+	FUICommandInfoDecl& Description( const FText& InDesc );
+
+	operator TSharedPtr<FUICommandInfo>() const;
+	operator TSharedRef<FUICommandInfo>() const;
+
+public:
+
+	FUICommandInfoDecl( const TSharedRef<class FBindingContext>& InContext, const FName InCommandName, const FText& InLabel, const FText& InDesc  );
+
+private:
+
+	TSharedPtr<class FUICommandInfo> Info;
+	const TSharedRef<FBindingContext>& Context;
+};
+
+
+/**
+ * Represents a context in which input bindings are valid
+ */
+class SLATE_API FBindingContext
+	: public TSharedFromThis<FBindingContext>
+{
+public:
+	/**
+	 * Constructor
+	 *
+	 * @param InContextName		The name of the context
+	 * @param InContextDesc		The localized description of the context
+	 * @param InContextParent	Optional parent context.  Bindings are not allowed to be the same between parent and child contexts
+	 * @param InStyleSetName	The style set to find the icons in, eg) FCoreStyle::Get().GetStyleSetName()
+	 */
+	FBindingContext( const FName InContextName, const FText& InContextDesc, const FName InContextParent, const FName InStyleSetName )
+		: ContextName( InContextName )
+		, ContextParent( InContextParent )
+		, ContextDesc( InContextDesc )
+		, StyleSetName( InStyleSetName )
+	{
+		check(!InStyleSetName.IsNone());
+	}
+
+	FBindingContext( const FBindingContext &Other )
+		: ContextName( Other.ContextName )
+		, ContextParent( Other.ContextParent )
+		, ContextDesc( Other.ContextDesc )
+		, StyleSetName( Other.StyleSetName )
+	{}
+
+	/**
+	 * Creates a new command declaration used to populate commands with data
+	 */
+	FUICommandInfoDecl NewCommand( const FName InCommandName, const FText& InCommandLabel, const FText& InCommandDesc );
+
+	/**
+	 * @return The name of the context
+	 */
+	FName GetContextName() const { return ContextName; }
+
+	/**
+	 * @return The name of the parent context (or NAME_None if there isnt one)
+	 */
+	FName GetContextParent() const { return ContextParent; }
+
+	/**
+	 * @return The name of the style set to find the icons in
+	 */
+	FName GetStyleSetName() const { return StyleSetName; }
+
+	/**
+	 * @return The localized description of this context
+	 */
+	const FText& GetContextDesc() const { return ContextDesc; }
+
+	friend uint32 GetTypeHash( const FBindingContext& Context )
+	{
+		return GetTypeHash( Context.ContextName );
+	}
+
+	bool operator==( const FBindingContext& Other ) const
+	{
+		return ContextName == Other.ContextName;
+	}
+
+	/** A delegate that is called when commands are registered or unregistered with a binding context */
+	static FSimpleMulticastDelegate CommandsChanged; 
+private:
+	/** The name of the context */
+	FName ContextName;
+	/** The name of the parent context */
+	FName ContextParent;
+	/** The description of the context */
+	FText ContextDesc;
+	/** The style set to find the icons in */
+	FName StyleSetName;
+};
+
+
+class SLATE_API FUICommandInfo
+{
+	friend class FInputBindingManager;
+	friend class FUICommandInfoDecl;
+
+public:
+	/** Default constructor */
+	FUICommandInfo( const FName InBindingContext )
+		: ActiveGesture( new FInputGesture )
+		, DefaultGesture( EKeys::Invalid, EModifierKey::None )
+		, BindingContext( InBindingContext )
+		, UserInterfaceType( EUserInterfaceActionType::Button )
+	{ }
+
+	/**
+	 * Returns the friendly, localized string name of the gesture that is required to perform the command
+	 *
+	 * @param	bAllGestures	If true a comma separated string with all secondary active gestures will be returned.  Otherwise just the main gesture string is returned
+	 * @return	Localized friendly text for the gesture
+	 */
+	const FText GetInputText() const;
+
+	/**
+	 * @return	Returns the active gesture for this command
+	 */
+	const TSharedRef<const FInputGesture> GetActiveGesture() const { return ActiveGesture; }
+
+	const FInputGesture& GetDefaultGesture() const { return DefaultGesture; }
+
+	/** Utility function to make an FUICommandInfo */
+	static void MakeCommandInfo( const TSharedRef<class FBindingContext>& InContext, TSharedPtr< FUICommandInfo >& OutCommand, const FName InCommandName, const FText& InCommandLabel, const FText& InCommandDesc, const FSlateIcon& InIcon, const EUserInterfaceActionType::Type InUserInterfaceType, const FInputGesture& InDefaultGesture );
+
+	/** @return The display label for this command */
+	const FText& GetLabel() const { return Label; }
+
+	/** @return The description of this command */
+	const FText& GetDescription() const { return Description; }
+
+	/** @return The icon to used when this command is displayed in UI that shows icons */
+	const FSlateIcon& GetIcon() const { return Icon; }
+
+	/** @return The type of command this is.  Used to determine what UI to create for it */
+	EUserInterfaceActionType::Type GetUserInterfaceType() const { return UserInterfaceType; }
+	
+	/** @return The name of the command */
+	FName GetCommandName() const { return CommandName; }
+
+	/** @return The name of the context where the command is valid */
+	FName GetBindingContext() const { return BindingContext; }
+
+	/** Sets the new active gesture for this command */
+	void SetActiveGesture( const FInputGesture& NewGesture );
+
+	/** Removes the active gesture from this command */
+	void RemoveActiveGesture();
+
+	/** 
+	 * Makes a tooltip for this command.
+	 * @param	InText	Optional dynamic text to be displayed in the tooltip.
+	 * @return	The tooltip widget
+	 */
+	TSharedRef<class SToolTip> MakeTooltip( const TAttribute<FText>& InText = TAttribute<FText>() , const TAttribute< EVisibility >& InToolTipVisibility = TAttribute<EVisibility>()) const;
+
+private:
+
+	/** Input command that executes this action */
+	TSharedRef<FInputGesture> ActiveGesture;
+
+	/** Default display name of the command */
+	FText Label;
+
+	/** Localized help text for this command */
+	FText Description;
+
+	/** The default input gesture for this command (can be invalid) */
+	FInputGesture DefaultGesture;
+
+	/** Brush name for icon to use in tool bars and menu items to represent this command */
+	FSlateIcon Icon;
+
+	/** Brush name for icon to use in tool bars and menu items to represent this command in its toggled on (checked) state*/
+	FName UIStyle;
+
+	/** Name of the command */
+	FName CommandName;
+
+	/** The context in which this command is active */
+	FName BindingContext;
+
+	/** The type of user interface to associated with this action */
+	EUserInterfaceActionType::Type UserInterfaceType;
+};
diff --git a/Engine/Source/Runtime/Slate/Public/Framework/Commands/UICommandList.h b/Engine/Source/Runtime/Slate/Public/Framework/Commands/UICommandList.h
new file mode 100644
--- /dev/null
+++ b/Engine/Source/Runtime/Slate/Public/Framework/Commands/UICommandList.h
@@ -0,0 +1,181 @@
+// Copyright 1998-2014 Epic Games, Inc. All Rights Reserved.
+
+/*=============================================================================
+	UICommandList.h: Declares the FUICommandList class.
+=============================================================================*/
+
+#pragma once
+
+
+class SLATE_API FUICommandList
+	: public TSharedFromThis<FUICommandList>
+{
+public:
+
+	/** Determines if this UICommandList is capable of producing an action for the supplied command */
+	DECLARE_DELEGATE_RetVal_OneParam(bool, FCanProduceActionForCommand, const TSharedRef<const FUICommandInfo>& /*Command*/);
+
+	/**
+	 * Maps a command info to a series of delegates that are executed by a multibox or mouse/keyboard input
+	 *
+	 * @param InUICommandInfo	The command info to map
+	 * @param ExecuteAction		The delegate to call when the command should be executed
+	 */
+	void MapAction( const TSharedPtr< const FUICommandInfo > InUICommandInfo, FExecuteAction ExecuteAction );
+
+	/**
+	 * Maps a command info to a series of delegates that are executed by a multibox or mouse/keyboard input
+	 *
+	 * @param InUICommandInfo	The command info to map
+	 * @param ExecuteAction		The delegate to call when the command should be executed
+	 * @param CanExecuteAction	The delegate to call to see if the command can be executed
+	 */
+	void MapAction( const TSharedPtr< const FUICommandInfo > InUICommandInfo, FExecuteAction ExecuteAction, FCanExecuteAction CanExecuteAction );
+
+	/**
+	 * Maps a command info to a series of delegates that are executed by a multibox or mouse/keyboard input
+	 *
+	 * @param InUICommandInfo	The command info to map
+	 * @param ExecuteAction		The delegate to call when the command should be executed
+	 * @param CanExecuteAction	The delegate to call to see if the command can be executed
+	 * @param IsCheckedDelegate	The delegate to call to see if the command should appear checked when visualized in a multibox
+	 */
+	void MapAction( const TSharedPtr< const FUICommandInfo > InUICommandInfo, FExecuteAction ExecuteAction, FCanExecuteAction CanExecuteAction, FIsActionChecked IsCheckedDelegate );
+
+	/**
+	 * Maps a command info to a series of delegates that are executed by a multibox or mouse/keyboard input
+	 *
+	 * @param InUICommandInfo	The command info to map
+	 * @param ExecuteAction		The delegate to call when the command should be executed
+	 * @param CanExecuteAction	The delegate to call to see if the command can be executed
+	 * @param IsCheckedDelegate	The delegate to call to see if the command should appear checked when visualized in a multibox
+	 * @param IsVisibleDelegate	The delegate to call to see if the command should appear or be hidden when visualized in a multibox
+	 */
+	void MapAction( const TSharedPtr< const FUICommandInfo > InUICommandInfo, FExecuteAction ExecuteAction, FCanExecuteAction CanExecuteAction, FIsActionChecked IsCheckedDelegate, FIsActionButtonVisible IsVisibleDelegate );
+
+	/**
+	 * Maps a command info to a series of delegates that are executed by a multibox or mouse/keyboard input
+	 *
+	 * @param InUICommandInfo	The command info to map
+	 * @param InUIAction		Action to map to this command
+	 */
+	void MapAction( const TSharedPtr< const FUICommandInfo > InUICommandInfo, const FUIAction& InUIAction );
+
+	/**
+	 * Append commands in InCommandsToAppend to this command list.
+	 */
+	void Append( const TSharedRef<FUICommandList>& InCommandsToAppend );
+
+	/**
+	 * Executes the action associated with the provided command info
+	 * Note: It is assumed at this point that CanExecuteAction was already checked
+	 *
+	 * @param InUICommandInfo	The command info execute
+	 */
+	bool ExecuteAction( const TSharedRef< const FUICommandInfo > InUICommandInfo ) const;
+
+	/**
+	 * Calls the CanExecuteAction associated with the provided command info to see if ExecuteAction can be called
+	 *
+	 * @param InUICommandInfo	The command info execute
+	 */
+	bool CanExecuteAction( const TSharedRef< const FUICommandInfo > InUICommandInfo ) const;
+
+	/**
+	 * Attempts to execute the action associated with the provided command info
+	 * Note: This will check if the action can be executed before finally executing the action
+	 *
+	 * @param InUICommandInfo	The command info execute
+	 */
+	bool TryExecuteAction( const TSharedRef< const FUICommandInfo > InUICommandInfo ) const;
+
+	/**
+	 * Calls the IsVisible delegate associated with the provided command info to see if the command should be visible in a toolbar
+	 *
+	 * @param InUICommandInfo	The command info execute
+	 */
+	EVisibility GetVisibility( const TSharedRef< const FUICommandInfo > InUICommandInfo ) const;
+
+	/**
+	 * Calls the IsChecked delegate to see if the visualization of this command in a multibox should appear checked
+	 *
+	 * @param InUICommandInfo	The command info execute
+	 */
+	bool IsChecked( const TSharedRef< const FUICommandInfo > InUICommandInfo ) const;
+
+	/**
+	 * Processes any UI commands which are activated by the specified keyboard event
+	 *
+	 * @param InKeyboardEvent	The keyboard event to check
+	 *
+	 * @return true if an action was processed
+	 */
+	bool ProcessCommandBindings( const FKeyboardEvent& InKeyboardEvent ) const;
+
+	/**
+	 * Processes any UI commands which are activated by the specified mouse event
+	 *
+	 * @param InKeyboardEvent	The mouse event to check
+	 *
+	 * @return true if an action was processed
+	 */
+	bool ProcessCommandBindings( const FPointerEvent& InMouseEvent ) const;
+
+	/** Sets the delegate that determines if this UICommandList is capable of producing an action for the supplied command */
+	void SetCanProduceActionForCommand( const FCanProduceActionForCommand& NewCanProduceActionForCommand ) { CanProduceActionForCommand = NewCanProduceActionForCommand; }
+
+	/** 
+	  * Attempts to find an action for the specified command in the current UICommandList. This is a wrapper for GetActionForCommandRecursively.
+	  *
+	  * @param Command				The UI command for which you are discovering an action
+	  */
+	const FUIAction* GetActionForCommand(TSharedPtr<const FUICommandInfo> Command) const;
+
+protected:
+
+	/**
+	 * Helper function to execute delegate or exec command associated with a command (if valid)
+	 *
+	 * @param Key		The current key that is pressed
+	 * @param bCtrl		True if control is pressed
+	 * @param bAlt		True if alt is pressed
+	 * @param bShift	True if shift is pressed
+	 * @param bRepeat	True if command is repeating (held)
+	 * @return True if a command was executed, False otherwise
+	 */
+	bool ConditionalProcessCommandBindings( const FKey Key, bool bCtrl, bool bAlt, bool bShift, bool bRepeat ) const;
+
+	/** 
+	  * Attempts to find an action for the specified command in the current UICommandList. If it is not found, the action for the
+	  * specified command is discovered in the children recursively then the parents recursively.
+	  *
+	  * @param Command				The UI command for which you are discovering an action
+	  * @param bIncludeChildren		If true, children of this command list will be searched in the event that the action is not found
+	  * @param bIncludeParents		If true, parents of this command list will be searched in the event that the action is not found
+	  * @param OutVisitedLists		The set of visited lists during recursion. This is used to prevent cycles.
+	  */
+	const FUIAction* GetActionForCommandRecursively(const TSharedRef<const FUICommandInfo>& Command, bool bIncludeChildren, bool bIncludeParents, TSet<TSharedRef<const FUICommandList>>& InOutVisitedLists) const;
+
+	/** Returns all contexts associated with this list. This is a wrapper for GatherContextsForListRecursively */
+	void GatherContextsForList(TSet<FName>& OutAllContexts) const;
+
+	/** Returns all contexts associated with this list. */
+	void GatherContextsForListRecursively(TSet<FName>& OutAllContexts, TSet<TSharedRef<const FUICommandList>>& InOutVisitedLists) const;
+
+private:
+
+	typedef TMap< const TSharedPtr< const FUICommandInfo >, FUIAction > FUIBindingMap;
+	
+	/** Known contexts in this list.  Each context must be known so we can quickly look up commands from bindings */
+	TSet<FName> ContextsInList;
+
+	/** Mapping of command to action */
+	FUIBindingMap UICommandBindingMap;
+
+	/** The list of parent and children UICommandLists */
+	TArray<TWeakPtr<FUICommandList>> ParentUICommandLists;
+	TArray<TWeakPtr<FUICommandList>> ChildUICommandLists;
+
+	/** Determines if this UICommandList is capable of producing an action for the supplied command */
+	FCanProduceActionForCommand CanProduceActionForCommand;
+};
diff --git a/Engine/Source/Runtime/Slate/Public/Framework/Docking/LayoutService.h b/Engine/Source/Runtime/Slate/Public/Framework/Docking/LayoutService.h
new file mode 100644
--- /dev/null
+++ b/Engine/Source/Runtime/Slate/Public/Framework/Docking/LayoutService.h
@@ -0,0 +1,53 @@
+// Copyright 1998-2014 Epic Games, Inc. All Rights Reserved.
+
+#pragma once
+
+
+struct SLATE_API FLayoutSaveRestore
+{
+	/** Gets the ini section label for the additional layout configs */
+	static const FString& GetAdditionalLayoutConfigIni();
+
+	/** Write the LayoutToSave out into a a config file uses GEditorUserSettingsIni
+	 *
+	 * @param LayoutToSave the layout to save.
+	 */
+	static void SaveTheLayout( const TSharedRef<FTabManager::FLayout>& LayoutToSave );
+
+	/** Write the layout out into a named config file.
+	 *
+	 * @param LayoutToSave the layout to save.
+	 * @param ConfigFileName file to be save to
+	 */
+	static void SaveToConfig( const TSharedRef<FTabManager::FLayout>& LayoutToSave, FString ConfigFileName );
+
+	/** Given a named DefaultLayout, return any saved versions of it. Otherwise return the default layout uses GEditorUserSettingsIni
+	 *
+	 * @param DefaultLayout the layout to be used if the file does not exist.
+	 *
+	 * @return Loaded layout or the default
+	 */
+	static TSharedRef<FTabManager::FLayout> LoadUserConfigVersionOf( const TSharedRef<FTabManager::FLayout>& DefaultLayout );
+
+	/** Given a named DefualtLayout, return any saved version of it from the given ini file, otherwise return the default, also default to open tabs based on bool
+	 *
+	 * @param DefaultLayout the layout to be used if the file does not exist.
+	 * @param ConfigFileName file to be used to load an existing layout
+	 *
+	 * @return Loaded layout or the default
+	 */
+	static TSharedRef<FTabManager::FLayout> LoadFromConfig( const TSharedRef<FTabManager::FLayout>& DefaultLayout, FString ConfigFileName);
+
+private:
+	/**
+	 * Make a Json string friendly for writing out to UE .ini config files.
+	 * The opposite of GetLayoutStringFromIni.
+	 */
+	static FString PrepareLayoutStringForIni(const FString& LayoutString);
+
+	/**
+	 * Convert from UE .ini Json string to a vanilla Json string.
+	 * The opposite of PrepareLayoutStringForIni.
+	 */
+	static FString GetLayoutStringFromIni(const FString& LayoutString);
+};
\ No newline at end of file
