Slate: Refactored core Slate implementation into SlateCore module in preparation for UMG.

diff --git a/Engine/Source/Runtime/Slate/Public/Framework/Docking/TabManager.h b/Engine/Source/Runtime/Slate/Public/Framework/Docking/TabManager.h
new file mode 100644
--- /dev/null
+++ b/Engine/Source/Runtime/Slate/Public/Framework/Docking/TabManager.h
@@ -0,0 +1,849 @@
+// Copyright 1998-2014 Epic Games, Inc. All Rights Reserved.
+
+#pragma once
+
+class SDockingArea;
+class SDockTab;
+class SDockingTabStack;
+struct FTabMatcher;
+class FTabManager;
+
+DECLARE_MULTICAST_DELEGATE_TwoParams(
+	FOnActiveTabChanged,
+	/** Previously active tab */
+	TSharedPtr<SDockTab>,
+	/** Newly active tab */
+	TSharedPtr<SDockTab> );
+	
+
+struct FTabId
+{
+	FTabId( )
+		: InstanceId(INDEX_NONE)
+	{ }
+
+	FTabId( const FName& InTabType, const int32 InInstanceId )
+		: TabType(InTabType)
+		, InstanceId(InInstanceId)
+	{ }
+
+	FTabId( const FName& InTabType )
+		: TabType(InTabType)
+		, InstanceId(INDEX_NONE)
+	{ }
+
+	/** Document tabs allow multiple instances of the same tab type. The placement rules for these tabs are left up for the specific use-cases. These tabs are not persisted. */
+	bool IsTabPersistable() const
+	{
+		return InstanceId == INDEX_NONE;
+	}
+
+	bool operator==( const FTabId& Other ) const
+	{
+		return TabType == Other.TabType && (InstanceId == INDEX_NONE || Other.InstanceId == INDEX_NONE || InstanceId == Other.InstanceId) ;
+	}
+
+	FString ToString() const
+	{
+		return (InstanceId == INDEX_NONE)
+			? TabType.ToString()
+			: FString::Printf( TEXT("%s : %d"), *(TabType.ToString()), InstanceId );
+	}
+
+	FText ToText() const
+	{
+		FFormatNamedArguments Args;
+		Args.Add( TEXT("TabType"), FText::FromName( TabType ) );
+		Args.Add( TEXT("InstanceIdNumber"), FText::AsNumber( InstanceId ) );
+
+		return (InstanceId == INDEX_NONE)
+			? FText::FromName( TabType )
+			: FText::Format( NSLOCTEXT("TabManager", "TabIdFormat", "{TabType} : {InstanceIdNumber}"), Args );
+	}
+
+	FName TabType;
+	int32 InstanceId;
+};
+
+
+class FSpawnTabArgs
+{
+	public:
+	FSpawnTabArgs( const TSharedPtr<SWindow>& InOwnerWindow, const FTabId& InTabBeingSpawenedId )
+	: TabIdBeingSpawned(InTabBeingSpawenedId)
+	, OwnerWindow(InOwnerWindow)
+	{
+	}
+
+	const TSharedPtr<SWindow>& GetOwnerWindow() const
+	{
+		return OwnerWindow;
+	}
+
+	const FTabId& GetTabId() const 
+	{
+		return TabIdBeingSpawned;
+	}
+
+	private:
+	FTabId TabIdBeingSpawned;
+	TSharedPtr<SWindow> OwnerWindow;
+};
+
+
+/**
+ * Invoked when a tab needs to be spawned.
+ */
+DECLARE_DELEGATE_RetVal_OneParam( TSharedRef<SDockTab>, FOnSpawnTab, const FSpawnTabArgs& );
+
+/**
+ * Allows users to provide custom logic when searching for a tab to reuse.
+ * The TabId that is being searched for is provided as a courtesy, but does not have to be respected.
+ */
+DECLARE_DELEGATE_RetVal_OneParam( TSharedPtr<SDockTab>, FOnFindTabToReuse, const FTabId& )
+
+/** An enum to describe how TabSpawnerEntries will be handled by menus. */
+namespace ETabSpawnerMenuType
+{
+	enum Type
+	{
+		Display,		// Display this spawner in menus
+		Hide			// Do not display this spawner in menus, it will be invoked manually
+	};
+}
+
+struct FTabSpawnerEntry : public FWorkspaceItem
+{
+	FTabSpawnerEntry( const FName& InTabType, const FOnSpawnTab& InSpawnTabMethod )
+		: FWorkspaceItem( FText(), FSlateIcon(), false )
+		, TabType( InTabType )
+		, OnSpawnTab( InSpawnTabMethod )
+		, OnFindTabToReuse()
+		, MenuType(ETabSpawnerMenuType::Display)
+		, SpawnedTabPtr()
+
+	{
+	}
+
+	FTabSpawnerEntry& SetIcon( const FSlateIcon& InIcon)
+	{
+		Icon = InIcon;
+		return *this;
+	}
+
+	FTabSpawnerEntry& SetDisplayName( const FText& InLegibleName )
+	{
+		DisplayName = InLegibleName;
+		return *this;
+	}
+
+	FTabSpawnerEntry& SetTooltipText( const FText& InTooltipText )
+	{
+		TooltipText = InTooltipText;
+		return *this;
+	}
+
+	FTabSpawnerEntry& SetGroup( const TSharedRef<FWorkspaceItem>& InGroup )
+	{
+		InGroup->AddItem(SharedThis(this));
+		return *this;
+	}
+
+	FTabSpawnerEntry& SetReuseTabMethod( const FOnFindTabToReuse& InReuseTabMethod )
+	{
+		OnFindTabToReuse = InReuseTabMethod;
+		return *this;
+	}
+
+	FTabSpawnerEntry& SetMenuType( ETabSpawnerMenuType::Type InMenuType )
+	{
+		MenuType = InMenuType;
+		return *this;
+	}
+
+	virtual TSharedPtr<FTabSpawnerEntry> AsSpawnerEntry() OVERRIDE
+	{
+		return SharedThis(this);
+	}
+
+private:
+	FName TabType;
+	FOnSpawnTab OnSpawnTab;
+	/** When this method is not provided, we assume that the tab should only allow 0 or 1 instances */
+	FOnFindTabToReuse OnFindTabToReuse;
+	ETabSpawnerMenuType::Type MenuType;
+
+	TWeakPtr<SDockTab> SpawnedTabPtr;
+
+	FORCENOINLINE bool IsSoleTabInstanceSpawned() const
+	{
+		// Items that allow multiple instances need a custom way to find and reuse tabs.
+		return SpawnedTabPtr.IsValid();
+	}
+
+	friend class FTabManager;
+};
+
+
+namespace ETabState
+{
+	enum Type
+	{
+		OpenedTab = 0x1 << 0,
+		ClosedTab = 0x1 << 1
+	};
+}
+
+
+class SLATE_API FTabManager : public TSharedFromThis<FTabManager>
+{
+	friend class FGlobalTabmanager;
+	public:
+
+		class FStack;
+		class FSplitter;
+		class FArea;
+		class FLayout;
+
+		DECLARE_DELEGATE_OneParam( FOnPersistLayout, const TSharedRef<FLayout>& );
+
+		class SLATE_API FLayoutNode : public TSharedFromThis<FLayoutNode>
+		{
+			friend class FTabManager;
+
+		public:
+			
+			virtual TSharedPtr<FStack> AsStack();
+
+			virtual TSharedPtr<FSplitter> AsSplitter();
+
+			virtual TSharedPtr<FArea> AsArea();
+
+			float GetSizeCoefficient() const { return SizeCoefficient; }
+
+		protected:
+			FLayoutNode()
+			: SizeCoefficient(1.0f)
+			{
+			}
+
+			float SizeCoefficient;
+		};
+
+		struct FTab
+		{
+			FTab( const FTabId& InTabId, ETabState::Type InTabState )
+			: TabId(InTabId)
+			, TabState(InTabState)
+			{
+			}
+
+			bool operator==( const FTab& Other ) const
+			{
+				return this->TabId == Other.TabId && this->TabState == Other.TabState;
+			}
+
+			FTabId TabId;
+			ETabState::Type TabState;
+		};
+
+		class SLATE_API FStack : public FLayoutNode
+		{
+				friend class FTabManager;
+				friend class FLayout;
+				friend class SDockingTabStack;
+
+			public:				
+
+				TSharedRef<FStack> AddTab( const FName& TabType, ETabState::Type TabState )
+				{
+					Tabs.Add( FTab( FTabId(TabType), TabState ) );
+					return SharedThis(this);
+				}
+
+				TSharedRef<FStack> AddTab( const FTabId& TabId, ETabState::Type TabState )
+				{
+					Tabs.Add( FTab(TabId, TabState) );
+					return SharedThis(this);
+				}
+
+				TSharedRef<FStack> SetSizeCoefficient( const float InSizeCoefficient )
+				{
+					SizeCoefficient = InSizeCoefficient;
+					return SharedThis(this);
+				}
+
+				TSharedRef<FStack> SetHideTabWell( const bool InHideTabWell )
+				{
+					bHideTabWell = InHideTabWell;
+					return SharedThis(this);
+				}
+
+				TSharedRef<FStack> SetForegroundTab( const FTabId& TabId )
+				{
+					ForegroundTabId = TabId;
+					return SharedThis(this);
+				}
+
+				virtual TSharedPtr<FStack> AsStack() OVERRIDE
+				{
+					return SharedThis(this);
+				}
+
+				virtual ~FStack()
+				{
+				}
+
+			protected:
+
+				FStack()
+				: Tabs()
+				, bHideTabWell(false)
+				, ForegroundTabId(NAME_None)
+				{
+				}
+
+				TArray<FTab> Tabs;
+				bool bHideTabWell;
+				FTabId ForegroundTabId;
+		};
+
+
+		class SLATE_API FSplitter : public FLayoutNode
+		{
+				friend class FTabManager;
+		
+			public:
+
+				TSharedRef<FSplitter> Split( TSharedRef<FLayoutNode> InNode )
+				{
+					ChildNodes.Add(InNode);
+					return SharedThis(this);
+				}
+
+				TSharedRef<FSplitter> SetSizeCoefficient( const float InSizeCoefficient )
+				{
+					SizeCoefficient = InSizeCoefficient;
+					return SharedThis(this);
+				}
+
+				TSharedRef<FSplitter> SetOrientation( const EOrientation InOrientation )
+				{
+					Orientation = InOrientation;
+					return SharedThis(this);
+				}
+
+				virtual TSharedPtr<FSplitter> AsSplitter() OVERRIDE
+				{
+					return SharedThis(this);
+				}
+
+				EOrientation GetOrientation() const { return Orientation; }
+
+				virtual ~FSplitter()
+				{
+				}
+
+			protected:
+				FSplitter()
+				: Orientation(Orient_Horizontal)
+				{
+				}
+
+				EOrientation Orientation;
+				TArray< TSharedRef<FLayoutNode> > ChildNodes;
+		};
+
+
+		class SLATE_API FArea : public FSplitter
+		{
+				friend class FTabManager;
+		
+			public:			
+				enum EWindowPlacement
+				{
+					Placement_NoWindow,
+					Placement_Automatic,
+					Placement_Specified
+				};
+
+				TSharedRef<FArea> Split( TSharedRef<FLayoutNode> InNode )
+				{
+					ChildNodes.Add(InNode);
+					return SharedThis(this);
+				}
+				
+				TSharedRef<FArea> SetOrientation( const EOrientation InOrientation )
+				{
+					Orientation = InOrientation;
+					return SharedThis(this);
+				}
+
+				TSharedRef<FArea> SetWindow( FVector2D InPosition, bool IsMaximized )
+				{
+					WindowPlacement = Placement_Specified;
+					WindowPosition = InPosition;
+					bIsMaximized = IsMaximized;
+					return SharedThis(this);
+				}
+
+				virtual TSharedPtr<FArea> AsArea() OVERRIDE
+				{
+					return SharedThis(this);
+				}
+
+				virtual ~FArea()
+				{
+				}
+
+			protected:
+				FArea( const float InWidth, const float InHeight )
+				: WindowPlacement(Placement_Automatic)
+				, WindowPosition(FVector2D(0,0))
+				, WindowSize(InWidth, InHeight)
+				, bIsMaximized( false )
+				{
+				}
+
+				EWindowPlacement WindowPlacement;
+				FVector2D WindowPosition;
+				FVector2D WindowSize;
+				bool bIsMaximized;
+		};
+
+
+		class SLATE_API FLayout : public TSharedFromThis<FLayout>
+		{
+				friend class FTabManager;
+				
+			public:
+
+				TSharedRef<FLayout> AddArea( const TSharedRef<FArea>& InArea )
+				{
+					Areas.Add( InArea );
+					return SharedThis(this);
+				}
+
+				const TWeakPtr<FArea>& GetPrimaryArea() const
+				{
+					return PrimaryArea;
+				}
+				
+			public:
+				static TSharedPtr<FTabManager::FLayout> NewFromString( const FString& LayoutAsText );
+				FName GetLayoutName() const;
+				FString ToString() const;
+
+			protected:				
+				static TSharedRef<class FJsonObject> PersistToString_Helper(const TSharedRef<FLayoutNode>& NodeToPersist);
+				static TSharedRef<FLayoutNode> NewFromString_Helper( TSharedPtr<FJsonObject> JsonObject );
+
+				FLayout(const FName& InLayoutName)
+					: LayoutName(InLayoutName)
+				{}
+
+				TWeakPtr< FArea > PrimaryArea;
+				TArray< TSharedRef<FArea> > Areas;
+				/** The layout will be saved into a config file with this name. E.g. LevelEditorLayout or MaterialEditorLayout */
+				FName LayoutName;
+		};	
+
+
+		friend class FPrivateApi;
+		class FPrivateApi
+		{
+			public:
+				FPrivateApi( FTabManager& InTabManager )
+				: TabManager( InTabManager )
+				{
+				}
+
+				TSharedPtr<SWindow> GetParentWindow() const;
+				void OnDockAreaCreated( const TSharedRef<SDockingArea>& NewlyCreatedDockArea );
+				/** Notify the tab manager that a tab has been relocated. If the tab now lives in a new window, the NewOwnerWindow should be a valid pointer. */
+				void OnTabRelocated( const TSharedRef<SDockTab>& RelocatedTab, const TSharedPtr<SWindow>& NewOwnerWindow );
+				void OnTabOpening( const TSharedRef<SDockTab>& TabBeingOpened );
+				void OnTabClosing( const TSharedRef<SDockTab>& TabBeingClosed );
+				void OnDockAreaClosing( const TSharedRef<SDockingArea>& DockAreaThatIsClosing );
+				void OnTabManagerClosing();
+				const TArray< TWeakPtr<SDockingArea> >& GetLiveDockAreas() const;
+				/**
+				 * Notify the tab manager that the NewForegroundTab was brought to front and the BackgroundedTab was send to the background as a result.
+				 */
+				void OnTabForegrounded( const TSharedPtr<SDockTab>& NewForegroundTab, const TSharedPtr<SDockTab>& BackgroundedTab );
+
+				void ShowWindows();
+				void HideWindows();
+
+			private:
+				FTabManager& TabManager;
+				
+		};
+
+		FTabManager::FPrivateApi& GetPrivateApi();
+
+	public:
+		static TSharedRef<FLayout> NewLayout( const FName LayoutName )
+		{
+			return MakeShareable( new FLayout(LayoutName) );
+		}
+		
+		static TSharedRef<FArea> NewPrimaryArea()
+		{
+			TSharedRef<FArea>Area = MakeShareable( new FArea(0,0) );
+			Area->WindowPlacement = FArea::Placement_NoWindow;
+			return Area;
+		}
+
+		static TSharedRef<FArea> NewArea( const float Width, const float Height )
+		{
+			return MakeShareable( new FArea( Width, Height ) );
+		}
+
+		static TSharedRef<FArea> NewArea( const FVector2D& WindowSize )
+		{
+			return MakeShareable( new FArea( WindowSize.X, WindowSize.Y ) );
+		}
+		
+		static TSharedRef<FStack> NewStack() 
+		{
+			return MakeShareable( new FStack() );
+		}
+		
+		static TSharedRef<FSplitter> NewSplitter()
+		{
+			return MakeShareable( new FSplitter() );
+		}
+
+		void SetOnPersistLayout( const FOnPersistLayout& InHandler );
+
+		/** Close all live areas and wipe all the persisted areas. */
+		void CloseAllAreas();
+
+		/** Gather the persistent layout */
+		TSharedRef<FTabManager::FLayout> PersistLayout() const;
+
+		/** Gather the persistent layout and execute the custom delegate for saving it to persistent storage (e.g. into config files) */
+		void SavePersistentLayout();
+
+		/**
+		 * Register a new tab spawner with the tab manager.  The spawner will be called when anyone calls
+		 * InvokeTab().
+		 * @param TabId The TabId to register the spawner for.
+		 * @param OnSpawnTab The callback that will be used to spawn the tab.
+		 * @return The registration entry for the spawner.
+		 */
+		FTabSpawnerEntry& RegisterTabSpawner( const FName TabId, const FOnSpawnTab& OnSpawnTab );
+		
+		/**
+		 * Unregisters the tab spawner matching the provided TabId.
+		 * @param TabId The TabId to remove the spawner for.
+		 * @return true if a spawner was found for this TabId, otherwise false.
+		 */
+		bool UnregisterTabSpawner( const FName TabId );
+
+		/**
+		 * Unregisters all tab spawners.
+		 */
+		void UnregisterAllTabSpawners();
+
+		TSharedPtr<SWidget> RestoreFrom( const TSharedRef<FLayout>& Layout, const TSharedPtr<SWindow>& ParentWindow, const bool bEmbedTitleAreaContent = false );
+
+		void PopulateTabSpawnerMenu( FMenuBuilder& PopulateMe, TSharedRef<FWorkspaceItem> MenuStructure );
+
+		void PopulateTabSpawnerMenu( FMenuBuilder &PopulateMe, const FName& TabType );
+
+		void DrawAttention( const TSharedRef<SDockTab>& TabToHighlight );
+
+		struct ESearchPreference
+		{
+			enum Type
+			{
+				PreferLiveTab,
+				RequireClosedTab
+			};
+		};
+		
+		/** Insert a new UnmanagedTab document tab next to an existing tab (closed or open) that has the PlaceholdId. */
+		void InsertNewDocumentTab( FName PlaceholderId, ESearchPreference::Type SearchPreference, const TSharedRef<SDockTab>& UnmanagedTab );
+
+		/**
+		 * Much like InsertNewDocumentTab, but the UnmanagedTab is not seen by the user as newly-created.
+		 * e.g. Opening an restores multiple previously opened documents; these are not seen as new tabs.
+		 */
+		void RestoreDocumentTab( FName PlaceholderId, ESearchPreference::Type SearchPreference, const TSharedRef<SDockTab>& UnmanagedTab );
+
+		/**
+		 * Opens tab if it is closed at the last known location.  If it already exists, it will draw attention to the tab.
+		 * 
+		 * @param TabId The tab identifier.
+		 * @return The existing or newly spawned tab instance.
+		 */
+		TSharedRef<SDockTab> InvokeTab( const FTabId& TabId );
+
+		virtual ~FTabManager()
+		{
+		}
+
+
+	protected:
+		void InvokeTabForMenu( FName TabId );
+
+	protected:
+		
+		void InsertDocumentTab( FName PlaceholderId, ESearchPreference::Type SearchPreference, const TSharedRef<SDockTab>& UnmanagedTab, bool bPlaySpawnAnim );
+			
+		void PopulateTabSpawnerMenu_Helper( FMenuBuilder& PopulateMe, struct FPopulateTabSpawnerMenu_Args Args );
+
+		void MakeSpawnerMenuEntry( FMenuBuilder &PopulateMe, const TSharedPtr<FTabSpawnerEntry> &SpawnerNode );
+
+		TSharedRef<SDockTab> InvokeTab_Internal( const FTabId& TabId );
+		TSharedPtr<SDockingTabStack> FindPotentiallyClosedTab( const FTabId& ClosedTabId );
+
+		typedef TMap< FName, TSharedRef<FTabSpawnerEntry> > FTabSpawner;
+		
+		static TSharedRef<FTabManager> New( const TSharedPtr<SDockTab>& InOwnerTab, const TSharedRef<FTabSpawner>& InNomadTabSpawner )
+		{
+			return MakeShareable( new FTabManager(InOwnerTab, InNomadTabSpawner) );
+		}
+
+		FTabManager( const TSharedPtr<SDockTab>& InOwnerTab, const TSharedRef<FTabManager::FTabSpawner> & InNomadTabSpawner );
+
+		TSharedRef<SDockingArea> RestoreArea( const TSharedRef<FArea>& AreaToRestore, const TSharedPtr<SWindow>& InParentWindow, const bool bEmbedTitleAreaContent = false );
+
+		TSharedRef<class SDockingNode> RestoreArea_Helper( const TSharedRef<FLayoutNode>& LayoutNode, const TSharedPtr<SWindow>& ParentWindow, const bool bEmbedTitleAreaContent );
+
+		void RestoreSplitterContent( const TSharedRef<FSplitter>& SplitterNode, const TSharedRef<class SDockingSplitter>& SplitterWidget, const TSharedPtr<SWindow>& ParentWindow );
+		
+		bool IsValidTabForSpawning( const FTab& SomeTab ) const;
+		TSharedRef<SDockTab> SpawnTab( const FTabId& TabId, const TSharedPtr<SWindow>& ParentWindow );
+
+		TSharedPtr<SDockTab> FindExistingLiveTab( const FTabId& TabId ) const;
+		TSharedPtr<class SDockingTabStack> FindTabInLiveAreas( const FTabMatcher& TabMatcher ) const;
+		static TSharedPtr<class SDockingTabStack> FindTabInLiveArea( const FTabMatcher& TabMatcher, const TSharedRef<SDockingArea>& InArea );
+
+		template<typename MatchFunctorType> static bool HasAnyMatchingTabs( const TSharedRef<FTabManager::FLayoutNode>& SomeNode, const MatchFunctorType& Matcher );
+		bool HasOpenTabs( const TSharedRef<FTabManager::FLayoutNode>& SomeNode ) const;
+		bool HasValidTabs( const TSharedRef<FTabManager::FLayoutNode>& SomeNode ) const;
+
+		/**
+		 * Notify the tab manager that the NewForegroundTab was brought to front and the BackgroundedTab was send to the background as a result.
+		 */
+		virtual void OnTabForegrounded( const TSharedPtr<SDockTab>& NewForegroundTab, const TSharedPtr<SDockTab>& BackgroundedTab );
+		virtual void OnTabRelocated( const TSharedRef<SDockTab>& RelocatedTab, const TSharedPtr<SWindow>& NewOwnerWindow );
+		virtual void OnTabOpening( const TSharedRef<SDockTab>& TabBeingOpened );
+		virtual void OnTabClosing( const TSharedRef<SDockTab>& TabBeingClosed );
+		/** Invoked when a tab manager is closing down. */
+		virtual void OnTabManagerClosing();
+		/** Check these all tabs to see if it is OK to close them. Ignore the TabsToIgnore */
+		virtual bool CanCloseManager( const TSet< TSharedRef<SDockTab> >& TabsToIgnore = TSet< TSharedRef<SDockTab> >() );
+
+		static void GetAllStacks( const TSharedRef<SDockingArea>& InDockArea, TArray< TSharedRef<SDockingTabStack> >& OutTabStacks );
+
+		/** @return the stack that is under NodeToSearchUnder and contains TabIdToFind; Invalid pointer if not found. */
+		static TSharedPtr<FTabManager::FStack> FindTabUnderNode( const FTabMatcher& Matcher, const TSharedRef<FTabManager::FLayoutNode>& NodeToSearchUnder );
+		int32 FindTabInCollapsedAreas( const FTabMatcher& Matcher );
+		void RemoveTabFromCollapsedAreas( const FTabMatcher& Matcher );
+
+		/** Called when tab(s) have been added or windows created */
+		virtual void UpdateStats();
+		
+	private:
+		/** Checks all dock areas and adds up the number of open tabs and unique parent windows in the manager */
+		void GetRecordableStats( int32& OutTabCount, TArray<TSharedPtr<SWindow>>& OutUniqueParentWindows ) const;
+
+	protected:
+		FTabSpawner TabSpawner;
+		TSharedRef<FTabSpawner> NomadTabSpawner;
+		TSharedPtr<FTabSpawnerEntry> FindTabSpawnerFor(FName TabId);
+
+		TArray< TWeakPtr<SDockingArea> > DockAreas;
+		TArray< TSharedRef<FTabManager::FArea> > CollapsedDockAreas;
+
+		/** A Major tab that contains this TabManager's widgets. */
+		TWeakPtr<SDockTab> OwnerTabPtr;
+
+		/** Protected private API that must only be accessed by the docking framework internals */
+		TSharedRef<FPrivateApi> PrivateApi;
+
+		/** The name of the layout being used */
+		FName ActiveLayoutName;
+
+		/** Invoked when the tab manager is about to close */
+		FOnPersistLayout OnPersistLayout_Handler;
+
+		/**
+		 * Instance ID for document tabs. Allows us to distinguish between different document tabs at runtime.
+		 * This ID is never meant to be persisted, simply used to disambiguate between different documents, since most of thenm
+		 * will have the same Tab Type (which is usually document).
+		 */
+		int32 LastDocumentUID;
+
+		/** The fallback size for a window */
+		const static FVector2D FallbackWindowSize;
+
+		/**
+		 * Defensive: True when we are saving the visual state.
+		 * Inevitably someone will ask us to save layout while saving layout.
+		 * We will ignore that request them.
+		 */
+		bool bIsSavingVisualState;
+};
+
+
+
+
+
+
+class SLATE_API FGlobalTabmanager : public FTabManager
+{
+public:	
+
+	static const TSharedRef<FGlobalTabmanager>& Get();
+
+	/** Subscribe to notifications about the active tab changing */
+	void OnActiveTabChanged_Subscribe( const FOnActiveTabChanged::FDelegate& InDelegate );
+
+	/** Unsubscribe to notifications about the active tab changing */
+	void OnActiveTabChanged_Unsubscribe( const FOnActiveTabChanged::FDelegate& InDelegate );
+
+	/** @return the currently active tab; NULL pointer if there is no active tab */
+	TSharedPtr<SDockTab> GetActiveTab() const;
+
+	/** Activate the NewActiveTab. If NewActiveTab is NULL, the active tab is cleared. */
+	void SetActiveTab( const TSharedPtr<SDockTab>& NewActiveTab );
+
+	FTabSpawnerEntry& RegisterNomadTabSpawner( const FName TabId, const FOnSpawnTab& OnSpawnTab );
+	
+	void UnregisterNomadTabSpawner( const FName TabId );
+
+	void SetApplicationTitle( const FText& AppTitle );
+
+	const FText& GetApplicationTitle() const;
+
+	static TSharedRef<FGlobalTabmanager> New()
+	{
+		return MakeShareable( new FGlobalTabmanager() );
+	}
+
+	virtual bool CanCloseManager( const TSet< TSharedRef<SDockTab> >& TabsToIgnore = TSet< TSharedRef<SDockTab> >()) OVERRIDE;
+
+	/** Draw the user's attention to a child tab manager */
+	void DrawAttentionToTabManager( const TSharedRef<FTabManager>& ChildManager );
+
+	TSharedRef<FTabManager> NewTabManager( const TSharedRef<SDockTab>& InOwnerTab );
+
+	/** Persist and serialize the layout of every TabManager and the custom visual state of every Tab. */
+	void SaveAllVisualState();
+
+	/** Provide a window under which all other windows in this application should nest. */
+	void SetRootWindow( const TSharedRef<SWindow> InRootWindow );
+
+	/** The window under which all other windows in our app nest; might be null */
+	TSharedPtr<SWindow> GetRootWindow() const;
+
+	/** Adds a legacy tab type to the tab type redirection map so tabs loaded with this type will be automatically converted to the new type */
+	void AddLegacyTabType(FName InLegacyTabType, FName InNewTabType);
+
+	/** Returns true if the specified tab type is registered as a legacy tab */
+	bool IsLegacyTabType(FName InTabType) const;
+
+	/** If the specified TabType is deprecated, returns the new replacement tab type. Otherwise, returns InTabType */
+	FName GetTabTypeForPotentiallyLegacyTab(FName InTabType) const;
+
+	/** Returns the highest number of tabs that were open simultaneously during this session */
+	int32 GetMaximumTabCount() const { return AllTabsMaxCount; }
+
+	/** Returns the highest number of parent windows that were open simultaneously during this session */
+	int32 GetMaximumWindowCount() const { return AllAreasWindowMaxCount; }
+
+protected:
+	virtual void OnTabForegrounded( const TSharedPtr<SDockTab>& NewForegroundTab, const TSharedPtr<SDockTab>& BackgroundedTab ) OVERRIDE;
+	virtual void OnTabRelocated( const TSharedRef<SDockTab>& RelocatedTab, const TSharedPtr<SWindow>& NewOwnerWindow ) OVERRIDE;
+	virtual void OnTabClosing( const TSharedRef<SDockTab>& TabBeingClosed ) OVERRIDE;
+	virtual void UpdateStats() OVERRIDE;
+
+public:
+	virtual void OnTabManagerClosing();
+
+
+private:
+	
+	/** Pairs of Major Tab and the TabManager that manages tabs within it. */
+	struct FSubTabManager
+	{
+		FSubTabManager( const TSharedRef<SDockTab>& InMajorTab, const TSharedRef<FTabManager>& InTabManager )
+		: MajorTab( InMajorTab )
+		, TabManager( InTabManager )
+		{
+		}
+
+		TWeakPtr<SDockTab> MajorTab;
+		TWeakPtr<FTabManager> TabManager;
+	};
+
+	struct FFoobar {
+			bool operator()(const FSubTabManager& InItem) const
+			{
+				return !InItem.MajorTab.IsValid();
+			}
+		};
+
+	struct FindByTab
+	{
+		FindByTab(const TSharedRef<SDockTab>& InTabToFind)
+			: TabToFind(InTabToFind)
+		{
+		}
+
+		bool Matches( const FGlobalTabmanager::FSubTabManager& TabManagerPair ) const
+		{
+			return TabManagerPair.TabManager.IsValid() && TabManagerPair.MajorTab.IsValid() && TabManagerPair.MajorTab.Pin() == TabToFind;
+		}
+
+		const TSharedRef<SDockTab>& TabToFind;
+	};
+
+	struct FindByManager
+	{
+		FindByManager(const TSharedRef<FTabManager>& InManagerToFind)
+			: ManagerToFind(InManagerToFind)
+		{
+		}
+
+		bool Matches( const FGlobalTabmanager::FSubTabManager& TabManagerPair ) const
+		{
+			return TabManagerPair.TabManager.IsValid() && TabManagerPair.MajorTab.IsValid() && TabManagerPair.TabManager.Pin() == ManagerToFind;
+		}
+
+		const TSharedRef<FTabManager>& ManagerToFind;
+	};
+
+	FGlobalTabmanager()
+	: FTabManager( TSharedPtr<SDockTab>(), MakeShareable( new FTabSpawner() ) )
+	, AllTabsMaxCount(0)
+	, AllAreasWindowMaxCount(0)
+	{
+	}
+
+	TArray< FSubTabManager > SubTabManagers;
+
+	/** The currently active tab; NULL if there is no active tab. */
+	TWeakPtr<SDockTab> ActiveTabPtr;
+
+	FOnActiveTabChanged OnActiveTabChanged;
+
+	FText AppTitle;
+
+	/** A window under which all of the windows in this application will nest. */
+	TWeakPtr<SWindow> RootWindowPtr;
+
+	/** A map that correlates deprecated tab types to new tab types */
+	TMap<FName, FName> LegacyTabTypeRedirectionMap;
+
+	/** Keeps track of the running-maximum number of tabs in all dock areas and sub-managers during this session */
+	int32 AllTabsMaxCount;
+
+	/** Keeps track of the running-maximum number of unique parent windows in all dock areas and sub-managers during this session */
+	int32 AllAreasWindowMaxCount;
+};
diff --git a/Engine/Source/Runtime/Slate/Public/Framework/Layout/InertialScrollManager.h b/Engine/Source/Runtime/Slate/Public/Framework/Layout/InertialScrollManager.h
new file mode 100644
--- /dev/null
+++ b/Engine/Source/Runtime/Slate/Public/Framework/Layout/InertialScrollManager.h
@@ -0,0 +1,53 @@
+// Copyright 1998-2014 Epic Games, Inc. All Rights Reserved.
+
+#pragma once
+
+/**
+ * A helper class to calculate inertial scrolling
+ */
+class SLATE_API FInertialScrollManager
+{
+public:
+	/**
+	  * Constructor
+	  * @param ScrollDecceleration	The acceleration against the velocity causing it to decay.
+	  * @param SampleTimeout		Samples older than this amount of time will be discarded.
+	  */
+	FInertialScrollManager(float ScrollDecceleration = 2048.f, double SampleTimeout = 0.1f);
+
+	/** Adds a scroll velocity sample to help calculate a smooth velocity */
+	void AddScrollSample(float Delta, double CurrentTime);
+
+	/** Updates the current scroll velocity. Call every frame. */
+	void UpdateScrollVelocity(const float InDeltaTime);
+
+	/** Instantly end the inertial scroll */
+	void ClearScrollVelocity();
+
+	/** Gets the calculated velocity of the scroll. */
+	float GetScrollVelocity() const { return ScrollVelocity; }
+
+private:
+	struct FScrollSample
+	{
+		double Time;
+		float Delta;
+
+		FScrollSample(float InDelta, double InTime)
+			: Time(InTime)
+			, Delta(InDelta)
+		{}
+	};
+
+	/** Used to calculate the appropriate scroll velocity over the last few frames while inertial scrolling */
+	TArray<FScrollSample> ScrollSamples;
+
+	/** The current velocity of the scroll */
+	float ScrollVelocity;
+
+	/** The acceleration against the velocity causing it to decay. */
+	float ScrollDecceleration;
+
+	/** Samples older than this amount of time will be discarded. */
+	double SampleTimeout;
+};
diff --git a/Engine/Source/Runtime/Slate/Public/Framework/Layout/ScrollyZoomy.h b/Engine/Source/Runtime/Slate/Public/Framework/Layout/ScrollyZoomy.h
new file mode 100644
--- /dev/null
+++ b/Engine/Source/Runtime/Slate/Public/Framework/Layout/ScrollyZoomy.h
@@ -0,0 +1,150 @@
+// Copyright 1998-2014 Epic Games, Inc. All Rights Reserved.
+
+/*=============================================================================
+	ScrollyZoomy.h: Declares the IScrollableZoomable interface and FScrollyZoomy class.
+=============================================================================*/
+
+#pragma once
+
+
+/**
+ * Have your widget implement this interface to allow FScrollyZoomy to interface with your widget
+ */
+class IScrollableZoomable
+{
+public:
+
+	/**
+	 * Override this to scroll your widget's content.
+	 *
+	 * @param Offset - The 2D offset to scroll by.
+	 *
+	 * @return True if anything was scrolled.
+	 */
+	virtual bool ScrollBy( const FVector2D& Offset ) = 0;
+
+	/**
+	 * Override this to zoom your widget's content.
+	 *
+	 * @param Amount - The amount to zoom by.
+	 *
+	 * @return True if anything was zoomed.
+	 */
+	virtual bool ZoomBy( const float Amount ) = 0;
+};
+
+
+/**
+ * Utility class that adds scrolling and zooming functionality to a widget.  Derived your widget class from IScrollableZoomable, then embed an
+ * instance of FScrollyZoomy as a widget member variable, and call this class's event handlers from your own widget's event handler callbacks.
+ */
+class SLATE_API FScrollyZoomy
+{
+public:
+
+	/** Default constructor for FScrollyZoomy */
+	FScrollyZoomy( );
+
+	/** 
+	 * Should be called every frame to update simulation state
+	 *
+	 * @param	DeltaTime			Time that's passed
+	 * @param	ScrollableZoomable	Interface to the widget to scroll/zoom
+	 */
+	void Tick( const float DeltaTime, IScrollableZoomable& ScrollableZoomable );
+
+	/**
+	 * Should be called when a mouse button is pressed
+	 *
+	 * @param	MouseEvent	The mouse event passed to a widget's OnMouseButtonDown() function
+	 *
+	 * @return	If FReply::Handled is returned, that should be passed as the result of the calling function
+	 */
+	FReply OnMouseButtonDown( const FPointerEvent& MouseEvent );
+
+	/**
+	 * Should be called when a mouse button is released
+	 *
+	 * @param	MyWidget	Pointer to the widget that owns this object
+	 * @param	MyGeometry	Geometry of the widget we're scrolling
+	 * @param	MouseEvent	The mouse event passed to a widget's OnMouseButtonUp() function
+	 *
+	 * @return	If FReply::Handled is returned, that should be passed as the result of the calling function
+	 */
+	FReply OnMouseButtonUp( const TSharedRef<SWidget> MyWidget, const FGeometry& MyGeometry, const FPointerEvent& MouseEvent );
+		
+	/**
+	 * Should be called when a mouse move event occurs
+	 *
+	 * @param	MyWidget	Pointer to the widget that owns this object
+	 * @param	ScrollableZoomable	Interface to the widget to scroll/zoom
+	 * @param	MyGeometry	Geometry of the widget we're scrolling
+	 * @param	MouseEvent	The mouse event passed to a widget's OnMouseMove() function
+	 *
+	 * @return	If FReply::Handled is returned, that should be passed as the result of the calling function
+	 */
+	FReply OnMouseMove( const TSharedRef<SWidget> MyWidget, IScrollableZoomable& ScrollableZoomable, const FGeometry& MyGeometry, const FPointerEvent& MouseEvent );
+
+	/**
+	 * Should be called from your widget's OnMouseLeave() override
+	 *
+	 * @param	MyWidget	Pointer to the widget that owns this object
+	 * @param	MouseEvent	The mouse leave event that was passed to the widget's OnMouseLeave() function
+	 */
+	void OnMouseLeave( const TSharedRef<SWidget> MyWidget, const FPointerEvent& MouseEvent );
+		
+	/**
+	 * Should be called by your widget when the mouse wheel is used
+	 *
+	 * @param	MouseEvent	The event passed to your widget's OnMouseWheel function
+	 * @param	ScrollableZoomable	Interface to the widget to scroll/zoom
+	 *
+	 * @return	If FReply::Handled is returned, that should be passed as the result of the calling function
+	 */
+	FReply OnMouseWheel( const FPointerEvent& MouseEvent, IScrollableZoomable& ScrollableZoomable );
+
+	/**
+	 * Call this from your widget's OnCursorQuery() function
+	 *
+	 * @return	The cursor reply to pass back
+	 */
+	FCursorReply OnCursorQuery( ) const;
+
+	/** @return Returns true if the user is actively scrolling */
+	bool IsRightClickScrolling( ) const;
+
+	/** @return Returns whether we need a software cursor to be rendered right now.  Should be called from your widget's OnPaint() function! */
+	bool NeedsSoftwareCursor( ) const;
+
+	/** @return Returns the position of the software cursor when NeedsSoftwareCursor() == true.  Should be called from your widget's OnPaint() function! */
+	const FVector2D& GetSoftwareCursorPosition( ) const;
+
+	/**
+	 * Call this from your widget's OnPaint to paint a software cursor, if needed 
+	 *
+	 * @param	AllottedGeometry		Widget geometry passed into OnPaint()
+	 * @param	MyClippingRect			Widget clipping rect passed into OnPaint()
+	 * @param	OutDrawElements			The draw element list
+	 * @param	LayerId					Layer Id
+	 *
+	 * @return	New layer Id
+	 */
+	int32 PaintSoftwareCursorIfNeeded( const FGeometry& AllottedGeometry, const FSlateRect& MyClippingRect, FSlateWindowElementList& OutDrawElements, int32 LayerId ) const;
+
+private:
+
+	// How much we scrolled while the right mouse button has been held.
+	float AmountScrolledWhileRightMouseDown;
+
+	// Whether the software cursor should be drawn in the view port.
+	bool bShowSoftwareCursor;
+
+	// The current position of the software cursor.
+	FVector2D SoftwareCursorPosition;
+
+	// Tracks simulation state for horizontal scrolling.
+	FInertialScrollManager HorizontalIntertia;
+
+	// Tracks simulation state for vertical scrolling.
+	FInertialScrollManager VerticalIntertia;
+};
diff --git a/Engine/Source/Runtime/Slate/Public/Framework/MarqueeRect.h b/Engine/Source/Runtime/Slate/Public/Framework/MarqueeRect.h
new file mode 100644
--- /dev/null
+++ b/Engine/Source/Runtime/Slate/Public/Framework/MarqueeRect.h
@@ -0,0 +1,60 @@
+// Copyright 1998-2014 Epic Games, Inc. All Rights Reserved.
+
+/**
+ * A convenient representation of a marquee selection
+ */
+struct FMarqueeRect
+{
+	/** Where the user began the marquee selection */
+	FVector2D StartPoint;
+	/** Where the user has dragged to so far */
+	FVector2D EndPoint;
+
+	/** Make a default marquee selection */
+	FMarqueeRect( FVector2D InStartPoint = FVector2D::ZeroVector )
+		: StartPoint( InStartPoint )
+		, EndPoint( InStartPoint )
+	{
+	}
+
+	/**
+	 * Update the location to which the user has dragged the marquee selection so far
+	 *
+	 * @param NewEndPoint   Where the user has dragged so far.
+	 */
+	void UpdateEndPoint( FVector2D NewEndPoint )
+	{
+		EndPoint = NewEndPoint;
+	}
+
+	/** @return true if this marquee selection is not too small to be considered real */
+	bool IsValid() const
+	{
+		return ! (EndPoint - StartPoint).IsNearlyZero();
+	}
+
+	/** @return the upper left point of the selection */
+	FVector2D GetUpperLeft() const
+	{
+		return FVector2D( FMath::Min(StartPoint.X, EndPoint.X), FMath::Min( StartPoint.Y, EndPoint.Y ) );
+	}
+
+	/** @return the lower right of the selection */
+	FVector2D GetLowerRight() const
+	{
+		return  FVector2D( FMath::Max(StartPoint.X, EndPoint.X), FMath::Max( StartPoint.Y, EndPoint.Y ) );
+	}
+
+	/** The size of the selection */
+	FVector2D GetSize() const
+	{
+		FVector2D SignedSize = EndPoint - StartPoint;
+		return FVector2D( FMath::Abs(SignedSize.X), FMath::Abs(SignedSize.Y) );
+	}
+
+	/** @return This marquee rectangle as a well-formed SlateRect */
+	FSlateRect ToSlateRect() const
+	{
+		return FSlateRect( FMath::Min(StartPoint.X, EndPoint.X), FMath::Min(StartPoint.Y, EndPoint.Y), FMath::Max(StartPoint.X, EndPoint.X), FMath::Max( StartPoint.Y, EndPoint.Y ) );
+	}
+};
\ No newline at end of file
diff --git a/Engine/Source/Runtime/Slate/Public/Framework/MultiBox/MultiBoxBuilder.h b/Engine/Source/Runtime/Slate/Public/Framework/MultiBox/MultiBoxBuilder.h
new file mode 100644
--- /dev/null
+++ b/Engine/Source/Runtime/Slate/Public/Framework/MultiBox/MultiBoxBuilder.h
@@ -0,0 +1,467 @@
+// Copyright 1998-2014 Epic Games, Inc. All Rights Reserved.
+
+#pragma once
+
+#include "MultiBoxDefs.h"
+#include "MultiBoxExtender.h"
+
+
+class FUICommandInfo;
+class FUICommandList;
+
+
+/** Delegate used by multi-box to call a user function to populate a new menu.  Used for spawning sub-menus and pull-down menus. */
+DECLARE_DELEGATE_OneParam( FNewMenuDelegate, class FMenuBuilder& );
+
+
+/**
+ * MultiBox builder
+ */
+class SLATE_API FMultiBoxBuilder
+{
+
+public:
+
+	/**
+	 * Constructor
+	 *
+	 * @param	InType	Type of MultiBox
+	 * @param	bInShouldCloseWindowAfterMenuSelection	Sets whether or not the window that contains this multibox should be destroyed after the user clicks on a menu item in this box
+	 * @param	InCommandList	The action list that maps command infos to delegates that should be called for each command associated with a multiblock widget.  This can be modified after the MultiBox is created by calling the PushCommandList() and PopCommandList() methods.
+	 */
+	FMultiBoxBuilder( const EMultiBoxType::Type InType, FMultiBoxCustomization InCustomization, const bool bInShouldCloseWindowAfterMenuSelection, const TSharedPtr< const FUICommandList >& InCommandList, TSharedPtr<FExtender> InExtender = TSharedPtr<FExtender>() );
+
+
+	/**
+	 * Adds an editable text entry
+	 *
+	 * @param	InLabel				The label to display in the menu
+	 * @param	InToolTip			The tool tip to display when the menu entry is hovered over
+	 * @param	InIcon				The icon to display to the left of the label
+	 * @param	InTextAttribute		The text string we're editing (often, a delegate will be bound to the attribute)
+	 * @param	InOnTextCommitted	Called when the user commits their change to the editable text control
+	 * @param	InOnTextChanged		Called when the text is changed interactively
+	 * @param	bInReadOnly			Whether or not the text block is read only
+	 */
+	void AddEditableText( const FText& InLabel, const FText& InToolTip, const FSlateIcon& InIcon, const TAttribute< FText >& InTextAttribute, const FOnTextCommitted& InOnTextCommitted = FOnTextCommitted(), const FOnTextChanged& InOnTextChanged = FOnTextChanged(), bool bInReadOnly = false);
+
+
+	/**
+	 * Creates a widget for this MultiBox
+	 *
+	 * @return  New widget object
+	 */
+	TSharedRef< class SWidget > MakeWidget();
+
+
+	/**
+	 * Pushes a new command list onto the stack.  This command list will be used for all subsequently-added multiblocks, until the command-list is popped.
+	 *
+	 * @param	CommandList		The new command list to use
+	 */
+	void PushCommandList( const TSharedRef< const FUICommandList > CommandList );
+
+
+	/**
+	 * Pops the current command list.
+	 */
+	void PopCommandList();
+	
+	/**
+	 * @return The top command list
+	 */
+	TSharedPtr<const FUICommandList> GetTopCommandList();
+
+	/**
+	 * Pushes a new extender onto the stack. This extender will be used for all subsequently-added multiblocks, until the extender is popped.
+	 *
+	 * @param	InExtender	The new extender to use
+	 */
+	void PushExtender( TSharedRef< FExtender > InExtender );
+
+
+	/**
+	 * Pops the current extender.
+	 */
+	void PopExtender();
+
+	/** @return The style set used by the multibox widgets */
+	const ISlateStyle* GetStyleSet() const;
+
+	/** @return The style name used by the multibox widgets */
+	const FName& GetStyleName() const;
+
+	/** Sets the style to use on the multibox widgets */
+	void SetStyle( const ISlateStyle* InStyleSet, const FName& InStyleName );
+
+	/** @return  The customization settings for the box being built */
+	FMultiBoxCustomization GetCustomization() const;
+
+protected:
+	/** Applies any potential extension hooks at the current place */
+	virtual void ApplyHook(FName InExtensionHook, EExtensionHook::Position HookPosition) = 0;
+	
+	/** Applies the beginning of a section, including the header and relevant separator */
+	virtual void ApplySectionBeginning() {}
+
+protected:
+
+	/** The MultiBox we're building up */
+	TSharedRef< class FMultiBox > MultiBox;
+
+	/** A stack of command lists which map command infos to delegates that should be called.  New multi blocks will always use
+		the command-list at the top of the stack at the time they are added. */
+	TArray< TSharedPtr< const FUICommandList > > CommandListStack;
+
+	/** The extender stack holding all the possible extensions for this menu builder */
+	TArray< TSharedPtr<class FExtender> > ExtenderStack;
+};
+
+
+
+/**
+ * Base menu builder
+ */
+class SLATE_API FBaseMenuBuilder : public FMultiBoxBuilder
+{
+
+public:
+
+	/**
+	 * Constructor
+	 *
+	 * @param	InType	Type of MultiBox
+	 * @param	bInShouldCloseWindowAfterMenuSelection	Sets whether or not the window that contains this multibox should be destroyed after the user clicks on a menu item in this box
+	 * @param	InCommandList	The action list that maps command infos to delegates that should be called for each command associated with a multiblock widget
+	 * @param	bInCloseSelfOnly	True if clicking on a menu entry closes itself only and its children but not the entire stack 
+	 */
+	FBaseMenuBuilder( const EMultiBoxType::Type InType, const bool bInShouldCloseWindowAfterMenuSelection, TSharedPtr< const FUICommandList > InCommandList, bool bInCloseSelfOnly, TSharedPtr<FExtender> InExtender = TSharedPtr<FExtender>(), const ISlateStyle* InStyleSet = &FCoreStyle::Get() );
+
+	/**
+	 * Adds a menu entry
+	 *
+	 * @param	InCommand			The command associated with this menu entry
+	 * @param	InLabelOverride		Optional label override.  If omitted, then the action's label will be used instead.
+	 * @param	InToolTipOverride	Optional tool tip override.	 If omitted, then the action's label will be used instead.
+	 * @param	InIconOverride		Optional name of the slate brush to use for the tool bar image.  If omitted, then the action's icon will be used instead.
+	 */
+	void AddMenuEntry( const TSharedPtr< const FUICommandInfo > InCommand, FName InExtensionHook = NAME_None, const TAttribute<FText>& InLabelOverride = TAttribute<FText>(), const TAttribute<FText>& InToolTipOverride = TAttribute<FText>(), const FSlateIcon& InIconOverride = FSlateIcon() );
+
+	/**
+	 * Adds a menu entry without the use of a command
+	 *
+	 * @param	InLabel		Label to show in the menu entry
+	 * @param	InToolTip	Tool tip used when hovering over the menu entry
+	 * @param	InIcon		The icon to use		
+	 * @param	UIAction	Actions to execute on this menu item.
+	 * @param	UserInterfaceActionType	Type of interface action
+	 */
+	void AddMenuEntry( const TAttribute<FText>& InLabel, const TAttribute<FText>& InToolTip, const FSlateIcon& InIcon, const FUIAction& UIAction, FName InExtensionHook = NAME_None, const EUserInterfaceActionType::Type UserInterfaceActionType = EUserInterfaceActionType::Button );
+
+	void AddMenuEntry( const FUIAction& UIAction, const TSharedRef< SWidget > Contents, const FName& InExtensionHook = NAME_None, const EUserInterfaceActionType::Type UserInterfaceActionType = EUserInterfaceActionType::Button );
+
+protected:
+	/** True if clicking on a menu entry closes itself only and its children and not the entire stack */
+	bool bCloseSelfOnly;
+};
+
+
+/**
+ * Vertical menu builder
+ */
+class SLATE_API FMenuBuilder : public FBaseMenuBuilder
+{
+public:
+
+	/**
+	 * Constructor
+	 *
+	 * @param	bInShouldCloseWindowAfterMenuSelection	Sets whether or not the window that contains this multibox should be destroyed after the user clicks on a menu item in this box
+	 * @param	InCommandList	The action list that maps command infos to delegates that should be called for each command associated with a multiblock widget
+	 * @param	bInCloseSelfOnly	True if clicking on a menu entry closes itself only and its children but not the entire stack
+	 */
+	FMenuBuilder( const bool bInShouldCloseWindowAfterMenuSelection, TSharedPtr< const FUICommandList > InCommandList, TSharedPtr<FExtender> InExtender = TSharedPtr<FExtender>(), const bool bCloseSelfOnly = false, const ISlateStyle* InStyleSet = &FCoreStyle::Get() )
+		: FBaseMenuBuilder( EMultiBoxType::Menu, bInShouldCloseWindowAfterMenuSelection, InCommandList, bCloseSelfOnly, InExtender, InStyleSet )
+		, bSectionNeedsToBeApplied(false)
+	{
+	}
+
+
+	/**
+	 * Adds a menu separator
+	 */
+	void AddMenuSeparator(FName InExtensionHook = NAME_None);
+	
+
+	/**
+	 * Starts a section on to the extender section hook stack
+	 * 
+	 * @param InExtensionHook	The section hook. Can be NAME_None
+	 * @param InHeadingText		The heading text to use. If none, only a separator is used
+	 */
+	void BeginSection( FName InExtensionHook, const TAttribute< FText >& InHeadingText = TAttribute<FText>() );
+
+	/**
+	 * Ends the current section
+	 */
+	void EndSection();
+
+
+	/**
+	 * Adds a sub-menu which is a menu within a menu
+	 * 
+	 * @param	InMenuLabel			The text that should be shown for the menu
+	 * @param	InToolTip			The tooltip that should be shown when the menu is hovered over
+	 * @param	InSubMenu			Sub-Menu object which creates menu entries for the sub-menu
+	 * @param bInOpenSubMenuOnClick Sub-menu will open only if the sub-menu entry is clicked
+	 */
+	void AddSubMenu( const FText& InMenuLabel, const FText& InToolTip, const FNewMenuDelegate& InSubMenu, const FUIAction& InUIAction, FName InExtensionHook, const EUserInterfaceActionType::Type InUserInterfaceActionType, const bool bInOpenSubMenuOnClick = false, const FSlateIcon& InIcon = FSlateIcon() );
+
+	void AddSubMenu( const FText& InMenuLabel, const FText& InToolTip, const FNewMenuDelegate& InSubMenu, const bool bInOpenSubMenuOnClick = false, const FSlateIcon& InIcon = FSlateIcon() );
+
+	void AddSubMenu( const TSharedRef< SWidget > Contents, const FNewMenuDelegate& InSubMenu, const bool bInOpenSubMenuOnClick = false );
+
+	void AddSubMenu( const FUIAction& UIAction, const TSharedRef< SWidget > Contents, const FNewMenuDelegate& InSubMenu );
+
+	/**
+	 * Adds any widget to the menu
+	 * 
+	 * @param	InWidget			The widget that should be shown in the menu
+	 * @param	InLabel				Optional label text to be added to the left of the content
+	 * @param	bNoIndent			If true, removes the padding from the left of the widget that lines it up with other menu items (default == false)
+	 */
+	void AddWidget( TSharedRef<SWidget> InWidget, const FText& Label, bool bNoIndent = false );
+	
+protected:
+	/** FMultiBoxBuilder interface */
+	virtual void ApplyHook(FName InExtensionHook, EExtensionHook::Position HookPosition) OVERRIDE;
+	virtual void ApplySectionBeginning() OVERRIDE;
+
+public:
+	// These classes need access to the AddWrapperSubMenu() methods
+	//friend class FWidgetBlock;
+	//friend class FToolBarComboButtonBlock;
+
+	/**
+	 * Adds a sub-menu which is a menu within a menu
+	 * 
+	 * @param	InMenuLabel			The text that should be shown for the menu
+	 * @param	InToolTip			The tooltip that should be shown when the menu is hovered over
+	 * @param	InSubMenu			Sub-Menu object which creates the sub-menu
+	 */
+	void AddWrapperSubMenu( const FText& InMenuLabel, const FText& InToolTip, const FOnGetContent& InSubMenu, const FSlateIcon& InIcon );
+
+	/**
+	 * Adds a sub-menu which is a menu within a menu
+	 * 
+	 * @param	InMenuLabel			The text that should be shown for the menu
+	 * @param	InToolTip			The tooltip that should be shown when the menu is hovered over
+	 * @param	InSubMenu			Sub-Menu object
+	 */
+	void AddWrapperSubMenu( const FText& InMenuLabel, const FText& InToolTip, const TSharedPtr<SWidget>& InSubMenu, const FSlateIcon& InIcon );
+
+private:
+	/** Current extension hook name for sections to determine where sections begin and end */
+	FName CurrentSectionExtensionHook;
+	
+	/** Any pending section's heading text */
+	FText CurrentSectionHeadingText;
+	
+	/** True if there is a pending section that needs to be applied */
+	bool bSectionNeedsToBeApplied;
+};
+
+
+
+/**
+ * Menu bar builder
+ */
+class SLATE_API FMenuBarBuilder : public FBaseMenuBuilder
+{
+
+public:
+
+	/**
+	 * Constructor
+	 *
+	 * @param	InCommandList	The action list that maps command infos to delegates that should be called for each command associated with a multiblock widget
+	 */
+	FMenuBarBuilder( TSharedPtr< const FUICommandList > InCommandList, TSharedPtr<FExtender> InExtender = TSharedPtr<FExtender>(), const ISlateStyle* InStyleSet = &FCoreStyle::Get() )
+		: FBaseMenuBuilder( EMultiBoxType::MenuBar, false, InCommandList, false, InExtender, InStyleSet )
+	{
+	}
+
+
+	/**
+	 * Adds a pull down menu
+	 *
+	 * @param	InMenuLabel			The text that should be shown for the menu
+	 * @param	InToolTip			The tooltip that should be shown when the menu is hovered over
+	 * @param	InPullDownMenu		Pull down menu object for the menu to add.
+	 */
+	void AddPullDownMenu( const FText& InMenuLabel, const FText& InToolTip, const FNewMenuDelegate& InPullDownMenu, FName InExtensionHook = NAME_None, FName InTutorialHighlightName = NAME_None );
+	
+protected:
+	/** FMultiBoxBuilder interface */
+	virtual void ApplyHook(FName InExtensionHook, EExtensionHook::Position HookPosition) OVERRIDE;
+};
+
+
+
+/**
+ * Tool bar builder
+ */
+class SLATE_API FToolBarBuilder : public FMultiBoxBuilder
+{
+
+public:
+
+	/**
+	 * Constructor
+	 *
+	 * @param	InCommandList	The action list that maps command infos to delegates that should be called for each command associated with a multiblock widget
+	 */
+	FToolBarBuilder( TSharedPtr< const FUICommandList > InCommandList, FMultiBoxCustomization InCustomization, TSharedPtr<FExtender> InExtender = TSharedPtr<FExtender>(), EOrientation Orientation = Orient_Horizontal, const bool InForceSmallIcons = false )
+		: FMultiBoxBuilder( (Orientation == Orient_Horizontal) ? EMultiBoxType::ToolBar : EMultiBoxType::VerticalToolBar, InCustomization, false, InCommandList, InExtender )
+		, bSectionNeedsToBeApplied(false)
+		, bIsFocusable(false)
+		, bForceSmallIcons(InForceSmallIcons)
+	{
+	}
+
+	void SetLabelVisibility( EVisibility InLabelVisibility ) { LabelVisibility  = InLabelVisibility ; }
+
+	void SetIsFocusable( bool bInIsFocusable ) { bIsFocusable = bInIsFocusable; }
+
+	/**
+	 * Adds a tool bar button
+	 *
+	 * @param	InCommand				The command associated with this tool bar button
+	 * @param	InLabelOverride			Optional label override.  If omitted, then the action's label will be used instead.
+	 * @param	InToolTipOverride		Optional tool tip override.	 If omitted, then the action's label will be used instead.
+	 * @param	InIconOverride			Optional name of the slate brush to use for the tool bar image.  If omitted, then the action's icon will be used instead.
+	 */
+	void AddToolBarButton( const TSharedPtr< const FUICommandInfo > InCommand, FName InExtensionHook = NAME_None, const TAttribute<FText>& InLabelOverride = TAttribute<FText>(), const TAttribute<FText>& InToolTipOverride = TAttribute<FText>(), const TAttribute<FSlateIcon>& InIconOverride = TAttribute<FSlateIcon>(), FName TutorialHighlightName = NAME_None );
+	
+	/**
+	 * Adds a tool bar button
+	 *
+	 * @param	UIAction	Actions to execute on this menu item.
+	 * @param	InLabel		Label to show in the menu entry
+	 * @param	InToolTip	Tool tip used when hovering over the menu entry
+	 * @param	InIcon		The icon to use		
+	 * @param	UserInterfaceActionType	Type of interface action
+	 */
+	void AddToolBarButton( const FUIAction& InAction, FName InExtensionHook = NAME_None, const TAttribute<FText>& InLabelOverride = TAttribute<FText>(), const TAttribute<FText>& InToolTipOverride = TAttribute<FText>(), const TAttribute<FSlateIcon>& InIconOverride = TAttribute<FSlateIcon>(), const EUserInterfaceActionType::Type UserInterfaceActionType = EUserInterfaceActionType::Button, FName TutorialHighlightName = NAME_None );
+
+	/**
+	 * Adds a combo button
+	 *
+	 * @param	InAction					UI action that sets the enabled state for this combo button
+	 * @param	InMenuContentGenerator		Delegate that generates a widget for this combo button's menu content.  Called when the menu is summoned.
+	 * @param	InLabelOverride				Optional label override.  If omitted, then the action's label will be used instead.
+	 * @param	InToolTipOverride			Optional tool tip override.	 If omitted, then the action's label will be used instead.
+	 * @param	InIconOverride				Optional icon to use for the tool bar image.  If omitted, then the action's icon will be used instead.
+	 * @param	bInSimpleComboBox			If true, the icon and label won't be displayed
+	 */
+	void AddComboButton( const FUIAction& InAction, const FOnGetContent& InMenuContentGenerator, const TAttribute<FText>& InLabelOverride = TAttribute<FText>(), const TAttribute<FText>& InToolTipOverride = TAttribute<FText>(), const TAttribute<FSlateIcon>& InIconOverride = TAttribute<FSlateIcon>(), bool bInSimpleComboBox = false );
+
+	/**
+	 * Adds any widget to the toolbar
+	 * 
+	 * @param	InWidget				The widget that should be shown in the toolbar
+	 * @param	TutorialHighlightName	Name to identify this widget and highlight during tutorials
+	 */
+	void AddWidget( TSharedRef<SWidget> InWidget, FName TutorialHighlightName = NAME_None );
+	
+	/**
+	 * Adds a toolbar separator
+	 */
+	void AddSeparator(FName InExtensionHook = NAME_None);
+	
+	/**
+	 * Starts a section on to the extender section hook stack
+	 * 
+	 * @param InExtensionHook	The section hook. Can be NAME_None
+	 */
+	void BeginSection( FName InExtensionHook );
+
+	/**
+	 * Ends the current section
+	 */
+	void EndSection();
+
+	/** 
+	 * Starts a new Group block, must be used in conjunction with EndBlockGroup
+	 */
+	void BeginBlockGroup();
+	
+	/** 
+	 * End a group block, must be used in conjunction with BeginBlockGroup.
+	 */
+	void EndBlockGroup();
+
+protected:
+	/** FMultiBoxBuilder interface */
+	virtual void ApplyHook(FName InExtensionHook, EExtensionHook::Position HookPosition) OVERRIDE;
+	virtual void ApplySectionBeginning() OVERRIDE;
+
+private:
+	/** Current extension hook name for sections to determine where sections begin and end */
+	FName CurrentSectionExtensionHook;
+
+	TOptional< EVisibility > LabelVisibility;
+
+	/** True if there is a pending section that needs to be applied */
+	bool bSectionNeedsToBeApplied;
+
+	/** Whether the buttons created can receive keyboard focus */
+	bool bIsFocusable;
+
+	/** Whether this toolbar should always use small icons, regardless of the current settings */
+	bool bForceSmallIcons;
+};
+
+
+/**
+ * Button grid builder
+ */
+class SLATE_API FButtonRowBuilder : public FMultiBoxBuilder
+{
+public:
+	/**
+	 * Constructor
+	 *
+	 * @param	InCommandList	The action list that maps command infos to delegates that should be called for each command associated with a multiblock widget
+	 */
+	FButtonRowBuilder( TSharedPtr< const FUICommandList > InCommandList, TSharedPtr<FExtender> InExtender = TSharedPtr<FExtender>() )
+		: FMultiBoxBuilder( EMultiBoxType::ButtonRow, FMultiBoxCustomization::None, false, InCommandList, InExtender )
+	{
+	}
+
+
+	/**
+	 * Adds a button to a row
+	 *
+	 * @param	InCommand				The command associated with this tool bar button
+	 * @param	InLabelOverride			Optional label override.  If omitted, then the action's label will be used instead.
+	 * @param	InToolTipOverride		Optional tool tip override.	 If omitted, then the action's label will be used instead.
+	 * @param	InIconOverride			Optional icon to use for the tool bar image.  If omitted, then the action's icon will be used instead.
+	 */
+	void AddButton( const TSharedPtr< const FUICommandInfo > InCommand, const TAttribute<FText>& InLabelOverride = TAttribute<FText>(), const TAttribute<FText>& InToolTipOverride = TAttribute<FText>(), const FSlateIcon& InIconOverride = FSlateIcon() );
+
+	/**
+	 * Adds a button to a row
+	 *
+	 * @param	InLabel					The button label to display
+	 * @param	InToolTipOverride		The tooltip for the button
+	 * @param	UIAction				Action to execute when the button is clicked or when state should be checked
+	 * @param	InIcon					The icon for the button
+	 * @param	UserInterfaceActionType	The style of button to display
+	 */
+	void AddButton( const FText& InLabel, const FText& InToolTip, const FUIAction& UIAction, const FSlateIcon& InIcon = FSlateIcon(), const EUserInterfaceActionType::Type UserInterfaceActionType = EUserInterfaceActionType::Button );
+	
+protected:
+	/** FMultiBoxBuilder interface */
+	virtual void ApplyHook(FName InExtensionHook, EExtensionHook::Position HookPosition) OVERRIDE {}
+};
diff --git a/Engine/Source/Runtime/Slate/Public/Framework/MultiBox/MultiBoxDefs.h b/Engine/Source/Runtime/Slate/Public/Framework/MultiBox/MultiBoxDefs.h
new file mode 100644
--- /dev/null
+++ b/Engine/Source/Runtime/Slate/Public/Framework/MultiBox/MultiBoxDefs.h
@@ -0,0 +1,112 @@
+// Copyright 1998-2014 Epic Games, Inc. All Rights Reserved.
+
+#pragma once
+
+
+/**
+ * Types of MultiBoxes
+ */
+namespace EMultiBoxType
+{
+	enum Type
+	{
+		/** Horizontal menu bar */
+		MenuBar,
+
+		/** Horizontal tool bar */
+		ToolBar,
+
+		/** Vertical tool bar */
+		VerticalToolBar,
+
+		/** Vertical menu (pull-down menu, or context menu) */
+		Menu,
+
+		/** Buttons arranged in rows, with a maximum number of buttons per row, like a toolbar but can have multiple rows*/
+		ButtonRow
+	};
+}
+
+class SLATE_API FMultiBoxSettings
+{
+public:
+
+	DECLARE_DELEGATE_RetVal_ThreeParams( TSharedRef< SToolTip >, FConstructToolTip, const TAttribute<FText>& /*ToolTipText*/, const TSharedPtr<SWidget>& /*OverrideContent*/, const TSharedPtr<const FUICommandInfo>& /*Action*/ );
+
+	/** Access to whether multiboxes use small icons or default sized icons */
+	static TAttribute<bool> UseSmallToolBarIcons;
+	static TAttribute<bool> DisplayMultiboxHooks;
+	static FConstructToolTip ToolTipConstructor;
+
+	FMultiBoxSettings();
+
+	/** Enables or disables toolbar editing */
+	static void ToggleToolbarEditing();
+
+	/** @return True if in toolbar edit mode */
+	static bool IsInToolbarEditMode() { return bInToolbarEditMode; }
+
+	static TSharedRef< SToolTip > ConstructDefaultToolTip( const TAttribute<FText>& ToolTipText, const TSharedPtr<SWidget>& OverrideContent, const TSharedPtr<const FUICommandInfo>& Action );
+
+	static void ResetToolTipConstructor();
+
+private:
+	static bool bInToolbarEditMode;
+};
+
+struct SLATE_API FMultiBoxCustomization
+{
+	static const FMultiBoxCustomization None;
+
+	static FMultiBoxCustomization AllowCustomization( FName InCustomizationName )
+	{
+		ensure( InCustomizationName != NAME_None );
+		return FMultiBoxCustomization( InCustomizationName );
+	}
+
+	FName GetCustomizationName() const { return CustomizationName; }
+
+	FMultiBoxCustomization( FName InCustomizationName )
+		: CustomizationName( InCustomizationName )
+	{}
+
+private:
+	/** The Name of the customization that uniquely identifies the multibox for saving and loading users data*/
+	FName CustomizationName;
+};
+
+/** 
+ * Block location information
+ */
+namespace EMultiBlockLocation
+{
+	enum Type
+	{
+		/** Default, either no other blocks in group or grouping style is disabled */
+		None = -1,
+		
+		/** Denotes the beginning of a group, currently left most first */
+		Start,
+
+		/** Denotes a middle block(s) of a group */
+		Middle,
+
+		/** Denotes the end of a group, currently the right most */
+		End,
+	};
+
+	/** returns the passed in style with the addition of the location information */
+	static FName ToName(FName StyleName, Type InLocation)
+	{
+		switch(InLocation)
+		{
+		case Start:
+			return ISlateStyle::Join(StyleName, ".Start");
+		case Middle:
+			return ISlateStyle::Join(StyleName, ".Middle");
+		case End:
+			return ISlateStyle::Join(StyleName, ".End");
+		}
+		return StyleName;
+	}
+}
diff --git a/Engine/Source/Runtime/Slate/Public/Framework/Notifications/NotificationManager.h b/Engine/Source/Runtime/Slate/Public/Framework/Notifications/NotificationManager.h
new file mode 100644
--- /dev/null
+++ b/Engine/Source/Runtime/Slate/Public/Framework/Notifications/NotificationManager.h
@@ -0,0 +1,91 @@
+// Copyright 1998-2014 Epic Games, Inc. All Rights Reserved.
+
+#pragma once
+
+#include "LockFreeList.h"
+
+/**
+ * A class which manages a group of notification windows                 
+ */
+class SLATE_API FSlateNotificationManager
+{
+	friend class SNotificationList;
+
+public:
+	/**
+	 * Gets the instance of this manager                   
+	 */
+	static FSlateNotificationManager& Get();
+
+	/** Update the manager */
+	void Tick();
+
+	/** Provide a window under which all notifications should nest. */
+	void SetRootWindow( const TSharedRef<SWindow> InRootWindow );
+
+	/**
+	 * Adds a floating notification
+	 * @param Info 		Contains various settings used to initialize the notification
+	 */
+	TSharedPtr<SNotificationItem> AddNotification(const FNotificationInfo& Info);
+
+	/**
+	 * Thread safe method of queuing a notification for presentation on the next tick
+	 * @param Info 		Pointer to heap allocated notification info. Released by FSlateNotificationManager once the notification is displayed;
+	 */
+	void QueueNotification(FNotificationInfo* Info);
+
+	/**
+	 * Called back from the SlateApplication when a window is activated/resized
+	 * We need to keep notifications topmost in the z-order so we manage it here directly
+	 * as there isn't a cross-platform OS-level way of making a 'topmost child'.
+	 * @param ActivateEvent 	Information about the activation event
+	 */
+	void ForceNotificationsInFront( const TSharedRef<SWindow> InWindow );
+
+	/**
+	 * Gets all the windows that represent notifications
+	 * @returns an array of windows that contain notifications.
+	 */
+	void GetWindows(TArray< TSharedRef<SWindow> >& OutWindows) const;
+
+protected:
+	/** Protect constructor as this is a singleton */
+	FSlateNotificationManager();
+
+	/** Arranges the active notifications in a stack */
+	void ArrangeNotifications();
+
+	/** Create a notification list for the specified screen rectangle */
+	TSharedRef<SNotificationList> CreateStackForArea(const FSlateRect& InRectangle);
+
+private:
+
+	/** A list of notifications, bound to a particular region */
+	struct FRegionalNotificationList
+	{
+		/** Constructor */
+		FRegionalNotificationList(const FSlateRect& InRectangle);
+
+		/** Arranges the notifications in a stack */
+		void Arrange();
+
+		/** Remove any dead notifications */
+		void RemoveDeadNotifications();
+
+		/** The notification list itself (one per notification) */
+		TArray<TSharedRef<SNotificationList>> Notifications;
+
+		/** The rectangle we use to determine the anchor point for this list */
+		FSlateRect Region;
+	};
+
+	/** A window under which all of the notification windows will nest. */
+	TWeakPtr<SWindow> RootWindowPtr;
+
+	/** An array of notification lists grouped by work area regions */
+	TArray< FRegionalNotificationList > RegionalLists;
+
+	/** Thread safe queue of notifications to display */
+	TLockFreePointerList<FNotificationInfo> PendingNotifications;
+};
diff --git a/Engine/Source/Runtime/Slate/Public/Framework/Styling/ButtonWidgetStyle.h b/Engine/Source/Runtime/Slate/Public/Framework/Styling/ButtonWidgetStyle.h
new file mode 100644
--- /dev/null
+++ b/Engine/Source/Runtime/Slate/Public/Framework/Styling/ButtonWidgetStyle.h
@@ -0,0 +1,24 @@
+// Copyright 1998-2014 Epic Games, Inc. All Rights Reserved.
+
+#pragma once
+
+#include "SlateWidgetStyleContainerBase.h"
+#include "ButtonWidgetStyle.generated.h"
+
+/**
+ */
+UCLASS(hidecategories=Object, MinimalAPI)
+class UButtonWidgetStyle : public USlateWidgetStyleContainerBase
+{
+	GENERATED_UCLASS_BODY()
+
+public:
+	/** The actual data describing the button's appearance. */
+	UPROPERTY(Category=Appearance, EditAnywhere, meta=(ShowOnlyInnerProperties))
+	FButtonStyle ButtonStyle;
+
+	virtual const struct FSlateWidgetStyle* const GetStyle() const OVERRIDE
+	{
+		return static_cast< const struct FSlateWidgetStyle* >( &ButtonStyle );
+	}
+};
diff --git a/Engine/Source/Runtime/Slate/Public/Framework/Styling/ComboBoxWidgetStyle.h b/Engine/Source/Runtime/Slate/Public/Framework/Styling/ComboBoxWidgetStyle.h
new file mode 100644
--- /dev/null
+++ b/Engine/Source/Runtime/Slate/Public/Framework/Styling/ComboBoxWidgetStyle.h
@@ -0,0 +1,24 @@
+// Copyright 1998-2014 Epic Games, Inc. All Rights Reserved.
+
+#pragma once
+
+#include "SlateWidgetStyleContainerBase.h"
+#include "ComboBoxWidgetStyle.generated.h"
+
+/**
+ */
+UCLASS(hidecategories=Object, MinimalAPI)
+class UComboBoxWidgetStyle : public USlateWidgetStyleContainerBase
+{
+	GENERATED_UCLASS_BODY()
+
+public:
+	/** The actual data describing the combo box's appearance. */
+	UPROPERTY(Category=Appearance, EditAnywhere, meta=(ShowOnlyInnerProperties))
+	FComboBoxStyle ComboBoxStyle;
+
+	virtual const struct FSlateWidgetStyle* const GetStyle() const OVERRIDE
+	{
+		return static_cast< const struct FSlateWidgetStyle* >( &ComboBoxStyle );
+	}
+};
diff --git a/Engine/Source/Runtime/Slate/Public/Framework/Styling/ComboButtonWidgetStyle.h b/Engine/Source/Runtime/Slate/Public/Framework/Styling/ComboButtonWidgetStyle.h
new file mode 100644
--- /dev/null
+++ b/Engine/Source/Runtime/Slate/Public/Framework/Styling/ComboButtonWidgetStyle.h
@@ -0,0 +1,24 @@
+// Copyright 1998-2014 Epic Games, Inc. All Rights Reserved.
+
+#pragma once
+
+#include "SlateWidgetStyleContainerBase.h"
+#include "ComboButtonWidgetStyle.generated.h"
+
+/**
+ */
+UCLASS(hidecategories=Object, MinimalAPI)
+class UComboButtonWidgetStyle : public USlateWidgetStyleContainerBase
+{
+	GENERATED_UCLASS_BODY()
+
+public:
+	/** The actual data describing the combo button's appearance. */
+	UPROPERTY(Category=Appearance, EditAnywhere, meta=(ShowOnlyInnerProperties))
+	FComboButtonStyle ComboButtonStyle;
+
+	virtual const struct FSlateWidgetStyle* const GetStyle() const OVERRIDE
+	{
+		return static_cast< const struct FSlateWidgetStyle* >( &ComboButtonStyle );
+	}
+};
diff --git a/Engine/Source/Runtime/Slate/Public/Framework/Styling/TextBlockWidgetStyle.h b/Engine/Source/Runtime/Slate/Public/Framework/Styling/TextBlockWidgetStyle.h
new file mode 100644
--- /dev/null
+++ b/Engine/Source/Runtime/Slate/Public/Framework/Styling/TextBlockWidgetStyle.h
@@ -0,0 +1,24 @@
+// Copyright 1998-2014 Epic Games, Inc. All Rights Reserved.
+
+#pragma once
+
+#include "SlateWidgetStyleContainerBase.h"
+#include "TextBlockWidgetStyle.generated.h"
+
+/**
+ */
+UCLASS(hidecategories=Object, MinimalAPI)
+class UTextBlockWidgetStyle : public USlateWidgetStyleContainerBase
+{
+	GENERATED_UCLASS_BODY()
+
+public:
+	/** The actual data describing the button's appearance. */
+	UPROPERTY(Category=Appearance, EditAnywhere, meta=(ShowOnlyInnerProperties))
+	FTextBlockStyle TextBlockStyle;
+
+	virtual const struct FSlateWidgetStyle* const GetStyle() const OVERRIDE
+	{
+		return static_cast< const struct FSlateWidgetStyle* >( &TextBlockStyle );
+	}
+};
diff --git a/Engine/Source/Runtime/Slate/Public/Framework/Text/GenericPlatformTextField.h b/Engine/Source/Runtime/Slate/Public/Framework/Text/GenericPlatformTextField.h
new file mode 100644
--- /dev/null
+++ b/Engine/Source/Runtime/Slate/Public/Framework/Text/GenericPlatformTextField.h
@@ -0,0 +1,14 @@
+// Copyright 1998-2014 Epic Games, Inc. All Rights Reserved.
+
+#pragma once
+
+#include "IPlatformTextField.h"
+
+class FGenericPlatformTextField : public IPlatformTextField
+{
+public:
+	virtual void ShowKeyboard(bool bShow, TSharedPtr<SVirtualKeyboardEntry> TextEntryWidget) OVERRIDE {};
+
+};
+
+typedef FGenericPlatformTextField FPlatformTextField;
diff --git a/Engine/Source/Runtime/Slate/Public/Framework/Text/IOS/IOSPlatformTextField.h b/Engine/Source/Runtime/Slate/Public/Framework/Text/IOS/IOSPlatformTextField.h
new file mode 100644
--- /dev/null
+++ b/Engine/Source/Runtime/Slate/Public/Framework/Text/IOS/IOSPlatformTextField.h
@@ -0,0 +1,33 @@
+// Copyright 1998-2014 Epic Games, Inc. All Rights Reserved.
+
+#pragma once
+
+#include "IPlatformTextField.h"
+
+#import <UIKit/UIKit.h>
+
+@class SlateTextField;
+
+class FIOSPlatformTextField : public IPlatformTextField
+{
+public:
+	FIOSPlatformTextField();
+	virtual ~FIOSPlatformTextField();
+
+	virtual void ShowKeyboard(bool bShow, TSharedPtr<SVirtualKeyboardEntry> TextEntryWidget) OVERRIDE;
+
+private:
+	SlateTextField* TextField;
+};
+
+typedef FIOSPlatformTextField FPlatformTextField;
+
+
+@interface SlateTextField : NSObject<UIAlertViewDelegate>
+{
+	TSharedPtr<SVirtualKeyboardEntry> TextWidget;
+}
+
+-(void)show:(TSharedPtr<SVirtualKeyboardEntry>)InTextWidget;
+
+@end
diff --git a/Engine/Source/Runtime/Slate/Public/Framework/Text/SlateWordWrapper.h b/Engine/Source/Runtime/Slate/Public/Framework/Text/SlateWordWrapper.h
new file mode 100644
--- /dev/null
+++ b/Engine/Source/Runtime/Slate/Public/Framework/Text/SlateWordWrapper.h
@@ -0,0 +1,32 @@
+// Copyright 1998-2014 Epic Games, Inc. All Rights Reserved.
+#pragma once
+
+#include "WordWrapper.h"
+
+struct FWrappedStringSlice;
+
+namespace SlateWordWrapper
+{
+	/**
+	 * Produce hard-wrapped string
+	 * 
+	 * @param InString				The unwrapped text.
+	 * @param InFontInfo			The font used to render the text.
+	 * @param InWrapWidth			The width available.
+	 * @param OutWrappedLineData	An optional array to fill with the indices from the source string marking the begin and end points of the wrapped lines
+	 *
+	 * @return The resulting string with newlines inserted.
+	 */
+	FString SLATE_API WrapText( const FString& InString, const FSlateFontInfo& InFontInfo, const int32 InWrapWidth, const float InFontScale, FWordWrapper::FWrappedLineData* const OutWrappedLineData = nullptr );
+
+	/**
+	 * Produce soft-wrapping information about TextLines
+	 * 
+	 * @param InTextLines           Lines of text, each string represents a run of text followed by a \n
+	 * @param InWrapWidth           The width (in Slate Units) at which we want the text wrapped.
+	 * @param OutWrappedText		Information about how the text would be wrapped given
+	 * @param OutWrappedLineData	An optional array to fill with the indices from the source string marking the begin and end points of the wrapped lines
+	 * 
+	 */
+	void SLATE_API WrapText( const TArray< TSharedRef<FString> >& InTextLines, const FSlateFontInfo& InFontInfo, const float InWrapWidth, const float InFontScale, TArray< FWrappedStringSlice >& OutWrappedText, FWordWrapper::FWrappedLineData* const OutWrappedLineData = nullptr );
+}
\ No newline at end of file
diff --git a/Engine/Source/Runtime/Slate/Public/Framework/Text/TextLayoutEngine.h b/Engine/Source/Runtime/Slate/Public/Framework/Text/TextLayoutEngine.h
new file mode 100644
--- /dev/null
+++ b/Engine/Source/Runtime/Slate/Public/Framework/Text/TextLayoutEngine.h
@@ -0,0 +1,27 @@
+// Copyright 1998-2014 Epic Games, Inc. All Rights Reserved.
+#pragma once
+
+// Generic Text Layout
+#include "TextRange.h"
+#include "TextHighlight.h"
+
+#include "IRun.h"
+#include "IRunHighlighter.h"
+#include "ILayoutBlock.h"
+
+#include "TextLayout.h"
+#include "DefaultLayoutBlock.h"
+#include "WidgetLayoutBlock.h"
+
+
+// Slate Specific Text Layout
+#include "ISlateRun.h"
+#include "ISlateRunHighlighter.h"
+
+#include "SlateTextLayout.h"
+
+#include "SlateTextRun.h"
+#include "SlateHyperlinkRun.h"
+#include "SlateImageRun.h"
+#include "SlateWidgetRun.h"
+#include "SlateSimpleRunHighlighter.h"
diff --git a/Engine/Source/Runtime/Slate/Public/Framework/Views/TableViewTypeTraits.h b/Engine/Source/Runtime/Slate/Public/Framework/Views/TableViewTypeTraits.h
new file mode 100644
--- /dev/null
+++ b/Engine/Source/Runtime/Slate/Public/Framework/Views/TableViewTypeTraits.h
@@ -0,0 +1,246 @@
+// Copyright 1998-2014 Epic Games, Inc. All Rights Reserved.
+
+#pragma once
+
+
+/**
+ * Lists/Trees only work with pointer types.
+ * Type traits to ensure that the user does not accidentally make a List/Tree of value types.
+ */
+template<typename T> struct TIsValidListItem											{ enum { Value = false }; };
+template<typename T> struct TIsValidListItem< TSharedRef<T, ESPMode::NotThreadSafe> >	{ enum { Value = true }; };
+template<typename T> struct TIsValidListItem< TSharedRef<T, ESPMode::ThreadSafe> >		{ enum { Value = true }; };
+template<typename T> struct TIsValidListItem< TSharedPtr<T, ESPMode::NotThreadSafe> >	{ enum { Value = true }; };
+template<typename T> struct TIsValidListItem< TSharedPtr<T, ESPMode::ThreadSafe> >		{ enum { Value = true }; };
+template<typename T> struct TIsValidListItem< T* >										{ enum { Value = true }; };
+template<typename T> struct TIsValidListItem< TWeakObjectPtr<T> >						{ enum { Value = true }; };
+
+/**
+ * Furthermore, ListViews of TSharedPtr<> work differently from lists of UObject*.
+ * ListTypeTraits provide the specialized functionality such as pointer testing, resetting,
+ * and optional serialization for UObject garbage collection.
+ */
+template <typename T> struct TListTypeTraits
+{
+};
+
+
+/**
+ * Pointer-related functionality (e.g. setting to null, testing for null) specialized for SharedPointers.
+ */
+template <typename T> struct TListTypeTraits< TSharedPtr<T, ESPMode::NotThreadSafe> >
+{
+public:
+	typedef TSharedPtr<T> NullableType;
+
+	static void AddReferencedObjects( FReferenceCollector&, TArray< TSharedPtr<T> >&, TSet< TSharedPtr<T> >&  )
+	{
+	}
+
+	static bool IsPtrValid( const TSharedPtr<T>& InPtr )
+	{
+		return InPtr.IsValid();
+	}
+
+	static void ResetPtr( TSharedPtr<T>& InPtr )
+	{
+		InPtr.Reset();
+	}
+
+	static TSharedPtr<T> MakeNullPtr()
+	{
+		return TSharedPtr<T>(NULL);
+	}
+
+	static TSharedPtr<T> NullableItemTypeConvertToItemType( const TSharedPtr<T>& InPtr )
+	{
+		return InPtr;
+	}
+
+	class SerializerType{};
+};
+
+
+template <typename T> struct TListTypeTraits< TSharedPtr<T, ESPMode::ThreadSafe> >
+{
+public:
+	typedef TSharedPtr<T, ESPMode::ThreadSafe> NullableType;
+
+	static void AddReferencedObjects( FReferenceCollector&, TArray< TSharedPtr<T, ESPMode::ThreadSafe> >&, TSet< TSharedPtr<T, ESPMode::ThreadSafe> >&  )
+	{
+	}
+
+	static bool IsPtrValid( const TSharedPtr<T, ESPMode::ThreadSafe>& InPtr )
+	{
+		return InPtr.IsValid();
+	}
+
+	static void ResetPtr( TSharedPtr<T, ESPMode::ThreadSafe>& InPtr )
+	{
+		InPtr.Reset();
+	}
+
+	static TSharedPtr<T, ESPMode::ThreadSafe> MakeNullPtr()
+	{
+		return TSharedPtr<T, ESPMode::ThreadSafe>(NULL);
+	}
+
+	static TSharedPtr<T, ESPMode::ThreadSafe> NullableItemTypeConvertToItemType( const TSharedPtr<T, ESPMode::ThreadSafe>& InPtr )
+	{
+		return InPtr;
+	}
+
+	class SerializerType{};
+};
+
+
+template <typename T> struct TListTypeTraits< TSharedRef<T, ESPMode::NotThreadSafe> >
+{
+public:
+	typedef TSharedPtr<T> NullableType;
+
+	static void AddReferencedObjects( FReferenceCollector&, TArray< TSharedRef<T> >&, TSet< TSharedRef<T> >&  )
+	{
+	}
+
+	static bool IsPtrValid( const TSharedPtr<T>& InPtr )
+	{
+		return InPtr.IsValid();
+	}
+
+	static void ResetPtr( TSharedPtr<T>& InPtr )
+	{
+		InPtr.Reset();
+	}
+
+	static TSharedPtr<T> MakeNullPtr()
+	{
+		return TSharedPtr<T>(NULL);
+	}
+
+	static TSharedRef<T> NullableItemTypeConvertToItemType( const TSharedPtr<T>& InPtr )
+	{
+		return InPtr.ToSharedRef();
+	}
+
+	class SerializerType{};
+};
+
+
+template <typename T> struct TListTypeTraits< TSharedRef<T, ESPMode::ThreadSafe> >
+{
+public:
+	typedef TSharedPtr<T, ESPMode::ThreadSafe> NullableType;
+
+	static void AddReferencedObjects( FReferenceCollector&, TArray< TSharedRef<T, ESPMode::ThreadSafe> >&, TSet< TSharedRef<T, ESPMode::ThreadSafe> >&  )
+	{
+	}
+
+	static bool IsPtrValid( const TSharedPtr<T, ESPMode::ThreadSafe>& InPtr )
+	{
+		return InPtr.IsValid();
+	}
+
+	static void ResetPtr( TSharedPtr<T, ESPMode::ThreadSafe>& InPtr )
+	{
+		InPtr.Reset();
+	}
+
+	static TSharedPtr<T, ESPMode::ThreadSafe> MakeNullPtr()
+	{
+		return TSharedPtr<T, ESPMode::ThreadSafe>(NULL);
+	}
+
+	static TSharedRef<T, ESPMode::ThreadSafe> NullableItemTypeConvertToItemType( const TSharedPtr<T, ESPMode::ThreadSafe>& InPtr )
+	{
+		return InPtr.ToSharedRef();
+	}
+
+	class SerializerType{};
+};
+
+
+/**
+ * Pointer-related functionality (e.g. setting to null, testing for null) specialized for SharedPointers.
+ */
+template <typename T> struct TListTypeTraits< TWeakObjectPtr<T> >
+{
+public:
+	typedef TWeakObjectPtr<T> NullableType;
+
+	static void AddReferencedObjects( FReferenceCollector&, TArray< TWeakObjectPtr<T> >&, TSet< TWeakObjectPtr<T> >&  )
+	{
+	}
+
+	static bool IsPtrValid( const TWeakObjectPtr<T>& InPtr )
+	{
+		return InPtr.IsValid();
+	}
+
+	static void ResetPtr( TWeakObjectPtr<T>& InPtr )
+	{
+		InPtr.Reset();
+	}
+
+	static TWeakObjectPtr<T> MakeNullPtr()
+	{
+		return TWeakObjectPtr<T>(NULL);
+	}
+
+	static TWeakObjectPtr<T> NullableItemTypeConvertToItemType( const TWeakObjectPtr<T>& InPtr )
+	{
+		return InPtr;
+	}
+
+	class SerializerType{};
+};
+
+
+/**
+ * Lists of pointer types only work if the pointers are deriving from UObject*.
+ * In addition to testing and setting the pointers to null, Lists of UObjects
+ * will serialize the objects they are holding onto.
+ */
+template <typename T> struct TListTypeTraits< T* >
+{
+public:
+	typedef T* NullableType;
+
+	static void AddReferencedObjects( FReferenceCollector& Collector, TArray< T* >& ItemsWithGeneratedWidgets, TSet< T* >& SelectedItems )
+	{
+		// Serialize generated items
+		const int32 NumGeneratedWidgets = ItemsWithGeneratedWidgets.Num();
+		for ( int32 ItemIndex = 0; ItemIndex < NumGeneratedWidgets; ++ItemIndex )
+		{
+			Collector.AddReferencedObject( ItemsWithGeneratedWidgets[ItemIndex] );
+		}
+		
+		// Serialize the selected items
+		for( typename TSet< T* >::TIterator SelecteItemIt(SelectedItems); SelecteItemIt; ++SelecteItemIt )
+		{
+			Collector.AddReferencedObject( *SelecteItemIt );
+		}
+	}
+
+	static bool IsPtrValid( T* InPtr )
+	{
+		return InPtr != NULL;
+	}
+
+	static void ResetPtr( T* InPtr )
+	{
+		InPtr = NULL;
+	}
+
+	static T* MakeNullPtr()
+	{
+		return NULL;
+	}
+
+	static T* NullableItemTypeConvertToItemType( T* InPtr )
+	{
+		return InPtr;
+	}
+
+	typedef FGCObject SerializerType;
+};
diff --git a/Engine/Source/Runtime/Slate/Public/Slate.h b/Engine/Source/Runtime/Slate/Public/Slate.h
--- a/Engine/Source/Runtime/Slate/Public/Slate.h
+++ b/Engine/Source/Runtime/Slate/Public/Slate.h
@@ -1,23 +1,23 @@
 // Copyright 1998-2014 Epic Games, Inc. All Rights Reserved.
 
-#ifndef _INC_SLATE
-#define _INC_SLATE
-
+#pragma once
 
 #include "Core.h"
 #include "CoreUObject.h"
+#include "Json.h"
 #include "SlateCore.h"
 
+#include "SlateClasses.h"
+
+
 // Enables or disables more slate stats.  These stats are usually per widget and can be gathered hundreds or thousands of times per frame
 // enabling these constantly causes a bit of a performance drop due to stat gathering overhead so these should only be enabled for profiling
 #define SLATE_HD_STATS 0
 
-SLATE_API DECLARE_LOG_CATEGORY_EXTERN(LogSlate, Log, All);
-SLATE_API DECLARE_LOG_CATEGORY_EXTERN(LogSlateStyles, Log, All);
-
 #define BEGIN_SLATE_FUNCTION_BUILD_OPTIMIZATION PRAGMA_DISABLE_OPTIMIZATION
 #define END_SLATE_FUNCTION_BUILD_OPTIMIZATION   PRAGMA_ENABLE_OPTIMIZATION
 
+
 /** Stat tracking for Slate */
 DECLARE_CYCLE_STAT_EXTERN( TEXT("Total Slate Tick Time"), STAT_SlateTickTime, STATGROUP_Slate , );
 DECLARE_CYCLE_STAT_EXTERN( TEXT("CacheDesiredSize"), STAT_SlateCacheDesiredSize, STATGROUP_Slate , );
@@ -37,94 +37,48 @@ DECLARE_CYCLE_STAT_EXTERN( TEXT("OnPaint STextBlock"), STAT_SlateOnPaint_STextBl
 DECLARE_CYCLE_STAT_EXTERN( TEXT("OnPaint SRichTextBlock"), STAT_SlateOnPaint_SRichTextBlock, STATGROUP_Slate , );
 DECLARE_CYCLE_STAT_EXTERN( TEXT("OnPaint SSlider"), STAT_SlateOnPaint_SSlider, STATGROUP_Slate , );
 DECLARE_CYCLE_STAT_EXTERN( TEXT("OnPaint SGraphPanel"), STAT_SlateOnPaint_SGraphPanel, STATGROUP_Slate , SLATE_API);
-DECLARE_CYCLE_STAT_EXTERN( TEXT("Slate Rendering GT Time"), STAT_SlateRenderingGTTime, STATGROUP_Slate , SLATE_API);
-DECLARE_CYCLE_STAT_EXTERN( TEXT("Font Caching"), STAT_SlateFontCachingTime, STATGROUP_Slate , );
-DECLARE_CYCLE_STAT_EXTERN( TEXT("Build Elements"), STAT_SlateBuildElements, STATGROUP_Slate , );
-DECLARE_CYCLE_STAT_EXTERN( TEXT("Find Batch"), STAT_SlateFindBatchTime, STATGROUP_Slate , );
-DECLARE_CYCLE_STAT_EXTERN( TEXT("Build Lines"), STAT_SlateBuildLines, STATGROUP_Slate , );
-DECLARE_CYCLE_STAT_EXTERN( TEXT("Build Splines"), STAT_SlateBuildSplines, STATGROUP_Slate , );
-DECLARE_CYCLE_STAT_EXTERN( TEXT("Build Gradients"), STAT_SlateBuildGradients, STATGROUP_Slate , );
-DECLARE_CYCLE_STAT_EXTERN( TEXT("Build Text"), STAT_SlateBuildText, STATGROUP_Slate , );
-DECLARE_CYCLE_STAT_EXTERN( TEXT("Build Border"), STAT_SlateBuildBorders, STATGROUP_Slate , );
-DECLARE_CYCLE_STAT_EXTERN( TEXT("Build Boxes"), STAT_SlateBuildBoxes, STATGROUP_Slate , );
-DECLARE_CYCLE_STAT_EXTERN( TEXT("Build Viewports"), STAT_SlateBuildViewports, STATGROUP_Slate , );
-DECLARE_CYCLE_STAT_EXTERN( TEXT("Update Buffers GT"), STAT_SlateUpdateBufferGTTime, STATGROUP_Slate , );
-
-
-DECLARE_CYCLE_STAT_EXTERN( TEXT("Slate Rendering RT Time"), STAT_SlateRenderingRTTime, STATGROUP_Slate , SLATE_API);
-DECLARE_CYCLE_STAT_EXTERN( TEXT("Update Buffers RT"), STAT_SlateUpdateBufferRTTime, STATGROUP_Slate , SLATE_API);
-DECLARE_CYCLE_STAT_EXTERN( TEXT("Draw Time"), STAT_SlateDrawTime, STATGROUP_Slate , SLATE_API);
-
-DECLARE_DWORD_COUNTER_STAT_EXTERN( TEXT("Num Layers"), STAT_SlateNumLayers, STATGROUP_Slate , );
-DECLARE_DWORD_COUNTER_STAT_EXTERN( TEXT("Num Batches"), STAT_SlateNumBatches, STATGROUP_Slate , );
-DECLARE_DWORD_COUNTER_STAT_EXTERN( TEXT("Num Vertices"), STAT_SlateVertexCount, STATGROUP_Slate , );
-
-DECLARE_CYCLE_STAT_EXTERN( TEXT("Measure String"), STAT_SlateMeasureStringTime, STATGROUP_Slate , );
-
 DECLARE_CYCLE_STAT_EXTERN( TEXT("Slate Misc Time"), STAT_SlateMiscTime, STATGROUP_Slate , );
 
-DECLARE_MEMORY_STAT_EXTERN( TEXT("Batch Vertex Memory"), STAT_SlateVertexBatchMemory, STATGROUP_SlateMemory , );
-DECLARE_MEMORY_STAT_EXTERN( TEXT("Batch Index Memory"), STAT_SlateIndexBatchMemory, STATGROUP_SlateMemory , );
-DECLARE_MEMORY_STAT_EXTERN( TEXT("Vertex Buffer Memory"), STAT_SlateVertexBufferMemory, STATGROUP_SlateMemory , SLATE_API);
-DECLARE_MEMORY_STAT_EXTERN( TEXT("Index Buffer Memory"), STAT_SlateIndexBufferMemory, STATGROUP_SlateMemory , SLATE_API);
-
-DECLARE_MEMORY_STAT_EXTERN( TEXT("Texture Data CPU Memory"), STAT_SlateTextureDataMemory, STATGROUP_SlateMemory, SLATE_API );
-DECLARE_MEMORY_STAT_EXTERN( TEXT("Texture Atlas Memory"), STAT_SlateTextureAtlasMemory, STATGROUP_SlateMemory, );
-DECLARE_MEMORY_STAT_EXTERN( TEXT("Font Kerning Table Memory"), STAT_SlateFontKerningTableMemory, STATGROUP_SlateMemory, );
-DECLARE_MEMORY_STAT_EXTERN( TEXT("Font Measure Memory"), STAT_SlateFontMeasureCacheMemory, STATGROUP_SlateMemory, );
-
-DECLARE_DWORD_ACCUMULATOR_STAT_EXTERN( TEXT("Num Texture Atlases"), STAT_SlateNumTextureAtlases, STATGROUP_SlateMemory, SLATE_API );
-DECLARE_DWORD_ACCUMULATOR_STAT_EXTERN( TEXT("Num Font Atlases"), STAT_SlateNumFontAtlases, STATGROUP_SlateMemory,  );
-DECLARE_DWORD_ACCUMULATOR_STAT_EXTERN( TEXT("Num Non-Atlased Textures"), STAT_SlateNumNonAtlasedTextures, STATGROUP_SlateMemory,SLATE_API );
-DECLARE_DWORD_ACCUMULATOR_STAT_EXTERN( TEXT("Num Dynamic Textures"), STAT_SlateNumDynamicTextures, STATGROUP_SlateMemory, SLATE_API);
-
 
 #if PLATFORM_IOS
-#include "IOS/IOSPlatformTextField.h"
+	#include "IOS/IOSPlatformTextField.h"
 #else
-#include "GenericPlatformTextField.h"
+	#include "GenericPlatformTextField.h"
 #endif
 
-#include "StyleDefaults.h"
-#include "CoreStyle.h"
 #include "InputCore.h"
-#include "ISlateSoundDevice.h"
 
-#include "WidgetPath.h"
-#include "SlateSprings.h"
-#include "SlateScrollHelper.h"
-#include "ISlateStyle.h"
-#include "SlateStyleRegistry.h"
-#include "Events.h"
-#include "CursorReply.h"
-#include "DragAndDrop.h"
-#include "Reply.h"
 #include "SlateDelegates.h"
-#include "SNullWidget.h"
-#include "DeclarativeSyntaxSupport.h"
-#include "RenderingCommon.h"
-#include "ThrottleManager.h"
-#include "InertialScrollManager.h"
+
+// Application
 #include "MenuStack.h"
-#include "EventLogger.h"
+#include "IPlatformTextField.h"
+#include "IWidgetReflector.h"
 #include "SlateApplication.h"
-#include "DragAndDrop.inl"
-#include "InputBinding.h"
+#include "SlateIcon.h"
+
+// Commands
+#include "InputGesture.h"
+#include "UIAction.h"
+#include "UICommandInfo.h"
+#include "InputBindingManager.h"
 #include "Commands.h"
+#include "UICommandList.h"
+
+// Legacy
+#include "SMissingWidget.h"
+#include "SlateScrollHelper.h"
+#include "SlateStyleRegistry.h"
+#include "UICommandDragDropOp.h"
+#include "InertialScrollManager.h"
 #include "GenericCommands.h"
-#include "Renderer.h"
-#include "DrawElements.h"
-#include "Children.h"
-#include "SWidget.h"
-#include "WidgetPath.inl"
-#include "SLeafWidget.h"
 #include "SWeakWidget.h"
 #include "TextLayoutEngine.h"
 #include "SPanel.h"
-#include "SOverlay.h"
 #include "SCompoundWidget.h"
 #include "SFxWidget.h"
 #include "SBorder.h"
+#include "SBorder.h"
 #include "SSeparator.h"
 #include "SSpacer.h"
 #include "SWrapBox.h"
@@ -136,7 +90,6 @@ DECLARE_DWORD_ACCUMULATOR_STAT_EXTERN( TEXT("Num Dynamic Textures"), STAT_SlateN
 #include "TextDecorators.h"
 #include "SRichTextBlock.h"
 #include "SBox.h"
-#include "SBoxPanel.h"
 #include "SHeader.h"
 #include "SGridPanel.h"
 #include "SUniformGridPanel.h"
@@ -152,37 +105,28 @@ DECLARE_DWORD_ACCUMULATOR_STAT_EXTERN( TEXT("Num Dynamic Textures"), STAT_SlateN
 #include "SErrorHint.h"
 #include "SPopUpErrorText.h"
 #include "SSplitter.h"
-#include "TableView/TableViewTypeTraits.h"
-#include "TableView/SExpanderArrow.h"
-#include "TableView/ITypedTableView.h"
-#include "TableView/STableViewBase.h"
-#include "TableView/SHeaderRow.h"
-#include "TableView/STableRow.h"
-#include "TableView/SListView.h"
-#include "TableView/STileView.h"
-#include "TableView/STreeView.h"
+#include "TableViewTypeTraits.h"
+#include "SExpanderArrow.h"
+#include "ITypedTableView.h"
+#include "STableViewBase.h"
+#include "SHeaderRow.h"
+#include "STableRow.h"
+#include "SListView.h"
+#include "STileView.h"
+#include "STreeView.h"
 #include "SScrollBox.h"
-#include "SWindow.h"
-#include "RenderingPolicy.h"
-#include "ElementBatcher.h"
-#include "Renderer.h"
-#include "TextureAtlas.h"
-#include "FontTypes.h"
-#include "FontMeasure.h"
-#include "TextureManager.h"
 #include "SViewport.h"
 #include "SColorBlock.h"
 #include "SCheckBox.h"
 #include "SSpinBox.h"
 #include "SSlider.h"
 #include "SVolumeControl.h"
-#include "SToolTip.h"
 #include "SColorSpectrum.h"
 #include "SColorWheel.h"
 #include "SColorThemes.h"
 #include "SColorPicker.h"
-#include "MultiBox/MultiBoxBuilder.h"
-#include "MultiBox/MultiBoxExtender.h"
+#include "MultiBoxBuilder.h"
+#include "MultiBoxExtender.h"
 #include "SComboButton.h"
 #include "SComboBox.h"
 #include "SHyperlink.h"
@@ -206,19 +150,17 @@ DECLARE_DWORD_ACCUMULATOR_STAT_EXTERN( TEXT("Num Dynamic Textures"), STAT_SlateN
 #include "STutorialWrapper.h"
 #include "ScrollyZoomy.h"
 #include "SSafeZone.h"
+#include "MarqueeRect.h"
 
 // Docking Framework
-#include "DockingFramework/WorkspaceItem.h"
-#include "DockingFramework/TabManager.h"
-#include "DockingFramework/SDockTab.h"
-#include "DockingFramework/LayoutService.h"
-
+#include "WorkspaceItem.h"
+#include "TabManager.h"
+#include "SDockTab.h"
+#include "LayoutService.h"
 
 // Old docking
-#include "Docking/SDockableTab.h"
-#include "Docking/SDockTabStack.h"
+#include "SDockableTab.h"
+#include "SDockTabStack.h"
 
 // Test suite
-#include "Testing/STestSuite.h"
-
-#endif // _INC_SLATE
+#include "STestSuite.h"
diff --git a/Engine/Source/Runtime/Slate/Public/Widgets/Colors/SColorBlock.h b/Engine/Source/Runtime/Slate/Public/Widgets/Colors/SColorBlock.h
new file mode 100644
--- /dev/null
+++ b/Engine/Source/Runtime/Slate/Public/Widgets/Colors/SColorBlock.h
@@ -0,0 +1,77 @@
+// Copyright 1998-2014 Epic Games, Inc. All Rights Reserved.
+
+#pragma once
+
+
+class SLATE_API SColorBlock : public SLeafWidget
+{
+public:
+
+	SLATE_BEGIN_ARGS( SColorBlock )
+		: _Color( FLinearColor::White )
+		, _ColorIsHSV( false )
+		, _IgnoreAlpha( false )
+		, _ShowBackgroundForAlpha( false )
+		, _UseSRGB(true)
+		, _OnMouseButtonDown()
+		, _Size( FVector2D(16,16) )
+		{}
+
+		/** The color to display for this color block */
+		SLATE_ATTRIBUTE( FLinearColor, Color )
+
+		/** Whether the color displayed is HSV or not */
+		SLATE_ATTRIBUTE( bool, ColorIsHSV )
+
+		/** Whether to ignore alpha entirely from the input color */
+		SLATE_ATTRIBUTE( bool, IgnoreAlpha )
+
+		/** Whether to display a background for viewing opacity. Irrelevant if ignoring alpha */
+		SLATE_ATTRIBUTE( bool, ShowBackgroundForAlpha )
+
+		/** Whether to display sRGB color */
+		SLATE_ATTRIBUTE( bool, UseSRGB )
+
+		/** A handler to activate when the mouse is pressed. */
+		SLATE_EVENT( FPointerEventHandler, OnMouseButtonDown )
+
+		/** How big should this color block be? */
+		SLATE_ATTRIBUTE( FVector2D, Size )
+	SLATE_END_ARGS()
+
+	/**
+	 * Construct this widget
+	 *
+	 * @param	InArgs	The declaration data for this widget
+	 */
+	void Construct( const FArguments& InArgs );
+
+private:
+	// SWidget interface
+	virtual int32 OnPaint( const FGeometry& AllottedGeometry, const FSlateRect& MyClippingRect, FSlateWindowElementList& OutDrawElements, int32 LayerId, const FWidgetStyle& InWidgetStyle, bool bParentEnabled ) const OVERRIDE;
+	virtual FReply OnMouseButtonDown( const FGeometry& MyGeometry, const FPointerEvent& MouseEvent ) OVERRIDE;
+	virtual FVector2D ComputeDesiredSize() const OVERRIDE;
+	// End of SWidget interface
+
+private:
+
+	/** The color to display for this color block */
+	TAttribute< FLinearColor > Color;
+
+	/** Whether the color displayed is HSV or not */
+	TAttribute< bool > ColorIsHSV;
+	
+	/** Whether to ignore alpha entirely from the input color */
+	TAttribute< bool > IgnoreAlpha;
+
+	/** Whether to display a background for viewing opacity. Irrelevant if ignoring alpha */
+	TAttribute< bool > ShowBackgroundForAlpha;
+
+	/** Whether to display sRGB color */
+	TAttribute<bool> bUseSRGB;
+	
+	/** A handler to activate when the mouse is pressed. */
+	FPointerEventHandler MouseButtonDownHandler;
+
+	TAttribute<FVector2D> ColorBlockSize;
+};
