Slate: Refactored core Slate implementation into SlateCore module in preparation for UMG.

diff --git a/Engine/Source/Runtime/Slate/Public/Widgets/Colors/SColorPicker.h b/Engine/Source/Runtime/Slate/Public/Widgets/Colors/SColorPicker.h
new file mode 100644
--- /dev/null
+++ b/Engine/Source/Runtime/Slate/Public/Widgets/Colors/SColorPicker.h
@@ -0,0 +1,526 @@
+// Copyright 1998-2014 Epic Games, Inc. All Rights Reserved.
+
+#pragma once
+
+
+/** Called when the color picker cancel button is pressed */
+DECLARE_DELEGATE_OneParam( FOnColorPickerCancelled, FLinearColor );
+
+
+namespace EColorPickerChannels
+{
+	/**
+	 * Enumerates color channels (do not reorder).
+	 */
+	enum Type
+	{
+		Red,
+		Green,
+		Blue,
+		Alpha,
+		Hue,
+		Saturation,
+		Value
+	};
+}
+
+
+namespace EColorPickerModes
+{
+	/**
+	 * Enumerates color picker modes.
+	 */
+	enum Type
+	{
+		Spectrum,
+		Wheel
+	};
+}
+
+
+/**
+ * Struct for holding individual pointers to float values.
+ */
+struct FColorChannels
+{
+	FColorChannels()
+	{
+		Red = Green = Blue = Alpha = NULL;
+	}
+
+	float* Red;
+	float* Green;
+	float* Blue;
+	float* Alpha;
+};
+
+
+/**
+ * Class for placing a color picker. If all you need is a standalone color picker,
+ * use the functions OpenColorPicker and DestroyColorPicker, since they hold a static
+ * instance of the color picker.
+ */
+class SLATE_API SColorPicker
+	: public SCompoundWidget
+{
+public:
+
+	SLATE_BEGIN_ARGS(SColorPicker)
+		: _TargetColorAttribute(FLinearColor(ForceInit))
+		, _TargetFColors()
+		, _TargetLinearColors()
+		, _TargetColorChannels()
+		, _UseAlpha(true)
+		, _OnlyRefreshOnMouseUp(false)
+		, _OnlyRefreshOnOk(false)
+		, _OnColorCommitted()
+		, _PreColorCommitted()
+		, _OnColorPickerCancelled()
+		, _OnColorPickerWindowClosed()
+		, _OnInteractivePickBegin()
+		, _OnInteractivePickEnd()
+		, _ParentWindow()
+		, _DisplayGamma(2.2f)
+		, _DisplayInlineVersion(false)
+		{}
+		
+		/** The color that is being targeted as a TAttribute */
+		SLATE_ATTRIBUTE(FLinearColor, TargetColorAttribute)
+		
+		/** An array of color pointers this color picker targets */
+		SLATE_ATTRIBUTE(TArray<FColor*>, TargetFColors)
+		
+		/** An array of linear color pointers this color picker targets */
+		SLATE_ATTRIBUTE(TArray<FLinearColor*>, TargetLinearColors)
+		
+		/**
+		 * An array of color pointer structs this color picker targets
+		 * Only to keep compatibility with wx. Should be removed once wx is gone.
+		 */
+		SLATE_ATTRIBUTE(TArray<FColorChannels>, TargetColorChannels)
+
+		/** Whether the ability to pick the alpha value is enabled */
+		SLATE_ATTRIBUTE(bool, UseAlpha)
+		
+		/** Prevents immediate refreshs for performance reasons. */
+		SLATE_ATTRIBUTE(bool, OnlyRefreshOnMouseUp)
+
+		/** Prevents multiple refreshes when requested. */
+		SLATE_ATTRIBUTE(bool, OnlyRefreshOnOk)
+		
+		/** The event called when the color is committed */
+		SLATE_EVENT( FOnLinearColorValueChanged, OnColorCommitted )
+
+		/** The event called before the color is committed */
+		SLATE_EVENT( FOnLinearColorValueChanged, PreColorCommitted )
+		
+		/** The event called when the color picker cancel button is pressed */
+		SLATE_EVENT( FOnColorPickerCancelled, OnColorPickerCancelled )
+
+		/** The event called when the color picker parent window is closed */
+		SLATE_EVENT( FOnWindowClosed, OnColorPickerWindowClosed )
+
+		/** The event called when a slider drag, color wheel drag or dropper grab starts */
+		SLATE_EVENT( FSimpleDelegate, OnInteractivePickBegin )
+
+		/** The event called when a slider drag, color wheel drag or dropper grab finishes */
+		SLATE_EVENT( FSimpleDelegate, OnInteractivePickEnd )
+
+		/** A pointer to the parent window */
+		SLATE_ATTRIBUTE(TSharedPtr<SWindow>, ParentWindow)
+
+		/** Sets the display Gamma setting - used to correct colors sampled from the screen */
+		SLATE_ATTRIBUTE(float, DisplayGamma)
+
+		/** If true, this color picker will be a stripped down version of the full color picker */
+		SLATE_ARGUMENT(bool, DisplayInlineVersion)
+
+	SLATE_END_ARGS()
+	
+	/** A default window size for the color picker which looks nice */
+	static const FVector2D DEFAULT_WINDOW_SIZE;
+
+public:
+
+	/**
+	 *	Default Destructor
+	 */
+	~SColorPicker();
+
+public:
+
+	/**
+	 * Construct the widget
+	 *
+	 * @param InArgs   Declaration from which to construct the widget.
+	 */
+	void Construct(const FArguments& InArgs );
+
+protected:	
+
+	// SWidget interface
+	virtual void Tick( const FGeometry& AllottedGeometry, const double InCurrentTime, const float InDeltaTime ) OVERRIDE;
+	// End of SWidget interface
+
+	/** Backup all the colors that are being modified */
+	void BackupColors();
+
+	/** Restore all the modified colors to their original state */
+	void RestoreColors();
+
+	/** Set all the colors to this new color */
+	void SetColors( const FLinearColor& InColor );
+
+	bool ApplyNewTargetColor( bool bForceUpdate = false );
+
+	void GenerateDefaultColorPickerContent( bool bAdvancedSectionExpanded );
+	void GenerateInlineColorPickerContent();
+
+	FLinearColor GetCurrentColor() const
+	{
+		return CurrentColorHSV;
+	}
+
+	/** Calls the user defined delegate for when the color changes are discarded */
+	void DiscardColor();
+
+	/** Sets new color in ether RGB or HSV */
+	bool SetNewTargetColorRGB( const FLinearColor& NewValue, bool bForceUpdate = false );
+	bool SetNewTargetColorHSV( const FLinearColor& NewValue, bool bForceUpdate = false );
+
+    void UpdateColorPick();
+	void UpdateColorPickMouseUp();
+
+	void BeginAnimation(FLinearColor Start, FLinearColor End);
+
+	void HideSmallTrash();
+	void ShowSmallTrash();
+
+	/**
+	 * Cycles the color picker's mode.
+	 */
+	void CycleMode( );
+
+	/**
+	 * Creates a color slider widget for the specified channel.
+	 *
+	 * @param Channel - The color channel to create the widget for.
+	 *
+	 * @return The new slider.
+	 */
+	TSharedRef<SWidget> MakeColorSlider( EColorPickerChannels::Type Channel ) const;
+
+	/**
+	 * Creates a color spin box widget for the specified channel.
+	 *
+	 * @param Channel - The color channel to create the widget for.
+	 *
+	 * @return The new spin box.
+	 */
+	TSharedRef<SWidget> MakeColorSpinBox( EColorPickerChannels::Type Channel ) const;
+
+	/**
+	 * Creates the color preview box widget.
+	 *
+	 * @return The new color preview box.
+	 */
+	TSharedRef<SWidget> MakeColorPreviewBox( ) const;
+
+private:
+
+	// Callback for getting the end color of a color spin box gradient.
+	FLinearColor GetGradientEndColor( EColorPickerChannels::Type Channel ) const;
+
+	// Callback for getting the start color of a color spin box gradient.
+	FLinearColor GetGradientStartColor( EColorPickerChannels::Type Channel ) const;
+
+	// Callback for handling expansion of the 'Advanced' area.
+	void HandleAdvancedAreaExpansionChanged( bool Expanded );
+
+	// Callback for getting the visibility of the alpha channel portion in color blocks.
+	EVisibility HandleAlphaColorBlockVisibility( ) const;
+
+	// Callback for clicking the Cancel button.
+	FReply HandleCancelButtonClicked();
+
+	// Callback for pressing a mouse button in the color area.
+	FReply HandleColorAreaMouseDown( const FGeometry& MyGeometry, const FPointerEvent& MouseEvent );
+
+	// Callback for clicking the color picker mode button.
+	FReply HandleColorPickerModeButtonClicked( );
+
+	// Callback for getting the visibility of the given color picker mode.
+	EVisibility HandleColorPickerModeVisibility( EColorPickerModes::Type Mode ) const;
+
+	// Callback for getting the end color of a color slider.
+	FLinearColor HandleColorSliderEndColor( EColorPickerChannels::Type Channel ) const;
+
+	// Callback for getting the start color of a color slider.
+	FLinearColor HandleColorSliderStartColor( EColorPickerChannels::Type Channel ) const;
+
+	// Callback for value changes in the color spectrum picker.
+	void HandleColorSpectrumValueChanged( FLinearColor NewValue );
+
+	// Callback for getting the value of a color spin box.
+	float HandleColorSpinBoxValue( EColorPickerChannels::Type Channel ) const;
+
+	// Callback for value changes in a color spin box.
+	void HandleColorSpinBoxValueChanged( float NewValue, EColorPickerChannels::Type Channel );
+
+	// Callback for completed eye dropper interactions.
+	void HandleEyeDropperButtonComplete( );
+
+	// Callback for getting the text in the hex box.
+	FText HandleHexBoxText( ) const;
+
+	// Callback for committed text in the hex input box.
+	void HandleHexInputTextCommitted( const FText& Text, ETextCommit::Type CommitType );
+
+	// Callback for changing the HSV value of the current color.
+	void HandleHSVColorChanged( FLinearColor NewValue );
+
+	// Callback for when interactive user input begins.
+	void HandleInteractiveChangeBegin( );
+
+	// Callback for when interactive user input ends.
+	void HandleInteractiveChangeEnd( );
+
+	// Callback for when interactive user input ends.
+	void HandleInteractiveChangeEnd( float NewValue );
+
+	// Callback for clicking the new color preview block.
+	FReply HandleNewColorBlockMouseButtonDown( const FGeometry& MyGeometry, const FPointerEvent& MouseEvent, bool bCheckAlpha );
+
+	// Callback for clicking the OK button.
+	FReply HandleOkButtonClicked( );
+
+	// Callback for clicking the old color preview block.
+	FReply HandleOldColorBlockMouseButtonDown( const FGeometry& MyGeometry, const FPointerEvent& MouseEvent, bool bCheckAlpha );
+
+	// Callback for checking whether sRGB colors should be rendered.
+	bool HandleColorPickerUseSRGB( ) const;
+
+	// Callback for when the parent window has been closed.
+	void HandleParentWindowClosed( const TSharedRef<SWindow>& Window );
+
+	// Callback for changing the RGB value of the current color.
+	void HandleRGBColorChanged( FLinearColor NewValue );
+
+	// Callback for changing the checked state of the sRGB check box.
+	void HandleSRGBCheckBoxCheckStateChanged( ESlateCheckBoxState::Type InIsChecked );
+
+	// Callback for determining whether the sRGB check box should be checked.
+	ESlateCheckBoxState::Type HandleSRGBCheckBoxIsChecked( ) const;
+
+	// Callback for selecting a color in the color theme bar.
+	void HandleThemeBarColorSelected( FLinearColor NewValue );
+
+	// Callback for getting the theme bar's color theme.
+	TSharedPtr<FColorTheme> HandleThemeBarColorTheme( ) const;
+
+	// Callback for getting the visibility of the theme bar hint text.
+	EVisibility HandleThemeBarHintVisibility( ) const;
+
+	// Callback for determining whether the theme bar should display the alpha channel.
+	bool HandleThemeBarUseAlpha( ) const;
+
+	// Callback for theme viewer changes.
+	void HandleThemesViewerThemeChanged( );
+
+private:
+	
+	/** The color that is being targeted as a TAttribute */
+	TAttribute<FLinearColor> TargetColorAttribute;
+
+	/** The current color being picked in HSV */
+	FLinearColor CurrentColorHSV;
+
+	/** The current color being picked in RGB */
+	FLinearColor CurrentColorRGB;
+	
+	/** The old color to be changed in HSV */
+	FLinearColor OldColor;
+	
+	/** Color end point to animate to */
+	FLinearColor ColorEnd;
+
+	/** Color start point to animate from */
+	FLinearColor ColorBegin;
+
+	// Holds the color picker's mode.
+	EColorPickerModes::Type CurrentMode;
+
+	/** Time, used for color animation */
+	float CurrentTime;
+
+	/** The max time allowed for updating before we shut off auto-updating */
+	static const double MAX_ALLOWED_UPDATE_TIME;
+
+	/** If true, then the performance is too bad to have auto-updating */
+	bool bPerfIsTooSlowToUpdate;
+
+	/** An array of color pointers this color picker targets */
+	TArray<FColor*> TargetFColors;
+	
+	/** An array of linear color pointers this color picker targets */
+	TArray<FLinearColor*> TargetLinearColors;
+	
+	/**
+	 * An array of color pointer structs this color picker targets
+	 * Only to keep compatibility with wx. Should be removed once wx is gone.
+	 */
+	TArray<FColorChannels> TargetColorChannels;
+
+	/** Backups of the TargetFColors */
+	TArray<FColor> OldTargetFColors;
+
+	/** Backups of the TargetLinearColors */
+	TArray<FLinearColor> OldTargetLinearColors;
+
+	/** Backups of the TargetColorChannels */
+	TArray<FLinearColor> OldTargetColorChannels;
+
+	/** Whether or not the color uses Alpha or not */
+	TAttribute<bool> bUseAlpha;
+
+	/** Prevents immediate refresh for performance reasons. */
+	bool bOnlyRefreshOnMouseUp;
+
+	/** Prevents multiple refreshes when requested. */
+	bool bOnlyRefreshOnOk;
+
+	/** true if the picker was closed via the OK or Cancel buttons, false otherwise */
+	bool bClosedViaOkOrCancel;
+
+	/** Pointer to the parent window, so we know to destroy it when done */
+	TWeakPtr<SWindow> ParentWindowPtr;
+
+	/** The widget which holds the currently selected theme */
+	TSharedPtr<class SThemeColorBlocksBar> CurrentThemeBar;
+
+	/** Widget which is either the button to show the color themes viewer, or to be a color trash */
+	TSharedPtr<SBorder> ColorThemeButtonOrSmallTrash;
+
+	/** The button to show the color themes viewer */
+	TSharedPtr<class SComboButton> ColorThemeComboButton;
+
+	/** The small color trash shown in place of the combo button */
+	TSharedPtr<SWidget> SmallTrash;
+
+	/** Sets the display Gamma setting - used to correct colors sampled from the screen */
+	TAttribute<float> DisplayGamma;
+
+	/** True if this color picker is an inline color picker */
+	bool bColorPickerIsInlineVersion;
+
+	/** Tracks whether the user is moving a value spin box, the color wheel and the dropper */
+	bool bIsInteractive;
+
+private:
+
+	/** Invoked when a new value is selected on the color wheel */
+	FOnLinearColorValueChanged OnColorCommitted;
+
+	/** Invoked before a new value is selected on the color wheel */
+	FOnLinearColorValueChanged PreColorCommitted;
+
+	/** Invoked when the color picker cancel button is pressed */
+	FOnColorPickerCancelled OnColorPickerCancelled;
+
+	/** Invoked when a slider drag, color wheel drag or dropper grab starts */
+	FSimpleDelegate OnInteractivePickBegin;
+
+	/** Invoked when a slider drag, color wheel drag or dropper grab finishes */
+	FSimpleDelegate OnInteractivePickEnd;
+
+	/** Invoked when the color picker window closes. */
+	FOnWindowClosed OnColorPickerWindowClosed;
+
+private:
+
+	/** A static pointer to the global color themes viewer */
+	static TSharedPtr<class SColorThemesViewer> ColorThemesViewer;
+};
+
+
+struct FColorPickerArgs
+{
+	/** Whether or not the new color picker is modal. */
+	bool bIsModal;
+
+	/** The parent for the new color picker window */
+	TSharedPtr<const SWidget> ParentWidget;
+
+	/** Whether or not to enable the alpha slider. */
+	bool bUseAlpha;
+
+	/** Whether to disable the refresh except on mouse up for performance reasons. */
+	bool bOnlyRefreshOnMouseUp;
+
+	/** Whether to disable the refresh until the picker closes. */
+	bool bOnlyRefreshOnOk;
+
+	/** The current display gamma used to correct colors picked from the display. */
+	TAttribute<float> DisplayGamma;
+
+	/** An array of FColors to target. */
+	const TArray<FColor*>* ColorArray;
+
+	/** An array of FLinearColors to target. */
+	const TArray<FLinearColor*>* LinearColorArray;
+
+	/** An array of FColorChannels to target. (deprecated now that wx is gone?) */
+	const TArray<FColorChannels>* ColorChannelsArray;
+
+	/** A delegate to be called when the color changes. */
+	FOnLinearColorValueChanged OnColorCommitted;
+
+	/** A delegate to be called before the color change is committed. */
+	FOnLinearColorValueChanged PreColorCommitted;
+
+	/** A delegate to be called when the color picker window closes. */
+	FOnWindowClosed OnColorPickerWindowClosed;
+
+	/** A delegate to be called when the color picker cancel button is pressed */
+	FOnColorPickerCancelled OnColorPickerCancelled;
+
+	/** A delegate to be called when a slider drag, color wheel drag or dropper grab starts */
+	FSimpleDelegate OnInteractivePickBegin;
+
+	/** A delegate to be called when a slider drag, color wheel drag or dropper grab finishes */
+	FSimpleDelegate OnInteractivePickEnd;
+
+	/** Overrides the initial color set on the color picker. */
+	FLinearColor InitialColorOverride;
+
+	FColorPickerArgs()
+		: bIsModal(false)
+		, ParentWidget(NULL)
+		, bUseAlpha(false)
+		, bOnlyRefreshOnMouseUp(false)
+		, bOnlyRefreshOnOk(false)
+		, DisplayGamma(2.2f)
+		, ColorArray(NULL)
+		, LinearColorArray(NULL)
+		, ColorChannelsArray(NULL)
+		, OnColorCommitted()
+		, PreColorCommitted()
+		, OnColorPickerWindowClosed()
+		, OnColorPickerCancelled()
+		, OnInteractivePickBegin()
+		, OnInteractivePickEnd()
+		, InitialColorOverride()
+	{ }
+};
+
+
+/**
+ * Opens up the static color picker, destroying any previously existing one.
+ */
+SLATE_API bool OpenColorPicker( const FColorPickerArgs& Args );
+
+/**
+ * Destroys the current color picker. Necessary if the values the color picker
+ * currently targets become invalid.
+ */
+SLATE_API void DestroyColorPicker();
diff --git a/Engine/Source/Runtime/Slate/Public/Widgets/Colors/SColorSpectrum.h b/Engine/Source/Runtime/Slate/Public/Widgets/Colors/SColorSpectrum.h
new file mode 100644
--- /dev/null
+++ b/Engine/Source/Runtime/Slate/Public/Widgets/Colors/SColorSpectrum.h
@@ -0,0 +1,99 @@
+// Copyright 1998-2014 Epic Games, Inc. All Rights Reserved.
+
+/*=============================================================================
+	SColorSpectrum.h: Declares the SColorSpectrum class.
+=============================================================================*/
+
+#pragma once
+
+
+/**
+ * Implements the color spectrum widget.
+ */
+class SColorSpectrum
+	: public SLeafWidget
+{
+public:
+
+	SLATE_BEGIN_ARGS(SColorSpectrum)
+		: _SelectedColor()
+		, _OnMouseCaptureBegin()
+		, _OnMouseCaptureEnd()
+		, _OnValueChanged()
+	{ }
+	
+		/** The current color selected by the user */
+		SLATE_ATTRIBUTE(FLinearColor, SelectedColor)
+		
+		/** Invoked when the mouse is pressed and a capture begins */
+		SLATE_EVENT(FSimpleDelegate, OnMouseCaptureBegin)
+
+		/** Invoked when the mouse is released and a capture ends */
+		SLATE_EVENT(FSimpleDelegate, OnMouseCaptureEnd)
+
+		/** Invoked when a new value is selected on the color wheel */
+		SLATE_EVENT(FOnLinearColorValueChanged, OnValueChanged)
+
+	SLATE_END_ARGS()
+	
+public:
+
+	/**
+	 * Construct this widget.
+	 *
+	 * @param InArgs - The declaration data for this widget.
+	 */
+	void Construct( const FArguments& InArgs );
+
+public:
+
+	// Begin overrides
+
+	virtual FVector2D ComputeDesiredSize( ) const;
+
+	virtual FReply OnMouseButtonDoubleClick( const FGeometry& InMyGeometry, const FPointerEvent& InMouseEvent )
+	{
+		return FReply::Handled();
+	}
+
+	virtual FReply OnMouseButtonDown( const FGeometry& MyGeometry, const FPointerEvent& MouseEvent );
+	
+	virtual FReply OnMouseButtonUp( const FGeometry& MyGeometry, const FPointerEvent& MouseEvent );
+
+	virtual FReply OnMouseMove( const FGeometry& MyGeometry, const FPointerEvent& MouseEvent );
+
+	virtual int32 OnPaint( const FGeometry& AllottedGeometry, const FSlateRect& MyClippingRect, FSlateWindowElementList& OutDrawElements, int32 LayerId, const FWidgetStyle& InWidgetStyle, bool bParentEnabled ) const OVERRIDE;
+
+	// End overrides
+	
+protected:
+
+	/**
+	 * Calculates the position of the color selection indicator.
+	 *
+	 * @return The position relative to the widget.
+	 */
+	FVector2D CalcRelativeSelectedPosition( ) const;
+
+private:
+
+	// The color wheel image to show.
+	const FSlateBrush* Image;
+	
+	// The current color selected by the user.
+	TAttribute< FLinearColor > SelectedColor;
+
+	// The color selector image to show.
+	const FSlateBrush* SelectorImage;
+
+private:
+
+	// Holds a delegate that is executed when the mouse is pressed and a capture begins.
+	FSimpleDelegate OnMouseCaptureBegin;
+
+	// Holds a delegate that is executed when the mouse is let up and a capture ends.
+	FSimpleDelegate OnMouseCaptureEnd;
+
+	// Holds a delegate that is executed when a new value is selected on the color wheel.
+	FOnLinearColorValueChanged OnValueChanged;
+};
diff --git a/Engine/Source/Runtime/Slate/Public/Widgets/Colors/SColorThemes.h b/Engine/Source/Runtime/Slate/Public/Widgets/Colors/SColorThemes.h
new file mode 100644
--- /dev/null
+++ b/Engine/Source/Runtime/Slate/Public/Widgets/Colors/SColorThemes.h
@@ -0,0 +1,640 @@
+// Copyright 1998-2014 Epic Games, Inc. All Rights Reserved.
+
+#pragma once
+
+
+/**
+ * A Color Theme is a name and an array of Colors.
+ * It also holds and array of refresh callbacks which it calls every time it changes at all.
+ */
+class FColorTheme
+{
+public:
+
+	FColorTheme(const FString& InName = TEXT(""), const TArray< TSharedPtr<FLinearColor> >& InColors = TArray< TSharedPtr<FLinearColor> >());
+
+	/** Get a list of all the colors in the theme */
+	const TArray< TSharedPtr<FLinearColor> >& GetColors() const { return Colors; }
+
+	/** Insert a color at a specific point in the list and broadcast change */
+	void InsertNewColor(TSharedPtr<FLinearColor> InColor, int32 InsertPosition);
+
+	/** Check to see if a color is already present in the list */
+	int32 FindApproxColor(const FLinearColor& InColor, float Tolerance = KINDA_SMALL_NUMBER) const;
+
+	/** Remove all colors from the list, broadcast change */
+	void RemoveAll();
+
+	/** Remove color at index from the list, broadcast change */
+	void RemoveColor(int32 ColorIndex);
+
+	/** Remove specific color from the list, broadcast change */
+	int32 RemoveColor(const TSharedPtr<FLinearColor> InColor);
+	
+	FString Name;
+
+	DECLARE_EVENT( FColorTheme, FRefreshEvent );
+	FRefreshEvent& OnRefresh() { return RefreshEvent; }
+
+private:
+
+	TArray< TSharedPtr<FLinearColor> > Colors;
+
+	FRefreshEvent RefreshEvent;
+};
+
+
+/**
+ * The SColorTrash is a multipurpose widget which allows FColorDragDrops
+ * to be dropped on to to be 
+ */
+class SColorTrash : public SCompoundWidget
+{
+public:
+	SLATE_BEGIN_ARGS(SColorTrash)
+		: _UsesSmallIcon(false)
+		{}
+
+		SLATE_ATTRIBUTE(bool, UsesSmallIcon)
+	SLATE_END_ARGS()
+
+	void Construct( const FArguments& InArgs );
+
+protected:	
+	/**
+	 * Called during drag and drop when the drag enters a widget.
+	 *
+	 * @param MyGeometry      The geometry of the widget receiving the event.
+	 * @param DragDropEvent   The drag and drop event.
+	 *
+	 * @return A reply that indicated whether this event was handled.
+	 */
+	virtual void OnDragEnter( const FGeometry& MyGeometry, const FDragDropEvent& DragDropEvent ) OVERRIDE;
+
+	/**
+	 * Called during drag and drop when the drag leaves a widget.
+	 *
+	 * @param DragDropEvent   The drag and drop event.
+	 *
+	 * @return A reply that indicated whether this event was handled.
+	 */
+	virtual void OnDragLeave( const FDragDropEvent& DragDropEvent ) OVERRIDE;
+
+	/**
+	 * Called when the user is dropping something onto a widget; terminates drag and drop.
+	 *
+	 * @param MyGeometry      The geometry of the widget receiving the event.
+	 * @param DragDropEvent   The drag and drop event.
+	 *
+	 * @return A reply that indicated whether this event was handled.
+	 */
+	virtual FReply OnDrop( const FGeometry& MyGeometry, const FDragDropEvent& DragDropEvent ) OVERRIDE;
+	
+	const FSlateBrush* GetBorderStyle() const;
+
+private:
+	/** Determines whether to draw the border to show activation */
+	bool bBorderActivated;
+};
+
+
+class SThemeColorBlocksBar;
+
+
+/**
+ * SThemeColorBlocks are Color Blocks which point to a Color in a ColorTheme.
+ * They can be dragged and dropped, and clicking on one in the Color Picker will
+ * give the color that they point to.
+ */
+class SThemeColorBlock : public SCompoundWidget
+{
+public:
+	SLATE_BEGIN_ARGS(SThemeColorBlock)
+		: _Color()
+		, _OnSelectColor()
+		, _Parent()
+		, _ShowTrashCallback()
+		, _HideTrashCallback()
+		, _UseSRGB()
+		, _UseAlpha()
+		{}
+
+		/** A pointer to the color this block uses */
+		SLATE_ATTRIBUTE(TSharedPtr<FLinearColor>, Color)
+		
+		/** Event called when this block is clicked */
+		SLATE_EVENT(FOnLinearColorValueChanged, OnSelectColor)
+
+		/** A pointer to the theme color blocks bar that is this block's origin */
+		SLATE_ATTRIBUTE(TSharedPtr<SThemeColorBlocksBar>, Parent)
+		
+		/** Callback to pass down to the FColorDragDrop for it to show the trash */
+		SLATE_EVENT(FSimpleDelegate, ShowTrashCallback)
+		/** Callback to pass down to the FColorDragDrop for it to hide the trash */
+		SLATE_EVENT(FSimpleDelegate, HideTrashCallback)
+		
+		/** Whether to display sRGB color */
+		SLATE_ATTRIBUTE(bool, UseSRGB)
+
+		/** Whether the ability to pick the alpha value is enabled */
+		SLATE_ATTRIBUTE(bool, UseAlpha)
+
+	SLATE_END_ARGS()
+
+	/**
+	 * Construct the widget
+	 *
+	 * @param InArgs   Declaration from which to construct the widget.
+	 */
+	void Construct(const FArguments& InArgs );
+	
+private:
+	virtual FReply OnMouseButtonDown( const FGeometry& MyGeometry, const FPointerEvent& MouseEvent ) OVERRIDE;
+	virtual FReply OnMouseButtonUp( const FGeometry& MyGeometry, const FPointerEvent& MouseEvent ) OVERRIDE;
+	virtual FReply OnDragDetected( const FGeometry& MyGeometry, const FPointerEvent& MouseEvent ) OVERRIDE;
+
+	
+	FLinearColor GetColor() const;
+
+	FSlateColor HandleBorderColor() const;
+	const FSlateBrush* HandleBorderImage() const;
+
+	FText GetRedText() const;
+	FText GetGreenText() const;
+	FText GetBlueText() const;
+	FText GetAlphaText() const;
+	FText GetHueText() const;
+	FText GetSaturationText() const;
+	FText GetValueText() const;
+
+	/**
+	 * Function for formatting text for the tooltip which has limited space.
+	 */
+	FText FormatToolTipText(const FText& ColorIdentifier, float Value) const;
+	
+	bool OnReadIgnoreAlpha() const;
+	bool OnReadShowBackgroundForAlpha() const;
+	EVisibility OnGetAlphaVisibility() const;
+
+	/** A pointer to the color this block uses */
+	TWeakPtr<FLinearColor> ColorPtr;
+	
+	/** A pointer to the theme color blocks bar that is this block's origin */
+	TWeakPtr<SThemeColorBlocksBar> ParentPtr;
+	
+	/** Event called when this block is clicked */
+	FOnLinearColorValueChanged OnSelectColor;
+	
+	/** Callback to pass down to the FColorDragDrop for it to show the trash */
+	FSimpleDelegate ShowTrashCallback;
+	/** Callback to pass down to the FColorDragDrop for it to hide the trash */
+	FSimpleDelegate HideTrashCallback;
+	
+	/** Whether to use display sRGB color */
+	TAttribute<bool> bUseSRGB;
+
+	/** Whether or not the color uses Alpha or not */
+	TAttribute<bool> bUseAlpha;
+
+	float DistanceDragged;
+};
+
+
+/**
+ * SThemeColorBlocksBars are panels for dragging and dropping SColorThemeBlocks onto and off of.
+ */
+class SThemeColorBlocksBar : public SPanel
+{
+public:
+	SLATE_BEGIN_ARGS(SThemeColorBlocksBar)
+		: _ColorTheme()
+		, _OnSelectColor()
+		, _ShowTrashCallback()
+		, _HideTrashCallback()
+		, _EmptyText()
+		, _UseSRGB()
+		, _UseAlpha()
+		{}
+
+		/** A pointer to the color theme that this bar should display */
+		SLATE_ATTRIBUTE(TSharedPtr<FColorTheme>, ColorTheme)
+
+		/** Event called when a color block is clicked */
+		SLATE_EVENT(FOnLinearColorValueChanged, OnSelectColor)
+
+		/** Callback to pass down to the FColorDragDrop for it to show the trash */
+		SLATE_EVENT(FSimpleDelegate, ShowTrashCallback)
+		/** Callback to pass down to the FColorDragDrop for it to hide the trash */
+		SLATE_EVENT(FSimpleDelegate, HideTrashCallback)
+		
+		/** Specify what the bar should display when no colors are present */
+		SLATE_TEXT_ARGUMENT(EmptyText)
+
+		/** Whether to display sRGB color */
+		SLATE_ATTRIBUTE(bool, UseSRGB)
+
+		/** Whether the ability to pick the alpha value is enabled */
+		SLATE_ATTRIBUTE(bool, UseAlpha)
+
+	SLATE_END_ARGS()
+
+	/**
+	 * Most panels do not create widgets as part of their implementation, so
+	 * they do not need to implement a Construct()
+	 */
+	void Construct(const FArguments& InArgs);
+	
+	/**
+	 * Panels arrange their children in a space described by the AllottedGeometry parameter. The results of the arrangement
+	 * should be returned by appending a FArrangedWidget pair for every child widget.
+	 *
+	 * @param AllottedGeometry    The geometry allotted for this widget by its parent.
+	 * @param ArrangedChildren    The array to which to add the WidgetGeometries that represent the arranged children.
+	 */
+	virtual void ArrangeChildren( const FGeometry& AllottedGeometry, FArrangedChildren& ArrangedChildren ) const OVERRIDE;
+
+	/**
+	 * A Panel's desired size in the space required to arrange of its children on the screen while respecting all of
+	 * the children's desired sizes and any layout-related options specified by the user. See StackPanel for an example.
+	 *
+	 * @return The desired size.
+	 */
+	virtual FVector2D ComputeDesiredSize() const OVERRIDE;
+
+	/**
+	 * All widgets must provide a way to access their children in a layout-agnostic way.
+	 * Panels store their children in Slots, which creates a dilemma. Most panels
+	 * can store their children in a TPanelChildren<Slot>, where the Slot class
+	 * provides layout information about the child it stores. In that case
+	 * GetChildren should simply return the TPanelChildren<Slot>.
+	 */
+	virtual FChildren* GetChildren() OVERRIDE;
+	
+	/**
+	 * Called during drag and drop when the drag enters a widget.
+	 *
+	 * @param MyGeometry      The geometry of the widget receiving the event.
+	 * @param DragDropEvent   The drag and drop event.
+	 *
+	 * @return A reply that indicated whether this event was handled.
+	 */
+	virtual void OnDragEnter( const FGeometry& MyGeometry, const FDragDropEvent& DragDropEvent ) OVERRIDE;
+	
+	/**
+	 * Called during drag and drop when the drag leaves a widget.
+	 *
+	 * @param DragDropEvent   The drag and drop event.
+	 *
+	 * @return A reply that indicated whether this event was handled.
+	 */
+	virtual void OnDragLeave( const FDragDropEvent& DragDropEvent ) OVERRIDE;
+
+	/**
+	 * Called during drag and drop when the the mouse is being dragged over a widget.
+	 *
+	 * @param MyGeometry      The geometry of the widget receiving the event.
+	 * @param DragDropEvent   The drag and drop event.
+	 *
+	 * @return A reply that indicated whether this event was handled.
+	 */
+	virtual FReply OnDragOver( const FGeometry& MyGeometry, const FDragDropEvent& DragDropEvent ) OVERRIDE;
+
+	/**
+	 * Called when the user is dropping something onto a widget; terminates drag and drop.
+	 *
+	 * @param MyGeometry      The geometry of the widget receiving the event.
+	 * @param DragDropEvent   The drag and drop event.
+	 *
+	 * @return A reply that indicated whether this event was handled.
+	 */
+	virtual FReply OnDrop( const FGeometry& MyGeometry, const FDragDropEvent& DragDropEvent ) OVERRIDE;
+
+	/**
+	 * Adds a new color block to the Bar
+	 * @param Color				The	color the new block will be
+	 * @param InsertPosition	The location the block will be at
+	 */
+	void AddNewColorBlock(FLinearColor Color, int32 InsertPosition);
+
+	/**
+	 * @param ColorToRemove		The pointer to the color that should be removed
+	 *
+	 * @return					The index of the removed color block, INDEX_NONE if it can't be found
+	 */
+	int32 RemoveColorBlock(TSharedPtr<FLinearColor> ColorToRemove);
+
+	/**
+	 * Refresh Callbacks are used by Color Themes to notify widgets to refresh
+	 * This removes a callback from the Color Theme that this bar has a pointer to
+	 */
+	void RemoveRefreshCallback();
+	/**
+	 * Adds a callback to the Color Theme that this bar has a pointer to
+	 */
+	void AddRefreshCallback();
+
+	/**
+	 * Rebuilds the entire bar, regenerating all the constituent color blocks
+	 */
+	void Refresh();
+
+	void SetPlaceholderGrabOffset(FVector2D GrabOffset);
+	
+private:
+	/**
+	 * Destroys the placeholder block in this widget
+	 */
+	void DestroyPlaceholders();
+
+
+	/** The children blocks of this panel */
+	TSlotlessChildren<SThemeColorBlock> ColorBlocks;
+
+
+	/**
+	 * A placeholder child, which exists as a "preview" to show what would happen when dropping the current
+	 * FColorDragDrop onto this widget. It "replaces" the FColorDragDrop widget with a preview SThemeColorBlock
+	 */
+	TSharedPtr<SThemeColorBlock> NewColorBlockPlaceholder;
+	/** The placeholder's color, also used to determine whether the placeholder is real or not (by IsValid) */
+	TSharedPtr<FLinearColor> NewColorPlaceholder;
+	/** Current x offset of the placeholder block */
+	float PlaceholderBlockOffset;
+	/** The initial grab offset when grabbing the placeholder */
+	FVector2D PlaceholderInitialGrabOffset;
+
+
+	/** A help text widget which appears when there are no children in this panel */
+	TSharedPtr<SWidget> EmptyHintTextBlock;
+
+
+	/**
+	 * The Color Theme that this SThemeColorBlockBar is displaying.
+	 * This is a TAttribute so it can re-get the theme when it changes rather than rely on a delegate to refresh it
+	 */
+	TAttribute< TSharedPtr<FColorTheme> > ColorTheme;
+	
+	/** Event called when a color block is clicked */
+	FOnLinearColorValueChanged OnSelectColor;
+	
+	/** Callback to pass to the Color Theme. Holds a handle to this bar's Refresh method */
+	FSimpleDelegate RefreshCallback;
+
+	/** Callback to pass down to the FColorDragDrop for it to show the trash */
+	FSimpleDelegate ShowTrashCallback;
+	/** Callback to pass down to the FColorDragDrop for it to hide the trash */
+	FSimpleDelegate HideTrashCallback;
+
+	/** Whether to use display sRGB color */
+	TAttribute<bool> bUseSRGB;
+
+	/** Whether or not the color uses Alpha or not */
+	TAttribute<bool> bUseAlpha;
+};
+
+
+
+DECLARE_DELEGATE_OneParam( FOnCurrentThemeChanged, TSharedPtr<FColorTheme> )
+
+
+
+/**
+ * SColorThemeBars include a ThemeColorBlocksBar in addition to a label.
+ * Clicking on one will select it and set the currently used color theme to it
+ */
+class SColorThemeBar : public SCompoundWidget
+{
+public:
+	SLATE_BEGIN_ARGS(SColorThemeBar)
+		: _ColorTheme()
+		, _OnCurrentThemeChanged()
+		, _ShowTrashCallback()
+		, _HideTrashCallback()
+		, _UseSRGB()
+		, _UseAlpha()
+		{}
+
+		/** The color theme that this bar is displaying */
+		SLATE_ATTRIBUTE(TSharedPtr<FColorTheme>, ColorTheme)
+
+		/** Event to be called when the current theme changes */
+		SLATE_EVENT(FOnCurrentThemeChanged, OnCurrentThemeChanged)
+
+		/** Callback to pass down to the FColorDragDrop for it to show the trash */
+		SLATE_EVENT(FSimpleDelegate, ShowTrashCallback)
+		/** Callback to pass down to the FColorDragDrop for it to hide the trash */
+		SLATE_EVENT(FSimpleDelegate, HideTrashCallback)
+		
+		/** Whether to display sRGB color */
+		SLATE_ATTRIBUTE(bool, UseSRGB)
+
+		/** Whether the ability to pick the alpha value is enabled */
+		SLATE_ATTRIBUTE(bool, UseAlpha)
+
+	SLATE_END_ARGS()
+
+	void Construct(const FArguments& InArgs);
+	
+	/**
+	 * The system calls this method to notify the widget that a mouse button was pressed within it. This event is bubbled.
+	 *
+	 * @param MyGeometry The Geometry of the widget receiving the event
+	 * @param MouseEvent Information about the input event
+	 *
+	 * @return Whether the event was handled along with possible requests for the system to take action.
+	 */
+	virtual FReply OnMouseButtonDown( const FGeometry& MyGeometry, const FPointerEvent& MouseEvent );
+	
+private:
+	FString GetThemeName() const;
+
+	/** Text Block which shows the Color Theme's name */
+	TSharedPtr<STextBlock> ThemeNameText;
+
+	/** Color Theme that this bar is displaying */
+	TWeakPtr<FColorTheme> ColorTheme;
+
+	/** Callback to execute when the global current theme has changed */
+	FOnCurrentThemeChanged OnCurrentThemeChanged;
+
+	/** Callback to pass down to the FColorDragDrop for it to show the trash */
+	FSimpleDelegate ShowTrashCallback;
+	/** Callback to pass down to the FColorDragDrop for it to hide the trash */
+	FSimpleDelegate HideTrashCallback;
+	
+	/** Whether to use display sRGB color */
+	TAttribute<bool> bUseSRGB;
+
+	/** Whether or not the color uses Alpha or not */
+	TAttribute<bool> bUseAlpha;
+};
+
+
+
+/**
+ * There should only ever be a single SColorThemesViewer. It is the widget which
+ * manages all color themes and displays in a nice list format.
+ */
+class SColorThemesViewer : public SCompoundWidget
+{
+public:
+	SLATE_BEGIN_ARGS(SColorThemesViewer)
+		: _UseAlpha()
+		{}
+
+		/** Whether the ability to pick the alpha value is enabled */
+		SLATE_ATTRIBUTE(bool, UseAlpha)
+
+	SLATE_END_ARGS()
+	
+	void Construct(const FArguments& InArgs);
+	
+	/** Gets the current color theme */
+	TSharedPtr<FColorTheme> GetCurrentColorTheme() const;
+
+	/** Sets the UseAlpha attribute */
+	void SetUseAlpha( const TAttribute<bool>& InUseAlpha );
+
+	/** Load the color theme settings from the config */
+	static void LoadColorThemesFromIni();
+	/** Save the color theme settings to the config */
+	static void SaveColorThemesToIni();
+	
+	/** Callbacks to execute whenever we change the global current theme */
+	DECLARE_EVENT( SColorThemesViewer, FCurrentThemeChangedEvent );
+	FCurrentThemeChangedEvent& OnCurrentThemeChanged() { return CurrentThemeChangedEvent; }
+
+	/** Whether to use display sRGB color */
+	static bool bSRGBEnabled;
+
+	void MenuToStandardNoReturn();
+
+private:
+	FReply NewColorTheme();
+	FReply DuplicateColorTheme();
+	FReply DeleteColorTheme();
+	FReply ChangeThemeName();
+
+	/** Sets the current color theme to the existing theme */
+	void SetCurrentColorTheme(TSharedPtr<FColorTheme> NewTheme);
+
+	FReply MenuToStandard();
+	FReply MenuToRename();
+	FReply MenuToDelete();
+	void MenuToTrash();
+	
+	bool OnReadUseSRGB() const;
+	bool OnReadUseAlpha() const;
+	
+	/** Refreshes the list, saves colors to ini, and sends the menu to standard */
+	void RefreshThemes();
+
+	TSharedRef<ITableRow> OnGenerateColorThemeBars(TSharedPtr<FColorTheme> InItem, const TSharedRef<STableViewBase>& OwnerTable);
+
+	/** Gets the default color theme, optionally creates it if not present */
+	static TSharedPtr<FColorTheme> GetDefaultColorTheme(bool bCreateNew = false);
+	/** Gets the color theme, creates it if not present */
+	static TSharedPtr<FColorTheme> GetColorTheme(const FString& ThemeName);
+	/** Checks to see if this is a color theme, returns success */
+	static TSharedPtr<FColorTheme> IsColorTheme(const FString& ThemeName);
+	/** Makes the passed theme name unique so it doesn't clash with pre-existing themes */
+	static FString MakeUniqueThemeName(const FString& ThemeName);
+	/** Creates a new theme, ensuring the name is unique */
+	static TSharedPtr<FColorTheme> NewColorTheme(const FString& ThemeName, const TArray< TSharedPtr<FLinearColor> >& ThemeColors = TArray< TSharedPtr<FLinearColor> >());
+	
+	/** A list of all the color themes for the list view to use */
+	TSharedPtr< SListView <TSharedPtr <FColorTheme> > > ColorThemeList;
+
+	/** The menu is a widget at the bottom of the SColorThemesViewer with variable content */
+	TSharedPtr<SBorder> Menu;
+	/** The menu alternates between a standard menu, a rename menu, a confirm delete menu, and a color trash */
+	TSharedPtr<SWidget> MenuStandard;
+	TSharedPtr<SWidget> MenuRename;
+	TSharedPtr<SWidget> MenuConfirmDelete;
+	TSharedPtr<SWidget> MenuTrashColor;
+	
+	/** The text box for renaming themes */
+	TSharedPtr<SEditableTextBox> RenameTextBox;
+
+	/** Callbacks to execute whenever we change the global current theme */
+	FCurrentThemeChangedEvent CurrentThemeChangedEvent;
+
+	/** Whether or not the color uses Alpha or not */
+	TAttribute<bool> bUseAlpha;
+
+	/** A static holder of the color themes for the entire program */
+	static TArray< TSharedPtr<FColorTheme> > ColorThemes;
+	/** A static pointer to the color theme that is currently selected for the entire program */
+	static TWeakPtr<FColorTheme> CurrentlySelectedThemePtr;
+};
+
+
+/**
+ * This operation is a color which can be dragged and dropped between widgets.
+ * Represents a SThemeColorBlock that is dragged around, and can be dropped into a color trash.
+ */
+class FColorDragDrop : public FDragDropOperation
+{
+public:
+	DRAG_DROP_OPERATOR_TYPE(FColorDragDrop, FDragDropOperation)
+
+	/**
+	 * Invoked when the drag and drop operation has ended.
+	 * 
+	 * @param bDropWasHandled   true when the drop was handled by some widget; false otherwise
+	 * @param MouseEvent        The mouse event which caused the on drop to be called.
+	 */
+	virtual void OnDrop( bool bDropWasHandled, const FPointerEvent& MouseEvent ) OVERRIDE;
+
+	/** 
+	 * Called when the mouse was moved during a drag and drop operation
+	 *
+	 * @param DragDropEvent    The event that describes this drag drop operation.
+	 */
+	virtual void OnDragged( const class FDragDropEvent& DragDropEvent ) OVERRIDE;
+	
+	/** Gets the widget that will serve as the decorator unless overriden. If you do not override, you will have no decorator */
+	virtual TSharedPtr<SWidget> GetDefaultDecorator() const OVERRIDE;
+	
+	/**
+	 * Makes a new FColorDragDrop to hold on to
+	 * @param InColor			The color to be dragged and dropped
+	 * @param bSRGB				Whether the color is sRGB
+	 * @param bUseAlpha			Whether the colors alpha is important
+	 * @param TrashShowCallback	Called when this operation is created
+	 * @param TrashHideCallback Called when this operation is dropped
+	 * @param Origin			The SThemeColorBlockBar that this operation is from
+	 * @param OriginPosition	The position in it's origin that it is from
+	 */
+	static TSharedRef<FColorDragDrop> New( FLinearColor InColor, bool bSRGB, bool bUseAlpha,
+		FSimpleDelegate TrashShowCallback = FSimpleDelegate(), FSimpleDelegate TrashHideCallback = FSimpleDelegate(),
+		TSharedPtr<SThemeColorBlocksBar> Origin = TSharedPtr<SThemeColorBlocksBar>(), int32 OriginPosition = 0 );
+
+	/** The color currently held onto by this drag drop operation */
+	FLinearColor Color;
+
+	/** Whether or not the color uses SRGB or not */
+	bool bUseSRGB;
+
+	/** Whether or not the color uses Alpha or not */
+	bool bUseAlpha;
+
+	/**
+	 * The origin is the SThemeColorBlockBar where this operation is from.
+	 * Upon dropping, this operation jumps back to it's origin if possible
+	 */
+	TWeakPtr<SThemeColorBlocksBar> OriginBar;
+	/** The origin position */
+	int32 OriginBarPosition;
+
+	/** Callback to show the trash when this widget is created */
+	FSimpleDelegate ShowTrash;
+	/** Callback to hide the trash when this widget is dropped */
+	FSimpleDelegate HideTrash;
+
+	/** Flag which ensures that OnDrop will not replace this block in it's origin */
+	bool bSetForDeletion;
+
+	/** The size of the drag and drop color block */
+	FVector2D BlockSize;
+};
diff --git a/Engine/Source/Runtime/Slate/Public/Widgets/Colors/SColorWheel.h b/Engine/Source/Runtime/Slate/Public/Widgets/Colors/SColorWheel.h
new file mode 100644
--- /dev/null
+++ b/Engine/Source/Runtime/Slate/Public/Widgets/Colors/SColorWheel.h
@@ -0,0 +1,99 @@
+// Copyright 1998-2014 Epic Games, Inc. All Rights Reserved.
+
+/*=============================================================================
+	SColorWheel.h: Declares the SColorWheel class.
+=============================================================================*/
+
+#pragma once
+
+
+/**
+ * Implements the color wheel widget.
+ */
+class SColorWheel
+	: public SLeafWidget
+{
+public:
+
+	SLATE_BEGIN_ARGS(SColorWheel)
+		: _SelectedColor()
+		, _OnMouseCaptureBegin()
+		, _OnMouseCaptureEnd()
+		, _OnValueChanged()
+	{ }
+	
+		/** The current color selected by the user */
+		SLATE_ATTRIBUTE(FLinearColor, SelectedColor)
+		
+		/** Invoked when the mouse is pressed and a capture begins */
+		SLATE_EVENT(FSimpleDelegate, OnMouseCaptureBegin)
+
+		/** Invoked when the mouse is released and a capture ends */
+		SLATE_EVENT(FSimpleDelegate, OnMouseCaptureEnd)
+
+		/** Invoked when a new value is selected on the color wheel */
+		SLATE_EVENT(FOnLinearColorValueChanged, OnValueChanged)
+
+	SLATE_END_ARGS()
+	
+public:
+
+	/**
+	 * Construct this widget.
+	 *
+	 * @param InArgs - The declaration data for this widget.
+	 */
+	void Construct( const FArguments& InArgs );
+
+public:
+
+	// Begin SWidget overrides
+
+	virtual FVector2D ComputeDesiredSize( ) const;
+
+	virtual FReply OnMouseButtonDoubleClick( const FGeometry& InMyGeometry, const FPointerEvent& InMouseEvent )
+	{
+		return FReply::Handled();
+	}
+
+	virtual FReply OnMouseButtonDown( const FGeometry& MyGeometry, const FPointerEvent& MouseEvent );
+	
+	virtual FReply OnMouseButtonUp( const FGeometry& MyGeometry, const FPointerEvent& MouseEvent );
+
+	virtual FReply OnMouseMove( const FGeometry& MyGeometry, const FPointerEvent& MouseEvent );
+
+	virtual int32 OnPaint( const FGeometry& AllottedGeometry, const FSlateRect& MyClippingRect, FSlateWindowElementList& OutDrawElements, int32 LayerId, const FWidgetStyle& InWidgetStyle, bool bParentEnabled ) const OVERRIDE;
+
+	// End SWidget overrides
+	
+protected:
+
+	/**
+	 * Calculates the position of the color selection indicator.
+	 *
+	 * @return The position relative to the widget.
+	 */
+	FVector2D CalcRelativeSelectedPosition( ) const;
+
+private:
+
+	/** The color wheel image to show */
+	const FSlateBrush* Image;
+	
+	/** The current color selected by the user */
+	TAttribute< FLinearColor > SelectedColor;
+
+	/** The color selector image to show */
+	const FSlateBrush* SelectorImage;
+
+private:
+
+	/** Invoked when the mouse is pressed and a capture begins */
+	FSimpleDelegate OnMouseCaptureBegin;
+
+	/** Invoked when the mouse is let up and a capture ends */
+	FSimpleDelegate OnMouseCaptureEnd;
+
+	/** Invoked when a new value is selected on the color wheel */
+	FOnLinearColorValueChanged OnValueChanged;
+};
diff --git a/Engine/Source/Runtime/Slate/Public/Widgets/Docking/SDockTab.h b/Engine/Source/Runtime/Slate/Public/Widgets/Docking/SDockTab.h
new file mode 100644
--- /dev/null
+++ b/Engine/Source/Runtime/Slate/Public/Widgets/Docking/SDockTab.h
@@ -0,0 +1,354 @@
+// Copyright 1998-2014 Epic Games, Inc. All Rights Reserved.
+
+#pragma once
+
+
+class SDockingTabWell;
+class SDockingTabStack;
+
+/** How will this tab be used. */
+namespace ETabRole
+{
+	enum Type
+	{
+		MajorTab,
+		PanelTab,
+		NomadTab,
+		DocumentTab,
+		NumRoles
+	};
+}
+
+namespace ETabActivationCause
+{
+	/** The cause of a tab activation */
+	enum Type
+	{
+		UserClickedOnTab,
+		SetDirectly
+	};
+}
+
+/**
+ * A tab widget that also holds on to some content that should be shown when this tab is selected.
+ * Intended to be used in conjunction with SDockingTabStack.
+ */
+class SLATE_API SDockTab : public SBorder
+{
+	friend class FTabManager;
+public:
+
+	/** Invoked when a tab is closing */
+	DECLARE_DELEGATE_OneParam(FOnTabClosedCallback, TSharedRef<SDockTab>);
+	
+	/** Invoked when a tab is activated */
+	DECLARE_DELEGATE_TwoParams(FOnTabActivatedCallback, TSharedRef<SDockTab>, ETabActivationCause::Type);
+
+	/** Invoked when this tab should save some information about its content. */
+	DECLARE_DELEGATE(FOnPersistVisualState);
+
+	/** Delegate called before a tab is closed.  Returning false will prevent the tab from closing */
+	DECLARE_DELEGATE_RetVal( bool, FCanCloseTab );
+
+	SLATE_BEGIN_ARGS(SDockTab)
+		: _Content()
+		, _TabWellContentLeft()
+		, _TabWellContentRight()
+		, _ContentPadding( FMargin( 5 ) )
+		, _TabRole(ETabRole::PanelTab)
+		, _Label()
+		, _Icon( FStyleDefaults::GetNoBrush() )
+		, _OnTabClosed()
+		, _OnTabActivated()
+		, _ShouldAutosize(false)
+		, _OnCanCloseTab()
+		, _OnPersistVisualState()
+		, _TabColorScale(0,0,0,0)
+		{}
+
+		SLATE_DEFAULT_SLOT( FArguments, Content )
+		SLATE_NAMED_SLOT( FArguments, TabWellContentLeft )
+		SLATE_NAMED_SLOT( FArguments, TabWellContentRight )
+		SLATE_ATTRIBUTE( FMargin, ContentPadding )
+		SLATE_ARGUMENT( ETabRole::Type, TabRole )
+		SLATE_ATTRIBUTE( FText, Label )
+		SLATE_ATTRIBUTE( const FSlateBrush*, Icon )
+		SLATE_EVENT( FOnTabClosedCallback, OnTabClosed )
+		SLATE_EVENT( FOnTabActivatedCallback, OnTabActivated )
+		SLATE_ARGUMENT( bool, ShouldAutosize )
+		SLATE_EVENT( FCanCloseTab, OnCanCloseTab )
+		SLATE_EVENT( FOnPersistVisualState, OnPersistVisualState )
+		SLATE_ARGUMENT( FLinearColor, TabColorScale )
+	SLATE_END_ARGS()
+
+	/** Construct the widget from the declaration. */
+	void Construct( const FArguments& InArgs );
+
+	virtual void Tick( const FGeometry& AllottedGeometry, const double InCurrentTime, const float InDeltaTime ) OVERRIDE;
+
+	virtual FReply OnMouseButtonDown( const FGeometry& MyGeometry, const FPointerEvent& MouseEvent ) OVERRIDE;
+
+	virtual FReply OnMouseButtonDoubleClick( const FGeometry& MyGeometry, const FPointerEvent& MouseEvent ) OVERRIDE;
+
+	virtual FReply OnDragDetected( const FGeometry& MyGeometry, const FPointerEvent& MouseEvent  ) OVERRIDE;
+	
+	virtual FReply OnMouseButtonUp( const FGeometry& MyGeometry, const FPointerEvent& MouseEvent ) OVERRIDE;
+
+	virtual void OnDragEnter( const FGeometry& MyGeometry, const FDragDropEvent& DragDropEvent ) OVERRIDE;
+
+	virtual void OnDragLeave( const FDragDropEvent& DragDropEvent ) OVERRIDE;
+
+	virtual FReply OnDrop( const FGeometry& MyGeometry, const FDragDropEvent& DragDropEvent ) OVERRIDE;
+
+	virtual FReply OnTouchStarted( const FGeometry& MyGeometry, const FPointerEvent& InTouchEvent ) OVERRIDE;
+
+	virtual FReply OnTouchEnded( const FGeometry& MyGeometry, const FPointerEvent& InTouchEvent ) OVERRIDE;
+
+	void SetContent( TSharedRef<SWidget> InContent );
+	void SetLeftContent( TSharedRef<SWidget> InContent );
+	void SetRightContent( TSharedRef<SWidget> InContent );
+
+	/** @return True if this tab is currently focused */
+	bool IsActive() const;
+
+	/** @return True if this tab appears active; False otherwise */
+	bool IsForeground() const;
+
+	/** Is this an MajorTab? A tool panel tab? */
+	ETabRole::Type GetTabRole() const;
+
+	/** Returns true if the tab role is NomadTab, but is being visualized as a major tab. */
+	bool IsNomadTabWithMajorTabStyle() const;
+
+	/**
+	 * What should the content area look like for this type of tab?
+	 * Documents, Apps, and Tool Panels have different backgrounds.
+	 */
+	const FSlateBrush* GetContentAreaBrush() const;
+
+	/** Depending on the tabs we put into the tab well, we want a different background brush. */
+	const FSlateBrush* GetTabWellBrush() const;
+
+	/** @return the content associated with this tab */
+	TSharedRef<SWidget> GetContent();
+	TSharedRef<SWidget> GetLeftContent();
+	TSharedRef<SWidget> GetRightContent();
+
+	/** Padding around the content when it is presented by the SDockingTabStack */
+	FMargin GetContentPadding() const;
+
+	/** Gets this tab's layout identifier */
+	const FTabId& GetLayoutIdentifier() const;
+
+	/** Sets the tab's tab well parent, or resets it if nothing is passed in */
+	void SetParent(TSharedPtr<SDockingTabWell> Parent = TSharedPtr<SDockingTabWell>());
+
+	/** Gets the tab's tab well parent, or nothing, if it has none */
+	TSharedPtr<SDockingTabWell> GetParent() const;
+
+	/** Gets the dock area that this resides in */
+	TSharedPtr<class SDockingArea> GetDockArea() const;
+
+	/** Get the window in which this tab's tabmanager has placed it */
+	TSharedPtr<SWindow> GetParentWindow() const;
+
+	/** The width that this tab will overlap with side-by-side tabs. */
+	float GetOverlapWidth();
+
+	/** The label on the tab */
+	FText GetTabLabel() const;
+
+	/** The label that appears on the tab. */
+	void SetLabel( const TAttribute<FText>& InTabLabel );
+
+	/** The tooltip text that appears on the tab. */
+	void SetTabToolTipWidget(TSharedPtr<SToolTip> InTabToolTipWidget);
+
+	/** Gets the tab icon */
+	const FSlateBrush* GetTabIcon() const;
+
+	/** Sets the tab icon */
+	void SetTabIcon( const TAttribute<const FSlateBrush*> InTabIcon );
+	
+	/**Is this dock tab being placed via a tab well or via a target*/
+	enum EViaTabwell
+	{
+		DockingViaTabWell,
+		DockingViaTarget
+	};
+
+	/** Checks to see if this tab can dock in this node. Some tabs can only dock via the tab well. */
+	bool CanDockInNode(const TSharedRef<class SDockingNode>& DockNode, EViaTabwell IsDockingViaTabwell ) const;
+
+	/** Should this tab be sized based on its content. */
+	bool ShouldAutosize() const;
+
+	/** @return true if the tab can be closed */
+	bool CanCloseTab() const;
+
+	/** Requests that the tab be closed.  Tabs may prevent closing depending on their state */	
+	void RequestCloseTab();
+
+	/** A chance for the tab's content to save any internal layout info */
+	void PersistVisualState();
+
+	/** 
+	 * Pulls this tab out of it's parent tab stack and destroys it
+	 * Note: This does not check if its safe to remove the tab.  Use RequestCloseTab to do this safely.
+	 */
+	void RemoveTabFromParent();
+
+	/** Protected constructor; Widgets may only be constructed via a FArguments (i.e.: SNew(SDockTab) ) */
+	SDockTab();
+
+	/** 
+	 * Make this tab active in its tabwell 
+	 * @param	InActivationMethod	How this tab was activated.
+	 */
+	void ActivateInParent(ETabActivationCause::Type InActivationCause);
+
+	/** Set the tab manager that is controlling this tab */
+	void SetTabManager( const TSharedPtr<class FTabManager>& InTabManager );
+
+	/**
+	 * Set the custom code to execute for saving visual state in this tab.
+	 * e.g. ContentBrowser saves the visible filters.
+	 */
+	void SetOnPersistVisualState( const FOnPersistVisualState& Handler );
+
+	/** Set the handler to be invoked when the user requests that this tab be closed. */
+	void SetCanCloseTab( const FCanCloseTab& InOnTabClosing );
+
+	/** Set the handler that will be invoked when the tab is closed */
+	void SetOnTabClosed( const FOnTabClosedCallback& InDelegate );
+
+	/** Set the handler that will be invoked when the tab is activated */
+	void SetOnTabActivated( const FOnTabActivatedCallback& InDelegate );
+
+	/** Get the tab manager currently managing this tab. Note that a user move the tab between Tab Managers, so this return value may change. */
+	TSharedRef<FTabManager> GetTabManager() const;
+
+	/** Provide a default tab label in case the spawner did not set one. */
+	void ProvideDefaultLabel( const FText& InDefaultLabel );
+
+	/** Provide a default tab icon in case the spawner did not set one. */
+	void ProvideDefaultIcon( const FSlateBrush* InDefaultIcon );
+
+	/** Play an animation showing this tab as opening */
+	void PlaySpawnAnim();
+
+	/** Flash the tab, used for drawing attention to it */
+	void FlashTab();
+
+	/** Used by the drag/drop operation to signal to this tab what it is dragging over. */
+	void SetDraggedOverDockArea( const TSharedPtr<SDockingArea>& Area );
+
+protected:
+
+	/** Gets the dock tab stack this dockable tab resides within, if any */
+	TSharedPtr<SDockingTabStack> GetParentDockTabStack() const;
+
+	/** @return the style currently applied to the dock tab */
+	const FDockTabStyle& GetCurrentStyle() const;
+
+	/** @return the image brush that best represents this tab's in its current state */
+	const FSlateBrush* GetImageBrush() const;
+
+	/** @return the padding for the tab widget */
+	FMargin GetTabPadding() const;
+
+	/** @return the image brush for the tab's color overlay */
+	const FSlateBrush* GetColorOverlayImageBrush() const;
+
+	/** @return the image brush for the tab's active state overlay */
+	const FSlateBrush* GetActiveTabOverlayImageBrush() const;
+
+	/** @return Returns a color to scale the background of this tab by */
+	FSlateColor GetTabColor() const;
+
+	/** @return the image brush for the tab's flasher overlay */
+	const FSlateBrush* GetFlashOverlayImageBrush() const;
+
+	/** @return Returns a color to flash the background of this tab with */
+	FSlateColor GetFlashColor() const;
+
+	/** Called when the close button is clicked on the tab. */
+	FReply OnCloseButtonClicked();
+
+	/** Specify the TabId that was used to spawn this tab. */
+	void SetLayoutIdentifier( const FTabId& TabId );
+
+protected:
+
+	/** The tab manager that created this tab. */
+	TWeakPtr<class FTabManager> MyTabManager;
+
+	/** The stuff to show when this tab is selected */
+	TSharedRef<SWidget> Content;
+	TSharedRef<SWidget> TabWellContentLeft;
+	TSharedRef<SWidget> TabWellContentRight;
+	
+	/** The tab's layout identifier */
+	FTabId LayoutIdentifier;
+
+	/** Is this an MajorTab? A tool panel tab? */
+	ETabRole::Type TabRole;
+
+	/** The tab's parent tab well. Null if it is a floating tab. */
+	TWeakPtr<SDockingTabWell> ParentPtr;
+
+	/** The label on the tab */
+	TAttribute<FText> TabLabel;
+
+	/** The icon on the tab */
+	TAttribute<const FSlateBrush*> TabIcon;
+	
+	/** Callback to call when this tab is destroyed */
+	FOnTabClosedCallback OnTabClosed;
+
+	/** Callback to call when this tab is activated */
+	FOnTabActivatedCallback OnTabActivated;
+
+	/** Delegate to execute to determine if we can close this tab */
+	FCanCloseTab OnCanCloseTab;
+
+	/**
+	 * Invoked during the Save Visual State pass; gives this tab a chance to save misc info about visual state.
+	 * e.g. Content Browser might save the current filters, current folder, whether some panel is collapsed, etc.
+	 */
+	FOnPersistVisualState OnPersistVisualState;
+
+	/** The styles used to draw the tab in its various states */
+	const FDockTabStyle* MajorTabStyle;
+	const FDockTabStyle* GenericTabStyle;
+
+	TAttribute<FMargin> ContentAreaPadding;
+
+	/** Should this tab be auto-sized based on its content? */
+	bool bShouldAutosize;
+
+	/** Color of this tab */
+	FLinearColor TabColorScale;
+
+	/** @return the scaling of the tab based on the opening/closing animation */
+	FVector2D GetAnimatedScale() const;
+	/** Animation that shows the tab opening up */
+	FCurveSequence SpawnAnimCurve;
+
+	/** Animation that causes the tab to flash */
+	FCurveSequence FlashTabCurve;
+	/** Get the desired color of tab. These change during flashing. */
+	float GetFlashValue() const;
+
+	/** The dock area this tab is currently being dragged over. Used in nomad tabs to change style */
+	TSharedPtr<SDockingArea> DraggedOverDockingArea;
+
+	/** Widget used to show the label on the tab */
+	TSharedPtr<STextBlock> LabelWidget;
+	/** Widget used to show the icon on the tab */
+	TSharedPtr<SImage> IconWidget;
+
+
+	/** The amount of time spent until we bring this tab to the front */
+	float DragTimer;
+};
diff --git a/Engine/Source/Runtime/Slate/Public/Widgets/Docking/SDockTabStack.h b/Engine/Source/Runtime/Slate/Public/Widgets/Docking/SDockTabStack.h
new file mode 100644
--- /dev/null
+++ b/Engine/Source/Runtime/Slate/Public/Widgets/Docking/SDockTabStack.h
@@ -0,0 +1,25 @@
+// Copyright 1998-2014 Epic Games, Inc. All Rights Reserved.
+
+#pragma once
+
+
+/**
+ * A node in the Docking/Tabbing hierarchy.
+ * A DockTabStack shows a row of tabs and the content of one selected tab.
+ * It also supports re-arranging tabs and dragging them out for the stack.
+ */
+class SLATE_API SDockTabStack 
+{
+public:
+	SLATE_BEGIN_ARGS(SDockTabStack)
+		: _IsDocumentArea(false)
+		{}
+
+		SLATE_ARGUMENT( bool, IsDocumentArea )
+	SLATE_END_ARGS()
+
+	void Construct( const FArguments& InArgs );
+
+	void AddTab( const TSharedRef<SDockableTab>& InTab, int32 AtLocation = INDEX_NONE );
+
+};
diff --git a/Engine/Source/Runtime/Slate/Public/Widgets/Docking/SDockableTab.h b/Engine/Source/Runtime/Slate/Public/Widgets/Docking/SDockableTab.h
new file mode 100644
--- /dev/null
+++ b/Engine/Source/Runtime/Slate/Public/Widgets/Docking/SDockableTab.h
@@ -0,0 +1,122 @@
+// Copyright 1998-2014 Epic Games, Inc. All Rights Reserved.
+
+#pragma once
+
+
+class SDockTabStack;
+
+DECLARE_DELEGATE_TwoParams(
+	FOnTabStackMenuOpening,
+	/** Builder for the menu that will be opened */
+	FMenuBuilder& ,
+	/** Tab stack that this menu is being opened in */
+	const TSharedRef<SDockTabStack>& );
+
+
+/**
+ * A tab widget that also holds on to some content that should be shown when this tab is selected.
+ * Intended to be used in conjunction with SDockTabStack.
+ */
+class SLATE_API SDockableTab
+{
+public:
+	DECLARE_DELEGATE_OneParam(FOnTabClosedCallback, TSharedRef<SDockableTab>);
+
+	SLATE_BEGIN_ARGS(SDockableTab)
+	{}
+
+	SLATE_END_ARGS()
+
+	/** Construct the widget from the declaration. */
+	void Construct( const FArguments& InArgs );
+
+	void SetContent( TSharedRef<SWidget> InContent );
+
+	/** @return True if this tab is currently focused */
+	bool IsActive() const;
+
+	/** @return True if this tab appears active; False otherwise */
+	bool IsForeground() const;
+
+
+	/** @return the content associated with this tab */
+	TSharedRef<SWidget> GetContent();
+
+	/** Gets the dock tab stack this dockable tab resides within, if any */
+	TSharedPtr<SDockTabStack> GetParentDockTabStack() const;
+
+	/** Brings this tab to the front of its parent's tab well, if applicable. */
+	void BringToFrontInParent();
+
+	/** Should this tab be sized based on its content. */
+	bool ShouldAutosize() const;
+
+	/** Set delegate to invoke when tab stack menu is opening */
+	void SetOnTabStackMenuOpening( const FOnTabStackMenuOpening& InHandler );
+
+	/**
+	 * Requests that the tab be closed.  Tabs may prevent closing depending on their state
+	 */	
+	void RequestCloseTab();
+
+	/** 
+	 * Pulls this tab out of it's parent tab stack and destroys it
+	 * Note: This does not check if its safe to remove the tab.  Use RequestCloseTab to do this safely.
+	 */
+	void RemoveTabFromParent();
+
+	/** 
+	 * Make this tab active in its tabwell 
+	 * @param	InActivationMethod	How this tab was activated.
+	 */
+	void ActivateInParent(ETabActivationCause::Type InActivationCause);
+
+protected:
+
+	/** The tab's layout identifier */
+	FName LayoutIdentifier;
+
+	/** Is this an MajorTab? A tool panel tab? */
+	ETabRole::Type TabRole;
+
+	/** The width that this tab will overlap with side-by-side tabs. */
+	float OverlapWidth;
+
+	/** The label on the tab */
+	TAttribute<FString> TabLabel;
+	
+	/** Callback to call when this tab is destroyed */
+	FOnTabClosedCallback OnTabClosed;
+
+	/** Delegate to execute to determine if we can close this tab */
+	SDockTab::FCanCloseTab OnCanCloseTab;
+
+	/** The brushes used to draw the tab in its various states */
+	const FSlateBrush* NoBrush;
+	const FSlateBrush* NormalBrush;
+	const FSlateBrush* ActiveBrush;
+	const FSlateBrush* ColorOverlayBrush;
+	const FSlateBrush* ForegroundBrush;
+	const FSlateBrush* HoveredBrush;
+
+	/**
+	 * The brush that the SDockTabStack should use to draw the content associated with this tab
+	 * Documents, Apps, and Tool Panels have different backgrounds
+	 */
+	const FSlateBrush* ContentAreaBrush;
+
+	TAttribute<FMargin> ContentAreaPadding;
+
+	const FSlateBrush* TabWellBrush;
+
+	FMargin TabWellPadding;
+
+	/** Called when the the tab stack's context menu is open; gives the currently focused tab a change to add custom  */
+	FOnTabStackMenuOpening OnTabStackMenuOpeningHandler;
+
+	/** Should this tab be auto-sized based on its content? */
+	bool bShouldAutosize;
+
+	/** Color of this tab */
+	FLinearColor TabColorScale;
+};
diff --git a/Engine/Source/Runtime/Slate/Public/Widgets/Images/SImage.h b/Engine/Source/Runtime/Slate/Public/Widgets/Images/SImage.h
new file mode 100644
--- /dev/null
+++ b/Engine/Source/Runtime/Slate/Public/Widgets/Images/SImage.h
@@ -0,0 +1,59 @@
+// Copyright 1998-2014 Epic Games, Inc. All Rights Reserved.
+
+#pragma once
+
+
+/** SImage displays a single image at a desired Width and Height. */
+class SLATE_API SImage : public SLeafWidget
+{
+public:
+	SLATE_BEGIN_ARGS( SImage )
+		: _Image( FCoreStyle::Get().GetDefaultBrush() )
+		, _ColorAndOpacity( FLinearColor::White )
+		, _OnMouseButtonDown()
+		{}
+
+		/** Image resource */
+		SLATE_ATTRIBUTE( const FSlateBrush*, Image )
+
+		/** Color and opacity */
+		SLATE_ATTRIBUTE( FSlateColor, ColorAndOpacity )
+
+		/** Invoked when the mouse is pressed in the widget. */
+		SLATE_EVENT( FPointerEventHandler, OnMouseButtonDown )
+
+	SLATE_END_ARGS()
+
+	/**
+	 * Construct this widget
+	 *
+	 * @param	InArgs	The declaration data for this widget
+	 */
+	void Construct( const FArguments& InArgs );
+
+	// SWidget interface
+	virtual int32 OnPaint( const FGeometry& AllottedGeometry, const FSlateRect& MyClippingRect, FSlateWindowElementList& OutDrawElements, int32 LayerId, const FWidgetStyle& InWidgetStyle, bool bParentEnabled ) const OVERRIDE;
+	virtual FReply OnMouseButtonDown( const FGeometry& MyGeometry, const FPointerEvent& MouseEvent ) OVERRIDE;
+	virtual FVector2D ComputeDesiredSize() const OVERRIDE;
+	// End of SWidget interface
+
+public:
+
+	/** The the color and opacity of this image */
+	void SetColorAndOpacity( const TAttribute<FSlateColor>& InColorAndOpacity );
+	/** The the color and opacity of this image */
+	void SetColorAndOpacity( FLinearColor InColorAndOpacity );
+
+
+protected:
+
+	/** The FName of the image resource to show */
+	TAttribute< const FSlateBrush* > Image;
+
+	/** Color and opacity scale for this image */
+	TAttribute<FSlateColor> ColorAndOpacity;
+
+	/** Invoked when the mouse is pressed in the image */
+	FPointerEventHandler OnMouseButtonDownHandler;
+
+};
diff --git a/Engine/Source/Runtime/Slate/Public/Widgets/Images/SSpinningImage.h b/Engine/Source/Runtime/Slate/Public/Widgets/Images/SSpinningImage.h
new file mode 100644
--- /dev/null
+++ b/Engine/Source/Runtime/Slate/Public/Widgets/Images/SSpinningImage.h
@@ -0,0 +1,41 @@
+// Copyright 1998-2014 Epic Games, Inc. All Rights Reserved.
+
+#pragma once
+
+
+/** A widget that displays a spinning image.*/
+class SLATE_API SSpinningImage : public SImage
+{
+public:
+	SLATE_BEGIN_ARGS(SSpinningImage)
+		: _Image( FCoreStyle::Get().GetDefaultBrush() )
+		, _ColorAndOpacity( FLinearColor::White )
+		, _OnMouseButtonDown()
+		, _Period( 1.0f )
+		{}
+
+		/** Image resource */
+		SLATE_ATTRIBUTE( const FSlateBrush*, Image )
+
+		/** Color and opacity */
+		SLATE_ATTRIBUTE( FSlateColor, ColorAndOpacity )
+
+		/** Invoked when the mouse is pressed in the widget. */
+		SLATE_EVENT( FPointerEventHandler, OnMouseButtonDown )
+
+		/** The amount of time in seconds for a full rotation */
+		SLATE_ARGUMENT( float, Period )
+
+	SLATE_END_ARGS()
+
+	/* Construct this widget */
+	void Construct( const FArguments& InArgs );
+
+	/* Override SImage's OnPaint to draw rotated */
+	virtual int32 OnPaint( const FGeometry& AllottedGeometry, const FSlateRect& MyClippingRect, FSlateWindowElementList& OutDrawElements, int32 LayerId, const FWidgetStyle& InWidgetStyle, bool bParentEnabled ) const OVERRIDE;
+
+private:
+	/** The sequence to drive the spinning animation */
+	FCurveSequence Sequence;
+	FCurveHandle Curve;
+};
\ No newline at end of file
diff --git a/Engine/Source/Runtime/Slate/Public/Widgets/Input/SButton.h b/Engine/Source/Runtime/Slate/Public/Widgets/Input/SButton.h
new file mode 100644
--- /dev/null
+++ b/Engine/Source/Runtime/Slate/Public/Widgets/Input/SButton.h
@@ -0,0 +1,163 @@
+// Copyright 1998-2014 Epic Games, Inc. All Rights Reserved.
+
+#pragma once
+
+
+/**
+ * Slate's Buttons are clickable Widgets that can contain arbitrary widgets as its Content().
+ */
+class SLATE_API SButton : public SBorder
+{
+public:
+
+	SLATE_BEGIN_ARGS( SButton )
+		: _Content()
+		, _ButtonStyle( &FCoreStyle::Get().GetWidgetStyle< FButtonStyle >( "Button" ) )
+		, _TextStyle( &FCoreStyle::Get().GetWidgetStyle< FTextBlockStyle >("NormalText") )
+		, _HAlign( HAlign_Fill )
+		, _VAlign( VAlign_Fill )
+		, _ContentPadding(FMargin(4.0, 2.0))
+		, _Text()
+		, _ClickMethod( EButtonClickMethod::DownAndUp )
+		, _DesiredSizeScale( FVector2D(1,1) )
+		, _ContentScale( FVector2D(1,1) )
+		, _ButtonColorAndOpacity(FLinearColor::White)
+		, _ForegroundColor( FCoreStyle::Get().GetSlateColor( "InvertedForeground" ) )
+		, _IsFocusable( true )
+		{}
+
+		/** Slot for this button's content (optional) */
+		SLATE_DEFAULT_SLOT( FArguments, Content )
+
+		/** The visual style of the button */
+		SLATE_STYLE_ARGUMENT( FButtonStyle, ButtonStyle )
+
+		/** The text style of the button */
+		SLATE_STYLE_ARGUMENT( FTextBlockStyle, TextStyle )
+
+		/** Horizontal alignment */
+		SLATE_ARGUMENT( EHorizontalAlignment, HAlign )
+
+		/** Vertical alignment */
+		SLATE_ARGUMENT( EVerticalAlignment, VAlign )
+	
+		/** Spacing between button's border and the content. */
+		SLATE_ATTRIBUTE( FMargin, ContentPadding )
+	
+		/** The text to display in this button, if no custom content is specified */
+		SLATE_TEXT_ATTRIBUTE( Text )
+	
+		/** Called when the button is clicked */
+		SLATE_EVENT( FOnClicked, OnClicked )
+
+		/** Sets the rules to use for determining whether the button was clicked.  This is an advanced setting and generally should be left as the default. */
+		SLATE_ARGUMENT( EButtonClickMethod::Type, ClickMethod )
+
+		SLATE_ATTRIBUTE( FVector2D, DesiredSizeScale )
+
+		SLATE_ATTRIBUTE( FVector2D, ContentScale )
+
+		SLATE_ATTRIBUTE( FSlateColor, ButtonColorAndOpacity )
+
+		SLATE_ATTRIBUTE( FSlateColor, ForegroundColor )
+
+		/** Sometimes a button should only be mouse-clickable and never keyboard focusable. */
+		SLATE_ARGUMENT( bool, IsFocusable )
+
+		/** The sound to play when the button is pressed */
+		SLATE_ARGUMENT( TOptional<FSlateSound>, PressedSoundOverride )
+
+		/** The sound to play when the button is hovered */
+		SLATE_ARGUMENT( TOptional<FSlateSound>, HoveredSoundOverride )
+
+	SLATE_END_ARGS()
+
+
+	/** @return An image that represents this button's border*/
+	const FSlateBrush* GetBorder() const;
+	
+	/**
+	 * See SWidget::SupportsKeyboardFocus().
+	 *
+	 * @return  True if this widget can take keyboard focus
+	 */
+	virtual bool SupportsKeyboardFocus() const OVERRIDE;
+
+	virtual FReply OnKeyDown( const FGeometry& MyGeometry, const FKeyboardEvent& InKeyboardEvent ) OVERRIDE;
+
+	virtual FReply OnMouseButtonDown( const FGeometry& MyGeometry, const FPointerEvent& MouseEvent ) OVERRIDE;
+
+	virtual FReply OnMouseButtonDoubleClick( const FGeometry& InMyGeometry, const FPointerEvent& InMouseEvent ) OVERRIDE;
+	
+	virtual FReply OnMouseButtonUp( const FGeometry& MyGeometry, const FPointerEvent& MouseEvent ) OVERRIDE;
+	
+	virtual FReply OnMouseMove( const FGeometry& MyGeometry, const FPointerEvent& MouseEvent ) OVERRIDE;
+	
+	virtual void OnMouseEnter( const FGeometry& MyGeometry, const FPointerEvent& MouseEvent ) OVERRIDE;
+
+	virtual void OnMouseLeave( const FPointerEvent& MouseEvent ) OVERRIDE;
+
+	/**
+	 * Returns true if this button is currently pressed
+	 *
+	 * @return	True if pressed, otherwise false
+	 */
+	virtual bool IsPressed() const
+	{
+		return bIsPressed;
+	}
+
+	/**
+	 * Construct this widget
+	 *
+	 * @param	InArgs	The declaration data for this widget
+	 */
+	void Construct( const FArguments& InArgs );
+
+protected:
+	/** @return combines the user-specified margin and the button's internal margin. */
+	FMargin GetCombinedPadding() const;
+
+	/** Padding specified by the user; it will be combind with the button's internal padding. */
+	TAttribute<FMargin> ContentPadding;
+
+	/** Padding that accounts for the button border */
+	FMargin BorderPadding;
+
+	/** Padding that accounts for the button border when pressed */
+	FMargin PressedBorderPadding;
+
+	/** True if this button is currently in a pressed state */
+	bool bIsPressed;
+
+	/** The delegate to execute when the button is clicked */
+	FOnClicked OnClicked;
+
+	/** Brush resource that represents a button */
+	const FSlateBrush* NormalImage;
+	/** Brush resource that represents a button when it is hovered */
+	const FSlateBrush* HoverImage;
+	/** Brush resource that represents a button when it is pressed */
+	const FSlateBrush* PressedImage;
+
+	/** Sets whether a click should be triggered on mouse down, mouse up, or that both a mouse down and up are required. */
+	EButtonClickMethod::Type ClickMethod;
+
+	/** Can this button be focused? */
+	bool bIsFocusable;
+
+	/** Play the pressed sound */
+	void PlayPressedSound() const;
+
+	/** Play the hovered sound */
+	void PlayHoverSound() const;
+
+	/** The Sound to play when the button is hovered  */
+	FSlateSound HoveredSound;
+
+	/** The Sound to play when the button is pressed */
+	FSlateSound PressedSound;
+
+private:
+	virtual FVector2D ComputeDesiredSize() const OVERRIDE;
+};
diff --git a/Engine/Source/Runtime/Slate/Public/Widgets/Input/SComboBox.h b/Engine/Source/Runtime/Slate/Public/Widgets/Input/SComboBox.h
new file mode 100644
--- /dev/null
+++ b/Engine/Source/Runtime/Slate/Public/Widgets/Input/SComboBox.h
@@ -0,0 +1,364 @@
+// Copyright 1998-2014 Epic Games, Inc. All Rights Reserved.
+
+#pragma once
+
+
+DECLARE_DELEGATE( FOnComboBoxOpening )
+
+
+template<typename OptionType>
+class SComboRow : public STableRow< OptionType >
+{
+public:
+
+	SLATE_BEGIN_ARGS( SComboRow )
+		: _RowContent()
+		{}
+
+		SLATE_ARGUMENT( TSharedPtr<SWidget>, RowContent )
+	SLATE_END_ARGS()
+
+public:
+
+	/**
+	 * Constructs this widget.
+	 */
+	void Construct( const FArguments& InArgs, const TSharedRef<STableViewBase>& InOwnerTable )
+	{
+		STableRow< OptionType >::Construct(
+			typename STableRow<OptionType>::FArguments()
+			.Content()
+			[
+				InArgs._RowContent->AsShared()
+			]
+			, InOwnerTable
+		);
+	}
+
+	// handle case where user clicks on an existing selected item
+	virtual FReply OnMouseButtonDown( const FGeometry& MyGeometry, const FPointerEvent& MouseEvent )
+	{
+		if (MouseEvent.GetEffectingButton() == EKeys::LeftMouseButton )
+		{
+			TSharedPtr< ITypedTableView<OptionType> > OwnerWidget = this->OwnerTablePtr.Pin();
+
+			const OptionType* MyItem = OwnerWidget->Private_ItemFromWidget( this );
+			const bool bIsSelected = OwnerWidget->Private_IsItemSelected( *MyItem );
+				
+			if (bIsSelected)
+			{
+				// Reselect content to ensure selection is taken
+				OwnerWidget->Private_SignalSelectionChanged(ESelectInfo::Direct);
+				return FReply::Handled();
+			}
+		}
+		return STableRow<OptionType>::OnMouseButtonDown(MyGeometry, MouseEvent);
+	}
+};
+
+
+/**
+ * A combo box that shows arbitrary content.
+ */
+template< typename OptionType >
+class SComboBox : public SComboButton
+{
+public:
+	/** Type of list used for showing menu options. */
+	typedef SListView< OptionType > SComboListType;
+	/** Delegate type used to generate widgets that represent Options */
+	typedef typename TSlateDelegates< OptionType >::FOnGenerateWidget FOnGenerateWidget;
+	typedef typename TSlateDelegates< OptionType >::FOnSelectionChanged FOnSelectionChanged;
+
+	SLATE_BEGIN_ARGS( SComboBox )
+		: _Content()
+		, _ComboBoxStyle( &FCoreStyle::Get().GetWidgetStyle< FComboBoxStyle >( "ComboBox" ) )
+		, _ButtonStyle(nullptr)
+		, _ContentPadding(FMargin(4.0, 2.0))
+		, _ForegroundColor(FCoreStyle::Get().GetSlateColor("InvertedForeground"))
+		, _OptionsSource()
+		, _OnSelectionChanged()
+		, _OnGenerateWidget()
+		, _InitiallySelectedItem( NULL )
+		, _Method( SMenuAnchor::CreateNewWindow )
+		, _MaxListHeight(450.0f)
+		, _HasDownArrow( true )
+		{}
+		
+		/** Slot for this button's content (optional) */
+		SLATE_DEFAULT_SLOT( FArguments, Content )
+
+		SLATE_STYLE_ARGUMENT( FComboBoxStyle, ComboBoxStyle )
+
+		/** The visual style of the button part of the combo box (overrides ComboBoxStyle) */
+		SLATE_STYLE_ARGUMENT( FButtonStyle, ButtonStyle )
+
+		SLATE_ATTRIBUTE( FMargin, ContentPadding )
+		SLATE_ATTRIBUTE( FSlateColor, ForegroundColor )
+
+		SLATE_ARGUMENT( const TArray< OptionType >*, OptionsSource )
+		SLATE_EVENT( FOnSelectionChanged, OnSelectionChanged )
+		SLATE_EVENT( FOnGenerateWidget, OnGenerateWidget )
+
+		/** Called when combo box is opened, before list is actually created */
+		SLATE_EVENT( FOnComboBoxOpening, OnComboBoxOpening )
+
+		/** The option that should be selected when the combo box is first created */
+		SLATE_ARGUMENT( OptionType, InitiallySelectedItem )
+
+		SLATE_ARGUMENT( SMenuAnchor::EMethod, Method )
+
+		/** The max height of the combo box menu */
+		SLATE_ARGUMENT(float, MaxListHeight)
+
+		/** The sound to play when the button is pressed (overrides ComboBoxStyle) */
+		SLATE_ARGUMENT( TOptional<FSlateSound>, PressedSoundOverride )
+
+		/** The sound to play when the selection changes (overrides ComboBoxStyle) */
+		SLATE_ARGUMENT( TOptional<FSlateSound>, SelectionChangeSoundOverride )
+
+		/**
+		 * When false, the down arrow is not generated and it is up to the API consumer
+		 * to make their own visual hint that this is a drop down.
+		 */
+		SLATE_ARGUMENT( bool, HasDownArrow )
+				
+	SLATE_END_ARGS()
+
+	/**
+	 * Construct the widget from a declaration
+	 *
+	 * @param InArgs   Declaration from which to construct the combo box
+	 */
+	void Construct( const FArguments& InArgs )
+	{
+		check(InArgs._ComboBoxStyle);
+
+		// Work out which values we should use based on whether we were given an override, or should use the style's version
+		const FComboButtonStyle& OurComboButtonStyle = InArgs._ComboBoxStyle->ComboButtonStyle;
+		const FButtonStyle* const OurButtonStyle = InArgs._ButtonStyle ? InArgs._ButtonStyle : &OurComboButtonStyle.ButtonStyle;
+		PressedSound = InArgs._PressedSoundOverride.Get(InArgs._ComboBoxStyle->PressedSlateSound);
+		SelectionChangeSound = InArgs._SelectionChangeSoundOverride.Get(InArgs._ComboBoxStyle->SelectionChangeSlateSound);
+
+		this->OnComboBoxOpening = InArgs._OnComboBoxOpening;
+		this->OnSelectionChanged = InArgs._OnSelectionChanged;
+		this->OnGenerateWidget = InArgs._OnGenerateWidget;
+
+		OptionsSource = InArgs._OptionsSource;
+
+		TSharedRef<SWidget> MenuContent = 
+			SNew(SVerticalBox)
+
+			+ SVerticalBox::Slot()
+			.AutoHeight()
+			.MaxHeight(InArgs._MaxListHeight)
+			[
+				SAssignNew( this->ComboListView, SComboListType )
+				.ListItemsSource( InArgs._OptionsSource )
+				.OnGenerateRow( this, &SComboBox< OptionType >::GenerateMenuItemRow )
+				.OnSelectionChanged( this, &SComboBox< OptionType >::OnSelectionChanged_Internal )
+				.SelectionMode( ESelectionMode::Single )
+			];
+
+		// Set up content
+		TSharedPtr<SWidget> ButtonContent = InArgs._Content.Widget;
+		if (InArgs._Content.Widget == SNullWidget::NullWidget)
+		{
+			 SAssignNew(ButtonContent, STextBlock)
+			.Text(NSLOCTEXT("SComboBox", "ContentWarning", "No Content Provided").ToString())
+			.ColorAndOpacity( FLinearColor::Red);
+		}
+
+
+		SComboButton::Construct( SComboButton::FArguments()
+			.ComboButtonStyle(&OurComboButtonStyle)
+			.ButtonStyle(OurButtonStyle)
+			.Method( InArgs._Method )
+			.ButtonContent()
+			[
+				ButtonContent.ToSharedRef()
+			]
+			.MenuContent()
+			[
+				MenuContent
+			]
+			.HasDownArrow( InArgs._HasDownArrow )
+			.ContentPadding( InArgs._ContentPadding )
+			.ForegroundColor( InArgs._ForegroundColor )
+		);
+		SetMenuContentWidgetToFocus(ComboListView);
+
+		// Need to establish the selected item at point of construction so its available for querying
+		// NB: If you need a selection to fire use SetItemSelection rather than setting an IntiallySelectedItem
+		this->SelectedItem = InArgs._InitiallySelectedItem;
+		if( TListTypeTraits<OptionType>::IsPtrValid( InArgs._InitiallySelectedItem ) )
+		{
+			ComboListView->Private_SetItemSelection( SelectedItem, true);
+		}
+	}
+
+	void ClearSelection( )
+	{
+		ComboListView->ClearSelection();
+	}
+
+	void SetSelectedItem( OptionType InSelectedItem )
+	{
+		ComboListView->SetSelection( InSelectedItem );
+	}
+
+	/** @return the item currently selected by the combo box. */
+	OptionType GetSelectedItem()
+	{
+		return SelectedItem;
+	}
+
+	/** 
+	 * Requests a list refresh after updating options 
+	 * Call SetSelectedItem to update the selected item if required
+	 * @see SetSelectedItem
+	 */
+	void RefreshOptions()
+	{
+		ComboListView->RequestListRefresh();
+	}
+
+protected:
+	/** Handle key presses that SListView ignores */
+	virtual FReply OnHandleKeyPressed(FKey KeyPressed)
+	{
+		if (KeyPressed == EKeys::Enter)
+		{
+			TArray<OptionType> SelectedItems = ComboListView->GetSelectedItems();
+			if (SelectedItems.Num() > 0)
+			{
+				ComboListView->SetSelection(SelectedItems[0]);
+			}
+			return FReply::Handled();
+		}
+		else if (KeyPressed == EKeys::Escape)
+		{
+			this->SetIsOpen(false);
+			return FReply::Handled();
+		}
+		return FReply::Unhandled();
+	}
+
+	FReply OnKeyDown( const FGeometry& MyGeometry, const FKeyboardEvent& InKeyboardEvent ) OVERRIDE
+	{
+		const FKey Key = InKeyboardEvent.GetKey();
+
+		if(Key == EKeys::Up)
+		{
+			const int32 SelectionIndex = OptionsSource->Find( GetSelectedItem() );
+			if ( SelectionIndex >= 1 )
+			{
+				// Select an item on the prev row
+				SetSelectedItem((*OptionsSource)[SelectionIndex - 1]);
+			}
+
+			return FReply::Handled();
+		}
+		else if(Key == EKeys::Down)
+		{
+			const int32 SelectionIndex = OptionsSource->Find( GetSelectedItem() );
+			if ( SelectionIndex < OptionsSource->Num() - 1 )
+			{
+				// Select an item on the next row
+				SetSelectedItem((*OptionsSource)[SelectionIndex + 1]);
+			}
+
+			return FReply::Handled();
+		}
+
+		return SComboButton::OnKeyDown( MyGeometry, InKeyboardEvent );
+	}
+
+	virtual bool SupportsKeyboardFocus() const OVERRIDE
+	{
+		return true;
+	}
+
+private:
+
+	/** Generate a row for the InItem in the combo box's list (passed in as OwnerTable). Do this by calling the user-specified OnGenerateWidget */
+	TSharedRef<ITableRow> GenerateMenuItemRow( OptionType InItem, const TSharedRef<STableViewBase>& OwnerTable)
+	{
+		return
+		SNew( SComboRow<OptionType>, OwnerTable)
+		.RowContent(OnGenerateWidget.Execute(InItem) );
+	}
+
+	/** Invoked when the selection in the list changes */
+	void OnSelectionChanged_Internal( OptionType ProposedSelection, ESelectInfo::Type SelectInfo )
+	{
+		// Ensure that the proposed selection is different
+		if(SelectInfo != ESelectInfo::OnNavigation)
+		{
+			// Ensure that the proposed selection is different from selected
+			if ( ProposedSelection != SelectedItem )
+			{
+				PlaySelectionChangeSound();
+				SelectedItem = ProposedSelection;
+				OnSelectionChanged.ExecuteIfBound( ProposedSelection, SelectInfo );
+			}
+			// close combo even if user reselected item
+			this->SetIsOpen( false );
+		}
+	}
+
+	/** Handle clicking on the content menu */
+	virtual FReply OnButtonClicked() OVERRIDE
+	{
+		// if user clicked to close the combo menu
+		if (this->IsOpen())
+		{
+			// Re-select first selected item, just in case it was selected by navigation previously
+			TArray<OptionType> SelectedItems = ComboListView->GetSelectedItems();
+			if (SelectedItems.Num() > 0)
+			{
+				OnSelectionChanged_Internal(SelectedItems[0], ESelectInfo::Direct);
+			}
+		}
+		else
+		{
+			PlayPressedSound();
+			OnComboBoxOpening.ExecuteIfBound();
+		}
+
+		return SComboButton::OnButtonClicked();
+	}
+
+	/** Play the pressed sound */
+	void PlayPressedSound() const
+	{
+		FSlateApplication::Get().PlaySound( PressedSound );
+	}
+
+	/** Play the selection changed sound */
+	void PlaySelectionChangeSound() const
+	{
+		FSlateApplication::Get().PlaySound( SelectionChangeSound );
+	}
+
+	/** The Sound to play when the button is pressed */
+	FSlateSound PressedSound;
+
+	/** The Sound to play when the selection is changed */
+	FSlateSound SelectionChangeSound;
+
+private:
+	/** Delegate that is invoked when the selected item in the combo box changes */
+	FOnSelectionChanged OnSelectionChanged;
+	/** The item currently selected in the combo box */
+	OptionType SelectedItem;
+	/** The ListView that we pop up; visualized the available options. */
+	TSharedPtr< SComboListType > ComboListView;
+	/** Delegate to invoke before the combo box is opening. */
+	FOnComboBoxOpening OnComboBoxOpening;
+	/** Delegate to invoke when we need to visualize an option as a widget. */
+	FOnGenerateWidget OnGenerateWidget;
+
+	const TArray< OptionType >* OptionsSource;
+};
+
diff --git a/Engine/Source/Runtime/Slate/Public/Widgets/Input/SComboButton.h b/Engine/Source/Runtime/Slate/Public/Widgets/Input/SComboButton.h
new file mode 100644
--- /dev/null
+++ b/Engine/Source/Runtime/Slate/Public/Widgets/Input/SComboButton.h
@@ -0,0 +1,109 @@
+// Copyright 1998-2014 Epic Games, Inc. All Rights Reserved.
+
+#pragma once
+
+DECLARE_DELEGATE( FOnComboBoxOpened )
+
+/**
+ * A button that, when clicked, brings up a popup.
+ */
+class SLATE_API SComboButton : public SMenuAnchor
+{
+public:
+
+	SLATE_BEGIN_ARGS( SComboButton )
+		: _ComboButtonStyle( &FCoreStyle::Get().GetWidgetStyle< FComboButtonStyle >( "ComboButton" ) )
+		, _ButtonStyle(nullptr)
+		, _ButtonContent()
+		, _MenuContent()
+		, _HasDownArrow(true)
+		, _ForegroundColor(FCoreStyle::Get().GetSlateColor("InvertedForeground"))
+		, _ButtonColorAndOpacity(FLinearColor::White)
+		, _ContentScale(FVector2D(1,1))
+		, _ContentPadding(FMargin(5))
+		, _MenuPlacement(MenuPlacement_ComboBox)
+		, _HAlign(HAlign_Fill)
+		, _VAlign(VAlign_Fill)
+		, _Method( SMenuAnchor::CreateNewWindow )
+		{}
+
+		SLATE_STYLE_ARGUMENT( FComboButtonStyle, ComboButtonStyle )
+
+		/** The visual style of the button (overrides ComboButtonStyle) */
+		SLATE_STYLE_ARGUMENT( FButtonStyle, ButtonStyle )
+		
+		SLATE_NAMED_SLOT( FArguments, ButtonContent )
+		SLATE_NAMED_SLOT( FArguments, MenuContent )
+		SLATE_EVENT( FOnGetContent, OnGetMenuContent )
+		SLATE_EVENT( FOnComboBoxOpened, OnComboBoxOpened )
+		SLATE_ARGUMENT( bool, IsFocusable )
+		SLATE_ARGUMENT( bool, HasDownArrow )
+
+		SLATE_ATTRIBUTE( FSlateColor, ForegroundColor )
+		SLATE_ATTRIBUTE( FSlateColor, ButtonColorAndOpacity )
+		SLATE_ATTRIBUTE( FVector2D, ContentScale )
+		SLATE_ATTRIBUTE( FMargin, ContentPadding )
+		SLATE_ATTRIBUTE( EMenuPlacement, MenuPlacement )
+		SLATE_ARGUMENT( EHorizontalAlignment, HAlign )
+		SLATE_ARGUMENT( EVerticalAlignment, VAlign )
+
+		/** Spawn a new window or reuse current window for this combo*/
+		SLATE_ARGUMENT( SMenuAnchor::EMethod, Method )
+		/** The max height of the combo menu list */
+		SLATE_ATTRIBUTE( FOptionalSize, MenuHeight )
+		/** The max width of the combo menu list */
+		SLATE_ATTRIBUTE( FOptionalSize, MenuWidth )
+	SLATE_END_ARGS()
+	
+	/**
+	 * Sets the content for this button
+	 *
+	 * @param	InContent	The widget to use as content for this button
+	 */
+	void SetMenuContent( const TSharedRef< SWidget >& InContent );
+
+	/**
+	 * Construct the widget from a declaration
+	 *
+	 * @param InArgs  The declaration from which to construct
+	 */
+	void Construct(const FArguments& InArgs);
+
+	void SetMenuContentWidgetToFocus( TWeakPtr<SWidget> InWidgetToFocusPtr )
+	{
+		WidgetToFocusPtr = InWidgetToFocusPtr;
+	}
+
+protected:
+	/**
+	 * Handle the button being clicked by summoning the ComboButton.
+	 */
+	virtual FReply OnButtonClicked();
+
+protected:
+	/** Area where the button's content resides */
+	SHorizontalBox::FSlot* ButtonContentSlot;
+
+	/** Delegate to execute to get the menu content of this button */
+	FOnGetContent OnGetMenuContent;
+
+	/** Delegate to execute when the combo list is opened */
+	FOnComboBoxOpened OnComboBoxOpened;
+
+	TWeakPtr<SWidget> WidgetToFocusPtr;
+
+	/** Brush to use to add a "menu border" around the drop-down content */
+	const FSlateBrush* MenuBorderBrush;
+
+	/** Padding to use to add a "menu border" around the drop-down content */
+	FMargin MenuBorderPadding;
+
+	/** The max height of the menu */
+	TAttribute<FOptionalSize> MenuWidth;
+
+	/** When specified, ignore the content's desired size and report the.HeightOverride as the Box's desired height. */
+	TAttribute<FOptionalSize> MenuHeight;
+
+	/** The content widget, if any, set by the user on creation */
+	TWeakPtr<SWidget> ContentWidgetPtr;
+};
