Slate: Refactored core Slate implementation into SlateCore module in preparation for UMG.

diff --git a/Engine/Source/Runtime/Slate/Public/Widgets/Input/SEditableText.h b/Engine/Source/Runtime/Slate/Public/Widgets/Input/SEditableText.h
new file mode 100644
--- /dev/null
+++ b/Engine/Source/Runtime/Slate/Public/Widgets/Input/SEditableText.h
@@ -0,0 +1,581 @@
+// Copyright 1998-2014 Epic Games, Inc. All Rights Reserved.
+
+#pragma once
+
+#include "MultiBox/MultiBoxExtender.h"
+#include "ITextEditorWidget.h"
+#include "ITextInputMethodSystem.h"
+
+/**
+ * Editable text widget
+ */
+class SLATE_API SEditableText : public SLeafWidget, public ITextEditorWidget
+{
+
+public:
+
+	SLATE_BEGIN_ARGS( SEditableText )
+		: _Text()
+		, _HintText()
+		, _Style(&FCoreStyle::Get().GetWidgetStyle< FEditableTextStyle >("NormalEditableText"))
+		, _Font()
+		, _ColorAndOpacity()
+		, _BackgroundImageSelected()
+		, _BackgroundImageSelectionTarget()
+		, _BackgroundImageComposing()
+		, _CaretImage()
+		, _IsReadOnly( false )
+		, _IsPassword( false )
+		, _IsCaretMovedWhenGainFocus( true )
+		, _SelectAllTextWhenFocused( false )
+		, _RevertTextOnEscape( false )
+		, _ClearKeyboardFocusOnCommit( true )
+		, _MinDesiredWidth( 0.0f )
+		, _SelectAllTextOnCommit( false )
+		{}
+
+		/** Sets the text content for this editable text widget */
+		SLATE_ATTRIBUTE( FText, Text )
+
+		/** The text that appears when there is nothing typed into the search box */
+		SLATE_ATTRIBUTE( FText, HintText )
+
+		/** The style of the text block, which dictates the font, color */
+		SLATE_STYLE_ARGUMENT( FEditableTextStyle, Style )
+
+		/** Sets the font used to draw the text (overrides Style) */
+		SLATE_ATTRIBUTE( FSlateFontInfo, Font )
+
+		/** Text color and opacity (overrides Style) */
+		SLATE_ATTRIBUTE( FSlateColor, ColorAndOpacity )
+
+		/** Background image for the selected text (overrides Style) */
+		SLATE_ATTRIBUTE( const FSlateBrush*, BackgroundImageSelected )
+
+		/** Background image for the selection targeting effect (overrides Style) */
+		SLATE_ATTRIBUTE( const FSlateBrush*, BackgroundImageSelectionTarget )
+
+		/** Background image for the composing text (overrides Style) */
+		SLATE_ATTRIBUTE( const FSlateBrush*, BackgroundImageComposing )
+
+		/** Image brush used for the caret (overrides Style) */
+		SLATE_ATTRIBUTE( const FSlateBrush*, CaretImage )
+
+		/** Sets whether this text box can actually be modified interactively by the user */
+		SLATE_ATTRIBUTE( bool, IsReadOnly )
+
+		/** Sets whether this text box is for storing a password */
+		SLATE_ATTRIBUTE( bool, IsPassword )
+
+		/** Workaround as we loose focus when the auto completion closes. */
+		SLATE_ATTRIBUTE( bool, IsCaretMovedWhenGainFocus )
+
+		/** Whether to select all text when the user clicks to give focus on the widget */
+		SLATE_ATTRIBUTE( bool, SelectAllTextWhenFocused )
+
+		/** Whether to allow the user to back out of changes when they press the escape key */
+		SLATE_ATTRIBUTE( bool, RevertTextOnEscape )
+
+		/** Whether to clear keyboard focus when pressing enter to commit changes */
+		SLATE_ATTRIBUTE( bool, ClearKeyboardFocusOnCommit )
+
+		/**
+		 * This is NOT for validating input!
+		 * 
+		 * Called whenever a character is typed.
+		 * Not called for copy, paste, or any other text changes!
+		 */
+		SLATE_EVENT( FOnIsTypedCharValid, OnIsTypedCharValid )
+
+		/** Called whenever the text is changed interactively by the user */
+		SLATE_EVENT( FOnTextChanged, OnTextChanged )
+
+		/** Called whenever the text is committed.  This happens when the user presses enter or the text box loses focus. */
+		SLATE_EVENT( FOnTextCommitted, OnTextCommitted )
+
+		/** Minimum width that a text block should be */
+		SLATE_ATTRIBUTE( float, MinDesiredWidth )
+
+		/** Whether to select all text when pressing enter to commit changes */
+		SLATE_ATTRIBUTE( bool, SelectAllTextOnCommit )
+
+		/** Menu extender for the right-click context menu */
+		SLATE_EVENT( FMenuExtensionDelegate, ContextMenuExtender )
+
+	SLATE_END_ARGS()
+
+	/** Constructor */
+	SEditableText();
+	virtual ~SEditableText();
+
+	/**
+	 * Construct this widget
+	 *
+	 * @param	InArgs	The declaration data for this widget
+	 */
+	void Construct( const FArguments& InArgs );
+
+	/**
+	 * Returns the text.
+	 *
+	 * @return  Text
+	 */
+	const FText& GetText() const
+	{
+		return Text.Get();
+	}
+
+	/**
+	 * Sets the text currently being edited 
+	 *
+	 * @param  InNewText  The new text
+	 */
+	void SetText( const TAttribute< FText >& InNewText );
+
+	/**
+	 * Checks to see if there is any text selected
+	 *
+	 * @return  True if any text is selected
+	 */
+	bool AnyTextSelected() const
+	{
+		return Selection.StartIndex != INDEX_NONE && Selection.StartIndex != Selection.FinishIndex;
+	}
+
+	/**
+	 * Restores the text to the original state
+	 */
+	void RestoreOriginalText();
+
+	/**
+	 *	Returns whether the current text varies from the original
+	 */
+	bool HasTextChangedFromOriginal() const;
+
+	/**
+	 * Sets the font used to draw the text
+	 *
+	 * @param  InNewFont	The new font to use
+	 */
+	void SetFont( const TAttribute< FSlateFontInfo >& InNewFont );
+
+protected:
+
+	friend class FTextInputMethodContext;
+	class FTextInputMethodContext : public ITextInputMethodContext
+	{
+	public:
+		FTextInputMethodContext(const TWeakPtr<SEditableText>& InOwningWidget)
+			:	OwningWidget(InOwningWidget)
+			,	IsComposing(false)
+			,	CompositionBeginIndex(INDEX_NONE)
+			,	CompositionLength(0)
+		{}
+
+		virtual ~FTextInputMethodContext() {}
+
+	private:
+		virtual bool IsReadOnly() OVERRIDE;
+		virtual uint32 GetTextLength() OVERRIDE;
+		virtual void GetSelectionRange(uint32& BeginIndex, uint32& Length, ECaretPosition& CaretPosition) OVERRIDE;
+		virtual void SetSelectionRange(const uint32 BeginIndex, const uint32 Length, const ECaretPosition CaretPosition) OVERRIDE;
+		virtual void GetTextInRange(const uint32 BeginIndex, const uint32 Length, FString& OutString) OVERRIDE;
+		virtual void SetTextInRange(const uint32 BeginIndex, const uint32 Length, const FString& InString) OVERRIDE;
+		virtual int32 GetCharacterIndexFromPoint(const FVector2D& Point) OVERRIDE;
+		virtual bool GetTextBounds(const uint32 BeginIndex, const uint32 Length, FVector2D& Position, FVector2D& Size) OVERRIDE;
+		virtual void GetScreenBounds(FVector2D& Position, FVector2D& Size) OVERRIDE;
+		virtual TSharedPtr<FGenericWindow> GetWindow() OVERRIDE;
+		virtual void BeginComposition() OVERRIDE;
+		virtual void UpdateCompositionRange(const int32 InBeginIndex, const uint32 InLength) OVERRIDE;
+		virtual void EndComposition() OVERRIDE;
+
+	private:
+		TWeakPtr<SEditableText> OwningWidget;
+
+	public:
+		FGeometry CachedGeometry;
+		bool IsComposing;
+		int32 CompositionBeginIndex;
+		uint32 CompositionLength;
+	};
+
+	/**
+	 * Range of selected characters
+	 */
+	class FTextSelection
+	{
+
+	public:
+
+		/** Index of character we started selection on.  Note that this can be larger than LastIndex! */
+		int32 StartIndex;
+
+		/** Index of last character included in the selection range.  Note that this can be smaller than StartIndex! */
+		int32 FinishIndex;
+
+
+	public:
+
+		/** Constructor */
+		FTextSelection()
+			: StartIndex( INDEX_NONE ),
+			  FinishIndex( INDEX_NONE )
+		{
+		}
+
+		/** Clears selection */
+		void Clear()
+		{
+			StartIndex = FinishIndex = INDEX_NONE;
+		}
+
+		/**
+		 * Returns the first character index in the selection range
+		 *
+		 * @return  First character index
+		 */
+		int32 GetMinIndex() const
+		{
+			return FMath::Min( StartIndex, FinishIndex );
+		}
+
+		/**
+		 * Returns the last character index in the selection range
+		 *
+		 * @return  Last character index
+		 */
+		int32 GetMaxIndex() const
+		{
+			return FMath::Max( StartIndex, FinishIndex );
+		}
+	};
+
+
+
+	/**
+	 * Stores a single undo level for editable text
+	 */
+	class FUndoState
+	{
+
+	public:
+
+		/** Text string */
+		FText Text;
+
+		/** Selection state */
+		FTextSelection Selection;
+
+		/** Caret position */
+		int32 CaretPosition;
+
+	public:
+
+		FUndoState()
+			: Text(),
+			  Selection(),
+			  CaretPosition( INDEX_NONE )
+		{
+		}
+	};
+
+	/**
+	 * Gets the height of the largest character in the font
+	 *
+	 * @return  The fonts height
+	 */
+	virtual float GetFontHeight() const;
+
+public:
+	// BEGIN ITextEditorWidget interface
+	virtual void StartChangingText() OVERRIDE;
+	virtual void FinishChangingText() OVERRIDE;
+	virtual bool GetIsReadOnly() const OVERRIDE;
+	virtual void BackspaceChar() OVERRIDE;
+	virtual void DeleteChar() OVERRIDE;
+	virtual bool CanTypeCharacter(const TCHAR CharInQuestion) const OVERRIDE;
+	virtual void TypeChar( const int32 Character ) OVERRIDE;
+	virtual FReply MoveCursor( ECursorMoveMethod::Type Method, int8 Direction, ECursorAction::Type Action ) OVERRIDE;
+	virtual void JumpTo(ETextLocation::Type JumpLocation, ECursorAction::Type Action) OVERRIDE;
+	virtual void ClearSelection() OVERRIDE;
+	virtual void SelectAllText() OVERRIDE;
+	virtual FReply OnEscape() OVERRIDE;
+	virtual void OnEnter() OVERRIDE;
+	virtual bool CanExecuteCut() const OVERRIDE;
+	virtual void CutSelectedTextToClipboard() OVERRIDE;
+	virtual bool CanExecuteCopy() const OVERRIDE;
+	virtual void CopySelectedTextToClipboard() OVERRIDE;
+	virtual bool CanExecutePaste() const OVERRIDE;
+	virtual void PasteTextFromClipboard() OVERRIDE;
+	virtual bool CanExecuteUndo() const OVERRIDE;
+	virtual void Undo() OVERRIDE;
+	virtual void Redo() OVERRIDE;
+	// END ITextEditorWidget interface
+
+protected:
+	// BEGIN SWidget interface
+	virtual void Tick( const FGeometry& AllottedGeometry, const double InCurrentTime, const float InDeltaTime ) OVERRIDE;
+	virtual FVector2D ComputeDesiredSize() const OVERRIDE;
+	virtual int32 OnPaint( const FGeometry& AllottedGeometry, const FSlateRect& MyClippingRect, FSlateWindowElementList& OutDrawElements, int32 LayerId, const FWidgetStyle& InWidgetStyle, bool bParentEnabled ) const OVERRIDE;
+	virtual FReply OnDragOver( const FGeometry& MyGeometry, const FDragDropEvent& DragDropEvent ) OVERRIDE;
+	virtual FReply OnDrop( const FGeometry& MyGeometry, const FDragDropEvent& DragDropEvent ) OVERRIDE;
+	virtual bool SupportsKeyboardFocus() const OVERRIDE;
+	virtual FReply OnKeyboardFocusReceived( const FGeometry& MyGeometry, const FKeyboardFocusEvent& InKeyboardFocusEvent ) OVERRIDE;
+	virtual void OnKeyboardFocusLost( const FKeyboardFocusEvent& InKeyboardFocusEvent ) OVERRIDE;
+	virtual FReply OnKeyChar( const FGeometry& MyGeometry,  const FCharacterEvent& InCharacterEvent ) OVERRIDE;
+	virtual FReply OnKeyDown( const FGeometry& MyGeometry, const FKeyboardEvent& InKeyboardEvent ) OVERRIDE;
+	virtual FReply OnKeyUp( const FGeometry& MyGeometry, const FKeyboardEvent& InKeyboardEvent ) OVERRIDE;
+	virtual FReply OnMouseButtonDown( const FGeometry& InMyGeometry, const FPointerEvent& InMouseEvent ) OVERRIDE;
+	virtual FReply OnMouseButtonUp( const FGeometry& InMyGeometry, const FPointerEvent& InMouseEvent ) OVERRIDE;
+	virtual FReply OnMouseMove( const FGeometry& InMyGeometry, const FPointerEvent& InMouseEvent ) OVERRIDE;
+	virtual FReply OnMouseButtonDoubleClick( const FGeometry& InMyGeometry, const FPointerEvent& InMouseEvent ) OVERRIDE;
+	virtual FCursorReply OnCursorQuery( const FGeometry& MyGeometry, const FPointerEvent& CursorEvent ) const OVERRIDE;
+	// END SWidget interface
+
+protected:
+	/**
+	 * Updates text selection after the caret moved
+	 *
+	 * @param  InOldCaretPosition  Previous location of the caret
+	 */
+	void SelectText( const int32 InOldCaretPosition );
+
+	/**
+	 * Deletes selected text
+	 */
+	void DeleteSelectedText();
+	
+	/**
+	 * Determines which character was clicked on
+	 *
+	 * @param  InMyGeometry  Widget geometry
+	 * @param  InMouseEvent  Mouse button event
+	 *
+	 * @return  The clicked character index
+	 */
+	int32 FindClickedCharacterIndex( const FGeometry& InMyGeometry, const FVector2D& InScreenCursorPosition ) const;
+
+	/**
+	 * Adds the specified undo state to the undo stack
+	 *
+	 * @param  InUndoState  The undo state to push.  Usually, this is created using MakeUndoState()
+	 */
+	void PushUndoState( const SEditableText::FUndoState& InUndoState );
+
+
+	/**
+	 * Clears the local undo buffer
+	 */
+	void ClearUndoStates();
+
+
+	/**
+	 * Creates an undo state structure for the control's current state and returns it
+	 *
+	 * @param  OutUndoState  Undo state for current state of control
+	 */
+	void MakeUndoState( SEditableText::FUndoState& OutUndoState );
+
+
+	/**
+	 * Loads text for editing, from our text attribute
+	 */
+	void LoadText();
+
+	/**
+	 * Saves edited text into our text attribute, if appropriate
+	 */
+	void SaveText();
+
+	/**
+	 * Summons a context menu that allows the user to perform certain text-related actions
+	 *
+	 * @param	InLocation	Screen space position to show the menu
+	 */
+	void SummonContextMenu( const FVector2D& InLocation );
+
+
+	/**
+	 * Executes action to delete selected text
+	 */
+	void ExecuteDeleteAction();
+
+	/**
+	 * Returns true if 'Select All' is currently a valid action
+	 *
+	 * @return	True to allow this action, otherwise false
+	 */
+	bool CanExecuteSelectAll() const;
+
+	/**
+	 * Returns true if 'Delete' is currently a valid action
+	 *
+	 * @return	True to allow this action, otherwise false
+	 */
+	bool CanExecuteDelete() const;
+
+	/**
+	 * Sets the caret position
+	 * 
+	 * @param Position The character index for the new position
+	 */
+	void SetCaretPosition( int32 Position );
+
+	/**
+	 * Called to restart selection target animation after most types of text selection changes
+	 */
+	void RestartSelectionTargetAnimation();
+
+	/**
+	 * Called when the content menu window is closed
+	 */
+	void OnWindowClosed(const TSharedRef<SWindow>&);
+
+	/** 
+	 * Calculate the width of the caret 
+	 * @param FontMaxCharHeight The height of the font to calculate the caret width for
+	 * @return The width of the caret (might be clamped for very small fonts)
+	 */
+	float CalculateCaretWidth(const float FontMaxCharHeight) const;
+
+private:
+	/** 
+	 * @return whether there is anything in the clipboard
+	 */
+	bool DoesClipboardHaveAnyText() const;
+
+	/**
+	 * @return  the string that needs to be rendered
+	 */
+	FString GetStringToRender() const;
+
+private:
+
+	/** A pointer to the editable text currently under the mouse cursor. NULL when there isn't one. */
+	static SEditableText* EditableTextUnderCursor;
+
+	/** The text content for this editable text widget */
+	TAttribute< FText > Text;
+
+	TAttribute< FText > HintText;
+
+	/** The font used to draw the text */
+	TAttribute< FSlateFontInfo > Font;
+
+	/** Text color and opacity */
+	TAttribute<FSlateColor> ColorAndOpacity;
+
+	/** Background image for the selected text */
+	TAttribute< const FSlateBrush* > BackgroundImageSelected;
+
+	/** Background image for the selection targeting effect */
+	TAttribute< const FSlateBrush* > BackgroundImageSelectionTarget;
+
+	/** Background image for the composing text */
+	TAttribute< const FSlateBrush* > BackgroundImageComposing;
+
+	/** Image brush used for the caret */
+	TAttribute< const FSlateBrush* > CaretImage;
+
+	/** Sets whether this text box can actually be modified interactively by the user */
+	TAttribute< bool > IsReadOnly;
+
+	/** Sets whether this text box is for storing a password */
+	TAttribute< bool > IsPassword;
+
+	/** Workaround as we loose focus when the auto completion closes. */
+	TAttribute< bool > IsCaretMovedWhenGainFocus;
+	
+	/** Whether to select all text when the user clicks to give focus on the widget */
+	TAttribute< bool > SelectAllTextWhenFocused;
+
+	/** Whether to allow the user to back out of changes when they press the escape key */
+	TAttribute< bool > RevertTextOnEscape;
+
+	/** Whether to clear keyboard focus when pressing enter to commit changes */
+	TAttribute< bool > ClearKeyboardFocusOnCommit;
+
+	/** Whether to select all text when pressing enter to commit changes */
+	TAttribute< bool > SelectAllTextOnCommit;
+
+	/** Called when a character is typed and we want to know if the text field supports typing this character. */
+	FOnIsTypedCharValid OnIsTypedCharValid;
+
+	/** Called whenever the text is changed interactively by the user */
+	FOnTextChanged OnTextChanged;
+
+	/** Called whenever the text is committed.  This happens when the user presses enter or the text box loses focus. */
+	FOnTextCommitted OnTextCommitted;
+
+	/** Text string currently being edited */
+	FText EditedText;
+
+	/** Original text prior to user edits.  This is used to restore the text if the user uses the escape key. */
+	FText OriginalText;
+
+	/** Current scrolled position, in text-area space; use ConvertGeometryFromRenderSpaceToTextSpace to convert before using values with ScrollHelper */
+	FScrollHelper ScrollHelper;
+
+	/** Caret position as an index into the editable text string's character array */
+	int32 CaretPosition;
+
+	/** Horizontal visual position of caret along the scrolling canvas */
+	FFloatSpring1D CaretVisualPositionSpring;
+
+	/** Stores the last time that the caret was moved by the user.  Used to prevent the caret from blinking. */
+	double LastCaretInteractionTime;
+
+	/** Current selection state */
+	FTextSelection Selection;
+
+	/** Selection "targeting" visual effect left position */
+	FFloatSpring1D SelectionTargetLeftSpring;
+
+	/** Selection "targeting" visual effect right position */
+	FFloatSpring1D SelectionTargetRightSpring;
+
+	/** Last time that the user had a major interaction with the selection */
+	double LastSelectionInteractionTime;
+
+	/** True if we're currently selecting text by dragging the mouse cursor with the left button held down */
+	bool bIsDragSelecting;
+
+	/** True if the last mouse down caused us to receive keyboard focus */
+	bool bWasFocusedByLastMouseDown;
+
+	/** True if characters were selected by dragging since the last keyboard focus.  Used for text selection. */
+	bool bHasDragSelectedSinceFocused;
+
+	/** True if characters were selected by dragging since the last mouse down. */
+	bool bHasDragSelectedSinceMouseDown;
+
+	/** Undo states */
+	TArray< FUndoState > UndoStates;
+
+	/** Current undo state level that we've rolled back to, or INDEX_NONE if we haven't undone.  Used for 'Redo'. */
+	int32 CurrentUndoLevel;
+
+	/** True if we're currently (potentially) changing the text string */
+	bool bIsChangingText;
+
+	/** Undo state that will be pushed if text is actually changed between calls to StartChangingText() and FinishChangingText() */
+	FUndoState StateBeforeChangingText;
+
+	/** Weak pointer to context menu window that's currently open, if there is one */
+	TWeakPtr< SWindow > ContextMenuWindow;
+
+	/** Prevents the editabletext from being smaller than desired in certain cases (e.g. when it is empty) */
+	TAttribute<float> MinDesiredWidth;
+
+	/** A list commands to execute if a user presses the corresponding keybinding in the text box */
+	TSharedRef< FUICommandList > UICommandList;
+
+	/** Menu extender for right-click context menu */
+	TSharedPtr<FExtender> MenuExtender;
+
+	/** Implemented context object for text input method systems. */
+	TSharedPtr<FTextInputMethodContext> TextInputMethodContext;
+
+	/** Notification interface object for text input method systems. */
+	TSharedPtr<ITextInputMethodChangeNotifier> TextInputMethodChangeNotifier;
+};
diff --git a/Engine/Source/Runtime/Slate/Public/Widgets/Input/SHyperlink.h b/Engine/Source/Runtime/Slate/Public/Widgets/Input/SHyperlink.h
new file mode 100644
--- /dev/null
+++ b/Engine/Source/Runtime/Slate/Public/Widgets/Input/SHyperlink.h
@@ -0,0 +1,74 @@
+// Copyright 1998-2014 Epic Games, Inc. All Rights Reserved.
+
+#pragma once
+
+/**
+ * A hyperlink widget is what you would expect from a browser hyperlink.
+ * When a hyperlink is clicked in invokes an OnNavigate() delegate.
+ */
+class SHyperlink : public SButton
+{
+public:
+
+	SLATE_BEGIN_ARGS(SHyperlink)
+		: _Text()
+		, _Style(&FCoreStyle::Get().GetWidgetStyle< FHyperlinkStyle >("Hyperlink"))
+		, _TextStyle(NULL)
+		, _UnderlineStyle(NULL)
+		, _Padding()
+		, _OnNavigate()
+		{}
+
+		SLATE_TEXT_ATTRIBUTE( Text )
+		SLATE_STYLE_ARGUMENT( FHyperlinkStyle, Style )
+		SLATE_STYLE_ARGUMENT( FTextBlockStyle, TextStyle )
+		SLATE_STYLE_ARGUMENT( FButtonStyle, UnderlineStyle )
+		SLATE_ATTRIBUTE( FMargin, Padding )
+		SLATE_EVENT( FSimpleDelegate, OnNavigate )
+	SLATE_END_ARGS()
+
+	/**
+	 * Construct the hyperlink widgets from a declaration
+	 * 
+	 * @param InArgs    Widget declaration from which to construct the hyperlink.
+	 */
+	void Construct( const FArguments& InArgs )
+	{
+		this->OnNavigate = InArgs._OnNavigate;
+
+		check (InArgs._Style);
+		const FButtonStyle* UnderlineStyle = InArgs._UnderlineStyle != NULL ? InArgs._UnderlineStyle : &InArgs._Style->UnderlineStyle;
+		const FTextBlockStyle* TextStyle = InArgs._TextStyle != NULL ? InArgs._TextStyle : &InArgs._Style->TextStyle;
+		TAttribute<FMargin> Padding = InArgs._Padding.IsSet() ? InArgs._Padding : InArgs._Style->Padding;
+
+		SButton::Construct(
+			SButton::FArguments()
+			.Text( InArgs._Text )
+			.ContentPadding( Padding )
+			.ButtonStyle( UnderlineStyle )
+			.TextStyle( TextStyle )
+			.OnClicked(this, &SHyperlink::Hyperlink_OnClicked)
+			.ForegroundColor(FSlateColor::UseForeground())			
+		);
+	}
+
+	// SWidget implementation
+	virtual FCursorReply OnCursorQuery( const FGeometry& MyGeometry, const FPointerEvent& CursorEvent ) const OVERRIDE
+	{
+		return FCursorReply::Cursor( EMouseCursor::Hand );
+	}
+	// End SWidget implementation
+
+protected:
+
+	/** Invoke the OnNavigate method */
+	FReply Hyperlink_OnClicked()
+	{
+		OnNavigate.ExecuteIfBound();
+
+		return FReply::Handled();
+	}
+
+	/** The delegate to invoke when someone clicks the hyperlink */
+	FSimpleDelegate OnNavigate;
+};
\ No newline at end of file
diff --git a/Engine/Source/Runtime/Slate/Public/Widgets/Input/SMenuAnchor.h b/Engine/Source/Runtime/Slate/Public/Widgets/Input/SMenuAnchor.h
new file mode 100644
--- /dev/null
+++ b/Engine/Source/Runtime/Slate/Public/Widgets/Input/SMenuAnchor.h
@@ -0,0 +1,118 @@
+// Copyright 1998-2014 Epic Games, Inc. All Rights Reserved.
+
+#pragma once
+
+/**
+ * A PopupAnchor summons a Popup relative to its content.
+ * Summoning a popup relative to the cursor is accomplished via the application.
+ */
+class SLATE_API SMenuAnchor : public SCompoundWidget
+{
+public:
+	enum EMethod
+	{
+		/** Creating a new window allows us to place popups outside of the window in which the menu anchor resides. */
+		CreateNewWindow,
+		/** Applications that intend to run in fullscreen cannot create new windows. */
+		UseCurrentWindow
+	};
+	
+public:
+	SLATE_BEGIN_ARGS( SMenuAnchor )
+		: _Content()
+		, _MenuContent( SNew(STextBlock) .Text( NSLOCTEXT("SMenuAnchor", "NoMenuContent", "No Menu Content Assigned; use .MenuContent") ) )
+		, _OnGetMenuContent()
+		, _Placement( MenuPlacement_BelowAnchor )
+		, _Method( CreateNewWindow )
+		{}
+		
+		SLATE_DEFAULT_SLOT( FArguments, Content )
+		
+		SLATE_ARGUMENT( TSharedPtr<SWidget>, MenuContent )
+
+		SLATE_EVENT( FOnGetContent, OnGetMenuContent )
+		
+		SLATE_ATTRIBUTE( EMenuPlacement, Placement )
+
+		SLATE_ARGUMENT( EMethod, Method )
+
+	SLATE_END_ARGS()
+
+	
+	/**
+	 * Construct this widget
+	 *
+	 * @param	InArgs	The declaration data for this widget
+	 */
+	void Construct( const FArguments& InArgs );
+	
+	SMenuAnchor();
+
+	// SWidget interface
+	virtual void Tick( const FGeometry& AllottedGeometry, const double InCurrentTime, const float InDeltaTime ) OVERRIDE;
+	// End of SWidget interface
+
+	
+	/**
+	 * Open or close the popup
+	 *
+	 * @param InIsOpen    If true, open the popup. Otherwise close it.
+	 * @param bFocusMenu  Shoudl we focus the popup as soon as it opens?
+	 */
+	void SetIsOpen( bool InIsOpen, const bool bFocusMenu = true );
+	
+	/** @return true if the popup is open; false otherwise. */
+	bool IsOpen() const;
+
+	/** @return true if we should open the menu due to a click. Sometimes we should not, if
+	   the same MouseDownEvent that just closed the menu is about to re-open it because it happens to land on the button.*/
+	bool ShouldOpenDueToClick() const;
+
+	/** @return The current menu position */
+	FVector2D GetMenuPosition() const;
+
+	/** @return Whether this menu has open submenus */
+	bool HasOpenSubMenus() const;
+
+protected:
+
+	/**
+	 * Computes the desired position for a menu that's attached to this anchor
+	 *
+	 * @param   AllottedGeometry    Our widget's geometry
+	 * @param   PopupDesiredSize    How big the popup needs to be for all of its contents to be visible
+	 * @param   WindowSize          (Out) New size of the window
+	 * @param   NewPosition         (Out) The computed window position
+	 */
+	void ComputeMenuPlacement( const FGeometry& AllottedGeometry, const FVector2D& PopupDesiredSize, FVector2D& WindowSize, FVector2D& NewPosition );
+
+	/** Invoked when the popup window is being closed by the application */
+	void RequestClosePopupWindow( const TSharedRef<SWindow>& PopupWindow );
+
+	/** Invoked when a click occurred outside the widget that is our popup; used for popup auto-dismissal */
+	void OnClickedOutsidePopup();
+
+protected:
+	/** A pointer to the window presenting this popup. Pointer is null when the popup is not visible */
+	TWeakPtr<SWindow> PopupWindowPtr;
+
+	/** When not making a new window for the popup, we rely on the SPopupLayer to achieve the effect. */
+	struct FPopupLayerSlot* PopupLayerSlot;
+	
+	/** Static menu content */
+	TSharedPtr<SWidget> MenuContent;
+	
+	/** Callback invoked when the popup is being summoned */
+	FOnGetContent OnGetMenuContent;
+	
+	/** How should the popup be placed relative to the anchor. */
+	TAttribute<EMenuPlacement> Placement;
+
+	/** Was the menu just dismissed this tick? */
+	bool bDismissedThisTick;
+
+	/** Should we summon a new window for this popup or  */
+	EMethod Method;
+
+
+};
diff --git a/Engine/Source/Runtime/Slate/Public/Widgets/Input/SNumericEntryBox.h b/Engine/Source/Runtime/Slate/Public/Widgets/Input/SNumericEntryBox.h
new file mode 100644
--- /dev/null
+++ b/Engine/Source/Runtime/Slate/Public/Widgets/Input/SNumericEntryBox.h
@@ -0,0 +1,465 @@
+// Copyright 1998-2014 Epic Games, Inc. All Rights Reserved.
+
+#pragma once
+
+/**
+ * Implementation for a box that only accepts a numeric value or that can display an undetermined value via a string
+ * Supports an optional spin box for manipulating a value by dragging with the mouse
+ * Supports an optional label inset in the text box
+ */ 
+template<typename NumericType>
+class SNumericEntryBox : public SCompoundWidget
+{
+public: 
+	static const FLinearColor RedLabelBackgroundColor;
+	static const FLinearColor GreenLabelBackgroundColor;
+	static const FLinearColor BlueLabelBackgroundColor;
+
+	/** Notification for numeric value change */
+	DECLARE_DELEGATE_OneParam( FOnValueChanged, NumericType );
+
+	/** Notification for numeric value committed */
+	DECLARE_DELEGATE_TwoParams( FOnValueCommitted, NumericType, ETextCommit::Type);
+
+public:
+	SLATE_BEGIN_ARGS( SNumericEntryBox<NumericType> )
+		: _EditableTextBoxStyle( &FCoreStyle::Get().GetWidgetStyle<FEditableTextBoxStyle>("NormalEditableTextBox") )
+		, _Label()
+		, _LabelVAlign( VAlign_Fill )
+		, _LabelPadding( FMargin(3,0) )
+		, _BorderForegroundColor(FCoreStyle::Get().GetSlateColor("InvertedForeground"))
+		, _BorderBackgroundColor(FLinearColor::White)
+		, _UndeterminedString( TEXT("---") )
+		, _Font( FCoreStyle::Get().GetFontStyle( TEXT("NormalFont") ) )
+		, _AllowSpin(false)
+		, _Delta(0)
+		, _MinValue(TNumericLimits<NumericType>::Lowest())
+		, _MaxValue(TNumericLimits<NumericType>::Max())
+		, _MinSliderValue(0)				
+		, _MaxSliderValue(100)
+		, _SliderExponent(1.f)
+	{}		
+
+		/** Style to use for the editable text box within this widget */
+		SLATE_STYLE_ARGUMENT( FEditableTextBoxStyle, EditableTextBoxStyle )
+
+		/** Slot for this button's content (optional) */
+		SLATE_NAMED_SLOT( FArguments, Label )
+		/** Vertical alignment of the label content */
+		SLATE_ARGUMENT( EVerticalAlignment, LabelVAlign )
+		/** Padding around the label content */
+		SLATE_ARGUMENT( FMargin, LabelPadding )
+		/** Border Foreground Color */
+		SLATE_ARGUMENT( FSlateColor, BorderForegroundColor )
+		/** Border Background Color */
+		SLATE_ARGUMENT( FSlateColor, BorderBackgroundColor )
+		/** The value that should be displayed.  This value is optional in the case where a value cannot be determined */
+		SLATE_ATTRIBUTE( TOptional<NumericType>, Value )
+		/** The string to display if the value cannot be determined */
+		SLATE_TEXT_ARGUMENT( UndeterminedString )
+		/** Font color and opacity */
+		SLATE_ATTRIBUTE( FSlateFontInfo, Font )
+		/** Whether or not the user should be able to change the value by dragging with the mouse cursor */
+		SLATE_ARGUMENT( bool, AllowSpin )
+		/** Delta to increment the value as the slider moves.  If not specified will determine automatically */
+		SLATE_ARGUMENT( NumericType, Delta )
+		/** The minimum value that can be entered into the text edit box */
+		SLATE_ATTRIBUTE( TOptional< NumericType >, MinValue )
+		/** The maximum value that can be entered into the text edit box */
+		SLATE_ATTRIBUTE( TOptional< NumericType >, MaxValue )
+		/** The minimum value that can be specified by using the slider */
+		SLATE_ATTRIBUTE( TOptional< NumericType >, MinSliderValue )
+		/** The maximum value that can be specified by using the slider */
+		SLATE_ATTRIBUTE( TOptional< NumericType >, MaxSliderValue )
+		/** Use exponential scale for the slider */
+		SLATE_ATTRIBUTE( float, SliderExponent )
+		/** Called whenever the text is changed interactively by the user */
+		SLATE_EVENT( FOnValueChanged, OnValueChanged )
+		/** Called whenever the text is committed.  This happens when the user presses enter or the text box loses focus. */
+		SLATE_EVENT( FOnValueCommitted, OnValueCommitted )
+		/** Called right before the slider begins to move */
+		SLATE_EVENT( FSimpleDelegate, OnBeginSliderMovement )
+		/** Called right after the slider handle is released by the user */
+		SLATE_EVENT( FOnValueChanged, OnEndSliderMovement )		
+		/** Menu extender for right-click context menu */
+		SLATE_EVENT( FMenuExtensionDelegate, ContextMenuExtender )		
+
+	SLATE_END_ARGS()
+	SNumericEntryBox()
+	{
+	}
+
+	void Construct( const FArguments& InArgs )
+	{
+		check(InArgs._EditableTextBoxStyle);
+
+		const bool bAllowSpin = InArgs._AllowSpin;
+		OnValueChanged = InArgs._OnValueChanged;
+		OnValueCommitted = InArgs._OnValueCommitted;
+		ValueAttribute = InArgs._Value;
+		UndeterminedString = InArgs._UndeterminedString;
+		
+		BorderImageNormal = &InArgs._EditableTextBoxStyle->BackgroundImageNormal;
+		BorderImageHovered = &InArgs._EditableTextBoxStyle->BackgroundImageHovered;
+		BorderImageFocused = &InArgs._EditableTextBoxStyle->BackgroundImageFocused;
+		const FMargin& TextMargin = InArgs._EditableTextBoxStyle->Padding;
+		
+		TSharedPtr<SWidget> FinalWidget;
+		if( bAllowSpin )
+		{
+			SAssignNew( SpinBox, SSpinBox<NumericType> )
+				.Style( FCoreStyle::Get(), "NumericEntrySpinBox" )
+				.Font( InArgs._Font )
+				.ContentPadding( TextMargin )
+				.Value( this, &SNumericEntryBox<NumericType>::OnGetValueForSpinBox )
+				.Delta( InArgs._Delta )
+				.OnValueChanged( OnValueChanged )
+				.OnValueCommitted( OnValueCommitted )
+				.MinSliderValue(InArgs._MinSliderValue)
+				.MaxSliderValue(InArgs._MaxSliderValue)
+				.MaxValue(InArgs._MaxValue)
+				.MinValue(InArgs._MinValue)
+				.SliderExponent(InArgs._SliderExponent)
+				.OnBeginSliderMovement(InArgs._OnBeginSliderMovement)
+				.OnEndSliderMovement(InArgs._OnEndSliderMovement);
+		}
+
+		// Always create an editable text box.  In the case of an undetermined value being passed in, we cant use the spinbox.
+		SAssignNew( EditableText, SEditableText )
+			.Text( this, &SNumericEntryBox<NumericType>::OnGetValueForTextBox )
+			.Visibility( bAllowSpin ? EVisibility::Collapsed : EVisibility::Visible )
+			.Font( InArgs._Font )
+			.SelectAllTextWhenFocused( true )
+			.ClearKeyboardFocusOnCommit( false )
+			.OnTextChanged( this, &SNumericEntryBox<NumericType>::OnTextChanged  )
+			.OnTextCommitted( this, &SNumericEntryBox<NumericType>::OnTextCommitted )
+			.SelectAllTextOnCommit( true )
+			.ContextMenuExtender( InArgs._ContextMenuExtender );
+
+		TSharedRef<SHorizontalBox> HorizontalBox = SNew( SHorizontalBox );
+		if( InArgs._Label.Widget != SNullWidget::NullWidget )
+		{
+			HorizontalBox->AddSlot()
+				.AutoWidth()
+				.HAlign(HAlign_Left)
+				.VAlign(InArgs._LabelVAlign)
+				.Padding(InArgs._LabelPadding)
+				[
+					InArgs._Label.Widget
+				];
+		}
+
+		// Add the spin box if we have one
+		if( bAllowSpin )
+		{
+			HorizontalBox->AddSlot()
+				.HAlign(HAlign_Fill) 
+				.VAlign(VAlign_Center) 
+				.FillWidth(1)
+				[
+					SpinBox.ToSharedRef()
+				];
+		}
+
+		HorizontalBox->AddSlot()
+			.HAlign(HAlign_Fill) 
+			.VAlign(VAlign_Center)
+			.Padding(TextMargin)
+			.FillWidth(1)
+			[
+				EditableText.ToSharedRef()
+			];
+
+
+		ChildSlot
+			[
+				SNew( SBorder )
+				.BorderImage( this, &SNumericEntryBox<NumericType>::GetBorderImage )
+				.BorderBackgroundColor( InArgs._BorderBackgroundColor )
+				.ForegroundColor( InArgs._BorderForegroundColor )
+				.Padding(0)
+				[
+					HorizontalBox
+				]
+			];
+	}
+
+
+	/** Build a generic label with specified text, foreground color and background color */
+	static TSharedRef<SWidget> BuildLabel( const FText& LabelText, const FSlateColor& ForegroundColor, const FSlateColor& BackgroundColor )
+	{
+		return
+			SNew(SBorder)
+			.BorderImage( FCoreStyle::Get().GetBrush("NumericEntrySpinBox.Decorator") )
+			.BorderBackgroundColor( BackgroundColor )
+			.ForegroundColor( ForegroundColor )
+			.VAlign(VAlign_Center)
+			.HAlign(HAlign_Left)
+			.Padding( FMargin(1, 0, 6, 0) )
+			[
+				SNew( STextBlock )
+				.Text( LabelText )
+			];
+	}
+
+private:
+
+	// Begin SWidget Interface
+	virtual bool SupportsKeyboardFocus() const OVERRIDE
+	{
+		return StaticCastSharedPtr<SWidget>(EditableText)->SupportsKeyboardFocus();
+	}
+
+	virtual FReply OnKeyboardFocusReceived( const FGeometry& MyGeometry, const FKeyboardFocusEvent& InKeyboardFocusEvent ) OVERRIDE
+	{
+		FReply Reply = FReply::Handled();
+
+		// The widget to forward focus to changes depending on whether we have a SpinBox or not.
+		TSharedPtr<SWidget> FocusWidget;
+		if (SpinBox.IsValid() && SpinBox->GetVisibility() == EVisibility::Visible) 
+		{
+			FocusWidget = SpinBox;
+		}
+		else
+		{
+			FocusWidget = EditableText;
+		}
+
+		if ( InKeyboardFocusEvent.GetCause() != EKeyboardFocusCause::Cleared )
+		{
+			// Forward keyboard focus to our chosen widget
+			Reply.SetKeyboardFocus( FocusWidget.ToSharedRef(), InKeyboardFocusEvent.GetCause() );
+		}
+
+		return Reply;
+	}
+
+	virtual FReply OnKeyDown( const FGeometry& MyGeometry, const FKeyboardEvent& InKeyboardEvent ) OVERRIDE
+	{
+		FKey Key = InKeyboardEvent.GetKey();
+
+		if( Key == EKeys::Escape && EditableText->HasKeyboardFocus() )
+		{
+			return FReply::Handled().SetKeyboardFocus( SharedThis( this ), EKeyboardFocusCause::Cleared );
+		}
+
+		return FReply::Unhandled();
+	}
+	// End SWidget Interface
+
+	/**
+	 * @return the Label that should be displayed                   
+	 */
+	FString GetLabel() const
+	{
+		// Should always be set if this is being called
+		return LabelAttribute.Get().GetValue();
+	}
+
+	/**
+	 * Called to get the value for the spin box                   
+	 */
+	NumericType OnGetValueForSpinBox() const
+	{
+		const auto& Value = ValueAttribute.Get();
+
+		// Get the value or 0 if its not set
+		if( Value.IsSet() == true )
+		{
+			return Value.GetValue();
+		}
+		else
+		{
+			return 0;
+		}
+	}
+
+	/**
+	 * Called to get the value for the text box as FText                 
+	 */
+	FText OnGetValueForTextBox() const
+	{
+		FString NewString = TEXT("");
+
+		if( EditableText->GetVisibility() == EVisibility::Visible )
+		{
+			const auto& Value = ValueAttribute.Get();
+
+			// If the value was set convert it to a string, otherwise the value cannot be determined
+			if( Value.IsSet() == true )
+			{
+				auto CurrentValue = Value.GetValue();
+				NewString = TTypeToString<NumericType>::ToSanitizedString(CurrentValue);
+			}
+			else
+			{
+				NewString = UndeterminedString;
+			}
+		}
+
+		// The box isnt visible, just return an empty string
+		return FText::FromString(NewString);
+	}
+
+
+	/**
+	 * Called when the text changes in the text box                   
+	 */
+	void OnTextChanged( const FText& NewValue )
+	{
+		const auto& Value = ValueAttribute.Get();
+
+		// Do not sent change events if the current value cannot be determined or else next tick the spin box could be swapped in if the value becomes determined
+		// while a user is typing in the box.  This causes keyboard focus switch which is bad
+		if( Value.IsSet() == true )
+		{
+			SendChangesFromText( NewValue, false, ETextCommit::Default );
+		}
+	}
+
+	/**
+	 * Called when the text is committed from the text box                   
+	 */
+	void OnTextCommitted( const FText& NewValue, ETextCommit::Type CommitInfo )
+	{
+		SendChangesFromText( NewValue, true, CommitInfo );
+	}
+
+	/**
+	 * Called to get the border image of the box                   
+	 */
+	const FSlateBrush* GetBorderImage() const
+	{
+		TSharedPtr<const SWidget> EditingWidget;
+		if (SpinBox.IsValid() && SpinBox->GetVisibility() == EVisibility::Visible) 
+		{
+			EditingWidget = SpinBox;
+		}
+		else
+		{
+			EditingWidget = EditableText;
+		}
+
+		if ( EditingWidget->HasKeyboardFocus() )
+		{
+			return BorderImageFocused;
+		}
+		else
+		{
+			if ( EditingWidget->IsHovered() )
+			{
+				return BorderImageHovered;
+			}
+			else
+			{
+				return BorderImageNormal;
+			}
+		}
+	}
+
+	/**
+	 * Calls the value commit or changed delegate set for this box when the value is set from a string
+	 *
+	 * @param NewValue	The new value as a string
+	 * @param bCommit	Whether or not to call the commit or changed delegate
+	 */
+	void SendChangesFromText( const FText& NewValue, bool bCommit, ETextCommit::Type CommitInfo )
+	{
+		// Only call the delegates if we have a valid numeric value
+		if( !NewValue.IsEmpty()  )
+		{
+			if( bCommit )
+			{
+				bool bEvalResult = false;
+				NumericType NumericValue = 0;
+
+				// Convert string to an underlying type in case text is not math equation
+				// Try to parse equation otherwise
+				if ( NewValue.IsNumeric() )
+				{
+					TTypeFromString<NumericType>::FromString( NumericValue, *NewValue.ToString() );
+					bEvalResult = true;
+				}
+				else
+				{
+					// Only evaluate equations on commit or else they could fail because the equation is still being typed
+					float Value = 0.f;
+					bEvalResult = FMath::Eval( NewValue.ToString(), Value );
+					NumericValue = NumericType(Value);
+				}
+				
+				if( bEvalResult )
+				{
+					OnValueCommitted.ExecuteIfBound(NumericValue, CommitInfo );
+				}
+			}
+			else if( NewValue.IsNumeric() )
+			{
+				NumericType Value = 0;
+				TTypeFromString<NumericType>::FromString( Value, *NewValue.ToString() );
+				OnValueChanged.ExecuteIfBound( Value );
+			}
+		}
+	}
+
+	/**
+	 * Caches the value and performs widget visibility maintenance
+	 */
+	virtual void Tick( const FGeometry& AllottedGeometry, const double InCurrentTime, const float InDeltaTime ) OVERRIDE
+	{
+		// Visibility toggle only matters if the spinbox is used
+		if( SpinBox.IsValid() )
+		{
+			const auto& Value = ValueAttribute.Get();
+
+			if( Value.IsSet() == true )
+			{
+				if( SpinBox->GetVisibility() != EVisibility::Visible )
+				{
+					// Set the visibility of the spinbox to visible if we have a valid value
+					SpinBox->SetVisibility( EVisibility::Visible );
+					// The text box should be invisible
+					EditableText->SetVisibility( EVisibility::Collapsed );
+				}
+			}
+			else
+			{
+				// The value isn't set so the spinbox should be hidden and the text box shown
+				SpinBox->SetVisibility(EVisibility::Collapsed);
+				EditableText->SetVisibility(EVisibility::Visible);
+			}
+		}
+	}
+
+private:
+	/** Attribute for getting the label */
+	TAttribute< TOptional<FString > > LabelAttribute;
+	/** Attribute for getting the value.  If the value is not set we display the undetermined string */
+	TAttribute< TOptional<NumericType> > ValueAttribute;
+	/** Spinbox widget */
+	TSharedPtr<SWidget> SpinBox;
+	/** Editable widget */
+	TSharedPtr<SEditableText> EditableText;
+	/** Delegate to call when the value changes */
+	FOnValueChanged OnValueChanged;
+	/** Delegate to call when the value is committed */
+	FOnValueCommitted OnValueCommitted;
+	/** The undetermined string to display when needed */
+	FString UndeterminedString;
+	/** Styling: border image to draw when not hovered or focused */
+	const FSlateBrush* BorderImageNormal;
+	/** Styling: border image to draw when hovered */
+	const FSlateBrush* BorderImageHovered;
+	/** Styling: border image to draw when focused */
+	const FSlateBrush* BorderImageFocused;
+};
+
+
+template <typename NumericType>
+const FLinearColor SNumericEntryBox<NumericType>::RedLabelBackgroundColor(0.594f,0.0197f,0.0f);
+
+template <typename NumericType>
+const FLinearColor SNumericEntryBox<NumericType>::GreenLabelBackgroundColor(0.1349f,0.3959f,0.0f);
+
+template <typename NumericType>
+const FLinearColor SNumericEntryBox<NumericType>::BlueLabelBackgroundColor(0.0251f,0.207f,0.85f);
diff --git a/Engine/Source/Runtime/Slate/Public/Widgets/Input/SRichTextHyperlink.h b/Engine/Source/Runtime/Slate/Public/Widgets/Input/SRichTextHyperlink.h
new file mode 100644
--- /dev/null
+++ b/Engine/Source/Runtime/Slate/Public/Widgets/Input/SRichTextHyperlink.h
@@ -0,0 +1,75 @@
+// Copyright 1998-2014 Epic Games, Inc. All Rights Reserved.
+#pragma once
+
+class SRichTextHyperlink : public SHyperlink
+{
+public:
+
+	SLATE_BEGIN_ARGS(SRichTextHyperlink)
+		: _Text()
+		, _Style(&FCoreStyle::Get().GetWidgetStyle< FHyperlinkStyle >("Hyperlink"))
+		, _OnNavigate()
+	{}
+		SLATE_TEXT_ATTRIBUTE( Text )
+		SLATE_STYLE_ARGUMENT( FHyperlinkStyle, Style )
+		SLATE_EVENT( FSimpleDelegate, OnNavigate )
+	SLATE_END_ARGS()
+
+public:
+
+	void Construct( const FArguments& InArgs, const TSharedRef< FSlateHyperlinkRun::FWidgetViewModel >& InViewModel )
+	{
+		ViewModel = InViewModel;
+
+		SHyperlink::Construct(
+			SHyperlink::FArguments()
+			.Text( InArgs._Text )
+			.Style( InArgs._Style )
+			.Padding( FMargin(0))
+			.OnNavigate( InArgs._OnNavigate )
+		);
+	}
+
+	virtual void OnMouseEnter( const FGeometry& MyGeometry, const FPointerEvent& MouseEvent ) OVERRIDE
+	{
+		SHyperlink::OnMouseEnter( MyGeometry, MouseEvent );
+		ViewModel->SetIsHovered( true );
+	}
+
+	virtual void OnMouseLeave( const FPointerEvent& MouseEvent ) OVERRIDE
+	{
+		SHyperlink::OnMouseLeave( MouseEvent );
+		ViewModel->SetIsHovered( false );
+	}
+
+	virtual FReply OnMouseButtonDown( const FGeometry& MyGeometry, const FPointerEvent& MouseEvent ) OVERRIDE
+	{
+		FReply Reply = SHyperlink::OnMouseButtonDown( MyGeometry, MouseEvent );
+		ViewModel->SetIsPressed( bIsPressed );
+
+		return Reply;
+	}
+
+	virtual FReply OnMouseButtonUp( const FGeometry& MyGeometry, const FPointerEvent& MouseEvent ) OVERRIDE
+	{
+		FReply Reply = SHyperlink::OnMouseButtonUp( MyGeometry, MouseEvent );
+		ViewModel->SetIsPressed( bIsPressed );
+
+		return Reply;
+	}
+
+	virtual bool IsHovered() const OVERRIDE
+	{
+		return ViewModel->IsHovered();
+	}
+
+	virtual bool IsPressed() const OVERRIDE
+	{
+		return ViewModel->IsPressed();
+	}
+
+
+private:
+
+	TSharedPtr< FSlateHyperlinkRun::FWidgetViewModel > ViewModel;
+};
diff --git a/Engine/Source/Runtime/Slate/Public/Widgets/Input/SRotatorInputBox.h b/Engine/Source/Runtime/Slate/Public/Widgets/Input/SRotatorInputBox.h
new file mode 100644
--- /dev/null
+++ b/Engine/Source/Runtime/Slate/Public/Widgets/Input/SRotatorInputBox.h
@@ -0,0 +1,69 @@
+// Copyright 1998-2014 Epic Games, Inc. All Rights Reserved.
+
+#pragma once
+
+
+
+/**
+ * FRotator Slate control
+ */
+class SLATE_API SRotatorInputBox : public SCompoundWidget
+{
+public:
+	SLATE_BEGIN_ARGS( SRotatorInputBox )
+		: _bColorAxisLabels( false )
+		, _Font( FCoreStyle::Get().GetFontStyle("NormalFont") )
+		, _AllowSpin(true)
+		{}
+		/** Roll component of the rotator */
+		SLATE_ATTRIBUTE( TOptional<float>, Roll )
+
+		/** Pitch component of the rotator */
+		SLATE_ATTRIBUTE( TOptional<float>, Pitch )
+
+		/** Yaw component of the rotator */
+		SLATE_ATTRIBUTE( TOptional<float>, Yaw )
+
+		/** Should the axis labels be colored */
+		SLATE_ARGUMENT( bool, bColorAxisLabels )		
+
+		/** Font to use for the text in this box */
+		SLATE_ATTRIBUTE( FSlateFontInfo, Font )
+
+		/** Whether or not values can be spun or if they should be typed in */
+		SLATE_ARGUMENT( bool, AllowSpin )
+
+		/** Called when the pitch value is changed */
+		SLATE_EVENT( FOnFloatValueChanged, OnPitchChanged )
+
+		/** Called when the yaw value is changed */
+		SLATE_EVENT( FOnFloatValueChanged, OnYawChanged )
+
+		/** Called when the roll value is changed */
+		SLATE_EVENT( FOnFloatValueChanged, OnRollChanged )
+
+		/** Called when the pitch value is committed */
+		SLATE_EVENT( FOnFloatValueCommitted, OnPitchCommitted )
+
+		/** Called when the yaw value is committed */
+		SLATE_EVENT( FOnFloatValueCommitted, OnYawCommitted )
+
+		/** Called when the roll value is committed */
+		SLATE_EVENT( FOnFloatValueCommitted, OnRollCommitted )
+
+		/** Called when the slider begins to move on any axis */
+		SLATE_EVENT( FSimpleDelegate, OnBeginSliderMovement )
+
+		/** Called when the slider for any axis is released */
+		SLATE_EVENT( FOnFloatValueChanged, OnEndSliderMovement )
+
+	SLATE_END_ARGS()
+
+
+	/**
+	 * Construct this widget
+	 *
+	 * @param	InArgs	The declaration data for this widget
+	 */
+	void Construct( const FArguments& InArgs );
+};
diff --git a/Engine/Source/Runtime/Slate/Public/Widgets/Input/SSlider.h b/Engine/Source/Runtime/Slate/Public/Widgets/Input/SSlider.h
new file mode 100644
--- /dev/null
+++ b/Engine/Source/Runtime/Slate/Public/Widgets/Input/SSlider.h
@@ -0,0 +1,144 @@
+// Copyright 1998-2014 Epic Games, Inc. All Rights Reserved.
+
+#pragma once
+
+
+/**
+ * A Slate slider control is a linear scale and draggable handle.
+ */
+class SLATE_API SSlider
+	: public SLeafWidget
+{
+public:
+
+	SLATE_BEGIN_ARGS(SSlider)
+		: _IndentHandle(true)
+		, _Locked(false)
+		, _Orientation(EOrientation::Orient_Horizontal)
+		, _SliderBarColor(FLinearColor::White)
+		, _SliderHandleColor(FLinearColor::White)
+		, _Style(&FCoreStyle::Get().GetWidgetStyle<FSliderStyle>("Slider"))
+		, _Value(1.f)
+		, _OnMouseCaptureBegin()
+		, _OnMouseCaptureEnd()
+		, _OnValueChanged()
+		{ }
+
+		/** Whether the slidable area should be indented to fit the handle. */
+		SLATE_ATTRIBUTE( bool, IndentHandle )
+
+		/** Whether the handle is interactive or fixed. */
+		SLATE_ATTRIBUTE( bool, Locked )
+
+		/** The slider's orientation. */
+		SLATE_ARGUMENT( EOrientation, Orientation)
+
+		/** The color to draw the slider bar in. */
+		SLATE_ATTRIBUTE( FSlateColor, SliderBarColor )
+
+		/** The color to draw the slider handle in. */
+		SLATE_ATTRIBUTE( FSlateColor, SliderHandleColor )
+
+		/** The style used to draw the slider. */
+		SLATE_STYLE_ARGUMENT( FSliderStyle, Style )
+
+		/** The volume value to display. */
+		SLATE_ATTRIBUTE( float, Value )
+
+		/** Invoked when the mouse is pressed and a capture begins. */
+		SLATE_EVENT(FSimpleDelegate, OnMouseCaptureBegin)
+
+		/** Invoked when the mouse is released and a capture ends. */
+		SLATE_EVENT(FSimpleDelegate, OnMouseCaptureEnd)
+
+
+		/** Called when the value is changed by slider or typing. */
+		SLATE_EVENT( FOnFloatValueChanged, OnValueChanged )
+
+	SLATE_END_ARGS()
+
+	/**
+	 * Construct the widget.
+	 * 
+	 * @param InDeclaration - A declaration from which to construct the widget.
+	 */
+	void Construct( const SSlider::FArguments& InDeclaration );
+
+	/**
+	 * Gets the slider's value.
+	 *
+	 * @return Slider value.
+	 */
+	float GetValue() const
+	{
+		return ValueAttribute.Get();
+	}
+
+public:
+
+	virtual int32 OnPaint( const FGeometry& AllottedGeometry, const FSlateRect& MyClippingRect, FSlateWindowElementList& OutDrawElements, int32 LayerId, const FWidgetStyle& InWidgetStyle, bool bParentEnabled ) const OVERRIDE;
+
+	virtual FVector2D ComputeDesiredSize() const;
+
+	virtual FReply OnMouseButtonDown( const FGeometry& MyGeometry, const FPointerEvent& MouseEvent ) OVERRIDE;
+	
+	virtual FReply OnMouseButtonUp( const FGeometry& MyGeometry, const FPointerEvent& MouseEvent ) OVERRIDE;
+	
+	virtual FReply OnMouseMove( const FGeometry& MyGeometry, const FPointerEvent& MouseEvent ) OVERRIDE;
+
+protected:
+
+	/**
+	 * Commits the specified slider value.
+	 *
+	 * @param NewValue - The value to commit.
+	 */
+	void CommitValue(float NewValue);
+
+	/**
+	 * Calculates the new value based on the given absolute coordinates.
+	 *
+	 * @param MyGeometry - The slider's geometry.
+	 * @param AbsolutePosition - The absolute position of the slider.
+	 *
+	 * @return The new value.
+	 */
+	float PositionToValue( const FGeometry& MyGeometry, const FVector2D& AbsolutePosition );
+
+private:
+
+	// Holds the image to use when the slider handle is in its disabled state.
+	const FSlateBrush* DisabledHandleImage;
+
+	// Holds the image to use when the slider handle is in its normal state.
+	const FSlateBrush* NormalHandleImage;
+
+	// Holds a flag indicating whether the slideable area should be indented to fit the handle.
+	TAttribute<bool> IndentHandle;
+
+	// Holds a flag indicating whether the slider is locked.
+	TAttribute<bool> LockedAttribute;
+
+	// Holds the slider's orientation.
+	EOrientation Orientation;
+
+	// Holds the color of the slider bar.
+	TAttribute<FSlateColor> SliderBarColor;
+
+	// Holds the color of the slider handle.
+	TAttribute<FSlateColor> SliderHandleColor;
+
+	// Holds the slider's current value.
+	TAttribute<float> ValueAttribute;
+
+private:
+
+	// Holds a delegate that is executed when the mouse is pressed and a capture begins.
+	FSimpleDelegate OnMouseCaptureBegin;
+
+	// Holds a delegate that is executed when the mouse is let up and a capture ends.
+	FSimpleDelegate OnMouseCaptureEnd;
+
+	// Holds a delegate that is executed when the slider's value changed.
+	FOnFloatValueChanged OnValueChanged;
+};
diff --git a/Engine/Source/Runtime/Slate/Public/Widgets/Input/SSpinBox.h b/Engine/Source/Runtime/Slate/Public/Widgets/Input/SSpinBox.h
new file mode 100644
--- /dev/null
+++ b/Engine/Source/Runtime/Slate/Public/Widgets/Input/SSpinBox.h
@@ -0,0 +1,690 @@
+// Copyright 1998-2014 Epic Games, Inc. All Rights Reserved.
+
+#pragma once
+
+/**
+ * A Slate SpinBox resembles traditional spin boxes in that it is a widget that provides
+ * keyboard-based and mouse-based manipulation of a numeric value.
+ * Mouse-based manipulation: drag anywhere on the spinbox to change the value.
+ * Keyboard-based manipulation: click on the spinbox to enter text mode.
+ */
+template<typename NumericType>
+class SSpinBox : public SCompoundWidget
+{
+public:
+	/** Notification for numeric value change */
+	DECLARE_DELEGATE_OneParam( FOnValueChanged, NumericType );
+
+	/** Notification for numeric value committed */
+	DECLARE_DELEGATE_TwoParams( FOnValueCommitted, NumericType, ETextCommit::Type);
+
+	SLATE_BEGIN_ARGS(SSpinBox<NumericType>)
+		: _Style(&FCoreStyle::Get().GetWidgetStyle<FSpinBoxStyle>("SpinBox"))
+		, _Value(0)
+		, _MinValue(0)
+		, _MaxValue(10)
+		, _Delta(0)
+		, _SliderExponent(1)
+		, _Font( FCoreStyle::Get().GetFontStyle( TEXT( "NormalFont" ) ) )
+		, _ContentPadding(  FMargin( 2.0f, 1.0f) )
+		, _OnValueChanged()
+		, _OnValueCommitted()
+		, _ClearKeyboardFocusOnCommit( false )
+		, _SelectAllTextOnCommit( true )
+		{}
+	
+		/** The style used to draw this spinbox */
+		SLATE_STYLE_ARGUMENT( FSpinBoxStyle, Style )
+
+		/** The value to display */
+		SLATE_ATTRIBUTE( NumericType, Value )
+		/** The minimum value that can be entered into the text edit box */
+		SLATE_ATTRIBUTE( TOptional< NumericType >, MinValue )
+		/** The maximum value that can be entered into the text edit box */
+		SLATE_ATTRIBUTE( TOptional< NumericType >, MaxValue )
+		/** The minimum value that can be specified by using the slider, defaults to MinValue */
+		SLATE_ATTRIBUTE( TOptional< NumericType >, MinSliderValue )
+		/** The maximum value that can be specified by using the slider, defaults to MaxValue */
+		SLATE_ATTRIBUTE( TOptional< NumericType >, MaxSliderValue )
+		/** Delta to increment the value as the slider moves.  If not specified will determine automatically */
+		SLATE_ATTRIBUTE( NumericType, Delta )
+		/** Use exponential scale for the slider */
+		SLATE_ATTRIBUTE( float, SliderExponent )
+		/** Font used to display text in the slider */
+		SLATE_ATTRIBUTE( FSlateFontInfo, Font )
+		/** Padding to add around this widget and its internal widgets */
+		SLATE_ARGUMENT( FMargin, ContentPadding )
+		/** Called when the value is changed by slider or typing */
+		SLATE_EVENT( FOnValueChanged, OnValueChanged )
+		/** Called when the value is committed (by pressing enter) */
+		SLATE_EVENT( FOnValueCommitted, OnValueCommitted )
+		/** Called right before the slider begins to move */
+		SLATE_EVENT( FSimpleDelegate, OnBeginSliderMovement )
+		/** Called right after the slider handle is released by the user */
+		SLATE_EVENT( FOnValueChanged, OnEndSliderMovement )
+		/** Whether to clear keyboard focus when pressing enter to commit changes */
+		SLATE_ATTRIBUTE( bool, ClearKeyboardFocusOnCommit )
+		/** Whether to select all text when pressing enter to commit changes */
+		SLATE_ATTRIBUTE( bool, SelectAllTextOnCommit )
+
+	SLATE_END_ARGS()
+
+	SSpinBox()
+	{
+	}
+		
+	/**
+	 * Construct the widget
+	 * 
+	 * @param InArgs   A declaration from which to construct the widget
+	 */
+	void Construct( const FArguments& InArgs )
+	{
+		check(InArgs._Style);
+
+		ForegroundColor = InArgs._Style->ForegroundColor;
+		ValueAttribute = InArgs._Value;
+		OnValueChanged = InArgs._OnValueChanged;
+		OnValueCommitted = InArgs._OnValueCommitted;
+		OnBeginSliderMovement = InArgs._OnBeginSliderMovement;
+		OnEndSliderMovement = InArgs._OnEndSliderMovement;
+	
+		MinValue = InArgs._MinValue;
+		MaxValue = InArgs._MaxValue;
+		MinSliderValue = (InArgs._MinSliderValue.Get().IsSet()) ? InArgs._MinSliderValue : MinValue;
+		MaxSliderValue = (InArgs._MaxSliderValue.Get().IsSet()) ? InArgs._MaxSliderValue : MaxValue;
+
+		bUnlimitedSpinRange = !((InArgs._MinValue.Get().IsSet() && InArgs._MaxValue.Get().IsSet()) || (InArgs._MinSliderValue.Get().IsSet() && InArgs._MaxSliderValue.Get().IsSet()));
+	
+		SliderExponent = InArgs._SliderExponent;
+
+		DistanceDragged = 0.0f;
+		CachedExternalValue = 0;
+		InternalValue = 0;
+		PreDragValue = 0;
+
+		Delta = InArgs._Delta.Get();
+	
+		BackgroundHoveredBrush = &InArgs._Style->HoveredBackgroundBrush;
+		BackgroundBrush = &InArgs._Style->BackgroundBrush;
+		ActiveFillBrush = &InArgs._Style->ActiveFillBrush;
+		InactiveFillBrush = &InArgs._Style->InactiveFillBrush;
+		const FMargin& TextMargin = InArgs._Style->TextPadding;
+
+		bDragging = false;
+		CachedExternalValue = ValueAttribute.Get();
+		InternalValue = ValueAttribute.Get();
+
+		this->ChildSlot
+		.Padding( InArgs._ContentPadding )
+		[
+			SNew(SHorizontalBox)
+			+SHorizontalBox::Slot()
+			.FillWidth(1.0f)
+			.Padding( TextMargin )
+			.HAlign(HAlign_Fill) 
+			.VAlign(VAlign_Center)
+			[
+				SAssignNew(TextBlock, STextBlock)
+				.Font(InArgs._Font)
+				.Text( this, &SSpinBox<NumericType>::GetValueAsString )
+			]
+			+SHorizontalBox::Slot()
+			.FillWidth(1.0f)
+			.Padding( TextMargin )
+			.HAlign(HAlign_Fill) 
+			.VAlign(VAlign_Center) 
+			[
+				SAssignNew(EditableText, SEditableText)
+				.Visibility( EVisibility::Collapsed )
+				.Font(InArgs._Font)
+				.SelectAllTextWhenFocused( true )
+				.Text( this, &SSpinBox<NumericType>::GetValueAsText )
+				.OnIsTypedCharValid_Static(TCharacterFilter<NumericType>::IsValid)
+				.OnTextCommitted( this, &SSpinBox<NumericType>::TextField_OnTextCommitted )
+				.ClearKeyboardFocusOnCommit( InArgs._ClearKeyboardFocusOnCommit )
+				.SelectAllTextOnCommit( InArgs._SelectAllTextOnCommit )
+			]			
+			+SHorizontalBox::Slot()
+			.AutoWidth()
+			.HAlign(HAlign_Fill) 
+			.VAlign(VAlign_Center) 
+			[
+				SNew(SImage)
+				.Image( &InArgs._Style->ArrowsImage )
+				.ColorAndOpacity( FSlateColor::UseForeground() )
+			]
+
+		];
+	}
+	
+	/**
+	* The widget should respond by populating the OutDrawElements array with FDrawElements 
+	* that represent it and any of its children.
+	*
+	* @param AllottedGeometry  The FGeometry that describes an area in which the widget should appear.
+	* @param MyClippingRect    The clipping rectangle allocated for this widget and its children.
+	* @param OutDrawElements   A list of FDrawElements to populate with the output.
+	* @param LayerId           The Layer onto which this widget should be rendered.
+	* @param InColorAndOpacity Color and Opacity to be applied to all the descendants of the widget being painted
+	* @param bParentEnabled	True if the parent of this widget is enabled.
+	*
+	* @return The maximum layer ID attained by this widget or any of its children.
+	*/
+	virtual int32 OnPaint( const FGeometry& AllottedGeometry, const FSlateRect& MyClippingRect, FSlateWindowElementList& OutDrawElements, int32 LayerId, const FWidgetStyle& InWidgetStyle, bool bParentEnabled ) const OVERRIDE
+	{
+		const bool bActiveFeedback = IsHovered() || bDragging;
+
+		const FSlateBrush* BackgroundImage = bActiveFeedback ?
+			BackgroundHoveredBrush :
+			BackgroundBrush;
+
+		const FSlateBrush* FillImage = bActiveFeedback ?
+			ActiveFillBrush :
+			InactiveFillBrush;
+
+		const int32 BackgroundLayer = LayerId;
+
+		const bool bEnabled = ShouldBeEnabled( bParentEnabled );
+		const ESlateDrawEffect::Type DrawEffects = bEnabled ? ESlateDrawEffect::None : ESlateDrawEffect::DisabledEffect;
+
+		FSlateDrawElement::MakeBox(
+			OutDrawElements,
+			BackgroundLayer,
+			AllottedGeometry.ToPaintGeometry(),
+			BackgroundImage,
+			MyClippingRect,
+			DrawEffects,
+			InWidgetStyle.GetColorAndOpacityTint()
+			);
+
+		const int32 FilledLayer = BackgroundLayer + 1;
+
+		//if there is a spin range limit, draw the filler bar
+		if (!bUnlimitedSpinRange)
+		{
+			double Value = ValueAttribute.Get();
+			if( Delta != 0.0f )
+			{
+				Value= Snap(Value, Delta); // snap floating point value to nearest Delta
+			}
+
+			float FractionFilled = Fraction(Value, GetMinSliderValue(), GetMaxSliderValue());
+			const float CachedSliderExponent = SliderExponent.Get();
+			if (CachedSliderExponent != 1)
+			{
+				FractionFilled = 1.0f - FMath::Pow( 1.0f - FractionFilled, CachedSliderExponent);
+			}
+			const FVector2D FillSize( AllottedGeometry.Size.X * FractionFilled, AllottedGeometry.Size.Y );
+
+			if ( ! IsInTextMode() )
+			{
+				FSlateDrawElement::MakeBox(
+					OutDrawElements,
+					FilledLayer,
+					AllottedGeometry.ToPaintGeometry(FVector2D(0,0), FillSize),
+					FillImage,
+					MyClippingRect,
+					DrawEffects,
+					InWidgetStyle.GetColorAndOpacityTint()
+					);
+			}
+		}
+
+		return FMath::Max( FilledLayer, SCompoundWidget::OnPaint(AllottedGeometry, MyClippingRect, OutDrawElements, FilledLayer, InWidgetStyle, bEnabled ) );
+	}
+
+	/**
+	 * The system calls this method to notify the widget that a mouse button was pressed within it. This event is bubbled.
+	 *
+	 * @param MyGeometry The Geometry of the widget receiving the event
+	 * @param MouseEvent Information about the input event
+	 *
+	 * @return Whether the event was handled along with possible requests for the system to take action.
+	 */
+	virtual FReply OnMouseButtonDown( const FGeometry& MyGeometry, const FPointerEvent& MouseEvent ) OVERRIDE
+	{
+		if ( MouseEvent.GetEffectingButton() == EKeys::LeftMouseButton )
+		{
+			DistanceDragged = 0;
+			PreDragValue = InternalValue;
+			return FReply::Handled().CaptureMouse( SharedThis(this) ).UseHighPrecisionMouseMovement( SharedThis(this) ).SetKeyboardFocus( SharedThis(this), EKeyboardFocusCause::Mouse);
+		}
+		else
+		{
+
+			return FReply::Unhandled();
+		}
+	}
+	
+	/**
+	 * The system calls this method to notify the widget that a mouse button was release within it. This event is bubbled.
+	 *
+	 * @param MyGeometry The Geometry of the widget receiving the event
+	 * @param MouseEvent Information about the input event
+	 *
+	 * @return Whether the event was handled along with possible requests for the system to take action.
+	 */
+	virtual FReply OnMouseButtonUp( const FGeometry& MyGeometry, const FPointerEvent& MouseEvent ) OVERRIDE
+	{
+		if ( MouseEvent.GetEffectingButton() == EKeys::LeftMouseButton && this->HasMouseCapture() )
+		{
+			if( bDragging )
+			{
+				NotifyValueCommitted();
+			}
+
+			bDragging = false;
+			if ( DistanceDragged < SlateDragStartDistance )
+			{
+				EnterTextMode();
+				return FReply::Handled().ReleaseMouseCapture().SetKeyboardFocus( EditableText.ToSharedRef(), EKeyboardFocusCause::SetDirectly ).SetMousePos(CachedMousePosition);
+			}
+
+			return FReply::Handled().ReleaseMouseCapture().SetMousePos(CachedMousePosition);
+			
+		}
+		else
+		{
+			return FReply::Unhandled();
+		}
+	}
+	
+	/**
+	 * The system calls this method to notify the widget that a mouse moved within it. This event is bubbled.
+	 *
+	 * @param MyGeometry The Geometry of the widget receiving the event
+	 * @param MouseEvent Information about the input event
+	 *
+	 * @return Whether the event was handled along with possible requests for the system to take action.
+	 */
+	virtual FReply OnMouseMove( const FGeometry& MyGeometry, const FPointerEvent& MouseEvent ) OVERRIDE
+	{
+		if ( this->HasMouseCapture() )
+		{
+			if (!bDragging)
+			{
+				DistanceDragged += FMath::Abs(MouseEvent.GetCursorDelta().X);
+				if ( DistanceDragged > SlateDragStartDistance )
+				{
+					ExitTextMode();
+					bDragging = true;
+				}
+
+				if( bDragging )
+				{
+					OnBeginSliderMovement.ExecuteIfBound();
+				}
+
+				// Cache the mouse, even if not dragging cache it.
+				CachedMousePosition = MouseEvent.GetScreenSpacePosition().IntPoint();
+			}
+			else
+			{
+				double NewValue = 0;
+
+				// Increments the spin based on delta mouse movement.
+
+				//if we have a range to draw in
+				if ( !bUnlimitedSpinRange) 
+				{
+					const float CachedSliderExponent = SliderExponent.Get();
+					// The amount currently filled in the spinbox, needs to be calculated to do deltas correctly.
+					float FractionFilled = Fraction(InternalValue, GetMinSliderValue(), GetMaxSliderValue());
+						
+					if (CachedSliderExponent != 1)
+					{
+						FractionFilled = 1.0f - FMath::Pow( 1.0f - FractionFilled, CachedSliderExponent);
+					}
+					FractionFilled *= MyGeometry.GetDrawSize().X;
+
+					// Now add the delta to the fraction filled, this causes the spin.
+					FractionFilled += MouseEvent.GetCursorDelta().X;
+						
+					// Clamp the fraction to be within the bounds of the geometry.
+					FractionFilled = FMath::Clamp(FractionFilled, 0.0f, MyGeometry.GetDrawSize().X);
+
+					// Convert the fraction filled to a percent.
+					float Percent = FMath::Clamp(FractionFilled / MyGeometry.GetDrawSize().X, 0.0f, 1.0f);
+					if (CachedSliderExponent != 1)
+					{
+						// Have to convert the percent to the proper value due to the exponent component to the spin.
+						Percent = 1.0f - FMath::Pow(1.0f - Percent, 1.0 / CachedSliderExponent);
+					}
+
+					NewValue = FMath::LerpStable<double>(GetMinSliderValue(), GetMaxSliderValue(), Percent);
+				}
+				else
+				{
+					//we have no range specified, so increment scaled by our current value		
+					const double CurrentValue = FMath::Clamp<double>(FMath::Abs(InternalValue),1.0,TNumericLimits<NumericType>::Max());
+					const float MouseDelta = FMath::Abs(MouseEvent.GetCursorDelta().X / MyGeometry.GetDrawSize().X);
+					const float Sign = (MouseEvent.GetCursorDelta().X > 0) ? 1.f : -1.f;
+					NewValue = InternalValue + (Sign * MouseDelta * FMath::Pow(CurrentValue, SliderExponent.Get()));
+				}
+			
+				CommitValue( NewValue, CommittedViaSpin, ETextCommit::OnEnter );
+			}
+
+			return FReply::Handled();
+		}
+
+		return FReply::Unhandled();
+	}
+
+	virtual FCursorReply OnCursorQuery( const FGeometry& MyGeometry, const FPointerEvent& CursorEvent ) const OVERRIDE
+	{
+		return bDragging ? 
+			FCursorReply::Cursor( EMouseCursor::None ) :
+			FCursorReply::Cursor( EMouseCursor::ResizeLeftRight );
+	}
+
+	virtual bool SupportsKeyboardFocus() const
+	{
+		// SSpinBox is focusable.
+		return true;
+	}
+
+
+	virtual FReply OnKeyboardFocusReceived( const FGeometry& MyGeometry, const FKeyboardFocusEvent& InKeyboardFocusEvent ) OVERRIDE
+	{
+		if ( InKeyboardFocusEvent.GetCause() == EKeyboardFocusCause::Keyboard || InKeyboardFocusEvent.GetCause() == EKeyboardFocusCause::SetDirectly )
+		{
+			EnterTextMode();
+			return FReply::Handled().SetKeyboardFocus( EditableText.ToSharedRef(), InKeyboardFocusEvent.GetCause() );
+		}
+		else
+		{
+			return FReply::Unhandled();
+		}
+	}
+	
+	virtual FReply OnKeyDown( const FGeometry& MyGeometry, const FKeyboardEvent& InKeyboardEvent ) OVERRIDE
+	{
+		const FKey Key = InKeyboardEvent.GetKey();
+		if ( Key == EKeys::Escape && HasMouseCapture())
+		{
+			bDragging = false;
+
+			InternalValue = PreDragValue;
+			NotifyValueCommitted();
+			return FReply::Handled().ReleaseMouseCapture().SetMousePos(CachedMousePosition);
+		}
+		else if ( Key == EKeys::Up || Key == EKeys::Right )
+		{
+			CommitValue( InternalValue + Delta, CommittedViaSpin, ETextCommit::OnEnter );
+			ExitTextMode();
+			return FReply::Handled();	
+		}
+		else if ( Key == EKeys::Down || Key == EKeys::Left )
+		{
+			CommitValue( InternalValue - Delta, CommittedViaSpin, ETextCommit::OnEnter );
+			ExitTextMode();
+			return FReply::Handled();
+		}
+		else if ( Key == EKeys::Enter )
+		{
+			EnterTextMode();
+			return FReply::Handled().SetKeyboardFocus( EditableText.ToSharedRef(), EKeyboardFocusCause::Keyboard );
+		}
+		else
+		{
+			return FReply::Unhandled();
+		}
+	}
+	
+	virtual bool HasKeyboardFocus() const
+	{
+		// The spinbox is considered focused when we are typing it text.
+		return SCompoundWidget::HasKeyboardFocus() || (EditableText.IsValid() && EditableText->HasKeyboardFocus());
+	}
+
+protected:
+	/** Make the spinbox switch to keyboard-based input mode. */
+	void EnterTextMode()
+	{
+		TextBlock->SetVisibility( EVisibility::Collapsed );
+		EditableText->SetVisibility( EVisibility::Visible );
+	}
+	
+	/** Make the spinbox switch to mouse-based input mode. */
+	void ExitTextMode()
+	{
+		TextBlock->SetVisibility( EVisibility::Visible );
+		EditableText->SetVisibility( EVisibility::Collapsed );
+	}
+
+	/** @return the value being observed by the spinbox as a string */
+	FString GetValueAsString() const
+	{
+		auto CurrentValue = ValueAttribute.Get();
+		if( Delta != 0 )
+		{
+			CurrentValue = (NumericType)Snap(CurrentValue, Delta);
+		}
+		
+		return TTypeToString<NumericType>::ToSanitizedString(CurrentValue);
+	}
+
+	/** @return the value being observed by the spinbox as FText - todo: spinbox FText support (reimplement me) */
+	FText GetValueAsText() const
+	{
+		return FText::FromString(GetValueAsString());
+	}
+	
+	/**
+	 * Invoked when the text field commits its text.
+	 *
+	 * @param NewText		The value of text coming from the editable text field.
+	 * @param CommitInfo	Information about the source of the commit
+	 */
+	void TextField_OnTextCommitted( const FText& NewText, ETextCommit::Type CommitInfo )
+	{
+		if (CommitInfo != ETextCommit::OnEnter)
+		{
+			ExitTextMode();
+		}
+
+		NumericType NewValue = 0;
+		bool bEvalResult = false;
+		FString NewString = NewText.ToString();
+		
+		if (NewString.IsNumeric())
+		{
+			TTypeFromString<NumericType>::FromString( NewValue, *NewString );
+			bEvalResult = true;
+		}
+		else
+		{
+			float FloatValue = 0.f;
+			bEvalResult = FMath::Eval( *NewString, FloatValue  );
+			NewValue = NumericType(FloatValue);
+		}
+				
+		if (bEvalResult)
+		{
+			CommitValue( NewValue, CommittedViaTypeIn, CommitInfo );
+		}
+	}
+
+
+	/** How user changed the value in the spinbox */
+	enum ECommitMethod
+	{
+		CommittedViaSpin,
+		CommittedViaTypeIn	
+	};
+
+	/**
+	 * Call this method when the user's interaction has changed the value
+	 *
+	 * @param NewValue               Value resulting from the user's interaction
+	 * @param CommitMethod           Did the user type in the value or spin to it.
+	 * @param OriginalCommitInfo	 If the user typed in the value, information about the source of the commit
+	 */
+	void CommitValue( double NewValue, ECommitMethod CommitMethod, ETextCommit::Type OriginalCommitInfo )
+	{
+		if( CommitMethod == CommittedViaSpin )
+		{
+			NewValue = FMath::Clamp<double>( NewValue, GetMinSliderValue(), GetMaxSliderValue() );
+		}
+
+		NewValue = FMath::Clamp<double>( NewValue, GetMinValue(), GetMaxValue() );
+
+		if ( !ValueAttribute.IsBound() )
+		{
+			if (TIsIntegralType<NumericType>::Value)
+			{
+				// Round floating point to avoid 'jumping' when spinning down integral values 
+				ValueAttribute.Set( (NumericType)FMath::FloorDouble(NewValue+0.5) );
+			}
+			else
+			{
+				ValueAttribute.Set( (NumericType)NewValue );
+			}
+		}
+
+		auto CurrentValue = ValueAttribute.Get();
+
+		// If not in spin mode, there is no need to jump to the value from the external source, continue to use the committed value.
+		if(CommitMethod == CommittedViaSpin)
+		{
+			// This will detect if an external force has changed the value. Internally it will abandon the delta calculated this tick and update the internal value instead.
+			if(CurrentValue != CachedExternalValue)
+			{
+				NewValue = CurrentValue;
+			}
+		}
+
+		// Update the internal value, this needs to be done before rounding.
+		InternalValue = NewValue;
+
+		// If needed, round this value to the delta. Internally the value is not held to the Delta but externally it appears to be.
+		if ( CommitMethod == CommittedViaSpin )
+		{
+			if( Delta != 0 )
+			{
+				NewValue = Snap(NewValue, Delta); // snap numeric point value to nearest Delta
+			}
+		}
+
+		if( CommitMethod == CommittedViaTypeIn )
+		{
+			OnValueCommitted.ExecuteIfBound( NumericType(NewValue), OriginalCommitInfo );
+		}
+
+		OnValueChanged.ExecuteIfBound( NumericType(NewValue) );
+
+		// Update the cache of the external value to what the user believes the value is now.
+		CachedExternalValue = ValueAttribute.Get();
+
+		// This ensures that dragging is cleared if focus has been removed from this widget in one of the delegate calls, such as when spawning a modal dialog.
+		if ( !this->HasMouseCapture() )
+		{
+			bDragging = false;
+		}
+	}
+	
+	void NotifyValueCommitted() const
+	{
+		const auto CurrentValue = (NumericType)InternalValue; // internal value should be snapped and clamped at this point
+		OnValueCommitted.ExecuteIfBound( CurrentValue, ETextCommit::OnEnter );
+		OnEndSliderMovement.ExecuteIfBound( CurrentValue );
+	}
+
+	/** @return true when we are in keyboard-based input mode; false otherwise */
+	bool IsInTextMode() const
+	{
+		return ( EditableText->GetVisibility() == EVisibility::Visible );
+	}
+
+	/** Calculates range fraction. Possible to use on full numeric range  */
+	static float Fraction(double InValue, NumericType InMinValue, NumericType InMaxValue)
+	{
+		const double HalfMax = InMaxValue*0.5;
+		const double HalfMin = InMinValue*0.5;
+		const double HalfVal = InValue*0.5;
+
+		return (float)FMath::Clamp((HalfVal - HalfMin)/(HalfMax - HalfMin), 0.0, 1.0);
+	}
+
+	/** Snaps value to a nearest grid multiple. Clamps value if it's outside numeric range  */
+	static double Snap(double InValue, double InGrid)
+	{
+		double Result = FMath::GridSnap(InValue, InGrid);
+		return FMath::Clamp<double>(Result, TNumericLimits<NumericType>::Lowest(), TNumericLimits<NumericType>::Max());
+	}
+	
+	/*
+	 * Helper structure to check if a given character is supported by underlying numeric type
+	 */
+	template<typename T, typename U = void>
+	struct TCharacterFilter
+	{
+		static bool IsValid(const TCHAR InChar)
+		{
+			static FString ValidChars(TEXT("1234567890-.")); // for integers and doubles
+			int32 OutUnused;
+			return ValidChars.FindChar( InChar, OutUnused );
+		}
+	};
+
+	template<typename U> // nested partial specialization is allowed
+	struct TCharacterFilter<float, U>
+	{
+		static bool IsValid(const TCHAR InChar)
+		{
+			static FString ValidChars(TEXT("1234567890()-+=\\/.,*^%%"));
+			int32 OutUnused;
+			return ValidChars.FindChar( InChar, OutUnused );		
+		}
+	};
+
+private:
+	TAttribute<NumericType> ValueAttribute;
+	FOnValueChanged OnValueChanged;
+	FOnValueCommitted OnValueCommitted;
+	FSimpleDelegate OnBeginSliderMovement;
+	FOnValueChanged OnEndSliderMovement;
+	TSharedPtr<STextBlock> TextBlock;
+	TSharedPtr<SEditableText> EditableText;
+
+	/** True when no range is specified, spinner can be spun indefinitely */
+	bool bUnlimitedSpinRange;
+
+	const FSlateBrush* BackgroundHoveredBrush;
+	const FSlateBrush* BackgroundBrush;
+	const FSlateBrush* ActiveFillBrush;
+	const FSlateBrush* InactiveFillBrush;
+
+	float DistanceDragged;
+	NumericType Delta;
+	TAttribute< TOptional<NumericType> > MinValue;
+	TAttribute< TOptional<NumericType> > MaxValue;
+	TAttribute< TOptional<NumericType> > MinSliderValue;
+	TAttribute< TOptional<NumericType> > MaxSliderValue;
+
+	NumericType GetMinValue() const { return MinValue.Get().Get(TNumericLimits<NumericType>::Lowest()); }
+	NumericType GetMaxValue() const { return MaxValue.Get().Get(TNumericLimits<NumericType>::Max()); }
+	NumericType GetMinSliderValue() const { return MinSliderValue.Get().Get(TNumericLimits<NumericType>::Lowest()); }
+	NumericType GetMaxSliderValue() const { return MaxSliderValue.Get().Get(TNumericLimits<NumericType>::Max()); }
+
+	TAttribute<float> SliderExponent;
+
+	bool bDragging;
+	
+	/** Cached mouse position to restore after scrolling. */
+	FIntPoint CachedMousePosition;
+
+	/** This value represents what the spinbox believes the value to be, regardless of delta and the user binding to an int. 
+		The spinbox will always count using floats between values, this is important to keep it flowing smoothly and feeling right, 
+		and most importantly not conflicting with the user truncating the value to an int. */
+	double InternalValue;
+
+	/** The state of InternalValue before a drag operation was started */
+	double PreDragValue;
+
+	/** This is the cached value the user believes it to be (usually different due to truncation to an int). Used for identifying 
+		external forces on the spinbox and syncing the internal value to them. Synced when a value is committed to the spinbox. */
+	NumericType CachedExternalValue;
+};
diff --git a/Engine/Source/Runtime/Slate/Public/Widgets/Input/SSuggestionTextBox.h b/Engine/Source/Runtime/Slate/Public/Widgets/Input/SSuggestionTextBox.h
new file mode 100644
--- /dev/null
+++ b/Engine/Source/Runtime/Slate/Public/Widgets/Input/SSuggestionTextBox.h
@@ -0,0 +1,302 @@
+// Copyright 1998-2014 Epic Games, Inc. All Rights Reserved.
+
+/*=============================================================================
+	SSuggestionTextBox.h: Declares the SSuggestionTextBox class.
+=============================================================================*/
+
+#pragma once
+
+
+/**
+ * Delegate type for showing the input history.
+ *
+ * The first parameter will contain the history of input strings to show.
+ */
+DECLARE_DELEGATE_OneParam(FOnShowingHistory, OUT TArray<FString>&)
+
+/**
+ * Delegate type for showing a list of suggestions.
+ *
+ * The first parameter is the current input text.
+ * The second parameter will contain the list of suggestions to show.
+ */
+DECLARE_DELEGATE_TwoParams(FOnShowingSuggestions, const FString&, OUT TArray<FString>&)
+
+
+/**
+ * Implements an editable text box that can show auto-complete histories and suggestions lists.
+ */
+class SLATE_API SSuggestionTextBox
+	: public SCompoundWidget
+{
+
+public:
+
+	SLATE_BEGIN_ARGS(SSuggestionTextBox)
+		: _BackgroundColor(FLinearColor::White)
+		, _BackgroundImage(FCoreStyle::Get().GetBrush("SuggestionTextBox.Background"))
+		, _ClearKeyboardFocusOnCommit(true)
+		, _ErrorReporting()
+		, _Font( FCoreStyle::Get().GetFontStyle( TEXT( "NormalFont" ) ) )
+		, _ForegroundColor( FCoreStyle::Get().GetSlateColor( "InvertedForeground" ) )
+		, _HintText()
+		, _IsCaretMovedWhenGainFocus (true)
+		, _MinDesiredWidth(0.0f)
+		, _RevertTextOnEscape(false)
+		, _SelectAllTextOnCommit(false)
+		, _SelectAllTextWhenFocused(false)
+		, _TextStyle( &FCoreStyle::Get().GetWidgetStyle<FEditableTextBoxStyle>("NormalEditableTextBox") )
+		, _SuggestionTextStyle( &FCoreStyle::Get().GetWidgetStyle<FTextBlockStyle>("SuggestionTextBox.Text") )
+		, _SuggestionListMaxHeight(250.0)
+		, _Text()
+		{ }
+
+		/**
+		 * The color of the background/border around the editable text.
+		 */
+		SLATE_ATTRIBUTE(FSlateColor, BackgroundColor)
+
+		/**
+		 * The image of the background/border around the editable text.
+		 */
+		SLATE_ARGUMENT(const FSlateBrush*, BackgroundImage)
+
+		/**
+		 * Whether to clear keyboard focus when pressing enter to commit changes.
+		 */
+		SLATE_ATTRIBUTE(bool, ClearKeyboardFocusOnCommit)
+
+		/**
+		 * Provides an alternative mechanism for error reporting.
+		 */
+		SLATE_ARGUMENT(TSharedPtr<class IErrorReportingWidget>, ErrorReporting)
+
+		/**
+		 * Font color and opacity.
+		 */
+		SLATE_ATTRIBUTE(FSlateFontInfo, Font)
+
+		/**
+		 * Text color and opacity.
+		 */
+		SLATE_ATTRIBUTE(FSlateColor, ForegroundColor)
+
+		/**
+		 * Hint text that appears when there is no text in the text box.
+		 */
+		SLATE_ATTRIBUTE(FText, HintText)
+
+		/**
+		 * Workaround as we loose focus when the auto completion closes.
+		 */
+		SLATE_ATTRIBUTE(bool, IsCaretMovedWhenGainFocus)
+
+		/**
+		 * Minimum width that a text block should be.
+		 */
+		SLATE_ATTRIBUTE(float, MinDesiredWidth)
+
+		/**
+		 * Called before the history list is shown.
+		 */
+		SLATE_EVENT(FOnShowingHistory, OnShowingHistory)
+
+		/**
+		 * Called before the suggestion list is shown.
+		 */
+		SLATE_EVENT(FOnShowingSuggestions, OnShowingSuggestions)
+
+		/**
+		 * Called whenever the text is changed interactively by the user.
+		 */
+		SLATE_EVENT(FOnTextChanged, OnTextChanged)
+
+		/**
+		 * Called when the text has been committed.
+		 */
+		SLATE_EVENT(FOnTextCommitted, OnTextCommitted)
+
+		/**
+		 * Whether to allow the user to back out of changes when they press the escape key.
+		 */
+		SLATE_ATTRIBUTE(bool, RevertTextOnEscape)
+
+		/**
+		 * Whether to select all text when pressing enter to commit changes.
+		 */
+		SLATE_ATTRIBUTE(bool, SelectAllTextOnCommit)
+
+		/**
+		 * Whether to select all text when the user clicks to give focus on the widget.
+		 */
+		SLATE_ATTRIBUTE(bool, SelectAllTextWhenFocused)
+
+		/**
+		 * The styling of the text box.
+		 */
+		SLATE_STYLE_ARGUMENT(FEditableTextBoxStyle, TextStyle)
+
+		/**
+		 * The styling of the suggestion text.
+		 */
+		SLATE_STYLE_ARGUMENT(FTextBlockStyle, SuggestionTextStyle)
+
+		/**
+		 * The maximum height of the suggestion list.
+		 */
+		SLATE_ATTRIBUTE(float, SuggestionListMaxHeight)
+
+		/**
+		 * Sets the text content for this editable text box widget.
+		 */
+		SLATE_ATTRIBUTE(FText, Text)
+
+	SLATE_END_ARGS()
+
+
+public:
+
+	/**
+	 * Default constructor.
+	 */
+	SSuggestionTextBox( );
+
+
+public:
+
+	/**
+	 * Construct this widget.  Called by the SNew() Slate macro.
+	 *
+	 * @param InArgs - Declaration used by the SNew() macro to construct this widget.
+	 */
+	void Construct( const FArguments& InArgs );
+
+	/**
+	 * Returns the text string.
+	 *
+	 * @return The text string.
+	 */
+	const FText& GetText( ) const
+	{
+		return TextBox->GetText();
+	}
+
+	/**
+	 * Sets an optional error string.
+	 *
+	 * If InError is a non-empty string the TextBox will the ErrorReporting provided during construction.
+	 * If no error reporting was provided, the TextBox will create a default error reporter.
+	 *
+	 * @param InError - The error string to set. 
+	 */
+	void SetError( const FString& InError )
+	{
+		TextBox->SetError(InError);
+	}
+
+	/**
+	 * Sets the text string currently being edited 
+	 *
+	 * @param InNewText - The new text string.
+	 */
+	void SetText( const TAttribute< FText >& InNewText );
+
+
+protected:
+	
+	/**
+	 * Clears the list of suggestions and hides the suggestions list.
+	 */
+	void ClearSuggestions( );
+
+	/**
+	 * Sets the keyboard focus to the text box.
+	 */
+	void FocusTextBox( );
+
+	/**
+	 * Gets the string value of the currently selected suggestion.
+	 *
+	 * @return Suggestion string, or empty if no suggestion is selected.
+	 */
+	FString GetSelectedSuggestionString( ) const;
+
+	/**
+	 * Highlights the selected suggestion in the suggestion list.
+	 */
+	void MarkActiveSuggestion( );
+
+	/**
+	 * Sets the list of suggestions.
+	 *
+	 * @param SuggestionStrings - The suggestion strings.
+	 * @param InHistoryMode - Whether the suggestions represent the input history.
+	 */
+	void SetSuggestions( TArray<FString>& SuggestionStrings, bool InHistoryMode );
+
+
+protected:
+
+	virtual void OnKeyboardFocusLost( const FKeyboardFocusEvent& InKeyboardFocusEvent ) OVERRIDE;
+
+	virtual FReply OnKeyDown( const FGeometry& MyGeometry, const FKeyboardEvent& KeyboardEvent ) OVERRIDE;
+
+	virtual bool SupportsKeyboardFocus( ) const OVERRIDE;
+
+
+private:
+
+	// Callback for generating a row widget for the suggestion list view.
+	TSharedRef<ITableRow> HandleSuggestionListViewGenerateRow( TSharedPtr<FString> Message, const TSharedRef<STableViewBase>& OwnerTable );
+
+	// Callback for changing the selected suggestion.
+	void HandleSuggestionListViewSelectionChanged( TSharedPtr<FString> NewValue, ESelectInfo::Type SelectInfo );
+
+	// Callback for getting the highlight string in a suggestion list row widget.
+	FText HandleSuggestionListWidgetHighlightText( ) const;
+
+	// Callback for changing the text box input.
+	void HandleTextBoxTextChanged( const FText& InText );
+
+	// Callback committing the text box input.
+	void HandleTextBoxTextCommitted( const FText& InText, ETextCommit::Type CommitInfo );
+
+
+private:
+
+	// Holds a flag to prevent recursive calls in UI callbacks.
+	bool IgnoreUIUpdate; 
+
+	// Holds the box for history and auto-complete suggestions.
+	TSharedPtr<SMenuAnchor> MenuAnchor;
+
+	// Holds the current list of suggestions.
+	TArray<TSharedPtr<FString> > Suggestions;
+
+	// Holds the list view that displays the suggestions.
+	TSharedPtr<SListView<TSharedPtr<FString> > > SuggestionListView;
+
+	// Holds the index of the selected suggestion (-1 for no selection).
+	int32 SelectedSuggestion;
+
+	// Holds the editable text box.
+	TSharedPtr<SEditableTextBox> TextBox;
+
+	// The styling of the suggestion text.
+	const FTextBlockStyle* SuggestionTextStyle;
+
+
+private:
+
+	// Holds a delegate that is executed before the input history is shown.
+	FOnShowingHistory OnShowingHistory;
+
+	// Holds a delegate that is executed before the suggestions list is shown.
+	FOnShowingSuggestions OnShowingSuggestions;
+
+	// Holds a delegate that is executed when the text has changed.
+	FOnTextChanged OnTextChanged;
+
+	// Holds a delegate that is executed when the text has been committed.
+	FOnTextCommitted OnTextCommitted;
+};
\ No newline at end of file
diff --git a/Engine/Source/Runtime/Slate/Public/Widgets/Input/SVectorInputBox.h b/Engine/Source/Runtime/Slate/Public/Widgets/Input/SVectorInputBox.h
new file mode 100644
--- /dev/null
+++ b/Engine/Source/Runtime/Slate/Public/Widgets/Input/SVectorInputBox.h
@@ -0,0 +1,84 @@
+// Copyright 1998-2014 Epic Games, Inc. All Rights Reserved.
+
+#pragma once
+
+#undef _X //@todo vogel: needs to be removed once Slate no longer uses _##Name
+
+/**
+ * Vector Slate control
+ */
+class SLATE_API SVectorInputBox : public SCompoundWidget
+{
+public:
+	SLATE_BEGIN_ARGS( SVectorInputBox )
+		: _Font( FCoreStyle::Get().GetFontStyle("NormalFont") ),
+		  _bColorAxisLabels( false )
+		{}
+		
+		/** X Component of the vector */
+		SLATE_ATTRIBUTE( TOptional<float>, X )
+
+		/** Y Component of the vector */
+		SLATE_ATTRIBUTE( TOptional<float>, Y )
+
+		/** Z Component of the vector */
+		SLATE_ATTRIBUTE( TOptional<float>, Z )
+
+		/** Font to use for the text in this box */
+		SLATE_ATTRIBUTE( FSlateFontInfo, Font )
+
+		/** Should the axis labels be colored */
+		SLATE_ARGUMENT( bool, bColorAxisLabels )		
+
+		/** Called when the x value of the vector is changed */
+		SLATE_EVENT( FOnFloatValueChanged, OnXChanged )
+
+		/** Called when the y value of the vector is changed */
+		SLATE_EVENT( FOnFloatValueChanged, OnYChanged )
+
+		/** Called when the z value of the vector is changed */
+		SLATE_EVENT( FOnFloatValueChanged, OnZChanged )
+
+		/** Called when the x value of the vector is committed */
+		SLATE_EVENT( FOnFloatValueCommitted, OnXCommitted )
+
+		/** Called when the y value of the vector is committed */
+		SLATE_EVENT( FOnFloatValueCommitted, OnYCommitted )
+
+		/** Called when the z value of the vector is committed */
+		SLATE_EVENT( FOnFloatValueCommitted, OnZCommitted )
+
+		/** Menu extender delegate for the X value */
+		SLATE_EVENT( FMenuExtensionDelegate, ContextMenuExtenderX )
+
+		/** Menu extender delegate for the Y value */
+		SLATE_EVENT( FMenuExtensionDelegate, ContextMenuExtenderY )
+
+		/** Menu extender delegate for the Z value */
+		SLATE_EVENT( FMenuExtensionDelegate, ContextMenuExtenderZ )
+
+	SLATE_END_ARGS()
+
+	/**
+	 * Construct this widget
+	 *
+	 * @param	InArgs	The declaration data for this widget
+	 */
+	void Construct( const FArguments& InArgs );
+private:
+
+	/**
+	 * Construct widgets for the X Value
+	 */
+	void ConstructX( const FArguments& InArgs, TSharedRef<SHorizontalBox> HorizontalBox );
+
+	/**
+	 * Construct widgets for the Y Value
+	 */
+	void ConstructY( const FArguments& InArgs, TSharedRef<SHorizontalBox> HorizontalBox );
+
+	/**
+	 * Construct widgets for the Z Value
+	 */
+	void ConstructZ( const FArguments& InArgs, TSharedRef<SHorizontalBox> HorizontalBox );
+};
diff --git a/Engine/Source/Runtime/Slate/Public/Widgets/Input/SVirtualJoystick.h b/Engine/Source/Runtime/Slate/Public/Widgets/Input/SVirtualJoystick.h
new file mode 100644
--- /dev/null
+++ b/Engine/Source/Runtime/Slate/Public/Widgets/Input/SVirtualJoystick.h
@@ -0,0 +1,223 @@
+// Copyright 1998-2014 Epic Games, Inc. All Rights Reserved.
+ 
+#pragma once
+
+/**
+ * A virtual joystsick
+ */
+class SLATE_API SVirtualJoystick : public SLeafWidget
+{
+
+public:
+	/** The settings and current state of each zone we render */ 
+	struct FControlInfo
+	{
+		FControlInfo()
+		{
+			// default to all 0
+			FMemory::Memzero(this, sizeof(*this));
+			CapturedPointerIndex = -1;
+			InputScale = FVector2D(1.f, 1.f);
+		}
+
+		// Set by the game
+
+		/** The brush to use to draw the background for joysticks, or unclicked for buttons */
+		TSharedPtr< FSlateDynamicImageBrush > Image1;
+
+		/** The brush to use to draw the thumb for joysticks, or clicked for buttons */
+		TSharedPtr< FSlateDynamicImageBrush > Image2;
+
+		/** The actual center of the control */
+		FVector2D Center;
+
+		/** The size of a joystick that can be re-centered within InteractionSize area */
+		FVector2D VisualSize;
+
+		/** The size of the thumb that can be re-centered within InteractionSize area */
+		FVector2D ThumbSize;
+
+		/** The size of a the interactable area around Center */
+		FVector2D InteractionSize;
+
+		/** The scale for control input */
+		FVector2D InputScale;
+
+		/** The input to send from this control (for sticks, this is the horizontal/X input) */ 
+		FKey MainInputKey;
+
+		/** The secondary input (for sticks, this is the vertical/Y input, unused for buttons) */ 
+		FKey AltInputKey;
+
+	private:
+		friend SVirtualJoystick;
+
+		/**
+		 * Reset the control to a centered/inactive state
+		 */
+		void Reset();
+
+		// Current state
+
+		/** The position of the thumb, in relation to the VisualCenter */
+		FVector2D ThumbPosition;
+
+		/** For recentered joysticks, this is the re-center location */
+		FVector2D VisualCenter;
+
+		/** Which pointer index is interacting with this control right now, or -1 if not interacting */
+		int32 CapturedPointerIndex;
+
+		/** Time to activate joystick **/
+		float ElapsedTime;
+
+		/** Visual center to be updated */
+		FVector2D NextCenter;
+
+		/** Whether or not to send one last "release" event next tick */
+		bool bSendOneMoreEvent;
+
+		/** Whether or not we need position the control against the geometry */
+		bool bHasBeenPositioned;
+
+		/** Whether or not to update center position */
+		bool bNeedUpdatedCenter;
+	};
+
+	SLATE_BEGIN_ARGS( SVirtualJoystick )
+		{}
+	SLATE_END_ARGS()
+
+	void Construct( const FArguments& InArgs );
+
+	/**
+	 * Static function to return if external users should create/activate/etc a touch interface
+	 * Note that this function is also used internally, so even if this returns false but an SVirtualJoystick
+	 * is created, it won't actually show any controls
+	 */
+	static bool ShouldDisplayTouchInterface();
+
+	/** 
+	 * Shows or hides the controls (for instance during cinematics
+	 */
+	void SetVisibility(const bool bVisible, const bool bFade);
+
+	void AddControl(const FControlInfo& Control)
+	{
+		Controls.Add(Control);
+	}
+
+	void ClearControls()
+	{
+		Controls.Empty();
+	}
+
+	void SetControls(const TArray<FControlInfo>& InControls)
+	{
+		if (ShouldDisplayTouchInterface())
+		{
+			Controls = InControls;
+		}
+	}
+
+	/**
+	 * Sets parameters that control all controls
+	 */
+	void SetGlobalParameters(float InActiveOpacity, float InInactiveOpacity, float InTimeUntilDeactive, float InTimeUntilReset, float InActivationDelay, bool InbPreventReCenter, float InStartupDelay);
+
+	virtual int32 OnPaint( const FGeometry& AllottedGeometry, const FSlateRect& MyClippingRect, FSlateWindowElementList& OutDrawElements, int32 LayerId, const FWidgetStyle& InWidgetStyle, bool bParentEnabled ) const OVERRIDE;
+	
+	virtual FVector2D ComputeDesiredSize() const OVERRIDE;
+
+	virtual FReply OnTouchStarted(const FGeometry& MyGeometry, const FPointerEvent& Event) OVERRIDE;
+	virtual FReply OnTouchMoved(const FGeometry& MyGeometry, const FPointerEvent& Event) OVERRIDE;
+	virtual FReply OnTouchEnded(const FGeometry& MyGeometry, const FPointerEvent& Event) OVERRIDE;
+
+	virtual void Tick( const FGeometry& AllottedGeometry, const double InCurrentTime, const float InDeltaTime ) OVERRIDE;
+
+	virtual bool SupportsKeyboardFocus() const OVERRIDE
+	{
+		return true;
+	}
+
+	void AlignBoxIntoScreen( FVector2D& Position, const FVector2D& Size, const FVector2D& ScreenSize )
+	{
+		if (Size.X > ScreenSize.X || Size.Y > ScreenSize.Y)
+		{
+			return;
+		}
+
+		// Align box to fit into screen
+		if (Position.X - Size.X * 0.5f < 0.f )
+		{
+			Position.X = Size.X * 0.5f;
+		}
+
+		if (Position.X + Size.X * 0.5f > ScreenSize.X )
+		{
+			Position.X = ScreenSize.X - Size.X * 0.5f;
+		}
+
+		if (Position.Y - Size.Y * 0.5f < 0.f)
+		{
+			Position.Y = Size.Y * 0.5f;
+		}
+
+		if (Position.Y + Size.Y * 0.5f > ScreenSize.Y )
+		{
+			Position.Y = ScreenSize.Y - Size.Y * 0.5f;
+		}
+	}
+
+private:
+	/**
+	 * Process a touch event (on movement and possibly on initial touch)
+	 *
+	 * @return true if the touch was successful
+	 */
+	bool HandleTouch(int32 ControlIndex, const FVector2D& LocalCoord, const FVector2D& ScreenSize);
+
+	/** 
+	 * Return the target opacity to lerp to given the current state
+	 */
+	FORCEINLINE float GetBaseOpacity();
+
+	/** List of controls set by the UTouchInterface */
+	TArray<FControlInfo> Controls;
+
+	/** Global settings from the UTouchInterface */
+	float ActiveOpacity;
+	float InactiveOpacity;
+	float TimeUntilDeactive;
+	float TimeUntilReset;
+	float ActivationDelay;
+	float StartupDelay;
+
+	enum EVirtualJoystickState
+	{
+		State_Active,
+		State_CountingDownToInactive,
+		State_CountingDownToReset,
+		State_Inactive,
+		State_WaitForStart,
+		State_CountingDownToStart,
+	};
+
+	/** The current state of all controls */
+	EVirtualJoystickState State;
+
+	/** True if the joystick should be visible */
+	uint32 bVisible:1;
+
+	/** If true, this zone will have it's "center" set when you touch it, otherwise the center will be set to the center of the zone */
+	uint32 bCenterOnEvent:1;
+
+	/** If true, ignore re-centering */
+	uint32 bPreventReCenter:1;
+
+	/** Target opacity */
+	float CurrentOpacity;
+
+	/* Countdown until next state change */
+	float Countdown;
+};
diff --git a/Engine/Source/Runtime/Slate/Public/Widgets/Input/SVirtualKeyboardEntry.h b/Engine/Source/Runtime/Slate/Public/Widgets/Input/SVirtualKeyboardEntry.h
new file mode 100644
--- /dev/null
+++ b/Engine/Source/Runtime/Slate/Public/Widgets/Input/SVirtualKeyboardEntry.h
@@ -0,0 +1,206 @@
+// Copyright 1998-2014 Epic Games, Inc. All Rights Reserved.
+
+#pragma once
+
+// @todo - hook up keyboard types
+enum EKeyboardType
+{
+	Keyboard_Default,
+	Keyboard_Number,
+	Keyboard_Web,
+	Keyboard_Email,
+	Keyboard_Password,
+};
+
+class SLATE_API SVirtualKeyboardEntry : public SLeafWidget
+{
+
+public:
+
+	SLATE_BEGIN_ARGS( SVirtualKeyboardEntry )
+		: _Text()
+		, _HintText()
+		, _Font( FCoreStyle::Get().GetFontStyle("NormalFont") )
+		, _ColorAndOpacity( FSlateColor::UseForeground() )
+		, _IsReadOnly( false )
+		, _ClearKeyboardFocusOnCommit( true )
+		, _MinDesiredWidth( 0.0f )
+		, _KeyboardType ( EKeyboardType::Keyboard_Default )
+		{}
+
+		/** Sets the text content for this editable text widget */
+		SLATE_ATTRIBUTE( FText, Text )
+
+		/** The text that appears when there is nothing typed into the search box */
+		SLATE_ATTRIBUTE( FText, HintText )
+
+		/** Sets the font used to draw the text */
+		SLATE_ATTRIBUTE( FSlateFontInfo, Font )
+
+		/** Text color and opacity */
+		SLATE_ATTRIBUTE( FSlateColor, ColorAndOpacity )
+
+		/** Sets whether this text box can actually be modified interactively by the user */
+		SLATE_ATTRIBUTE( bool, IsReadOnly )
+
+		/** Whether to clear keyboard focus when pressing enter to commit changes */
+		SLATE_ATTRIBUTE( bool, ClearKeyboardFocusOnCommit )
+
+		/** Called whenever the text is changed interactively by the user */
+		SLATE_EVENT( FOnTextChanged, OnTextChanged )
+
+		/** Called whenever the text is committed.  This happens when the user presses enter or the text box loses focus. */
+		SLATE_EVENT( FOnTextCommitted, OnTextCommitted )
+
+		/** Minimum width that a text block should be */
+		SLATE_ATTRIBUTE( float, MinDesiredWidth )
+
+		/** Sets the text content for this editable text widget */
+		SLATE_ATTRIBUTE( EKeyboardType, KeyboardType )
+
+		SLATE_END_ARGS()
+
+
+	/** Constructor */
+	SVirtualKeyboardEntry();
+	
+	/**
+	 * Construct this widget
+	 *
+	 * @param	InArgs	The declaration data for this widget
+	 */
+	void Construct( const FArguments& InArgs );
+
+	/**
+	 * Returns the text.
+	 *
+	 * @return  Text
+	 */
+	const FText& GetText() const
+	{
+		return Text.Get();
+	}
+
+	/**
+	 * Sets the text currently being edited 
+	 *
+	 * @param  InNewText  The new text
+	 */
+	void SetText( const TAttribute< FText >& InNewText );
+
+	/**
+	 * Returns the hint text.
+	 *
+	 * @return  HintText
+	 */
+	const FText& GetHintText() const
+	{
+		return HintText.Get();
+	}
+
+	/**
+	 * Returns the keyboard type.
+	 *
+	 * @return  KeyboardType
+	 */
+	EKeyboardType GetKeyboardType() const
+	{
+		return KeyboardType.Get();
+	}
+
+	/**
+	 * Restores the text to the original state
+	 */
+	void RestoreOriginalText();
+
+	/**
+	 *	Returns whether the current text varies from the original
+	 */
+	bool HasTextChangedFromOriginal() const;
+
+	/**
+	 * Sets the font used to draw the text
+	 *
+	 * @param  InNewFont	The new font to use
+	 */
+	void SetFont( const TAttribute< FSlateFontInfo >& InNewFont );
+
+	/**
+	 * Can the user edit the text?
+	 */
+	bool GetIsReadOnly() const;
+
+
+protected:
+
+	/**
+	 * Gets the height of the largest character in the font
+	 *
+	 * @return  The fonts height
+	 */
+	virtual float GetFontHeight() const;
+
+	// SWidget interface
+	virtual void Tick( const FGeometry& AllottedGeometry, const double InCurrentTime, const float InDeltaTime ) OVERRIDE;
+	virtual FVector2D ComputeDesiredSize() const OVERRIDE;
+	virtual int32 OnPaint( const FGeometry& AllottedGeometry, const FSlateRect& MyClippingRect, FSlateWindowElementList& OutDrawElements, int32 LayerId, const FWidgetStyle& InWidgetStyle, bool bParentEnabled ) const OVERRIDE;
+	virtual bool SupportsKeyboardFocus() const OVERRIDE;
+	virtual FReply OnKeyboardFocusReceived( const FGeometry& MyGeometry, const FKeyboardFocusEvent& InKeyboardFocusEvent ) OVERRIDE;
+	virtual void OnKeyboardFocusLost( const FKeyboardFocusEvent& InKeyboardFocusEvent ) OVERRIDE;
+	// End of SWidget interface
+
+private:
+
+	/**
+	 * @return  the string that needs to be rendered
+	 */
+	FString GetStringToRender() const;
+
+private:
+
+	/** The text content for this editable text widget */
+	TAttribute< FText > Text;
+
+	TAttribute< FText > HintText;
+
+	/** The font used to draw the text */
+	TAttribute< FSlateFontInfo > Font;
+
+	/** Text color and opacity */
+	TAttribute<FSlateColor> ColorAndOpacity;
+
+	/** Sets whether this text box can actually be modified interactively by the user */
+	TAttribute< bool > IsReadOnly;
+
+	/** Whether to clear keyboard focus when pressing enter to commit changes */
+	TAttribute< bool > ClearKeyboardFocusOnCommit;
+
+	/** Called whenever the text is changed interactively by the user */
+	FOnTextChanged OnTextChanged;
+
+	/** Called whenever the text is committed.  This happens when the user presses enter or the text box loses focus. */
+	FOnTextCommitted OnTextCommitted;
+
+	/** Text string currently being edited */
+	FText EditedText;
+
+	/** Original text prior to user edits.  This is used to restore the text if the user uses the escape key. */
+	FText OriginalText;
+
+	/** Current scrolled position */
+	FScrollHelper ScrollHelper;
+
+	/** True if the last mouse down caused us to receive keyboard focus */
+	bool bWasFocusedByLastMouseDown;
+
+	/** True if we're currently (potentially) changing the text string */
+	bool bIsChangingText;
+
+	/** Prevents the editabletext from being smaller than desired in certain cases (e.g. when it is empty) */
+	TAttribute<float> MinDesiredWidth;
+
+	TAttribute<EKeyboardType> KeyboardType;
+
+	bool bNeedsUpdate;
+
+};
