Slate: Refactored core Slate implementation into SlateCore module in preparation for UMG.

diff --git a/Engine/Source/Runtime/Slate/Public/Widgets/Layout/SBorder.h b/Engine/Source/Runtime/Slate/Public/Widgets/Layout/SBorder.h
new file mode 100644
--- /dev/null
+++ b/Engine/Source/Runtime/Slate/Public/Widgets/Layout/SBorder.h
@@ -0,0 +1,116 @@
+// Copyright 1998-2014 Epic Games, Inc. All Rights Reserved.
+
+#pragma once
+
+/**
+ * A border is an widget that can be used to contain other widgets.
+ * It has a BorderImage property, which allows it to take on different appearances.
+ * Border also has a Content() slot as well as some parameters controlling the
+ * arrangement of said content.
+ */
+class SLATE_API SBorder : public SCompoundWidget
+{
+public:
+
+	SLATE_BEGIN_ARGS(SBorder)
+		: _Content()
+		, _HAlign( HAlign_Fill )
+		, _VAlign( VAlign_Fill )
+		, _Padding( FMargin(2.0f) )
+		, _OnMouseButtonDown()
+		, _OnMouseButtonUp()
+		, _OnMouseMove()
+		, _OnMouseDoubleClick()
+		, _BorderImage( FCoreStyle::Get().GetBrush( "Border" ) )
+		, _ContentScale( FVector2D(1,1) )
+		, _DesiredSizeScale( FVector2D(1,1) )
+		, _ColorAndOpacity( FLinearColor(1,1,1,1) )
+		, _BorderBackgroundColor( FLinearColor::White )
+		, _ForegroundColor( FSlateColor::UseForeground() )
+		, _ShowEffectWhenDisabled( true )
+		{}
+
+		SLATE_DEFAULT_SLOT( FArguments, Content )
+
+		SLATE_ARGUMENT( EHorizontalAlignment, HAlign )
+		SLATE_ARGUMENT( EVerticalAlignment, VAlign )
+		SLATE_ATTRIBUTE( FMargin, Padding )
+		
+		SLATE_EVENT( FPointerEventHandler, OnMouseButtonDown )
+		SLATE_EVENT( FPointerEventHandler, OnMouseButtonUp )
+		SLATE_EVENT( FPointerEventHandler, OnMouseMove )
+		SLATE_EVENT( FPointerEventHandler, OnMouseDoubleClick )
+
+		SLATE_ATTRIBUTE( const FSlateBrush*, BorderImage )
+
+		SLATE_ATTRIBUTE( FVector2D, ContentScale )
+
+		SLATE_ATTRIBUTE( FVector2D, DesiredSizeScale )
+
+		/** ColorAndOpacity is the color and opacity of content in the border */
+		SLATE_ATTRIBUTE( FLinearColor, ColorAndOpacity )
+		/** BorderBackgroundColor refers to the actual color and opacity of the supplied border image.*/
+		SLATE_ATTRIBUTE( FSlateColor, BorderBackgroundColor )
+		/** The foreground color of text and some glyphs that appear as the border's content. */
+		SLATE_ATTRIBUTE( FSlateColor, ForegroundColor )
+		/** Whether or not to show the disabled effect when this border is disabled */
+		SLATE_ATTRIBUTE( bool, ShowEffectWhenDisabled )
+	SLATE_END_ARGS()
+
+	/**
+	 * Default constructor.
+	 */
+	SBorder();
+
+	/**
+	 * Construct this widget
+	 *
+	 * @param	InArgs	The declaration data for this widget
+	 */
+	void Construct( const FArguments& InArgs );
+
+	/**
+	 * Sets the content for this border
+	 *
+	 * @param	InContent	The widget to use as content for the border
+	 */
+	void SetContent( const TSharedRef< SWidget >& InContent );
+
+	/**
+	 * Gets the content for this border
+	 *
+	 * @return The widget used as content for the border
+	 */
+	const TSharedRef< SWidget >& GetContent() const;
+
+	/** Clears out the content for the border */
+	void ClearContent();
+
+	/** Sets the color and opacity of the background image of this border. */
+	void SetBorderBackgroundColor( const TAttribute<FSlateColor>& InColorAndOpacity );
+
+	/** Set the desired size scale multiplier */
+	void SetDesiredSizeScale( const TAttribute<FVector2D>& InDesiredSizeScale );
+
+public:
+	// SWidget interface
+	virtual int32 OnPaint( const FGeometry& AllottedGeometry, const FSlateRect& MyClippingRect, FSlateWindowElementList& OutDrawElements, int32 LayerId, const FWidgetStyle& InWidgetStyle, bool bParentEnabled ) const OVERRIDE;
+	virtual FReply OnMouseButtonDown( const FGeometry& MyGeometry, const FPointerEvent& MouseEvent ) OVERRIDE;
+	virtual FReply OnMouseButtonUp( const FGeometry& MyGeometry, const FPointerEvent& MouseEvent ) OVERRIDE;
+	virtual FReply OnMouseMove( const FGeometry& MyGeometry, const FPointerEvent& MouseEvent ) OVERRIDE;
+	virtual FReply OnMouseButtonDoubleClick( const FGeometry& MyGeometry, const FPointerEvent& MouseEvent ) OVERRIDE;
+	virtual FVector2D ComputeDesiredSize() const OVERRIDE;
+	// End of SWidget interface
+
+ protected:
+	
+	TAttribute<const FSlateBrush*> BorderImage;
+	TAttribute<FSlateColor> BorderBackgroundColor;
+	TAttribute<FVector2D> DesiredSizeScale;
+	/** Whether or not to show the disabled effect when this border is disabled */
+	TAttribute<bool> ShowDisabledEffect;
+	FPointerEventHandler MouseButtonDownHandler;
+	FPointerEventHandler MouseButtonUpHandler;
+	FPointerEventHandler MouseMoveHandler;
+	FPointerEventHandler MouseDoubleClickHandler;
+ };
diff --git a/Engine/Source/Runtime/Slate/Public/Widgets/Layout/SBox.h b/Engine/Source/Runtime/Slate/Public/Widgets/Layout/SBox.h
new file mode 100644
--- /dev/null
+++ b/Engine/Source/Runtime/Slate/Public/Widgets/Layout/SBox.h
@@ -0,0 +1,71 @@
+// Copyright 1998-2014 Epic Games, Inc. All Rights Reserved.
+
+#pragma once
+
+
+/**
+ * SBox is the simplest layout element.
+ */
+class SLATE_API SBox : public SPanel
+{
+	public:
+		class FBoxSlot : public TSupportsContentAlignmentMixin<FBoxSlot>, public TSupportsContentPaddingMixin<FBoxSlot>, public TSupportsOneChildMixin<SWidget, FBoxSlot>
+		{
+			public:
+				FBoxSlot()
+				: TSupportsContentAlignmentMixin<FBoxSlot>(HAlign_Fill, VAlign_Fill)
+				{
+				}
+		};
+
+		SLATE_BEGIN_ARGS(SBox)
+			: _HAlign( HAlign_Fill )
+			, _VAlign( VAlign_Fill )
+			, _Padding( 0.0f )
+			, _Content()
+			, _WidthOverride(FOptionalSize())
+			, _HeightOverride(FOptionalSize())
+			{
+				_Visibility = EVisibility::SelfHitTestInvisible;
+			}
+
+			/** Horizontal alignment of content in the area allotted to the SBox by its parent */
+			SLATE_ARGUMENT( EHorizontalAlignment, HAlign )
+
+			/** Vertical alignment of content in the area allotted to the SBox by its parent */
+			SLATE_ARGUMENT( EVerticalAlignment, VAlign )
+
+			/** Padding between the SBox and the content that it presents. Padding affects desired size. */
+			SLATE_ATTRIBUTE( FMargin, Padding )
+
+			/** The widget content presented by the SBox */
+			SLATE_DEFAULT_SLOT( FArguments, Content )
+
+			/** When specified, ignore the content's desired size and report the WidthOverride as the Box's desired width. */
+			SLATE_ATTRIBUTE( FOptionalSize, WidthOverride )
+
+			/** When specified, ignore the content's desired size and report the HeightOverride as the Box's desired height. */
+			SLATE_ATTRIBUTE( FOptionalSize, HeightOverride)
+
+		SLATE_END_ARGS()
+
+		void Construct( const FArguments& InArgs );
+
+
+		virtual FVector2D ComputeDesiredSize() const OVERRIDE;
+		virtual void ArrangeChildren( const FGeometry& AllottedGeometry, FArrangedChildren& ArrangedChildren ) const OVERRIDE;
+		virtual FChildren* GetChildren() OVERRIDE;
+		virtual int32 OnPaint( const FGeometry& AllottedGeometry, const FSlateRect& MyClippingRect, FSlateWindowElementList& OutDrawElements, int32 LayerId, const FWidgetStyle& InWidgetStyle, bool bParentEnabled ) const OVERRIDE;
+
+private:
+		FBoxSlot ChildSlot;
+
+		/** When specified, ignore the content's desired size and report the.WidthOverride as the Box's desired width. */
+		TAttribute<FOptionalSize> WidthOverride;
+
+		/** When specified, ignore the content's desired size and report the.HeightOverride as the Box's desired height. */
+		TAttribute<FOptionalSize> HeightOverride;
+	
+};
+
+
diff --git a/Engine/Source/Runtime/Slate/Public/Widgets/Layout/SDPIScaler.h b/Engine/Source/Runtime/Slate/Public/Widgets/Layout/SDPIScaler.h
new file mode 100644
--- /dev/null
+++ b/Engine/Source/Runtime/Slate/Public/Widgets/Layout/SDPIScaler.h
@@ -0,0 +1,48 @@
+// Copyright 1998-2014 Epic Games, Inc. All Rights Reserved.
+
+#pragma once
+
+/**
+ * Controls the DPI scale of its content. Can be used for zooming or shrinking arbitrary widget content.
+ */
+class SLATE_API SDPIScaler : public SPanel
+{
+public:
+	SLATE_BEGIN_ARGS( SDPIScaler )
+		: _DPIScale( 1.0f )
+		, _Content()
+		{
+			_Visibility = EVisibility::SelfHitTestInvisible;
+		}
+
+		/**
+		 * The scaling factor between Pixels and Slate Units (SU).
+		 * For example, a value of 2 means 2 pixels for every linear Slate Unit and 4 pixels for every square Slate Unit.
+		 */
+		SLATE_ATTRIBUTE( float, DPIScale )
+
+		/** The content being presented and scaled. */
+		SLATE_DEFAULT_SLOT( FArguments, Content )
+
+	SLATE_END_ARGS()
+
+	void Construct( const FArguments& InArgs );
+
+	virtual void ArrangeChildren( const FGeometry& AllottedGeometry, FArrangedChildren& ArrangedChildren ) const OVERRIDE;
+	
+	virtual FVector2D ComputeDesiredSize() const OVERRIDE;
+
+	virtual FChildren* GetChildren() OVERRIDE;
+
+protected:
+
+	/**
+	 * The scaling factor from 1:1 pixel to slate unit (SU) ratio.
+	 * e.g. Value of 2 would mean that there are two pixels for every slate unit in every dimension.
+	 *      This means that every 1x1 SU square is represented by 2x2=4 pixels.
+	 */
+	TAttribute<float> DPIScale;
+
+	/** The content being scaled. */
+	FSimpleSlot ChildSlot;
+};
\ No newline at end of file
diff --git a/Engine/Source/Runtime/Slate/Public/Widgets/Layout/SEnableBox.h b/Engine/Source/Runtime/Slate/Public/Widgets/Layout/SEnableBox.h
new file mode 100644
--- /dev/null
+++ b/Engine/Source/Runtime/Slate/Public/Widgets/Layout/SEnableBox.h
@@ -0,0 +1,27 @@
+// Copyright 1998-2014 Epic Games, Inc. All Rights Reserved.
+
+#pragma once
+
+/**
+ * A SEnableBox contains a widget that is lied to about whether the parent hierarchy is enabled or not, being told the parent is always enabled
+ */
+class SEnableBox : public SBox
+{
+public:
+	SLATE_BEGIN_ARGS(SEnableBox) {}
+		/** The widget content to be presented as if the parent were enabled */
+		SLATE_DEFAULT_SLOT(FArguments, Content)
+	SLATE_END_ARGS()
+
+	void Construct(const FArguments& InArgs)
+	{
+		SBox::Construct(SBox::FArguments().Content()[InArgs._Content.Widget]);
+	}
+
+	// SWidget interface
+	virtual int32 OnPaint(const FGeometry& AllottedGeometry, const FSlateRect& MyClippingRect, FSlateWindowElementList& OutDrawElements, int32 LayerId, const FWidgetStyle& InWidgetStyle, bool bParentEnabled) const OVERRIDE
+	{
+		return SBox::OnPaint(AllottedGeometry, MyClippingRect, OutDrawElements, LayerId, InWidgetStyle, /*bParentEnabled=*/ true);
+	}
+	// End of SWidget interface
+};
diff --git a/Engine/Source/Runtime/Slate/Public/Widgets/Layout/SFxWidget.h b/Engine/Source/Runtime/Slate/Public/Widgets/Layout/SFxWidget.h
new file mode 100644
--- /dev/null
+++ b/Engine/Source/Runtime/Slate/Public/Widgets/Layout/SFxWidget.h
@@ -0,0 +1,77 @@
+// Copyright 1998-2014 Epic Games, Inc. All Rights Reserved.
+
+#pragma once
+
+/**
+ * Enables visual effects (zooming, sliding, fading, etc.) to be applied to arbitrary widget content.
+ * 
+ * Unless specified properties do not affect layout.
+ */
+class SLATE_API SFxWidget : public SCompoundWidget
+{
+public:
+	SLATE_BEGIN_ARGS( SFxWidget )
+	: _RenderScale( 1.0f )
+	, _RenderScaleOrigin( FVector2D::ZeroVector )
+	, _LayoutScale( 1.0f )
+	, _VisualOffset( FVector2D::ZeroVector )
+	, _IgnoreClipping( true )
+	, _ColorAndOpacity( FLinearColor::White )
+	, _Content()
+	{}
+		/** Scale the visuals of this widget. Geometry is not affected. */
+		SLATE_ATTRIBUTE( float, RenderScale )
+
+		/** The origin of the visual scale transform in the 0..1 range (0 being upper left, 1 being lower right)  */
+		SLATE_ATTRIBUTE( FVector2D, RenderScaleOrigin )
+
+		/** Just like visual scale only affects Geometry. */
+		SLATE_ATTRIBUTE( float, LayoutScale )
+
+		/** Offset the widget by some fraction of its size in either dimension. */
+		SLATE_ATTRIBUTE( FVector2D, VisualOffset )
+
+		/** Should the FX widget disable all clipping and show through regardless of its parents' bounds.*/
+		SLATE_ATTRIBUTE( bool, IgnoreClipping )
+
+		/** Multiply the contents of the SFxWidget by this color and opacity when drawing */
+		SLATE_ATTRIBUTE( FLinearColor, ColorAndOpacity )
+
+		/** The content that should be modified. */
+		SLATE_DEFAULT_SLOT( FArguments, Content )
+
+	SLATE_END_ARGS()
+
+	void Construct( const FArguments& InArgs );
+
+	/** @see VisualOffset */
+	void SetVisualOffset( const TAttribute<FVector2D>& InOffset );
+
+	/** @see VisualOffset */
+	void SetVisualOffset( FVector InOffset );
+
+	/** @see RenderScale */
+	void SetRenderScale( const TAttribute<float>& InScale );
+
+	/** @see RenderScale */
+	void SetRenderScale( float InScale );
+
+	/** @see ColorAndOpacity */
+	void SetColorAndOpacity( const TAttribute<FLinearColor>& InColorAndOpacity );
+
+	/** @see ColorAndOpacity */
+	void SetColorAndOpacity( FLinearColor InColorAndOpacity );
+
+	
+
+protected:
+	virtual int32 OnPaint( const FGeometry& AllottedGeometry, const FSlateRect& MyClippingRect, FSlateWindowElementList& OutDrawElements, int32 LayerId, const FWidgetStyle& InWidgetStyle, bool bParentEnabled ) const OVERRIDE;
+	virtual FVector2D ComputeDesiredSize() const OVERRIDE;
+	virtual void ArrangeChildren( const FGeometry& AllottedGeometry, FArrangedChildren& ArrangedChildren ) const OVERRIDE;
+
+	TAttribute<float> RenderScale;
+	TAttribute<FVector2D> RenderScaleOrigin;
+	TAttribute<float> LayoutScale;
+	TAttribute<FVector2D> VisualOffset;
+	TAttribute<bool> bIgnoreClipping;
+};
\ No newline at end of file
diff --git a/Engine/Source/Runtime/Slate/Public/Widgets/Layout/SGridPanel.h b/Engine/Source/Runtime/Slate/Public/Widgets/Layout/SGridPanel.h
new file mode 100644
--- /dev/null
+++ b/Engine/Source/Runtime/Slate/Public/Widgets/Layout/SGridPanel.h
@@ -0,0 +1,232 @@
+// Copyright 1998-2014 Epic Games, Inc. All Rights Reserved.
+
+#pragma once
+
+class SLATE_API SGridPanel : public SPanel
+{
+public:
+	// Used by the mandatory named parameter in FSlot
+	class Layer
+	{
+	public:
+		explicit Layer(int32 InLayer)
+			: TheLayer(InLayer)
+		{
+
+		}
+
+		int TheLayer;
+	};
+
+	class FSlot : public TSupportsOneChildMixin<SWidget, FSlot>, public TSupportsContentAlignmentMixin<FSlot>, public TSupportsContentPaddingMixin<FSlot>
+	{
+		public:
+			/** Default values for a slot. */
+			FSlot( int32 Column, int32 Row, int32 InLayer )
+				: TSupportsContentAlignmentMixin<FSlot>(HAlign_Fill, VAlign_Fill)
+				, ColumnParam( Column )
+				, ColumnSpanParam( 1 )
+				, RowParam( Row )
+				, RowSpanParam( 1 )
+				, LayerParam( InLayer )
+				, NudgeParam( FVector2D::ZeroVector )
+			{
+			}
+
+			/** The panel that contains this slot */
+			TWeakPtr<SGridPanel> Panel;
+
+			/** Which column in the grid this cell belongs to */
+			int32 ColumnParam;
+
+			/** How many columns this slot spans over */
+			int32 ColumnSpanParam;
+
+			/** Which row in the grid this cell belongs to */
+			int32 RowParam;
+
+			/** How many rows this this slot spans over */
+			int32 RowSpanParam;
+
+			/** Positive values offset this cell to be hit-tested and drawn on top of others. Default is 0; i.e. no offset. */
+			int32 LayerParam;
+
+			/** Offset this slot's content by some amount; positive values offset to lower right*/
+			FVector2D NudgeParam;
+
+
+			/** How many columns this slot spans over */
+			FSlot& ColumnSpan( int32 ColumnSpan )
+			{
+				ColumnSpanParam = FMath::Max(1,ColumnSpan);
+
+				if(Panel.IsValid())
+				{
+					Panel.Pin()->NotifySlotChanged(this);
+				}
+
+				return *this;
+			}
+
+			/** How many rows this this slot spans over */
+			FSlot& RowSpan( int32 RowSpan )
+			{
+				RowSpanParam = FMath::Max(1,RowSpan);
+
+				if(Panel.IsValid())
+				{
+					Panel.Pin()->NotifySlotChanged(this);
+				}
+
+				return *this;
+			}
+
+			/** Offset this slot's content by some amount; positive values offset to lower right */
+			FSlot& Nudge( const FVector2D& Nudge )
+			{
+				NudgeParam = Nudge;
+				return *this;
+			}
+
+	};
+
+	/**
+	 * Used by declarative syntax to create a Slot in the specified Column, Row and Layer.
+	 */
+	static FSlot& Slot( int32 Column, int32 Row, Layer InLayer = Layer(0) )
+	{
+		return *(new FSlot( Column, Row, InLayer.TheLayer ));
+	}
+
+	/**
+	 * Dynamically add a new slot to the UI at specified Column and Row. Optionally, specify a layer at which this slot should be added.
+	 *
+	 * @return A reference to the newly-added slot
+	 */
+	FSlot& AddSlot( int32 Column, int32 Row, Layer InLayer = Layer(0) );
+
+	SLATE_BEGIN_ARGS( SGridPanel )
+		{
+			_Visibility = EVisibility::SelfHitTestInvisible;
+		}
+		SLATE_SUPPORTS_SLOT( FSlot )
+
+		/** Specify a column to stretch instead of sizing to content. */
+		FArguments& FillColumn( int32 ColumnId, float Coefficient )
+		{
+			if (ColFillCoefficients.Num() <= ColumnId)
+			{
+				ColFillCoefficients.AddZeroed( ColumnId - ColFillCoefficients.Num() + 1 );
+			}
+			ColFillCoefficients[ColumnId] = Coefficient;
+			return Me();
+		}
+
+		/** Specify a row to stretch instead of sizing to content. */
+		FArguments& FillRow( int32 RowId, float Coefficient )
+		{
+			if (RowFillCoefficients.Num() <= RowId)
+			{
+				RowFillCoefficients.AddZeroed( RowId - RowFillCoefficients.Num() + 1 );
+			}
+			RowFillCoefficients[RowId] = Coefficient;
+			return Me();
+		}
+
+		/** Coefficients for columns that need to stretch instead of size to content */
+		TArray<float> ColFillCoefficients;
+		
+		/** Coefficients for rows that need to stretch instead of size to content */
+		TArray<float> RowFillCoefficients;
+		
+	SLATE_END_ARGS()
+
+	/** Removes all slots from the panel */
+	void ClearChildren();
+	
+	void Construct( const FArguments& InArgs );
+
+	// SWidget interface
+	virtual int32 OnPaint( const FGeometry& AllottedGeometry, const FSlateRect& MyClippingRect, FSlateWindowElementList& OutDrawElements, int32 LayerId, const FWidgetStyle& InWidgetStyle, bool bParentEnabled ) const OVERRIDE;
+	virtual void ArrangeChildren( const FGeometry& AllottedGeometry, FArrangedChildren& ArrangedChildren ) const OVERRIDE;
+	virtual void CacheDesiredSize() OVERRIDE;
+	virtual FVector2D ComputeDesiredSize() const OVERRIDE;
+	virtual FChildren* GetChildren() OVERRIDE;
+	// End of SWidget interface
+
+	
+	/**
+	 * GetDesiredSize of a subregion in the graph.
+	 *
+	 * @param StartCell   The cell (inclusive) in the upper left of the region.
+	 * @param Size        Number of cells in the X and Y directions to get the size for.
+	 *
+	 * @return FVector2D  The desired size of the region of cells specified.
+	 */
+	FVector2D GetDesiredSize( const FIntPoint& StartCell, int32 Width, int32 Height ) const;
+
+	/** Specify a column to stretch instead of sizing to content. */
+	void SetColumnFill( int32 ColumnId, float Coefficient );
+
+	/** Specify a row to stretch instead of sizing to content. */
+	void SetRowFill( int32 RowId, float Coefficient );
+
+private:
+
+	/**
+	 * Given an array of values, re-populate the array such that every contains the partial sums up to that element.
+	 * i.e. Array[N] = Array.Sum(0 .. N-1)
+	 *
+	 * The resulting array is 1-element longer.
+	 */
+	static void ComputePartialSums( TArray<float>& TurnMeIntoPartialSums );
+	
+	/** Given a SizeContribution, distribute it to the elements in DistributeOverMe at indexes from [StartIndex .. UpperBound) */
+	static void DistributeSizeContributions( float SizeContribution, TArray<float>& DistributeOverMe, int32 StartIndex, int32 UpperBound );
+
+	/**
+	 * Inser the slot InSlot into the list of Slots based on its LayerParam, such that Slots are sorter by layer.
+	 *
+	 * @param The newly-allocated slot to insert.
+	 *
+	 * @param A reference to the added slot
+	 */
+	FSlot& InsertSlot( FSlot* InSlot );	
+
+	/** Compute the sizes of columns and rows needed to fit all the slots in this grid. */
+	void ComputeDesiredCellSizes( TArray<float>& OutColumns, TArray<float>& OutRows ) const;
+
+	/** Draw the debug grid of rows and colummns; useful for inspecting the GridPanel's logic. See OnPaint() for parameter meaning */
+	int32 LayoutDebugPaint(const FGeometry& AllottedGeometry, const FSlateRect& MyClippingRect, FSlateWindowElementList& OutDrawElements, int32 LayerId ) const;
+	
+	/** 
+	 * Callback used to resize our internal arrays if a slot (or slot span) is changed. 
+	 *
+	 * @param The slot that has just changed.
+	 */
+	void NotifySlotChanged( FSlot* InSlot );
+
+	/** The slots that are placed into various grid locations */
+	TPanelChildren<FSlot> Slots;
+
+	/**
+	 * Offsets of each column from the beginning of the grid.
+	 * Includes a faux value at the end of the array for finding the size of the last cell.
+	 */
+	TArray<float> Columns;
+	
+	/**
+	 * Offsets of each row from the beginning of the grid.
+	 * Includes a faux value at the end of the array for finding the size of the last cell.
+	 */
+	TArray<float> Rows;
+
+	/** Total desires size along each axis. */
+	FVector2D TotalDesiredSizes;
+
+	/** Fill coefficients for the columns */
+	TArray<float> ColFillCoefficients;
+
+	/** Fill coefficients for the rows */
+	TArray<float> RowFillCoefficients;
+};
diff --git a/Engine/Source/Runtime/Slate/Public/Widgets/Layout/SMissingWidget.h b/Engine/Source/Runtime/Slate/Public/Widgets/Layout/SMissingWidget.h
new file mode 100644
--- /dev/null
+++ b/Engine/Source/Runtime/Slate/Public/Widgets/Layout/SMissingWidget.h
@@ -0,0 +1,23 @@
+// Copyright 1998-2014 Epic Games, Inc. All Rights Reserved.
+
+/*=============================================================================
+	SMissingWidget.h: Declares the SMissingWidget class.
+=============================================================================*/
+
+#pragma once
+
+
+/**
+ * Like a null widget, but visualizes itself as being explicitly missing.
+ */
+class SLATE_API SMissingWidget
+{
+public:
+
+	/**
+	 * Creates a new instance.
+	 *
+	 * @return The widget.
+	 */
+	static TSharedRef<class SWidget> MakeMissingWidget( );
+};
diff --git a/Engine/Source/Runtime/Slate/Public/Widgets/Layout/SSafeZone.h b/Engine/Source/Runtime/Slate/Public/Widgets/Layout/SSafeZone.h
new file mode 100644
--- /dev/null
+++ b/Engine/Source/Runtime/Slate/Public/Widgets/Layout/SSafeZone.h
@@ -0,0 +1,80 @@
+// Copyright 1998-2014 Epic Games, Inc. All Rights Reserved.
+
+#pragma once
+
+
+/**
+	Here's how you could make use of TitleSafe and ActionSafe areas:
+
+		SNew(SOverlay)
+		+SOverlay::Slot()
+		[
+			// ActioSafe
+			SNew(SSafeZone) .IsTitleSafe( false )
+		]
+		+SOverlay::Slot()
+		[
+			// TitleSafe
+			SNew(SSafeZone) .IsTitleSafe( true )
+		] 
+*/
+
+class SLATE_API SSafeZone : public SBox
+{
+	SLATE_BEGIN_ARGS(SSafeZone)
+		: _HAlign( HAlign_Fill )
+		, _VAlign( VAlign_Fill )
+		, _Padding( 0.0f )
+		, _Content()
+		, _IsTitleSafe( false )
+		 
+		{}
+
+		/** Horizontal alignment of content in the area allotted to the SBox by its parent */
+		SLATE_ARGUMENT( EHorizontalAlignment, HAlign )
+
+		/** Vertical alignment of content in the area allotted to the SBox by its parent */
+		SLATE_ARGUMENT( EVerticalAlignment, VAlign )
+
+		/** Padding between the SBox and the content that it presents. Padding affects desired size. */
+		SLATE_ATTRIBUTE( FMargin, Padding )
+
+		/** The widget content presented by the SBox */
+		SLATE_DEFAULT_SLOT( FArguments, Content )
+	
+		/** True if the zone is TitleSafe, otherwise it's ActionSafe */
+		SLATE_ARGUMENT( bool, IsTitleSafe )
+
+	SLATE_END_ARGS()
+
+	void Construct( const FArguments& InArgs )
+	{
+		SBox::Construct( SBox::FArguments()
+			.HAlign(InArgs._HAlign)
+			.VAlign(InArgs._VAlign)
+			.Padding( this, &SSafeZone::GetSafeZonePadding )
+			[
+				InArgs._Content.Widget
+			]
+		);
+
+		IsTitleSafe = InArgs._IsTitleSafe;
+		Padding = InArgs._Padding;
+	}
+
+	FMargin GetSafeZonePadding() const
+	{
+		// @todo: should we have a function that just returns the safe area size?
+		FDisplayMetrics Metrics;
+		FSlateApplication::Get().GetDisplayMetrics(Metrics);
+
+		// return either the TitleSafe or the ActionSafe size, added to the user padding
+		return Padding.Get() + (IsTitleSafe ? FMargin(Metrics.TitleSafePaddingSize.X, Metrics.TitleSafePaddingSize.Y) : 
+			FMargin(Metrics.ActionSafePaddingSize.X, Metrics.ActionSafePaddingSize.Y));
+	}
+
+private:
+	/** Cached values from the args */
+	TAttribute<FMargin> Padding;
+	bool IsTitleSafe;
+};
diff --git a/Engine/Source/Runtime/Slate/Public/Widgets/Layout/SScrollBar.h b/Engine/Source/Runtime/Slate/Public/Widgets/Layout/SScrollBar.h
new file mode 100644
--- /dev/null
+++ b/Engine/Source/Runtime/Slate/Public/Widgets/Layout/SScrollBar.h
@@ -0,0 +1,137 @@
+// Copyright 1998-2014 Epic Games, Inc. All Rights Reserved.
+ 
+#pragma once
+
+
+DECLARE_DELEGATE_OneParam(
+	FOnUserScrolled,
+	float );	/** ScrollOffset as a fraction between 0 and 1 */
+
+
+class SScrollBarTrack;
+
+class SLATE_API SScrollBar : public SBorder
+{
+public:
+
+	SLATE_BEGIN_ARGS( SScrollBar )
+		: _Style( &FCoreStyle::Get().GetWidgetStyle<FScrollBarStyle>("Scrollbar") )
+		, _OnUserScrolled()
+		, _AlwaysShowScrollbar(false)
+		, _Orientation( Orient_Vertical )
+		, _Thickness( FVector2D(12.0f, 12.0f) )
+		{}
+
+		/** The style to use for this scrollbar */
+		SLATE_STYLE_ARGUMENT( FScrollBarStyle, Style )
+		SLATE_EVENT( FOnUserScrolled, OnUserScrolled )
+		SLATE_ARGUMENT( bool, AlwaysShowScrollbar )
+		SLATE_ARGUMENT( EOrientation, Orientation )
+		/** The thickness of the scrollbar thumb */
+		SLATE_ATTRIBUTE( FVector2D, Thickness )
+	SLATE_END_ARGS()
+
+	/**
+	 * Construct this widget
+	 *
+	 * @param	InArgs	The declaration data for this widget
+	 */
+	void Construct(const FArguments& InArgs);
+
+	/**
+	 * Set the handler to be invoked when the user scrolls.
+	 *
+	 * @param InHandler   Method to execute when the user scrolls the scrollbar
+	 */
+	void SetOnUserScrolled( const FOnUserScrolled& InHandler );
+
+	/**
+	 * Set the offset and size of the track's thumb.
+	 * Note that the maximum offset is 1.0-ThumbSizeFraction.
+	 * If the user can view 1/3 of the items in a single page, the maximum offset will be ~0.667f
+	 *
+	 * @param InOffsetFraction     Offset of the thumbnail from the top as a fraction of the total available scroll space.
+	 * @param InThumbSizeFraction  Size of thumbnail as a fraction of the total available scroll space.
+	 */
+	void SetState( float InOffsetFraction, float InThumbSizeFraction );
+
+	/**
+	 * The system calls this method to notify the widget that a mouse button was pressed within it. This event is bubbled.
+	 *
+	 * @param MyGeometry The Geometry of the widget receiving the event
+	 * @param MouseEvent Information about the input event
+	 *
+	 * @return Whether the event was handled along with possible requests for the system to take action.
+	 */
+	virtual FReply OnMouseButtonDown( const FGeometry& MyGeometry, const FPointerEvent& MouseEvent );
+	
+	/**
+	 * The system calls this method to notify the widget that a mouse button was release within it. This event is bubbled.
+	 *
+	 * @param MyGeometry The Geometry of the widget receiving the event
+	 * @param MouseEvent Information about the input event
+	 *
+	 * @return Whether the event was handled along with possible requests for the system to take action.
+	 */
+	virtual FReply OnMouseButtonUp( const FGeometry& MyGeometry, const FPointerEvent& MouseEvent );
+	
+	/**
+	 * The system calls this method to notify the widget that a mouse moved within it. This event is bubbled.
+	 *
+	 * @param MyGeometry The Geometry of the widget receiving the event
+	 * @param MouseEvent Information about the input event
+	 *
+	 * @return Whether the event was handled along with possible requests for the system to take action.
+	 */
+	virtual FReply OnMouseMove( const FGeometry& MyGeometry, const FPointerEvent& MouseEvent );
+
+	/** @return true if scrolling is possible; false if the view is big enough to fit all the content */
+	bool IsNeeded() const;
+
+	/** @return normalized distance from top */
+	float DistanceFromTop() const;
+
+	/** @return normalized distance from bottom */
+	float DistanceFromBottom() const;
+
+	/** @return the scrollbar's visibility as a product of internal rules and user-specified visibility */
+	EVisibility ShouldBeVisible() const;
+
+	/** @return True if the user is scrolling by dragging the scroll bar thumb. */
+	bool IsScrolling() const;
+
+	void SetUserVisibility(TAttribute<EVisibility> InUserVisibility) { UserVisibility = InUserVisibility; }
+
+	SScrollBar();
+
+protected:
+	
+	/** Execute the on user scrolled delegate */
+	void ExecuteOnUserScrolled( const FGeometry& MyGeometry, const FPointerEvent& MouseEvent );
+
+	/** We fade the scroll track unless it is being hovered*/
+	FSlateColor GetTrackOpacity() const;
+
+	/** We always show a subtle scroll thumb, but highlight it extra when dragging */
+	FLinearColor GetThumbOpacity() const;
+
+	/** @return the name of an image for the scrollbar thumb based on whether the user is dragging or hovering it. */
+	const FSlateBrush* GetDragThumbImage() const;
+
+	/** The scrollbar's visibility as specified by the user. Will be compounded with internal visibility rules. */
+	TAttribute<EVisibility> UserVisibility;
+
+	TSharedPtr<SBorder> DragThumb;
+	bool bDraggingThumb;
+	TSharedPtr<SScrollBarTrack> Track;
+	FOnUserScrolled OnUserScrolled;
+	float DragGrabOffset;
+	EOrientation Orientation;
+
+	/** Image to use when the scrollbar thumb is in its normal state */
+	const FSlateBrush* NormalThumbImage;
+	/** Image to use when the scrollbar thumb is in its hovered state */
+	const FSlateBrush* HoveredThumbImage;
+	/** Image to use when the scrollbar thumb is in its dragged state */
+	const FSlateBrush* DraggedThumbImage;
+};
diff --git a/Engine/Source/Runtime/Slate/Public/Widgets/Layout/SScrollBox.h b/Engine/Source/Runtime/Slate/Public/Widgets/Layout/SScrollBox.h
new file mode 100644
--- /dev/null
+++ b/Engine/Source/Runtime/Slate/Public/Widgets/Layout/SScrollBox.h
@@ -0,0 +1,122 @@
+// Copyright 1998-2014 Epic Games, Inc. All Rights Reserved.
+
+#pragma once
+
+/** SScrollBox can scroll through an arbitrary number of widgets. */
+class SLATE_API SScrollBox : public SCompoundWidget
+{
+public:
+	/** A Slot that provides layout options for the contents of a scrollable box. */
+	class SLATE_API FSlot : public TSupportsOneChildMixin<SWidget, FSlot>, public TSupportsContentPaddingMixin<FSlot>
+	{
+	public:
+		FSlot()
+			: HAlignment(HAlign_Fill)
+		{
+
+		}
+
+		FSlot& HAlign( EHorizontalAlignment InHAlignment )
+		{
+			HAlignment = InHAlignment;
+			return *this;
+		}
+		
+		EHorizontalAlignment HAlignment;
+	};
+
+	SLATE_BEGIN_ARGS(SScrollBox)
+		: _Style( &FCoreStyle::Get().GetWidgetStyle<FScrollBoxStyle>("ScrollBox") )
+		{}
+		
+		SLATE_SUPPORTS_SLOT( FSlot )
+
+		/** Style used to draw this scrollbox */
+		SLATE_STYLE_ARGUMENT( FScrollBoxStyle, Style )
+
+	SLATE_END_ARGS()
+
+
+	/** @return a new slot. Slots contain children for SScrollBox */
+	static FSlot& Slot();
+
+	void Construct( const FArguments& InArgs );
+
+	/** Adds a slot to SScrollBox */
+	SScrollBox::FSlot& AddSlot();
+
+	/** Removes a slot at the specified location */
+	void RemoveSlot( const TSharedRef<SWidget>& WidgetToRemove );
+
+	/** Removes all children from the box */
+	void ClearChildren();
+
+	/** @return Returns true if the user is currently interactively scrolling the view by holding
+		        the right mouse button and dragging. */
+	bool IsRightClickScrolling() const;
+
+	void SetScrollOffset( float NewScrollOffset );
+
+public:
+
+	// SWidget interface
+	virtual void Tick( const FGeometry& AllottedGeometry, const double InCurrentTime, const float InDeltaTime ) OVERRIDE;
+	virtual FReply OnMouseButtonDown( const FGeometry& MyGeometry, const FPointerEvent& MouseEvent ) OVERRIDE;
+	virtual FReply OnMouseButtonUp( const FGeometry& MyGeometry, const FPointerEvent& MouseEvent ) OVERRIDE;
+	virtual FReply OnMouseMove( const FGeometry& MyGeometry, const FPointerEvent& MouseEvent ) OVERRIDE;
+	virtual void OnMouseLeave( const FPointerEvent& MouseEvent ) OVERRIDE;
+	virtual FReply OnMouseWheel( const FGeometry& MyGeometry, const FPointerEvent& MouseEvent ) OVERRIDE;
+	virtual FReply OnDragDetected( const FGeometry& MyGeometry, const FPointerEvent& MouseEvent ) OVERRIDE;
+	virtual FCursorReply OnCursorQuery( const FGeometry& MyGeometry, const FPointerEvent& CursorEvent ) const OVERRIDE;
+	virtual int32 OnPaint( const FGeometry& AllottedGeometry, const FSlateRect& MyClippingRect, FSlateWindowElementList& OutDrawElements, int32 LayerId, const FWidgetStyle& InWidgetStyle, bool bParentEnabled ) const OVERRIDE;
+	// End of SWidget interface
+
+private:
+
+	/** Scroll offset that the user asked for. We will clamp it before actually scrolling there. */
+	float DesiredScrollOffset;
+
+	/**
+	 * Scroll the view by ScrollAmount given its currently AllottedGeometry.
+	 *
+	 * @param AllottedGeometry  The geometry allotted for this SScrollBox by the parent
+	 * @param ScrollAmount      
+	 * @return Whether or not the scroll was fully handled
+	 */
+	bool ScrollBy( const FGeometry& AllottedGeometry, float ScrollAmount, bool InAnimateScroll = true );
+
+	/** Invoked when the user scroll via the scrollbar */
+	void ScrollBar_OnUserScrolled( float InScrollOffsetFraction );
+
+	/** Does the user need a hint that they can scroll up? */
+	FSlateColor GetTopShadowOpacity() const;
+	
+	/** Does the user need a hint that they can scroll down? */
+	FSlateColor GetBottomShadowOpacity() const;
+
+	/** Does the user need to see the scroll bar? */
+	EVisibility GetScrollBarVisibility() const;
+
+	TSharedPtr<class SScrollPanel> ScrollPanel;
+	TSharedPtr<SScrollBar> ScrollBar;
+	
+	/** Are we actively scrolling right now */
+	bool bIsScrolling;
+
+	/** Should the current scrolling be animated or immediately jump to the desired scroll offer */
+	bool bAnimateScroll;
+
+	/** How much we scrolled while the rmb has been held */
+	float AmountScrolledWhileRightMouseDown;
+
+	/** Helper object to manage inertial scrolling */
+	FInertialScrollManager InertialScrollManager;
+
+	/**	The current position of the software cursor */
+	FVector2D SoftwareCursorPosition;
+
+	/**	Whether the software cursor should be drawn in the viewport */
+	bool bShowSoftwareCursor;
+};
+
+
diff --git a/Engine/Source/Runtime/Slate/Public/Widgets/Layout/SSpacer.h b/Engine/Source/Runtime/Slate/Public/Widgets/Layout/SSpacer.h
new file mode 100644
--- /dev/null
+++ b/Engine/Source/Runtime/Slate/Public/Widgets/Layout/SSpacer.h
@@ -0,0 +1,44 @@
+// Copyright 1998-2014 Epic Games, Inc. All Rights Reserved.
+
+#pragma once
+
+
+class SLATE_API SSpacer : public SLeafWidget
+{
+public:
+
+	SLATE_BEGIN_ARGS( SSpacer )
+		: _Size(FVector2D::ZeroVector)
+		{
+			_Visibility = EVisibility::SelfHitTestInvisible;
+		}
+
+		SLATE_ATTRIBUTE( FVector2D, Size )
+	SLATE_END_ARGS()
+
+	/**
+	 * Construct this widget
+	 *
+	 * @param	InArgs	The declaration data for this widget
+	 */
+	void Construct( const FArguments& InArgs );
+
+	// SWidget interface
+	virtual int32 OnPaint( const FGeometry& AllottedGeometry, const FSlateRect& MyClippingRect, FSlateWindowElementList& OutDrawElements, int32 LayerId, const FWidgetStyle& InWidgetStyle, bool bParentEnabled ) const OVERRIDE;
+	virtual FVector2D ComputeDesiredSize() const OVERRIDE;
+	// End of SWidget interface
+
+	FVector2D GetSize() const
+	{
+		return SpacerSize.Get();
+	}
+
+	void SetSize( TAttribute<FVector2D> InSpacerSize )
+	{
+		SpacerSize = InSpacerSize;
+	}
+
+
+private:
+	TAttribute<FVector2D> SpacerSize;
+};
diff --git a/Engine/Source/Runtime/Slate/Public/Widgets/Layout/SSplitter.h b/Engine/Source/Runtime/Slate/Public/Widgets/Layout/SSplitter.h
new file mode 100644
--- /dev/null
+++ b/Engine/Source/Runtime/Slate/Public/Widgets/Layout/SSplitter.h
@@ -0,0 +1,490 @@
+// Copyright 1998-2014 Epic Games, Inc. All Rights Reserved.
+
+#pragma once
+
+namespace ESplitterResizeMode
+{
+	enum Type
+	{
+		Fixed,		
+		Fill,
+	};
+}
+
+/**
+ * SSPlitter divides its allotted area into N segments, where N is the number of children it has.
+ * It allows the users to resize the children along the splitters axis: that is, horizontally or vertically.
+ */
+class SLATE_API SSplitter : public SPanel
+{
+
+public:
+	/** How should a child's size be determined */
+	enum ESizeRule
+	{
+		/** Get the DesiredSize() of the content */
+		SizeToContent,
+		/** Use a fraction of the parent's size */
+		FractionOfParent
+	};
+
+	DECLARE_DELEGATE_OneParam(
+		FOnSlotResized,
+		/** The new size coefficient of the slot */
+		float );
+
+public:
+	class FSlot
+	{
+	public:		
+		FSlot()
+		: SizingRule( FractionOfParent )
+		, SizeValue( 1 )
+		, Widget( SNullWidget::NullWidget )
+		{
+		}
+
+		FSlot& operator[]( TSharedRef<SWidget> InWidget )
+		{
+			Widget = InWidget;
+			return *this;
+		}
+
+		FSlot& Value( const TAttribute<float>& InValue )
+		{
+			SizeValue = InValue;
+			return *this;
+		}
+		
+		FSlot& OnSlotResized( const FOnSlotResized& InHandler )
+		{
+			OnSlotResized_Handler = InHandler;
+			return *this;
+		}
+
+		FSlot& SizeRule( const TAttribute<ESizeRule>& InSizeRule ) 
+		{
+			SizingRule = InSizeRule;
+			return *this;
+		}
+
+		FSlot& Expose( FSlot*& OutVarToInit )
+		{
+			OutVarToInit = this;
+			return *this;
+		}
+
+		TAttribute<ESizeRule> SizingRule;
+		TAttribute<float> SizeValue;
+		TSharedRef<SWidget> Widget;
+		FOnSlotResized OnSlotResized_Handler;
+	};
+
+	/** @return a new SSplitter::FSlot() */
+	static FSlot& Slot();
+	
+	/**
+	 * Add a slot to the splitter at the specified index
+	 * Sample usage:
+	 *     SomeSplitter->AddSlot()
+	 *     [
+	 *       SNew(SSomeWidget)
+	 *     ];
+	 *
+	 * @return the new slot.
+	 */
+	FSlot& AddSlot( int32 AtIndex = INDEX_NONE );
+
+
+	SLATE_BEGIN_ARGS(SSplitter)
+		: _Style( &FCoreStyle::Get().GetWidgetStyle<FSplitterStyle>("Splitter") )
+		, _Orientation( Orient_Horizontal )
+		, _ResizeMode( ESplitterResizeMode::Fixed )
+		, _PhysicalSplitterHandleSize( 5.0f )
+		, _HitDetectionSplitterHandleSize( 5.0f )
+		, _OnSplitterFinishedResizing()
+		{}
+
+		SLATE_SUPPORTS_SLOT(FSlot)
+
+		/** Style used to draw this splitter */
+		SLATE_STYLE_ARGUMENT( FSplitterStyle, Style )
+
+		SLATE_ARGUMENT( EOrientation, Orientation )
+
+		SLATE_ARGUMENT( ESplitterResizeMode::Type, ResizeMode )
+
+		SLATE_ARGUMENT( float, PhysicalSplitterHandleSize )
+
+		SLATE_ARGUMENT( float, HitDetectionSplitterHandleSize )
+
+		SLATE_EVENT( FSimpleDelegate, OnSplitterFinishedResizing )
+
+	SLATE_END_ARGS()
+
+	SSplitter();
+
+	/**
+	 * Construct this widget
+	 *
+	 * @param	InArgs	The declaration data for this widget
+	 */
+	void Construct( const FArguments& InArgs );
+
+public:
+
+	/**
+	 * Get the slot at the specified index
+	 *
+	 * @param SlotIndex    Replace the child at this index.
+	 *
+	 * @return Slot at the index specified by SlotIndex
+	 */
+	SSplitter::FSlot& SlotAt( int32 SlotIndex );
+
+
+	/**
+	 * Remove the child at IndexToRemove
+	 *
+	 * @param IndexToRemove     Remove the slot and child at this index.
+	 */
+	void RemoveAt( int32 IndexToRemove );
+
+public:
+
+	virtual void ArrangeChildren( const FGeometry& AllottedGeometry, FArrangedChildren& ArrangedChildren ) const OVERRIDE;
+
+
+	virtual int32 OnPaint( const FGeometry& AllottedGeometry, const FSlateRect& MyClippingRect, FSlateWindowElementList& OutDrawElements, int32 LayerId, const FWidgetStyle& InWidgetStyle, bool bParentEnabled ) const OVERRIDE;
+
+
+	/**
+	 * A Panel's desired size in the space required to arrange of its children on the screen while respecting all of
+	 * the children's desired sizes and any layout-related options specified by the user. See StackPanel for an example.
+	 */
+	virtual FVector2D ComputeDesiredSize() const OVERRIDE;
+
+	/**
+	 * All widgets must provide a way to access their children in a layout-agnostic way.
+	 * Panels store their children in Slots, which creates a dilemma. Most panels
+	 * can store their children in a TPanelChildren<Slot>, where the Slot class
+	 * provides layout information about the child it stores. In that case
+	 * GetChildren should simply return the TPanelChildren<Slot>. See StackPanel for an example.
+	 */
+	virtual FChildren* GetChildren() OVERRIDE;
+
+	/**
+	 * The system calls this method to notify the widget that a mouse button was pressed within it. This event is bubbled.
+	 *
+	 * @param MyGeometry The Geometry of the widget receiving the event
+	 * @param MouseEvent Information about the input event
+	 *
+	 * @return Whether the event was handled along with possible requests for the system to take action.
+	 */
+	virtual FReply OnMouseButtonDown( const FGeometry& MyGeometry, const FPointerEvent& MouseEvent ) OVERRIDE;
+	
+	/**
+	 * The system calls this method to notify the widget that a mouse button was release within it. This event is bubbled.
+	 *
+	 * @param MyGeometry The Geometry of the widget receiving the event
+	 * @param MouseEvent Information about the input event
+	 *
+	 * @return Whether the event was handled along with possible requests for the system to take action.
+	 */
+	virtual FReply OnMouseButtonUp( const FGeometry& MyGeometry, const FPointerEvent& MouseEvent ) OVERRIDE;
+	
+	virtual FReply OnMouseMove( const FGeometry& MyGeometry, const FPointerEvent& MouseEvent ) OVERRIDE;
+
+	virtual void OnMouseLeave( const FPointerEvent& MouseEvent ) OVERRIDE;
+
+	/**
+	 * The system asks each widget under the mouse to provide a cursor. This event is bubbled.
+	 * 
+	 * @return FCursorReply::Unhandled() if the event is not handled; return FCursorReply::Cursor() otherwise.
+	 */
+	virtual FCursorReply OnCursorQuery( const FGeometry& MyGeometry, const FPointerEvent& CursorEvent ) const OVERRIDE;
+
+	/**
+	 * Change the orientation of the splitter
+	 *
+	 * @param NewOrientation  Should the splitter be horizontal or vertical
+	 */
+	void SetOrientation( EOrientation NewOrientation );
+
+	/**
+	 * @return the current orientation of the splitter.
+	 */
+	EOrientation GetOrientation() const;	
+
+
+protected:
+
+	/**
+	 * Given the index of the dragged handle and the children, find a child above/left_of of the dragged handle that can be resized.
+	 *
+	 * @return INDEX_NONE if no such child can be found.
+	 */
+	static int32 FindResizeableSlotBeforeHandle( int32 DraggedHandle, const TPanelChildren<FSlot>& Children );
+
+	/**
+	 * Given the index of the dragged handle and the children, find a child below/right_of the dragged handle that can be resized
+	 *
+	 * @return Children.Num() if no such child can be found.
+	 */
+	static int32 FindResizeableSlotAfterHandle( int32 DraggedHandle, const TPanelChildren<FSlot>& Children );
+
+	static void FindAllResizeableSlotsAfterHandle( int32 DraggedHandle, const TPanelChildren<FSlot>& Children, TArray< int32 >& OutSlotIndicies );
+
+	/**
+	 * Resizes the children based on user input. The template parameter Orientation corresponds to the splitter being horizontal or vertical.
+	 * 
+	 * @param DraggedHandle    The index of the handle that the user is dragging.
+	 * @param LocalMousePos    The position of the mouse in this widgets local space.
+	 * @param Children         A reference to this splitter's children array; we will modify the children's layout values.
+	 * @param ChildGeometries  The arranged children; we need their sizes and positions so that we can perform a resizing.
+	 */
+	template<EOrientation SplitterOrientation>
+	static void HandleResizing( const float PhysicalSplitterHandleSize, const ESplitterResizeMode::Type ResizeMode, int32 DraggedHandle, const FVector2D& LocalMousePos, TPanelChildren<FSlot>& Children, FArrangedChildren& ChildGeometries );
+
+	/**
+	 * @param ProposedSize  A size that a child would like to be
+	 *
+	 * @return A size that is clamped against the minimum size allowed for children.
+	 */
+	static float ClampChild( float ProposedSize );
+
+	/**
+	 * Given a mouse position within the splitter, figure out which resize handle we are hovering (if any).
+	 *
+
+	 * @param LocalMousePos  The mouse position within this splitter.
+	 * @param ChildGeometris The arranged children and their geometries; we need to test the mouse against them.
+	 *
+	 * @return The index of the handle being hovered, or INDEX_NONE if we are not hovering a handle.
+	 */
+	template<EOrientation SplitterOrientation>
+	static int32 GetHandleBeingResizedFromMousePosition(  float PhysicalSplitterHandleSize, float HitDetectionSplitterHandleSize, FVector2D LocalMousePos, FArrangedChildren& ChildGeometries );
+
+	TPanelChildren< FSlot > Children;
+
+	int32 HoveredHandleIndex;
+	bool bIsResizing;
+	EOrientation Orientation;
+	ESplitterResizeMode::Type ResizeMode;
+
+	FSimpleDelegate OnSplitterFinishedResizing;
+
+	/** The thickness of the grip area that the user uses to resize a splitter */
+	float PhysicalSplitterHandleSize;
+	float HitDetectionSplitterHandleSize;
+
+	const FSplitterStyle* Style;
+};
+
+/**
+ * SSplitter2x2													
+ * A splitter which has exactly 4 children and allows simultaneous		
+ * of all children along an axis as well as resizing all children
+ * by dragging the center of the splitter.
+ */
+class SLATE_API SSplitter2x2 : public SPanel
+{
+public:
+	class FSlot
+	{
+	public:	
+		/** Default Constructor.  Initially each slot takes up a quarter of the entire space */
+		FSlot()
+			: PercentageAttribute( FVector2D(0.5, 0.5) )
+			, Widget( SNullWidget::NullWidget )
+		{
+		}
+
+		/** Copy Constructor */
+		FSlot( const TSharedRef<SWidget>& InWidget )
+			: PercentageAttribute( FVector2D(0.5, 0.5) )
+			, Widget( InWidget )
+		{
+		}
+
+		/** Declarative syntax support */
+		FSlot& operator[]( TSharedRef<SWidget> InWidget )
+		{
+			Widget = InWidget;
+			return *this;
+		}
+
+		/**
+		 * Sets the percentage attribute
+		 *
+		 * @param Value The new percentage value
+		 */
+		FSlot& SetPercentage( const FVector2D& Value )
+		{
+			PercentageAttribute.Set( Value );
+			return *this;
+		}
+	public:
+		/** The percentage of the alloted space of the splitter that this slot requires */
+		TAttribute<FVector2D> PercentageAttribute;
+		/** The widget in this slot */
+		TSharedRef<SWidget> Widget;
+	};
+
+	SLATE_BEGIN_ARGS( SSplitter2x2 ){}
+		SLATE_NAMED_SLOT( FArguments, TopLeft )
+		SLATE_NAMED_SLOT( FArguments, BottomLeft )
+		SLATE_NAMED_SLOT( FArguments, TopRight )
+		SLATE_NAMED_SLOT( FArguments, BottomRight )
+	SLATE_END_ARGS()
+
+
+	void Construct( const FArguments& InArgs );
+
+	/**
+	 * Returns the widget displayed in the splitter top left area
+	 *
+	 * @return	Top left widget
+	 */
+	TSharedRef< SWidget > GetTopLeftContent();
+
+	/**
+	 * Returns the widget displayed in the splitter bottom left area
+	 *
+	 * @return	Bottom left widget
+	 */
+	TSharedRef< SWidget > GetBottomLeftContent();
+
+	/**
+	 * Returns the widget displayed in the splitter top right area
+	 *
+	 * @return	Top right widget
+	 */
+	TSharedRef< SWidget > GetTopRightContent();
+
+	/**
+	 * Returns the widget displayed in the splitter bottom right area
+	 *
+	 * @return	Bottom right widget
+	 */
+	TSharedRef< SWidget > GetBottomRightContent();
+
+	/**
+	 * Sets the widget to be displayed in the splitter top left area
+	 *
+	 * @param	TopLeftContent	The top left widget
+	 */
+	void SetTopLeftContent( TSharedRef< SWidget > TopLeftContent );
+
+	/**
+	 * Sets the widget to be displayed in the splitter bottom left area
+	 *
+	 * @param	BottomLeftContent	The bottom left widget
+	 */
+	void SetBottomLeftContent( TSharedRef< SWidget > BottomLeftContent );
+
+	/**
+	 * Sets the widget to be displayed in the splitter top right area
+	 *
+	 * @param	TopRightContent	The top right widget
+	 */
+	void SetTopRightContent( TSharedRef< SWidget > TopRightContent );
+
+	/**
+	 * Sets the widget to be displayed in the splitter bottom right area
+	 *
+	 * @param	BottomRightContent	The bottom right widget
+	 */
+	void SetBottomRightContent( TSharedRef< SWidget > BottomRightContent );
+
+	/** Returns an array of size percentages for the children in this order: TopLeft, BottomLeft, TopRight, BottomRight */
+	void GetSplitterPercentages( TArray< FVector2D >& OutPercentages ) const;
+	
+	/** Sets the size percentages for the children in this order: TopLeft, BottomLeft, TopRight, BottomRight */
+	void SetSplitterPercentages( const TArray< FVector2D >& InPercentages );
+
+
+private:
+
+	virtual void ArrangeChildren( const FGeometry& AllottedGeometry, FArrangedChildren& ArrangedChildren ) const OVERRIDE;
+
+	/**
+	 * A Panel's desired size in the space required to arrange of its children on the screen while respecting all of
+	 * the children's desired sizes and any layout-related options specified by the user. See StackPanel for an example.
+	 */
+	virtual FVector2D ComputeDesiredSize() const OVERRIDE;
+
+	/**
+	 * All widgets must provide a way to access their children in a layout-agnostic way.
+	 * Panels store their children in Slots, which creates a dilemma. Most panels
+	 * can store their children in a TPanelChildren<Slot>, where the Slot class
+	 * provides layout information about the child it stores. In that case
+	 * GetChildren should simply return the TPanelChildren<Slot>. See StackPanel for an example.
+	 */
+	virtual FChildren* GetChildren() OVERRIDE;
+
+	/**
+	 * The system calls this method to notify the widget that a mouse button was pressed within it. This event is bubbled.
+	 *
+	 * @param MyGeometry The Geometry of the widget receiving the event
+	 * @param MouseEvent Information about the input event
+	 *
+	 * @return Whether the event was handled along with possible requests for the system to take action.
+	 */
+ 	virtual FReply OnMouseButtonDown( const FGeometry& MyGeometry, const FPointerEvent& MouseEvent ) OVERRIDE;
+	
+	/**
+	 * The system calls this method to notify the widget that a mouse button was release within it. This event is bubbled.
+	 *
+	 * @param MyGeometry The Geometry of the widget receiving the event
+	 * @param MouseEvent Information about the input event
+	 *
+	 * @return Whether the event was handled along with possible requests for the system to take action.
+	 */
+ 	virtual FReply OnMouseButtonUp( const FGeometry& MyGeometry, const FPointerEvent& MouseEvent ) OVERRIDE;
+
+	/**
+	 * The system calls this method to notify the widget that a mouse moved within it. This event is bubbled.
+	 *
+	 * @param MyGeometry The Geometry of the widget receiving the event
+	 * @param MouseEvent Information about the input event
+	 *
+	 * @return Whether the event was handled along with possible requests for the system to take action.
+	 */
+	virtual FReply OnMouseMove( const FGeometry& MyGeometry, const FPointerEvent& MouseEvent ) OVERRIDE;
+
+	/**
+	 * @return The cursor that should be visible
+	 */
+	virtual FCursorReply OnCursorQuery( const FGeometry& MyGeometry, const FPointerEvent& CursorEvent ) const OVERRIDE;
+
+	/**
+	 * Calculates the axis being resized
+	 * 
+	 * @param MyGeometry	The geometry of this widget
+	 * @param LocalMousePos	The local space current mouse position
+	 */
+	int32 CalculateResizingAxis( const FGeometry& MyGeometry, const FVector2D& LocalMousePos ) const;
+
+	/**
+	 * Resizes all children based on a user moving the splitter handles
+	 * 
+	 * @param ArrangedChildren	The current geometry of all arranged children before the user moved the splitter
+	 * @param LocalMousePos		The current mouse position		
+	 */
+	void ResizeChildren( FArrangedChildren& ArrangedChildren, FVector2D LocalMousePos );
+
+
+private:
+
+	/** The children of the splitter. There can only be four */
+	TPanelChildren<FSlot> Children;
+
+	/** The axis currently being resized or INDEX_NONE if no resizing */
+	int32 ResizingAxis;
+
+	/** true if a splitter axis is currently being resized. */
+	bool bIsResizing;
+
+	float SplitterHandleSize;
+};
diff --git a/Engine/Source/Runtime/Slate/Public/Widgets/Layout/SUniformGridPanel.h b/Engine/Source/Runtime/Slate/Public/Widgets/Layout/SUniformGridPanel.h
new file mode 100644
--- /dev/null
+++ b/Engine/Source/Runtime/Slate/Public/Widgets/Layout/SUniformGridPanel.h
@@ -0,0 +1,78 @@
+// Copyright 1998-2014 Epic Games, Inc. All Rights Reserved.
+
+#pragma once
+
+/** A panel that evenly divides up available space between all of its children. */
+class SLATE_API SUniformGridPanel : public SPanel
+{
+public:
+	/** Stores the per-child info for this panel type */
+	struct FSlot : public TSupportsOneChildMixin<SWidget, FSlot>, public TSupportsContentAlignmentMixin<FSlot>
+	{
+		FSlot( int32 InColumn, int32 InRow )
+		: TSupportsContentAlignmentMixin<FSlot>(HAlign_Fill, VAlign_Fill)
+		, Column( InColumn )
+		, Row( InRow )
+		{
+		}
+
+		int32 Column;
+		int32 Row;
+	};
+
+	/**
+	 * Used by declarative syntax to create a Slot in the specified Column, Row.
+	 */
+	static FSlot& Slot( int32 Column, int32 Row )
+	{
+		return *(new FSlot( Column, Row ));
+	}
+
+	SLATE_BEGIN_ARGS( SUniformGridPanel )
+		: _SlotPadding( FMargin(0.0f) )
+		, _MinDesiredSlotWidth( 0.0f )
+		, _MinDesiredSlotHeight( 0.0f )
+		{
+			_Visibility = EVisibility::SelfHitTestInvisible;
+		}
+
+		/** Slot type supported by this panel */
+		SLATE_SUPPORTS_SLOT(FSlot)
+		
+		/** Padding given to each slot */
+		SLATE_ATTRIBUTE(FMargin, SlotPadding)
+
+		/** The minimum desired width of the slots */
+		SLATE_ATTRIBUTE(float, MinDesiredSlotWidth)
+
+		/** The minimum desired height of the slots */
+		SLATE_ATTRIBUTE(float, MinDesiredSlotHeight)
+
+	SLATE_END_ARGS()
+
+	void Construct( const FArguments& InArgs );
+
+	// BEGIN SPanel INTERFACE	
+	virtual void ArrangeChildren( const FGeometry& AllottedGeometry, FArrangedChildren& ArrangedChildren ) const OVERRIDE;
+	virtual FVector2D ComputeDesiredSize() const OVERRIDE;
+	virtual FChildren* GetChildren() OVERRIDE;
+	// END SPanel INTERFACE
+
+	/**
+	 * Dynamically add a new slot to the UI at specified Column and Row.
+	 *
+	 * @return A reference to the newly-added slot
+	 */
+	FSlot& AddSlot( int32 Column, int32 Row );
+
+	/** Removes all slots from the panel */
+	void ClearChildren();
+
+private:
+	TPanelChildren<FSlot> Children;
+	TAttribute<FMargin> SlotPadding;
+	int32 NumColumns;
+	int32 NumRows;
+	TAttribute<float> MinDesiredSlotWidth;
+	TAttribute<float> MinDesiredSlotHeight;
+};
\ No newline at end of file
diff --git a/Engine/Source/Runtime/Slate/Public/Widgets/Layout/SWidgetSwitcher.h b/Engine/Source/Runtime/Slate/Public/Widgets/Layout/SWidgetSwitcher.h
new file mode 100644
--- /dev/null
+++ b/Engine/Source/Runtime/Slate/Public/Widgets/Layout/SWidgetSwitcher.h
@@ -0,0 +1,162 @@
+// Copyright 1998-2014 Epic Games, Inc. All Rights Reserved.
+
+/*=============================================================================
+	SWidgetSwitcher.h: Declares the SWidgetSwitcher class.
+=============================================================================*/
+
+#pragma once
+
+
+/**
+ * Implements a widget switcher.
+ *
+ * A widget switcher is like a tab control, but without tabs. At most one widget is visible at time.
+ */
+class SLATE_API SWidgetSwitcher
+	: public SPanel
+{
+public:
+
+	SLATE_BEGIN_ARGS(SWidgetSwitcher)
+		: _WidgetIndex(0)
+		{
+			_Visibility = EVisibility::SelfHitTestInvisible;
+		}
+
+		SLATE_SUPPORTS_SLOT(FSimpleSlot)
+
+		/** Holds the index of the initial widget to be displayed (INDEX_NONE = default). */
+		SLATE_ATTRIBUTE(int32, WidgetIndex)
+
+	SLATE_END_ARGS()
+
+public:
+
+	/**
+	 * Adds a slot to the widget switcher at the specified location.
+	 *
+	 * @param SlotIndex - The index at which to insert the slot, or INDEX_NONE to append.
+	 */
+	FSimpleSlot& AddSlot( int32 SlotIndex = INDEX_NONE );
+
+	/**
+	 * Constructs the widget.
+	 */
+	void Construct( const FArguments& InArgs );
+
+	/**
+	 * Gets the active widget.
+	 *
+	 * @return Active widget.
+	 */
+	TSharedPtr<SWidget> GetActiveWidget( ) const;
+
+	/**
+	 * Gets the slot index of the currently active widget.
+	 *
+	 * @return The slot index, or INDEX_NONE if no widget is active.
+	 */
+	int32 GetActiveWidgetIndex( ) const
+	{
+		return WidgetIndex.Get();
+	}
+
+	/**
+	 * Gets the number of widgets that this switcher manages.
+	 */
+	int32 GetNumWidgets( ) const
+	{
+		return AllChildren.Num();
+	}
+
+	/**
+	 * Gets the widget in the specified slot.
+	 *
+	 * @param SlotIndex - The slot index of the widget to get.
+	 *
+	 * @return The widget, or NULL if the slot does not exist.
+	 */
+	TSharedPtr<SWidget> GetWidget( int32 SlotIndex ) const;
+
+	/**
+	 * Gets the slot index of the specified widget.
+	 *
+	 * @param Widget - The widget to get the index for.
+	 *
+	 * @return The slot index, or INDEX_NONE if the widget does not exist.
+	 */
+	int32 GetWidgetIndex( TSharedRef<SWidget> Widget ) const;
+
+	/**
+	 * Sets the active widget.
+	 *
+	 * @param Widget - The widget to activate.
+	 */
+	void SetActiveWidget( TSharedRef<SWidget> Widget )
+	{
+		SetActiveWidgetIndex(GetWidgetIndex(Widget));
+	}
+
+	/**
+	 * Activates the widget at the specified index.
+	 *
+	 * @param Index - The slot index.
+	 */
+	void SetActiveWidgetIndex( int32 Index );
+
+public:
+
+	/**
+	 * Creates a new widget slot.
+	 *
+	 * @return A new slot.
+	 */
+	static FSimpleSlot& Slot( )
+	{
+		return *(new FSimpleSlot());
+	}
+
+protected:
+
+	// Begin SCompoundWidget interface
+
+	virtual void ArrangeChildren( const FGeometry& AllottedGeometry, FArrangedChildren& ArrangedChildren ) const OVERRIDE;
+	
+	virtual FVector2D ComputeDesiredSize( ) const OVERRIDE;
+
+	virtual FChildren* GetChildren( ) OVERRIDE;
+
+	// End SCompoundWidget interface
+
+private:
+
+	/** Required to implement GetChildren() in a way that can dynamically return the currently active child. */
+	class SLATE_API FOneDynamicChild
+		: public FChildren
+	{
+	public:
+
+		FOneDynamicChild( TPanelChildren<FSimpleSlot>* InAllChildren = NULL, const TAttribute<int32>* InWidgetIndex = NULL )
+			: AllChildren( InAllChildren )
+			, WidgetIndex( InWidgetIndex )
+		{ }
+		
+		virtual int32 Num() const OVERRIDE { return AllChildren->Num() > 0 ? 1 : 0; }
+		
+		virtual TSharedRef<SWidget> GetChildAt( int32 Index ) OVERRIDE { check(Index == 0); return (*AllChildren)[WidgetIndex->Get()].Widget; }
+		
+		virtual TSharedRef<const SWidget> GetChildAt( int32 Index ) const OVERRIDE { check(Index == 0); return (*AllChildren)[WidgetIndex->Get()].Widget; }
+		
+	private:
+
+		TPanelChildren<FSimpleSlot>* AllChildren;
+		const TAttribute<int32>* WidgetIndex;
+
+	} OneDynamicChild;
+
+	/** Holds the desired widget index */
+	TAttribute<int32> WidgetIndex;
+
+	// Holds the collection of widgets.
+	TPanelChildren<FSimpleSlot> AllChildren;
+};
diff --git a/Engine/Source/Runtime/Slate/Public/Widgets/Layout/SWrapBox.h b/Engine/Source/Runtime/Slate/Public/Widgets/Layout/SWrapBox.h
new file mode 100644
--- /dev/null
+++ b/Engine/Source/Runtime/Slate/Public/Widgets/Layout/SWrapBox.h
@@ -0,0 +1,109 @@
+// Copyright 1998-2014 Epic Games, Inc. All Rights Reserved.
+
+#pragma once
+
+/** A slot that support alignment of content and padding */
+class SLATE_API FWrapSlot : public TSupportsOneChildMixin<SWidget, FWrapSlot>, public TSupportsContentAlignmentMixin<FWrapSlot>, public TSupportsContentPaddingMixin<FWrapSlot>
+{
+public:
+	FWrapSlot()
+		: TSupportsContentAlignmentMixin<FWrapSlot>( HAlign_Fill, VAlign_Fill )
+		, bSlotFillEmptySpace( false )
+	{
+	}
+
+	/** If the total available space in the wrap panel drops below this threshold, this slot will attempt to fill an entire line. */
+	FWrapSlot& FillLineWhenWidthLessThan( TOptional<float> InFillLineWhenWidthLessThan )
+	{
+		SlotFillLineWhenWidthLessThan = InFillLineWhenWidthLessThan;
+		return *( static_cast<FWrapSlot*>( this ) );
+	}
+
+	/** Should this slot fill the remaining space on the line? */
+	FWrapSlot& FillEmptySpace( bool bInFillEmptySpace )
+	{
+		bSlotFillEmptySpace = bInFillEmptySpace;
+		return *( static_cast<FWrapSlot*>( this ) );
+	}
+
+	TOptional<float> SlotFillLineWhenWidthLessThan;
+	bool bSlotFillEmptySpace;
+};
+
+/**
+ * Arranges widgets left-to-right.
+ * When the widgets exceed the PreferredWidth
+ * the SWrapBox will place widgets on the next line.
+ *
+ * Illustration:
+ *                      +-----Preferred Width
+ *                      |
+ *       [-----------][-|-]
+ *       [--][------[--]|
+ *       [--------------|]
+ *       [---]          |
+ */
+class SLATE_API SWrapBox : public SPanel
+{
+
+public:
+
+	SLATE_BEGIN_ARGS(SWrapBox)
+		: _PreferredWidth( 100.f )
+		, _InnerSlotPadding( FVector2D::ZeroVector )
+		, _UseAllottedWidth( false )
+		{
+			_Visibility = EVisibility::SelfHitTestInvisible;
+		}
+
+		/** The slot supported by this panel */
+		SLATE_SUPPORTS_SLOT( FWrapSlot )
+
+		/** The preferred width, if not set will fill the space */
+		SLATE_ATTRIBUTE( float, PreferredWidth )
+
+		/** The inner slot padding goes between slots sharing borders */
+		SLATE_ARGUMENT( FVector2D, InnerSlotPadding )
+
+		/** if true, the PreferredWidth will always match the room available to the SWrapBox  */
+		SLATE_ARGUMENT( bool, UseAllottedWidth )
+	SLATE_END_ARGS()
+
+	static FWrapSlot& Slot();
+
+	FWrapSlot& AddSlot();
+
+	/** Removes a slot from this box panel which contains the specified SWidget
+	 *
+	 * @param SlotWidget The widget to match when searching through the slots
+	 * @returns The index in the children array where the slot was removed and -1 if no slot was found matching the widget
+	 */
+	int32 RemoveSlot( const TSharedRef<SWidget>& SlotWidget );
+
+	void Construct( const FArguments & InArgs );
+
+	virtual void Tick( const FGeometry& AllottedGeometry, const double InCurrentTime, const float InDeltaTime ) OVERRIDE;
+
+	virtual void ArrangeChildren( const FGeometry& AllottedGeometry, FArrangedChildren& ArrangedChildren ) const OVERRIDE;
+
+	void ClearChildren();
+
+	virtual FVector2D ComputeDesiredSize() const OVERRIDE;
+
+	virtual FChildren* GetChildren() OVERRIDE;
+
+private:
+
+	/** How wide this panel should appear to be. Any widgets past this line will be wrapped onto the next line. */
+	TAttribute<float> PreferredWidth;
+
+	/** The slots that contain this panel's children. */
+	TPanelChildren<FWrapSlot> Slots;
+
+	/** When two slots end up sharing a border, this will put that much padding between then, but otherwise wont. */
+	FVector2D InnerSlotPadding;
+
+	/** If true the box will have a preferred width equal to its alloted width  */
+	bool bUseAllottedWidth;
+};
+
diff --git a/Engine/Source/Runtime/Slate/Public/Widgets/Notifications/SErrorHint.h b/Engine/Source/Runtime/Slate/Public/Widgets/Notifications/SErrorHint.h
new file mode 100644
--- /dev/null
+++ b/Engine/Source/Runtime/Slate/Public/Widgets/Notifications/SErrorHint.h
@@ -0,0 +1,39 @@
+// Copyright 1998-2014 Epic Games, Inc. All Rights Reserved.
+ 
+#pragma once
+
+class SLATE_API SErrorHint : public SCompoundWidget, public IErrorReportingWidget
+{
+public:
+	SLATE_BEGIN_ARGS( SErrorHint )
+		: _ErrorText()
+		{}
+		SLATE_TEXT_ARGUMENT(ErrorText)
+	SLATE_END_ARGS()
+
+	void Construct(const FArguments& InArgs);
+
+	// IErrorReportingWidget interface
+
+	virtual void SetError( const FText& InErrorText ) OVERRIDE;
+	virtual void SetError( const FString& InErrorText ) OVERRIDE;
+
+	virtual bool HasError() const OVERRIDE;
+
+	virtual TSharedRef<SWidget> AsWidget() OVERRIDE;
+
+	// IErrorReportingWidget interface
+
+private:
+	TAttribute<EVisibility> CustomVisibility;
+	EVisibility MyVisibility() const;
+
+	FVector2D GetDesiredSizeScale() const;
+	FCurveSequence ExpandAnimation;
+
+	TSharedPtr<SWidget> ImageWidget;
+	FText ErrorText;
+	FText GetErrorText() const;
+};
+
+
diff --git a/Engine/Source/Runtime/Slate/Public/Widgets/Notifications/SErrorText.h b/Engine/Source/Runtime/Slate/Public/Widgets/Notifications/SErrorText.h
new file mode 100644
--- /dev/null
+++ b/Engine/Source/Runtime/Slate/Public/Widgets/Notifications/SErrorText.h
@@ -0,0 +1,50 @@
+// Copyright 1998-2014 Epic Games, Inc. All Rights Reserved.
+ 
+#pragma once
+
+class SLATE_API IErrorReportingWidget
+{
+public:
+	virtual void SetError( const FString& InErrorText ) = 0;
+	virtual void SetError( const FText& InErrorText ) = 0;
+	virtual bool HasError() const = 0;
+	virtual TSharedRef<SWidget> AsWidget() = 0;
+};
+
+class SLATE_API SErrorText : public SBorder, public IErrorReportingWidget
+{
+public:
+	SLATE_BEGIN_ARGS( SErrorText )
+		: _ErrorText()
+		, _BackgroundColor(FCoreStyle::Get().GetColor("ErrorReporting.BackgroundColor"))
+		{}
+
+		SLATE_TEXT_ARGUMENT(ErrorText)
+		SLATE_ATTRIBUTE(FSlateColor, BackgroundColor)
+		
+	SLATE_END_ARGS()
+
+	void Construct(const FArguments& InArgs);
+
+	// IErrorReportingWidget interface
+
+	virtual void SetError( const FText& InErrorText ) OVERRIDE;
+	virtual void SetError( const FString& InErrorText ) OVERRIDE;
+
+	virtual bool HasError() const OVERRIDE;
+
+	virtual TSharedRef<SWidget> AsWidget() OVERRIDE;
+
+	// IErrorReportingWidget interface
+
+private:
+	TAttribute<EVisibility> CustomVisibility;
+	EVisibility MyVisibility() const;
+
+	TSharedPtr<class STextBlock> TextBlock;
+
+	FVector2D GetDesiredSizeScale() const;
+	FCurveSequence ExpandAnimation;
+};
+
+
diff --git a/Engine/Source/Runtime/Slate/Public/Widgets/Notifications/SNotificationList.h b/Engine/Source/Runtime/Slate/Public/Widgets/Notifications/SNotificationList.h
new file mode 100644
--- /dev/null
+++ b/Engine/Source/Runtime/Slate/Public/Widgets/Notifications/SNotificationList.h
@@ -0,0 +1,185 @@
+// Copyright 1998-2014 Epic Games, Inc. All Rights Reserved.
+
+#pragma once
+
+/**
+ * Interface class for an item in the event message list.
+ * Real implementation is found in SEventMessageItemImpl
+ */
+class SLATE_API SNotificationItem : public SCompoundWidget
+{
+public:
+	enum ECompletionState
+	{
+		CS_None,
+		CS_Pending,
+		CS_Success,
+		CS_Fail
+	};
+
+	/** Sets the text for message element */
+	virtual void SetText( const TAttribute< FText >& InText ) = 0;
+
+	/** Sets the text and delegate for the hyperlink */
+	virtual void SetHyperlink( const FSimpleDelegate& InHyperlink, const TAttribute< FText >& InHyperlinkText = TAttribute< FText >() ) = 0;
+
+	/** Gets the visibility state of the throbber, success, and fail images */
+	virtual ECompletionState GetCompletionState() const = 0;
+
+	/** Sets the visibility state of the throbber, success, and fail images */
+	virtual void SetCompletionState(ECompletionState State) = 0;
+
+	/** Waits for the ExpireDuration then begins to fade out */
+	virtual void ExpireAndFadeout() = 0;
+
+	/** Begin the fade out */
+	virtual void Fadeout() = 0;
+};
+
+/**
+	* Setup class to initialize buttons on a notification
+	*/
+struct FNotificationButtonInfo
+{
+	FNotificationButtonInfo( const FText& InText, const FText& InToolTip, FSimpleDelegate InCallback, SNotificationItem::ECompletionState VisibleInState = SNotificationItem::CS_Pending )
+		: Text( InText )
+		, ToolTip( InToolTip )
+		, Callback( InCallback )
+		, VisibilityOnNone( VisibleInState == SNotificationItem::CS_None ? EVisibility::Visible : EVisibility::Hidden )
+		, VisibilityOnPending( VisibleInState == SNotificationItem::CS_Pending ? EVisibility::Visible : EVisibility::Hidden )
+		, VisibilityOnSuccess( VisibleInState == SNotificationItem::CS_Success ? EVisibility::Visible : EVisibility::Hidden )
+		, VisibilityOnFail( VisibleInState == SNotificationItem::CS_Fail ? EVisibility::Visible : EVisibility::Hidden )
+	{}
+
+	/** Message on the button */
+	FText Text;
+	/** Tip displayed when moused over */
+	FText ToolTip;
+	/** Method called when button clicked */
+	FSimpleDelegate Callback;
+
+	/** Visibility of the button when the completion state of the button is SNotificationItem::ECompletionState::None */
+	EVisibility VisibilityOnNone;
+	/** Visibility of the button when the completion state of the button is SNotificationItem::ECompletionState::Pending */
+	EVisibility VisibilityOnPending;
+	/** Visibility of the button when the completion state of the button is SNotificationItem::ECompletionState::Success */
+	EVisibility VisibilityOnSuccess;
+	/** Visibility of the button when the completion state of the button is SNotificationItem::ECompletionState::Fail */
+	EVisibility VisibilityOnFail;
+};
+
+/**
+ * Setup class to initialize a notification.
+ */
+struct FNotificationInfo
+{
+	/**
+	 * FNotifcationInfo initialization constructor
+	 *
+	 * @param	InText	Text string to display for this notification
+	 */
+	FNotificationInfo( const FText& InText )
+		: Text(InText),
+		ButtonDetails(),
+		Image(NULL),
+		FadeInDuration(0.5f),
+		FadeOutDuration(2.0f),
+		ExpireDuration(1.0f),
+		bUseThrobber(true),
+		bUseSuccessFailIcons(true),
+		bUseLargeFont(true),
+		WidthOverride(),
+		bFireAndForget(true),
+		Hyperlink(),
+		HyperlinkText( NSLOCTEXT("EditorNotification", "DefaultHyperlinkText", "Show Log") ),
+		bAllowThrottleWhenFrameRateIsLow(true)
+	{
+	};
+
+	/** The text displayed in this text block */
+	FText Text;
+
+	/** Setup information for the buttons on the notification */ 
+	TArray<FNotificationButtonInfo> ButtonDetails;
+
+	/** The icon image to display next to the text */
+	const FSlateBrush* Image;
+
+	/** The fade in duration for this element */
+	float FadeInDuration;
+
+	/** The fade out duration for this element */
+	float FadeOutDuration;
+
+	/** The duration before a fadeout for this element */
+	float ExpireDuration;
+
+	/** Controls whether or not to add the animated throbber */
+	bool bUseThrobber;
+
+	/** Controls whether or not to display the success and fail icons */
+	bool bUseSuccessFailIcons;
+
+	/** When true the larger bolder font will be used to display the message */
+	bool bUseLargeFont;
+
+	/** When set this forces the width of the box, used to stop resizing on text change */
+	FOptionalSize WidthOverride;
+
+	/** When true the notification will automatically time out after the expire duration. */ 
+	bool bFireAndForget;
+
+	/** When set this will display as a hyperlink on the right side of the notification. */
+	FSimpleDelegate Hyperlink;
+
+	/** Text to display for the hyperlink message */
+	TAttribute< FText > HyperlinkText;
+
+	/** True if we should throttle the editor while the notification is transitioning and performance is poor, to make sure the user can see the animation */
+	bool bAllowThrottleWhenFrameRateIsLow;
+};
+
+/**
+ * A list of non-intrusive messages about the status of currently active work
+ */
+class SLATE_API SNotificationList : public SCompoundWidget
+{
+	friend class SNotificationExtendable;
+	friend class FSlateNotificationManager;
+
+public:
+	SLATE_BEGIN_ARGS( SNotificationList ){}
+
+	SLATE_END_ARGS()
+	
+	/**
+	 * Constructs this widget
+	 *
+	 * @param InArgs    Declaration from which to construct the widget
+	 */
+	void Construct( const FArguments& InArgs );
+
+	/**
+	 * Adds a floating notification
+	 *@param Info 		Contains various settings used to initialize the notification
+	 */
+	virtual TSharedRef<SNotificationItem> AddNotification(const FNotificationInfo& Info);
+
+protected:
+	/**
+	 * Invoked when a notification item has finished fading out. Called by the faded item.
+	 *
+	 * @param NotificationItem The item which finished fading out.
+	 */
+	virtual void NotificationItemFadedOut (const TSharedRef<SNotificationItem>& NotificationItem);
+
+protected:
+	/** The vertical box containing the list of message items */
+	TSharedPtr<SVerticalBox> MessageItemBoxPtr;
+
+	/** The parent window of this list */
+	TWeakPtr<SWindow> ParentWindowPtr;
+
+	/** Flag to auto-destroy this list */
+	bool bDone;
+};
diff --git a/Engine/Source/Runtime/Slate/Public/Widgets/Notifications/SPopUpErrorText.h b/Engine/Source/Runtime/Slate/Public/Widgets/Notifications/SPopUpErrorText.h
new file mode 100644
--- /dev/null
+++ b/Engine/Source/Runtime/Slate/Public/Widgets/Notifications/SPopUpErrorText.h
@@ -0,0 +1,33 @@
+// Copyright 1998-2014 Epic Games, Inc. All Rights Reserved.
+ 
+#pragma once
+
+#include "SComboButton.h"
+
+class SLATE_API SPopupErrorText : public SComboButton, public IErrorReportingWidget
+{
+public:
+	SLATE_BEGIN_ARGS(SPopupErrorText)
+		: _ShowInNewWindow( false )
+	{}
+		/** The popup appears in a new window instead of in the same window that this widget is in */
+		SLATE_ARGUMENT( bool, ShowInNewWindow )
+	SLATE_END_ARGS()
+
+	virtual void Construct( const FArguments& InArgs );
+
+	// IErrorReportingWidget interface
+
+	virtual void SetError( const FText& InErrorText ) OVERRIDE;
+	virtual void SetError( const FString& InErrorText ) OVERRIDE;
+
+	virtual bool HasError() const OVERRIDE;
+
+	virtual TSharedRef<SWidget> AsWidget() OVERRIDE;
+
+	// IErrorReportingWidget interface
+
+private:
+	TSharedPtr<SErrorText> HasErrorSymbol;
+	TSharedPtr<SErrorText> ErrorText;
+};
\ No newline at end of file
diff --git a/Engine/Source/Runtime/Slate/Public/Widgets/Notifications/SProgressBar.h b/Engine/Source/Runtime/Slate/Public/Widgets/Notifications/SProgressBar.h
new file mode 100644
--- /dev/null
+++ b/Engine/Source/Runtime/Slate/Public/Widgets/Notifications/SProgressBar.h
@@ -0,0 +1,66 @@
+// Copyright 1998-2014 Epic Games, Inc. All Rights Reserved.
+
+#pragma once
+
+/** A progress bar widget.*/
+class SLATE_API SProgressBar : public SLeafWidget
+{
+
+public:
+	SLATE_BEGIN_ARGS(SProgressBar)
+		: _Style( &FCoreStyle::Get().GetWidgetStyle<FProgressBarStyle>("ProgressBar") )
+		, _Percent( TOptional<float>() )
+		, _FillColorAndOpacity( FLinearColor::White )
+		, _BorderPadding( FVector2D(1,0) )
+		{}
+
+		/** Style used for the progress bar */
+		SLATE_STYLE_ARGUMENT( FProgressBarStyle, Style )
+
+		/** Used to determine the fill position of the progress bar ranging 0..1 */
+		SLATE_ATTRIBUTE( TOptional<float>, Percent )
+
+		/** Fill Color and Opacity */
+		SLATE_ATTRIBUTE( FSlateColor, FillColorAndOpacity )
+
+		/** Border Padding around fill bar */
+		SLATE_ATTRIBUTE( FVector2D, BorderPadding )
+
+	SLATE_END_ARGS()
+
+	/**
+	 * Construct the widget
+	 * 
+	 * @param InArgs   A declaration from which to construct the widget
+	 */
+	void Construct(const FArguments& InArgs);
+
+	virtual int32 OnPaint( const FGeometry& AllottedGeometry, const FSlateRect& MyClippingRect, FSlateWindowElementList& OutDrawElements, int32 LayerId, const FWidgetStyle& InWidgetStyle, bool bParentEnabled ) const OVERRIDE;
+
+	FVector2D ComputeDesiredSize() const;
+
+private:
+
+	/** The text displayed over the progress bar */
+	TAttribute< TOptional<float> > Percent;
+
+	/** Background image to use for the progress bar */
+	const FSlateBrush* BackgroundImage;
+
+	/** Foreground image to use for the progress bar */
+	const FSlateBrush* FillImage;
+
+	/** Image to use for marquee mode */
+	const FSlateBrush* MarqueeImage;
+
+	/** Fill Color and Opacity */
+	TAttribute<FSlateColor> FillColorAndOpacity;
+
+	/** Border Padding */
+	TAttribute<FVector2D> BorderPadding;
+
+	/** Curve sequence to drive progress bar animation */
+	FCurveSequence CurveSequence;
+
+};
+
diff --git a/Engine/Source/Runtime/Slate/Public/Widgets/SCanvas.h b/Engine/Source/Runtime/Slate/Public/Widgets/SCanvas.h
new file mode 100644
--- /dev/null
+++ b/Engine/Source/Runtime/Slate/Public/Widgets/SCanvas.h
@@ -0,0 +1,183 @@
+// Copyright 1998-2014 Epic Games, Inc. All Rights Reserved.
+
+#pragma once
+
+
+/**
+ * Canvas is a layout widget that allows you to arbitrary position and size child widgets in a relative coordinate space
+ */
+class SLATE_API SCanvas : public SPanel
+{
+public:
+
+	/**
+	 * Canvas slots allow child widgets to be positioned and sized
+	 */
+	class FSlot
+	{
+	public:		
+		FSlot& Position( const TAttribute<FVector2D>& InPosition )
+		{
+			PositionAttr = InPosition;
+			return *this;
+		}
+
+		FSlot& Size( const TAttribute<FVector2D>& InSize )
+		{
+			SizeAttr = InSize;
+			return *this;
+		}
+
+		FSlot& operator[]( TSharedRef<SWidget> InWidget )
+		{
+			Widget = InWidget;
+			return *this;
+		}
+
+		FSlot& Expose( FSlot*& OutVarToInit )
+		{
+			OutVarToInit = this;
+			return *this;
+		}
+
+		FSlot& HAlign( EHorizontalAlignment InHAlignment )
+		{
+			HAlignment = InHAlignment;
+			return *this;
+		}
+
+		FSlot& VAlign( EVerticalAlignment InVAlignment )
+		{
+			VAlignment = InVAlignment;
+			return *this;
+		}
+
+		/** The child widget contained in this slot. */
+		TSharedRef<SWidget> Widget;
+
+		/** Position */
+		TAttribute<FVector2D> PositionAttr;
+
+		/** Size */
+		TAttribute<FVector2D> SizeAttr;
+
+		/** Horizontal Alignment 
+		*  Given a top aligned slot, where '+' represents the 
+		*  anchor point defined by PositionAttr.
+		*  
+		*   Left				Center				Right
+			+ _ _ _ _            _ _ + _ _          _ _ _ _ +
+			|		  |		   | 		   |	  |		    |
+			| _ _ _ _ |        | _ _ _ _ _ |	  | _ _ _ _ |
+		* 
+		*  Note: FILL is NOT supported.
+		*/
+		EHorizontalAlignment HAlignment;
+
+		/** Vertical Alignment 
+		*   Given a left aligned slot, where '+' represents the 
+		*   anchor point defined by PositionAttr.
+		*  
+		*   Top					Center			  Bottom
+		*	+_ _ _ _ _          _ _ _ _ _		 _ _ _ _ _ 
+		*	|		  |		   | 		 |		|		  |
+		*	| _ _ _ _ |        +		 |		|		  |
+		*					   | _ _ _ _ |		+ _ _ _ _ |
+		* 
+		*  Note: FILL is NOT supported.
+		*/
+		EVerticalAlignment VAlignment;
+
+
+		/** Default values for a slot. */
+		FSlot()
+			: Widget( SNullWidget::NullWidget )
+			, PositionAttr( FVector2D::ZeroVector )
+			, SizeAttr( FVector2D( 1.0f, 1.0f ) )
+			, HAlignment(HAlign_Left)
+			, VAlignment(VAlign_Top)
+		{
+		}
+	};
+
+
+
+	SLATE_BEGIN_ARGS( SCanvas )
+		{
+			_Visibility = EVisibility::SelfHitTestInvisible;
+		}
+
+		SLATE_SUPPORTS_SLOT( SCanvas::FSlot )
+
+	SLATE_END_ARGS()
+
+	/**
+	 * Construct this widget
+	 *
+	 * @param	InArgs	The declaration data for this widget
+	 */
+	void Construct( const FArguments& InArgs );
+
+	static FSlot& Slot()
+	{
+		return *(new FSlot());
+	}
+
+	FSlot& AddSlot()
+	{
+		SCanvas::FSlot& NewSlot = SCanvas::Slot();
+		this->Children.Add( &NewSlot );
+		return NewSlot;
+	}
+
+	/**
+	 * Panels arrange their children in a space described by the AllottedGeometry parameter. The results of the arrangement
+	 * should be returned by appending a FArrangedWidget pair for every child widget. See StackPanel for an example
+	 *
+	 * @param AllottedGeometry    The geometry allotted for this widget by its parent.
+	 * @param ArrangedChildren    The array to which to add the WidgetGeometries that represent the arranged children.
+	 */
+	virtual void ArrangeChildren( const FGeometry& AllottedGeometry, FArrangedChildren& ArrangedChildren ) const OVERRIDE;
+	
+	/**
+	 * The widget should respond by populating the OutDrawElements array with FDrawElements 
+	 * that represent it and any of its children.
+	 *
+	 * @param AllottedGeometry  The FGeometry that describes an area in which the widget should appear.
+	 * @param MyClippingRect    The clipping rectangle allocated for this widget and its children.
+	 * @param OutDrawElements   A list of FDrawElements to populate with the output.
+	 * @param LayerId           The Layer onto which this widget should be rendered.
+	 * @param InColorAndOpacity Color and Opacity to be applied to all the descendants of the widget being painted
+	 * @param bParentEnabled	True if the parent of this widget is enabled.
+	 *
+	 * @return The maximum layer ID attained by this widget or any of its children.
+	 */
+	virtual int32 OnPaint( const FGeometry& AllottedGeometry, const FSlateRect& MyClippingRect, FSlateWindowElementList& OutDrawElements, int32 LayerId, const FWidgetStyle& InWidgetStyle, bool bParentEnabled ) const OVERRIDE;
+
+	/**
+	 * A Panel's desired size in the space required to arrange of its children on the screen while respecting all of
+	 * the children's desired sizes and any layout-related options specified by the user. See StackPanel for an example.
+	 *
+	 * @return The desired size.
+	 */
+	virtual FVector2D ComputeDesiredSize() const OVERRIDE;
+
+	/** @return  The children of a panel in a slot-agnostic way. */
+	virtual FChildren* GetChildren() OVERRIDE;
+
+	/** 
+	 * Removes a slot from this box panel which contains the specified SWidget
+	 *
+	 * @param SlotWidget The widget to match when searching through the slots
+	 * @returns The index in the children array where the slot was removed and -1 if no slot was found matching the widget
+	 */
+	int32 RemoveSlot( const TSharedRef<SWidget>& SlotWidget );
+
+	/** Removes all slots from the panel */
+	void ClearChildren();
+
+protected:
+
+	/** The canvas widget's children. */
+	TPanelChildren< FSlot > Children;
+};
diff --git a/Engine/Source/Runtime/Slate/Public/Widgets/SToolTip.h b/Engine/Source/Runtime/Slate/Public/Widgets/SToolTip.h
new file mode 100644
--- /dev/null
+++ b/Engine/Source/Runtime/Slate/Public/Widgets/SToolTip.h
@@ -0,0 +1,85 @@
+// Copyright 1998-2014 Epic Games, Inc. All Rights Reserved.
+
+#pragma once
+
+
+/**
+ * Slate tool tip widget
+ */
+class SLATE_API SToolTip
+	: public SBorder
+	, public IToolTip
+{
+public:
+
+	SLATE_BEGIN_ARGS( SToolTip )
+		: _Text()
+		, _Content()
+		, _Font(FCoreStyle::Get().GetFontStyle("ToolTip.Font"))
+		, _ColorAndOpacity( FSlateColor::UseForeground())
+		, _TextMargin(FMargin(8.0f))
+		, _BorderImage(FCoreStyle::Get().GetBrush("ToolTip.Background"))
+		, _IsInteractive(false)
+	{ }
+
+		/** The text displayed in this tool tip */
+		SLATE_TEXT_ATTRIBUTE(Text)
+
+		/** Arbitrary content to be displayed in the tool tip; overrides any text that may be set. */
+		SLATE_DEFAULT_SLOT(FArguments, Content)
+
+		/** The font to use for this tool tip */
+		SLATE_ATTRIBUTE(FSlateFontInfo, Font)
+		
+		/** Font color and opacity */
+		SLATE_ATTRIBUTE(FSlateColor, ColorAndOpacity)
+		
+		/** Margin between the tool tip border and the text content */
+		SLATE_ATTRIBUTE(FMargin, TextMargin)
+
+		/** The background/border image to display */
+		SLATE_ATTRIBUTE(const FSlateBrush*, BorderImage)
+
+		/** Whether the tooltip should be considered interactive */
+		SLATE_ATTRIBUTE(bool, IsInteractive)
+
+	SLATE_END_ARGS()
+
+	/**
+	 * Construct this widget
+	 *
+	 * @param	InArgs	The declaration data for this widget
+	 */
+	void Construct( const FArguments& InArgs );
+
+public:
+
+	// Begin IToolTip interface
+
+	virtual TSharedRef<class SWidget> AsWidget( ) OVERRIDE
+	{
+		return AsShared();
+	}
+
+	virtual const TSharedRef<SWidget>& GetContentWidget( ) OVERRIDE
+	{
+		return GetContent();
+	}
+
+	virtual bool IsEmpty( ) const OVERRIDE;
+
+	virtual bool IsInteractive( ) const OVERRIDE;
+
+	// End IToolTip interface
+
+private:
+
+	// Text block widget.
+	TAttribute<FString> TextContent;
+
+	// Content widget.
+	TWeakPtr<SWidget> WidgetContent;
+
+	// Whether the tooltip should be considered interactive.
+	TAttribute<bool> bIsInteractive;
+};
diff --git a/Engine/Source/Runtime/Slate/Public/Widgets/SViewport.h b/Engine/Source/Runtime/Slate/Public/Widgets/SViewport.h
new file mode 100644
--- /dev/null
+++ b/Engine/Source/Runtime/Slate/Public/Widgets/SViewport.h
@@ -0,0 +1,166 @@
+// Copyright 1998-2014 Epic Games, Inc. All Rights Reserved.
+
+#pragma once
+
+
+class SLATE_API SViewport : public SCompoundWidget
+{
+public:
+	
+	SLATE_BEGIN_ARGS( SViewport )
+		: _Content()
+		, _ShowEffectWhenDisabled(true)
+		, _RenderDirectlyToWindow(false)
+		, _EnableGammaCorrection(true)
+		, _EnableBlending(false)
+		, _IgnoreTextureAlpha(true)
+		{}
+
+		SLATE_DEFAULT_SLOT( FArguments, Content )
+		/** Whether or not to show the disabled effect when this viewport is disabled */
+		SLATE_ATTRIBUTE( bool, ShowEffectWhenDisabled )
+		/** 
+		 * Whether or not to render directly to the window's backbuffer or an offscreen render target that is applied to the window later 
+		 * Rendering to an offscreen target is the most common option in the editor where there may be many frames which this viewport's interface may wish to not re-render but use a cached buffer instead
+		 * Rendering directly to the backbuffer is the most common option in the game where you want to update each frame without the cost of writing to an intermediate target first.
+		 */
+		SLATE_ARGUMENT( bool, RenderDirectlyToWindow )
+
+		/**
+		 * Whether or not to enable gamma correction.  Doesn't apply when rendering directly to a backbuffer 
+		 */
+		SLATE_ARGUMENT( bool, EnableGammaCorrection )
+
+		/**
+		 * Allow this viewport to blend with its background
+		 */
+		SLATE_ARGUMENT( bool, EnableBlending )
+
+		/**
+		 * If true, the viewport's texture alpha is ignored when performing blending.  In this case only the viewport tint opacity is used
+		 * If false, the texture alpha is used during blending
+		 */
+		SLATE_ARGUMENT( bool, IgnoreTextureAlpha )
+	SLATE_END_ARGS()
+
+	/**
+	 * Construct the widget.
+	 *
+	 * @param InArgs  Declaration from which to construct the widget.
+	 */
+	void Construct(const FArguments& InArgs);
+	
+	SViewport();
+
+	/** SViewport wants keyboard focus */
+	virtual bool SupportsKeyboardFocus() const OVERRIDE { return true; }
+
+	/** 
+	 * Computes the ideal size necessary to display this widget.
+	 */
+	virtual FVector2D ComputeDesiredSize() const OVERRIDE
+	{
+		return FVector2D( 320, 240 );
+	}
+
+	/**
+	 * Sets the interface to be used by this viewport for rendering and I/O
+	 *
+	 * @param InViewportInterface	The interface to use
+	 */
+	void SetViewportInterface( TSharedRef<ISlateViewport> InViewportInterface )
+	{
+		ViewportInterface = InViewportInterface;
+	}
+
+	/**
+	 * Sets the content for this widget
+	 *
+	 * @param InContent	The new content (can be null)
+	 */
+	void SetContent( TSharedPtr<SWidget> InContent )
+	{
+		ChildSlot
+		[
+			InContent.IsValid() ? InContent.ToSharedRef() : (TSharedRef<SWidget>)SNullWidget::NullWidget
+		];
+	}
+
+	const TSharedPtr<SWidget> GetContent() const { return ChildSlot.Widget; }
+
+	/**
+	 * A delegate called when the viewports top level window is being closed
+	 *
+	 * @param InWindowBeingClosed	The window that is about to be closed
+	 */
+	void OnWindowClosed( const TSharedRef<SWindow>& InWindowBeingClosed );
+
+
+	/** @return Whether or not this viewport renders directly to the backbuffer */
+	bool ShouldRenderDirectly() const { return bRenderDirectlyToWindow; }
+
+	/**
+	 * Sets a widget that should become focused when this window is next activated
+	 *
+	 * @param	InWidget	The widget to set focus to when this window is activated
+	 */
+	void SetWidgetToFocusOnActivate( const TSharedPtr< SWidget >& InWidget )
+	{
+		
+		WidgetToFocusOnActivate = InWidget;
+	}
+
+	/** Removes the widget to focus on activate so the viewport will be focused */
+	void ClearWidgetToFocusOnActivate()
+	{
+		WidgetToFocusOnActivate.Reset();
+	}
+
+	// SWidget interface
+	virtual int32 OnPaint( const FGeometry& AllottedGeometry, const FSlateRect& MyClippingRect, FSlateWindowElementList& OutDrawElements, int32 LayerId, const FWidgetStyle& InWidgetStyle, bool bParentEnabled ) const OVERRIDE;
+	virtual void Tick( const FGeometry& AllottedGeometry, const double InCurrentTime, const float InDeltaTime ) OVERRIDE;
+	virtual FReply OnTouchStarted( const FGeometry& MyGeometry, const FPointerEvent& TouchEvent ) OVERRIDE;
+	virtual FReply OnTouchMoved( const FGeometry& MyGeometry, const FPointerEvent& TouchEvent ) OVERRIDE;
+	virtual FReply OnTouchEnded( const FGeometry& MyGeometry, const FPointerEvent& TouchEvent ) OVERRIDE;
+	virtual FReply OnTouchGesture( const FGeometry& MyGeometry, const FPointerEvent& GestureEvent ) OVERRIDE;
+	virtual FCursorReply OnCursorQuery( const FGeometry& MyGeometry, const FPointerEvent& CursorEvent ) const OVERRIDE;
+	virtual FReply OnMouseButtonDown( const FGeometry& MyGeometry, const FPointerEvent& MouseEvent ) OVERRIDE;
+	virtual FReply OnMouseButtonUp( const FGeometry& MyGeometry, const FPointerEvent& MouseEvent ) OVERRIDE;
+	virtual void OnMouseEnter( const FGeometry& MyGeometry, const FPointerEvent& MouseEvent ) OVERRIDE;
+	virtual void OnMouseLeave( const FPointerEvent& MouseEvent ) OVERRIDE;
+	virtual FReply OnMouseMove( const FGeometry& MyGeometry, const FPointerEvent& MouseEvent ) OVERRIDE;
+	virtual FReply OnMouseWheel( const FGeometry& MyGeometry, const FPointerEvent& MouseEvent ) OVERRIDE;
+	virtual FReply OnMouseButtonDoubleClick( const FGeometry& MyGeometry, const FPointerEvent& MouseEvent ) OVERRIDE;
+	virtual FReply OnKeyDown( const FGeometry& MyGeometry, const FKeyboardEvent& KeyboardEvent ) OVERRIDE;
+	virtual FReply OnKeyUp( const FGeometry& MyGeometry, const FKeyboardEvent& KeyboardEvent ) OVERRIDE;
+	virtual FReply OnKeyChar( const FGeometry& MyGeometry, const FCharacterEvent& CharacterEvent ) OVERRIDE;
+	virtual FReply OnKeyboardFocusReceived( const FGeometry& MyGeometry, const FKeyboardFocusEvent& InKeyboardFocusEvent ) OVERRIDE;
+	virtual void OnKeyboardFocusLost( const FKeyboardFocusEvent& InKeyboardFocusEvent ) OVERRIDE;
+	virtual FReply OnControllerButtonPressed( const FGeometry& MyGeometry, const FControllerEvent& ControllerEvent ) OVERRIDE;
+	virtual FReply OnControllerButtonReleased( const FGeometry& MyGeometry, const FControllerEvent& ControllerEvent ) OVERRIDE;
+	virtual FReply OnControllerAnalogValueChanged( const FGeometry& MyGeometry, const FControllerEvent& ControllerEvent ) OVERRIDE;
+	virtual FReply OnMotionDetected( const FGeometry& MyGeometry, const FMotionEvent& MotionEvent ) OVERRIDE;
+	// End of SWidget interface
+
+
+private:
+	FSimpleSlot& Decl_GetContent() { return ChildSlot; }
+	
+private:
+	/** Interface to the rendering and I/O implementation of the viewport */
+	TWeakPtr<ISlateViewport> ViewportInterface;
+	/** Whether or not to show the disabled effect when this viewport is disabled */
+	TAttribute<bool> ShowDisabledEffect;
+	/** Whether or not this viewport renders directly to the window backbuffer */
+	bool bRenderDirectlyToWindow;
+	/** Whether or not to apply gamma correction on the render target supplied by the ISlateViewport */
+	bool bEnableGammaCorrection;
+	/** Whether or not to blend this viewport with the background */
+	bool bEnableBlending;
+	/** Whether or not to allow texture alpha to be used in blending calculations */
+	bool bIgnoreTextureAlpha;
+
+	/** Widget to transfer keyboard focus to when this window becomes active, if any.  This is used to
+	    restore focus to a widget after a popup has been dismissed. */
+	TWeakPtr< SWidget > WidgetToFocusOnActivate;
+};
diff --git a/Engine/Source/Runtime/Slate/Public/Widgets/SWeakWidget.h b/Engine/Source/Runtime/Slate/Public/Widgets/SWeakWidget.h
new file mode 100644
--- /dev/null
+++ b/Engine/Source/Runtime/Slate/Public/Widgets/SWeakWidget.h
@@ -0,0 +1,42 @@
+// Copyright 1998-2014 Epic Games, Inc. All Rights Reserved.
+
+#pragma once
+
+
+/**
+ * Weak widgets present a piece of content without owning it.
+ * e.g. A tooltip is owned by the hovered widget but displayed
+ *      by a floating window. That window cannot own the tooltip
+ *      and must therefore use an SWeakWidget.
+ */
+class SLATE_API SWeakWidget : public SWidget
+{
+
+public:
+	SLATE_BEGIN_ARGS(SWeakWidget)
+		{
+			_Visibility = EVisibility::SelfHitTestInvisible;
+		}
+		SLATE_ARGUMENT( TSharedPtr<SWidget>, PossiblyNullContent )
+	SLATE_END_ARGS()
+
+	void Construct(const FArguments& InArgs);
+
+	virtual FVector2D ComputeDesiredSize() const OVERRIDE;
+
+	virtual void ArrangeChildren( const FGeometry& AllottedGeometry, FArrangedChildren& ArrangedChildren ) const OVERRIDE;
+		
+	virtual FChildren* GetChildren() OVERRIDE;
+
+	virtual int32 OnPaint( const FGeometry& AllottedGeometry, const FSlateRect& MyClippingRect, FSlateWindowElementList& OutDrawElements, int32 LayerId, const FWidgetStyle& InWidgetStyle, bool bParentEnabled ) const OVERRIDE;
+
+	void SetContent(TWeakPtr<SWidget> InWidget);
+
+	bool ChildWidgetIsValid() const;
+
+	TWeakPtr<SWidget> GetChildWidget() const;
+
+private:
+	TWeakChild<SWidget> WeakChild;
+
+};
diff --git a/Engine/Source/Runtime/Slate/Public/Widgets/Testing/STestSuite.h b/Engine/Source/Runtime/Slate/Public/Widgets/Testing/STestSuite.h
new file mode 100644
--- /dev/null
+++ b/Engine/Source/Runtime/Slate/Public/Widgets/Testing/STestSuite.h
@@ -0,0 +1,8 @@
+// Copyright 1998-2014 Epic Games, Inc. All Rights Reserved.
+
+#pragma once
+
+
+SLATE_API void RestoreSlateTestSuite();
+
+SLATE_API void MakeSplitterTest();
\ No newline at end of file
diff --git a/Engine/Source/Runtime/Slate/Public/Widgets/Text/STextBlock.h b/Engine/Source/Runtime/Slate/Public/Widgets/Text/STextBlock.h
new file mode 100644
--- /dev/null
+++ b/Engine/Source/Runtime/Slate/Public/Widgets/Text/STextBlock.h
@@ -0,0 +1,210 @@
+// Copyright 1998-2014 Epic Games, Inc. All Rights Reserved.
+ 
+#pragma once
+
+#include "WordWrapper.h"
+
+namespace ETextRole
+{
+	enum Type
+	{
+		Custom,
+		ButtonText,
+		ComboText
+	};
+}
+
+
+/**
+ * A simple static text widget
+ */
+class SLATE_API STextBlock : public SLeafWidget
+{
+
+public:
+
+	SLATE_BEGIN_ARGS( STextBlock )
+		: _Text()
+		, _TextStyle( &FCoreStyle::Get().GetWidgetStyle<FTextBlockStyle>( "NormalText" ) )
+		, _Font()
+		, _ColorAndOpacity()
+		, _ShadowOffset()
+		, _ShadowColorAndOpacity()
+		, _HighlightColor()
+		, _HighlightShape()
+		, _HighlightText()
+		, _WrapTextAt(0.0f)
+		, _AutoWrapText(false)
+		{}
+
+		/** The text displayed in this text block */
+		SLATE_TEXT_ATTRIBUTE( Text )
+
+		/** Pointer to a style of the text block, which dictates the font, color, and shadow options. */
+		SLATE_STYLE_ARGUMENT( FTextBlockStyle, TextStyle )
+
+		/** Sets the font used to draw the text */
+		SLATE_ATTRIBUTE( FSlateFontInfo, Font )
+
+		/** Text color and opacity */
+		SLATE_ATTRIBUTE( FSlateColor, ColorAndOpacity )
+
+		/** Drop shadow offset in pixels */
+		SLATE_ATTRIBUTE( FVector2D, ShadowOffset )
+
+		/** Shadow color and opacity */
+		SLATE_ATTRIBUTE( FLinearColor, ShadowColorAndOpacity )
+
+		/** The color used to highlight the specified text */
+		SLATE_ATTRIBUTE( FLinearColor, HighlightColor )
+
+		/** The brush used to highlight the specified text*/
+		SLATE_ATTRIBUTE( const FSlateBrush*, HighlightShape )
+
+		/** Highlight this text in the text block */
+		SLATE_ATTRIBUTE( FText, HighlightText )
+
+		/** Whether text wraps onto a new line when it's length exceeds this width; if this value is zero or negative, no wrapping occurs. */
+		SLATE_ATTRIBUTE( float, WrapTextAt )
+
+		/** Whether to wrap text automatically based on the widget's computed horizontal space.  IMPORTANT: Using automatic wrapping can result
+		    in visual artifacts, as the the wrapped size will computed be at least one frame late!  Consider using WrapTextAt instead.  The initial 
+			desired size will not be clamped.  This works best in cases where the text block's size is not affecting other widget's layout. */
+		SLATE_ATTRIBUTE( bool, AutoWrapText )
+
+		/** Called when this text is double clicked */
+		SLATE_EVENT( FOnClicked, OnDoubleClicked )
+
+	SLATE_END_ARGS()
+
+	/**
+	 * Construct this widget
+	 *
+	 * @param	InArgs	The declaration data for this widget
+	 */
+	void Construct( const FArguments& InArgs );
+
+	/**
+	 * Gets the text assigned to this text block
+	 *
+	 * @return	This text block's text string
+	 */
+	const FString& GetText() const
+	{
+		return Text.Get();
+	}
+
+
+	/**
+	 * Sets the text for this text block
+	 *
+	 * @param	InText	The new text to display
+	 */
+	void SetText( const TAttribute< FString >& InText )
+	{
+		Text = InText;
+		bRequestCache = true;
+	}
+
+	void SetText( const FString& InText )
+	{
+		Text = InText;
+		bRequestCache = true;
+	}
+
+	static FString PassThroughAttribute( TAttribute< FText > TextAttribute )
+	{
+		return TextAttribute.Get( FText::GetEmpty() ).ToString();
+	}
+
+	/**
+	 * Sets the text for this text block
+	 *
+	 * @param	InText	The new text to display
+	 */
+	void SetText( const TAttribute< FText >& InText );
+
+	void SetText( const FText& InText )
+	{
+		Text = InText.ToString();
+		bRequestCache = true;
+	}
+
+	/**
+	 * Sets the font used to draw the text
+	 *
+	 * @param	InFont	The new font to use
+	 */
+	void SetFont( const TAttribute< FSlateFontInfo >& InFont )
+	{
+		Font = InFont;
+		bRequestCache = true;
+	}
+
+	/** Set the color and opacity of the text */
+	void SetForegroundColor( const TAttribute<FSlateColor>& InSlateColor );
+
+	// SWidget interface
+	virtual int32 OnPaint( const FGeometry& AllottedGeometry, const FSlateRect& MyClippingRect, FSlateWindowElementList& OutDrawElements, int32 LayerId, const FWidgetStyle& InWidgetStyle, bool bParentEnabled ) const OVERRIDE;
+	virtual FReply OnMouseButtonDoubleClick( const FGeometry& InMyGeometry, const FPointerEvent& InMouseEvent ) OVERRIDE;
+	virtual void CacheDesiredSize() OVERRIDE;
+	virtual FVector2D ComputeDesiredSize() const OVERRIDE;
+	// End of SWidget interface
+
+private:
+	/** Caches string size to avoid remeasuring if possible */
+	void CacheStringSizeIfNeeded();
+
+private:
+	/** The text displayed in this text block */
+	TAttribute< FString > Text;
+
+	/** Sets the font used to draw the text */
+	TAttribute< FSlateFontInfo > Font;
+
+	/** Text color and opacity */
+	TAttribute<FSlateColor> ForegroundColor;
+
+	/** Drop shadow offset in pixels */
+	TAttribute< FVector2D > ShadowOffset;
+
+	/** Shadow color and opacity */
+	TAttribute<FLinearColor> ShadowColorAndOpacity;
+
+	TAttribute<FLinearColor> HighlightColor;
+
+	TAttribute< const FSlateBrush* > HighlightShape;
+
+	/** Highlight this text in the textblock */
+	TAttribute<FText> HighlightText;
+
+	/** Whether text wraps onto a new line when it's length exceeds this width; if this value is zero or negative, no wrapping occurs. */
+	TAttribute<float> WrapTextAt;
+
+	/** True if we're wrapping text automatically based on the computed horizontal space for this widget */
+	TAttribute<bool> AutoWrapText;
+
+	/** The delegate to execute when this text is double clicked */
+	FOnClicked OnDoubleClicked;
+
+	/** Cached Font that this text is using. This is used when determining if the cached string size should be updated */
+	FSlateFontInfo CachedFont;
+
+	/** Cached wrap width that this text is using. This is used when determining if the cached string size should be updated */
+	float CachedWrapTextWidth;
+
+	/** Cached auto-wrap width that this text is using. This is used when determining if the cached string size should be updated */
+	mutable float CachedAutoWrapTextWidth;
+
+	/** Text that was used to generate CachedWrappedString */
+	FString CachedOriginalString;
+
+	/** Cached wrapped text. This gets cached once in tick so it does not have to be re-generated at paint time every frame */
+	FString CachedWrappedString;
+
+	/** Line break data for the original string -> cached string */
+	FWordWrapper::FWrappedLineData CachedWrappedLineData;
+
+	/** If a request to cache the string size occurred and we should recache the text */
+	mutable bool bRequestCache;
+};
