Slate: Refactored core Slate implementation into SlateCore module in preparation for UMG.

diff --git a/Engine/Source/Runtime/Slate/Public/Widgets/Tutorials/STutorialWrapper.h b/Engine/Source/Runtime/Slate/Public/Widgets/Tutorials/STutorialWrapper.h
new file mode 100644
--- /dev/null
+++ b/Engine/Source/Runtime/Slate/Public/Widgets/Tutorials/STutorialWrapper.h
@@ -0,0 +1,57 @@
+// Copyright 1998-2014 Epic Games, Inc. All Rights Reserved.
+
+#pragma once
+
+/** Delegate for a named widget being highlighted */
+DECLARE_DELEGATE_RetVal_OneParam(bool, FOnWidgetHighlight, const FName&);
+
+class STutorialWrapper : public SBorder
+{
+public:
+	SLATE_BEGIN_ARGS( STutorialWrapper )
+		: _Content()
+	{
+		_Visibility = EVisibility::SelfHitTestInvisible;
+	}
+
+	/** Slot for the wrapped content (optional) */
+	SLATE_DEFAULT_SLOT( FArguments, Content )
+
+	SLATE_END_ARGS()
+
+	SLATE_API void Construct(const FArguments& InArgs, const FName& Name);
+
+	/** Begin SWidget interface */
+	SLATE_API virtual void Tick( const FGeometry& AllottedGeometry, const double InCurrentTime, const float InDeltaTime ) OVERRIDE;
+	SLATE_API virtual int32 OnPaint(const FGeometry& AllottedGeometry, const FSlateRect& MyClippingRect, FSlateWindowElementList& OutDrawElements, int32 LayerId, const FWidgetStyle& InWidgetStyle, bool bParentEnabled ) const OVERRIDE;
+	/** End SWidget interface */
+
+	/** 
+	 * Get the delegate fired when a named widget wants to draw its highlight
+	 * @returns a reference to the delegate
+	 */
+	SLATE_API static FOnWidgetHighlight& OnWidgetHighlight();
+
+private:
+
+	/** Get the values that the animation drives */
+	void GetAnimationValues(float& OutAlphaFactor, float& OutPulseFactor, FLinearColor& OutShadowTint, FLinearColor& OutBorderTint) const;
+
+private:
+
+	/** The name of the widget */
+	FName Name;
+
+	/** Flag for whether we are playing animation or not */
+	bool bIsPlaying;
+
+	/** Animation curves for displaying border */
+	FCurveSequence BorderPulseAnimation;
+	FCurveSequence BorderIntroAnimation;
+
+	/** Geometry cached from Tick() */
+	FGeometry CachedGeometry;
+
+	/** The delegate fired when a named widget wants to draw its highlight */
+	static FOnWidgetHighlight OnWidgetHighlightDelegate;
+};
diff --git a/Engine/Source/Runtime/Slate/Public/Widgets/Views/SExpanderArrow.h b/Engine/Source/Runtime/Slate/Public/Widgets/Views/SExpanderArrow.h
new file mode 100644
--- /dev/null
+++ b/Engine/Source/Runtime/Slate/Public/Widgets/Views/SExpanderArrow.h
@@ -0,0 +1,36 @@
+// Copyright 1998-2014 Epic Games, Inc. All Rights Reserved.
+
+#pragma once
+
+/** Expander arrow and indentation component that can be placed in a TableRow of a TreeView. Intended for use by TMultiColumnRow in TreeViews. */
+class SLATE_API SExpanderArrow : public SCompoundWidget
+{
+public:
+	SLATE_BEGIN_ARGS( SExpanderArrow ){}
+		SLATE_ATTRIBUTE(float, IndentAmount)
+	SLATE_END_ARGS()
+
+	void Construct( const FArguments& InArgs, const TSharedPtr<class ITableRow>& TableRow );
+
+protected:
+
+	/** Invoked when the expanded button is clicked (toggle item expansion) */
+	FReply OnArrowClicked();
+
+	/** @return Visible when has children; invisible otherwise */
+	EVisibility GetExpanderVisibility() const;
+
+	/** @return the margin corresponding to how far this item is indented */
+	FMargin GetExpanderPadding() const;
+
+	/** @return the name of an image that should be shown as the expander arrow */
+	const FSlateBrush* GetExpanderImage() const;
+
+	TWeakPtr<class ITableRow> OwnerRowPtr;
+
+	/** A reference to the expander button */
+	TSharedPtr<SButton> ExpanderArrow;
+
+	/** The amount of space to indent at each level */
+	TAttribute<float> IndentAmount;
+};
diff --git a/Engine/Source/Runtime/Slate/Public/Widgets/Views/SListView.h b/Engine/Source/Runtime/Slate/Public/Widgets/Views/SListView.h
new file mode 100644
--- /dev/null
+++ b/Engine/Source/Runtime/Slate/Public/Widgets/Views/SListView.h
@@ -0,0 +1,1326 @@
+// Copyright 1998-2014 Epic Games, Inc. All Rights Reserved.
+
+#pragma once
+
+/**
+ * A ListView widget observes an array of data items and creates visual representations of these items.
+ * ListView relies on the property that holding a reference to a value ensures its existence. In other words,
+ * neither SListView<FString> nor SListView<FString*> are valid, while SListView< TSharedPtr<FString> > and
+ * SListView< UObject* > are valid.
+ *
+ * A trivial use case appear below:
+ *
+ *   Given: TArray< FString* > Items;
+ *
+ *   SNew( SListView< FString* > )
+ *     .ItemHeight(24)
+ *     .ListItemsSource( &Items )
+ *     .OnGenerateRow( SListView< TSharedPtr<FString> >::MakeOnGenerateWidget( this, &MyClass::OnGenerateRowForList ) )
+ *
+ * In the example we make all our widgets be 24 screen units tall. The ListView will create widgets based on data items
+ * in the Items TArray. When the ListView needs to generate an item, it will do so using the OnGenerateWidgetForList method.
+ *
+ * A sample implementation of OnGenerateWidgetForList would simply return a TextBlock with the corresponding text:
+ *
+ * TSharedRef<ITableRow> OnGenerateWidgetForList( FString* InItem )
+ * {
+ *     return SNew(STextBlock).Text( (*InItem) )
+ * }
+ *
+ */
+template <typename ItemType>
+class SListView : public STableViewBase, TListTypeTraits<ItemType>::SerializerType, public ITypedTableView< ItemType >
+{
+	checkAtCompileTime( TIsValidListItem<ItemType>::Value, ItemType_must_be_a_pointer_or_TSharedPtr );
+
+public:
+	typedef typename TListTypeTraits< ItemType >::NullableType NullableItemType;
+
+	typedef typename TSlateDelegates< ItemType >::FOnGenerateRow FOnGenerateRow;
+	typedef typename TSlateDelegates< ItemType >::FOnItemScrolledIntoView FOnItemScrolledIntoView;
+	typedef typename TSlateDelegates< NullableItemType >::FOnSelectionChanged FOnSelectionChanged;
+	typedef typename TSlateDelegates< ItemType >::FOnMouseButtonDoubleClick FOnMouseButtonDoubleClick;
+
+
+public:
+
+	class FColumnHeaderSlot
+	{
+	public:
+		FColumnHeaderSlot& operator[]( const TSharedRef< SHeaderRow >& InColumnHeaders )
+		{
+			HeaderRow = InColumnHeaders;
+		}			
+		TSharedPtr<SHeaderRow> HeaderRow;
+	};
+
+	SLATE_BEGIN_ARGS( SListView<ItemType> )
+		: _OnGenerateRow()
+		, _ListItemsSource( static_cast<const TArray<ItemType>*>(NULL) ) //@todo Slate Syntax: Initializing from NULL without a cast
+		, _ItemHeight(16)
+		, _OnContextMenuOpening()
+		, _OnMouseButtonDoubleClick()
+		, _OnSelectionChanged()
+		, _SelectionMode(ESelectionMode::Multi)
+		, _ClearSelectionOnClick(true)
+		, _ExternalScrollbar()
+	{}
+
+		SLATE_EVENT( FOnGenerateRow, OnGenerateRow )
+
+		SLATE_EVENT( FOnItemScrolledIntoView, OnItemScrolledIntoView )
+
+		SLATE_ARGUMENT( const TArray<ItemType>* , ListItemsSource )
+
+		SLATE_ATTRIBUTE( float, ItemHeight )
+
+		SLATE_EVENT( FOnContextMenuOpening, OnContextMenuOpening )
+
+		SLATE_EVENT( FOnMouseButtonDoubleClick, OnMouseButtonDoubleClick )
+
+		SLATE_EVENT( FOnSelectionChanged, OnSelectionChanged )
+
+		SLATE_ATTRIBUTE( ESelectionMode::Type, SelectionMode )
+
+		SLATE_ARGUMENT( TSharedPtr<SHeaderRow>, HeaderRow )
+
+		SLATE_ARGUMENT ( bool, ClearSelectionOnClick )
+
+		SLATE_ARGUMENT( TSharedPtr<SScrollBar>, ExternalScrollbar )
+
+		SLATE_ATTRIBUTE( EVisibility, ScrollbarVisibility)
+
+		SLATE_END_ARGS()
+
+		/**
+		 * Construct this widget
+		 *
+		 * @param	InArgs	The declaration data for this widget
+		 */
+		void Construct( const typename SListView<ItemType>::FArguments& InArgs )
+	{
+		this->OnGenerateRow = InArgs._OnGenerateRow;
+		this->OnItemScrolledIntoView = InArgs._OnItemScrolledIntoView;
+
+		this->ItemsSource = InArgs._ListItemsSource;
+		this->OnContextMenuOpening = InArgs._OnContextMenuOpening;
+		this->OnDoubleClick = InArgs._OnMouseButtonDoubleClick;
+		this->OnSelectionChanged = InArgs._OnSelectionChanged;
+		this->SelectionMode = InArgs._SelectionMode;
+
+		this->bClearSelectionOnClick = InArgs._ClearSelectionOnClick;
+
+		// Check for any parameters that the coder forgot to specify.
+		FString ErrorString;
+		{
+			if ( !this->OnGenerateRow.IsBound() )
+			{
+				ErrorString += TEXT("Please specify an OnGenerateRow. \n");
+			}
+
+			if ( this->ItemsSource == NULL )
+			{
+				ErrorString += TEXT("Please specify a ListItemsSource. \n");
+			}
+		}
+
+		if (ErrorString.Len() > 0)
+		{
+			// Let the coder know what they forgot
+			this->ChildSlot
+				.HAlign(HAlign_Center)
+				.VAlign(VAlign_Center)
+				[
+					SNew(STextBlock)
+					.Text(ErrorString)
+				];
+		}
+		else
+		{
+			// Make the TableView
+			ConstructChildren( 0, InArgs._ItemHeight, InArgs._HeaderRow, InArgs._ExternalScrollbar );
+			if(ScrollBar.IsValid())
+			{
+				ScrollBar->SetUserVisibility(InArgs._ScrollbarVisibility);
+			}
+
+		}
+	}
+
+	SListView( ETableViewMode::Type InListMode = ETableViewMode::List )
+		: STableViewBase( InListMode )
+		, SelectorItem( NullableItemType(NULL) )
+		, RangeSelectionStart( NullableItemType(NULL) )
+		, ItemsSource( NULL )
+		, ItemToScrollIntoView( NullableItemType(NULL) )
+		, ItemToNotifyWhenInView( NullableItemType(NULL) ) 
+	{
+	}
+
+public:
+	// Inherited from SWidget
+	virtual FReply OnKeyDown( const FGeometry& MyGeometry, const FKeyboardEvent& InKeyboardEvent ) OVERRIDE
+	{
+		const TArray<ItemType>& ItemsSourceRef = (*this->ItemsSource);
+
+		// Don't respond to key-presses containing "Alt" as a modifier
+		if ( ItemsSourceRef.Num() > 0 && !InKeyboardEvent.IsAltDown() )
+		{
+			bool bWasHandled = false;
+			NullableItemType ItemNavigatedTo( NULL );
+
+			// Check for selection manipulation keys (Up, Down, Home, End, PageUp, PageDown)
+			if ( InKeyboardEvent.GetKey() == EKeys::Home )
+			{
+				// Select the first item
+				ItemNavigatedTo = ItemsSourceRef[0];
+				bWasHandled = true;
+			}
+			else if ( InKeyboardEvent.GetKey() == EKeys::End )
+			{
+				// Select the last item
+				ItemNavigatedTo = ItemsSourceRef.Last();
+				bWasHandled = true;
+			}
+			else if ( InKeyboardEvent.GetKey() == EKeys::PageUp )
+			{
+				int32 SelectionIndex = 0;
+				if( TListTypeTraits<ItemType>::IsPtrValid(SelectorItem) )
+				{
+					SelectionIndex = ItemsSourceRef.Find( TListTypeTraits<ItemType>::NullableItemTypeConvertToItemType( SelectorItem ) );
+				}
+
+				int32 NumItemsInAPage = GetNumLiveWidgets();
+				int32 Remainder = NumItemsInAPage % GetNumItemsWide();
+				NumItemsInAPage -= Remainder;
+
+				if ( SelectionIndex >= NumItemsInAPage )
+				{
+					// Select an item on the previous page
+					ItemNavigatedTo = ItemsSourceRef[SelectionIndex - NumItemsInAPage];
+				}
+				else
+				{
+					ItemNavigatedTo = ItemsSourceRef[0];
+				}
+
+				bWasHandled = true;
+			}
+			else if ( InKeyboardEvent.GetKey() == EKeys::PageDown )
+			{
+				int32 SelectionIndex = 0;
+				if( TListTypeTraits<ItemType>::IsPtrValid(SelectorItem) )
+				{
+					SelectionIndex = ItemsSourceRef.Find( TListTypeTraits<ItemType>::NullableItemTypeConvertToItemType( SelectorItem ) );
+				}
+
+				int32 NumItemsInAPage = GetNumLiveWidgets();
+				int32 Remainder = NumItemsInAPage % GetNumItemsWide();
+				NumItemsInAPage -= Remainder;
+
+				if ( SelectionIndex < ItemsSourceRef.Num() - NumItemsInAPage )
+				{
+					// Select an item on the next page
+					ItemNavigatedTo = ItemsSourceRef[SelectionIndex + NumItemsInAPage];
+				}
+				else
+				{
+					ItemNavigatedTo = ItemsSourceRef.Last();
+				}
+
+				bWasHandled = true;
+			}
+			else if ( InKeyboardEvent.GetKey() == EKeys::Up )
+			{
+				int32 SelectionIndex = 0;
+				if( TListTypeTraits<ItemType>::IsPtrValid(SelectorItem) )
+				{
+					SelectionIndex = ItemsSourceRef.Find( TListTypeTraits<ItemType>::NullableItemTypeConvertToItemType( SelectorItem ) );
+				}
+
+				int32 NumItemsWide = this->GetNumItemsWide();
+
+				if ( SelectionIndex >= NumItemsWide )
+				{
+					// Select an item on the previous row
+					ItemNavigatedTo = ItemsSourceRef[SelectionIndex - NumItemsWide];
+				}
+
+				bWasHandled = true;
+			}
+			else if ( InKeyboardEvent.GetKey() == EKeys::Down )
+			{
+				// Begin at INDEX_NONE so the first item will get selected
+				int32 SelectionIndex = INDEX_NONE;
+				if( TListTypeTraits<ItemType>::IsPtrValid(SelectorItem) )
+				{
+					SelectionIndex = ItemsSourceRef.Find( TListTypeTraits<ItemType>::NullableItemTypeConvertToItemType( SelectorItem ) );
+				}
+
+				int32 NumItemsWide = this->GetNumItemsWide();
+
+				if ( SelectionIndex < ItemsSourceRef.Num() - NumItemsWide )
+				{
+					// Select an item on the next row
+					ItemNavigatedTo = ItemsSourceRef[SelectionIndex + NumItemsWide];
+				}
+
+				bWasHandled = true;
+			}
+
+			if( TListTypeTraits<ItemType>::IsPtrValid(ItemNavigatedTo) )
+			{
+				ItemType ItemToSelect( TListTypeTraits<ItemType>::NullableItemTypeConvertToItemType( ItemNavigatedTo ) );
+				KeyboardSelect( ItemToSelect, InKeyboardEvent );
+			}
+			else
+			{
+				// Change selected status of item.
+				if( TListTypeTraits<ItemType>::IsPtrValid(SelectorItem) && InKeyboardEvent.GetKey() == EKeys::SpaceBar )
+				{
+					ItemType SelectorItemDereference( TListTypeTraits<ItemType>::NullableItemTypeConvertToItemType( SelectorItem ) );
+
+					// Deselect.
+					if( InKeyboardEvent.IsControlDown() || this->SelectionMode == ESelectionMode::SingleToggle )
+					{
+						this->Private_SetItemSelection( SelectorItemDereference, !( this->Private_IsItemSelected( SelectorItemDereference ) ), true );
+						this->Private_SignalSelectionChanged( ESelectInfo::OnKeyPress );
+						bWasHandled = true;
+					}
+					else
+					{
+						// Already selected, don't handle.
+						if( this->Private_IsItemSelected( SelectorItemDereference ) )
+						{
+							bWasHandled = false;
+						}
+						// Select.
+						else
+						{
+							this->Private_SetItemSelection( SelectorItemDereference, true, true );
+							this->Private_SignalSelectionChanged( ESelectInfo::OnKeyPress );
+							bWasHandled = true;
+						}
+					}
+
+					RangeSelectionStart = SelectorItem;
+
+					// If the selector is not in the view, scroll it into view.
+					TSharedPtr<ITableRow> WidgetForItem = this->WidgetGenerator.GetWidgetForItem( SelectorItemDereference );
+					if ( !WidgetForItem.IsValid() )
+					{
+						this->RequestScrollIntoView( SelectorItemDereference );
+					}
+				}
+				// Select all items
+				else if ( (!InKeyboardEvent.IsShiftDown() && !InKeyboardEvent.IsAltDown() && InKeyboardEvent.IsControlDown() && InKeyboardEvent.GetKey() == EKeys::A) && this->SelectionMode == ESelectionMode::Multi )
+				{
+					this->Private_ClearSelection();
+
+					for ( int32 ItemIdx = 0; ItemIdx < ItemsSourceRef.Num(); ++ItemIdx )
+					{
+						this->Private_SetItemSelection( ItemsSourceRef[ItemIdx], true );
+					}
+
+					this->Private_SignalSelectionChanged(ESelectInfo::OnKeyPress);
+
+					bWasHandled = true;
+				}
+			}
+
+			return (bWasHandled ? FReply::Handled() : FReply::Unhandled());
+		}
+
+		return STableViewBase::OnKeyDown(MyGeometry, InKeyboardEvent);
+	}
+
+	virtual FReply OnMouseButtonDown( const FGeometry& MyGeometry, const FPointerEvent& MouseEvent ) OVERRIDE
+	{
+		if ( bClearSelectionOnClick
+			&& SelectionMode != ESelectionMode::None
+			&& MouseEvent.GetEffectingButton() == EKeys::LeftMouseButton
+			&& !MouseEvent.IsControlDown()
+			&& !MouseEvent.IsShiftDown()
+			)
+		{
+			// Left clicking on a list (but not an item) will clear the selection on mouse button down.
+			// Right clicking is handled on mouse up.
+			if ( this->Private_GetNumSelectedItems() > 0 )
+			{
+				this->Private_ClearSelection();
+				this->Private_SignalSelectionChanged(ESelectInfo::OnMouseClick);
+			}
+
+			return FReply::Handled();
+		}
+
+		return STableViewBase::OnMouseButtonDown(MyGeometry, MouseEvent);
+	}
+
+	virtual FReply OnMouseButtonUp( const FGeometry& MyGeometry, const FPointerEvent& MouseEvent ) OVERRIDE
+	{
+		if ( bClearSelectionOnClick
+			&& SelectionMode != ESelectionMode::None
+			&& MouseEvent.GetEffectingButton() == EKeys::RightMouseButton
+			&& !MouseEvent.IsControlDown()
+			&& !MouseEvent.IsShiftDown()
+			&& !this->IsRightClickScrolling()
+			)
+		{
+			// Right clicking on a list (but not an item) will clear the selection on mouse button up.
+			// Left clicking is handled on mouse down
+			if ( this->Private_GetNumSelectedItems() > 0 )
+			{
+				this->Private_ClearSelection();
+				this->Private_SignalSelectionChanged(ESelectInfo::OnMouseClick);
+			}
+		}
+
+		return STableViewBase::OnMouseButtonUp(MyGeometry, MouseEvent);
+	}
+
+private:
+
+	/**
+	 * A WidgetGenerator is a component responsible for creating widgets from data items.
+	 * It also provides mapping from currently generated widgets to the data items which they
+	 * represent.
+	 */
+	class FWidgetGenerator
+	{
+	public:
+		/**
+		 * Find a widget for this item if it has already been constructed.
+		 *
+		 * @param Item  The item for which to find the widget.
+		 *
+		 * @return A pointer to the corresponding widget if it exists; otherwise NULL.
+		 */
+		TSharedPtr<ITableRow> GetWidgetForItem( const ItemType& Item )
+		{
+			TSharedRef<ITableRow>* LookupResult = ItemToWidgetMap.Find( Item );
+			if ( LookupResult != NULL )
+			{
+				return *LookupResult;
+			}
+			else
+			{
+				return TSharedPtr<ITableRow>(NULL);
+			}
+		}
+
+		/**
+		 * Keep track of every item and corresponding widget during a generation pass.
+		 *
+		 * @param InItem             The DataItem which is in view.
+		 * @param InGeneratedWidget  The widget generated for this item; it may have been newly generated.
+		 */
+		void OnItemSeen( ItemType InItem, TSharedRef<ITableRow> InGeneratedWidget)
+		{
+			TSharedRef<ITableRow>* LookupResult = ItemToWidgetMap.Find( InItem );
+			const bool bWidgetIsNewlyGenerated = (LookupResult == NULL);
+			if ( bWidgetIsNewlyGenerated )
+			{
+				// It's a newly generated item!
+				ItemToWidgetMap.Add( InItem, InGeneratedWidget );
+				WidgetMapToItem.Add( &InGeneratedWidget.Get(), InItem );		
+			}
+
+			// We should not clean up this item's widgets because it is in view.
+			ItemsToBeCleanedUp.Remove(InItem);
+			ItemsWithGeneratedWidgets.Add(InItem);
+		}
+
+		/**
+		 * Called at the beginning of the generation pass.
+		 * Begins tracking of which widgets were in view and which were not (so we can clean them up)
+		 * 
+		 * @param InNumDataItems   The total number of data items being observed
+		 */
+		void OnBeginGenerationPass()
+		{
+			// Assume all the previously generated items need to be cleaned up.				
+			ItemsToBeCleanedUp = ItemsWithGeneratedWidgets;
+			ItemsWithGeneratedWidgets.Empty();
+		}
+
+		/**
+		 * Called at the end of the generation pass.
+		 * Cleans up any widgets associated with items that were not in view this frame.
+		 */
+		void OnEndGenerationPass()
+		{
+			for( int32 ItemIndex = 0; ItemIndex < ItemsToBeCleanedUp.Num(); ++ItemIndex )
+			{
+				ItemType ItemToBeCleanedUp = ItemsToBeCleanedUp[ItemIndex];
+				const TSharedRef<ITableRow>* FindResult = ItemToWidgetMap.Find( ItemToBeCleanedUp );
+				if ( FindResult != NULL )
+				{
+					const TSharedRef<ITableRow> WidgetToCleanUp = *FindResult;
+					ItemToWidgetMap.Remove( ItemToBeCleanedUp );
+					WidgetMapToItem.Remove( &WidgetToCleanUp.Get() );
+				}				
+			}
+
+			checkf( ItemToWidgetMap.Num() == WidgetMapToItem.Num(), TEXT("ItemToWidgetMap length (%d) does not match WidgetMapToItem length (%d)"), ItemToWidgetMap.Num(), WidgetMapToItem.Num() );
+			checkf( WidgetMapToItem.Num() == ItemsWithGeneratedWidgets.Num(), TEXT("WidgetMapToItem length (%d) does not match ItemsWithGeneratedWidgets length (%d). This is often because the same item is in the list more than once."), WidgetMapToItem.Num(), ItemsWithGeneratedWidgets.Num() );
+			ItemsToBeCleanedUp.Reset();
+		}
+
+		/**
+		 * Clear everything so widgets will be regenerated
+		*/
+		void Clear()
+		{
+			// Clean up all the previously generated items			
+			ItemsToBeCleanedUp = ItemsWithGeneratedWidgets;
+			ItemsWithGeneratedWidgets.Empty();
+
+			for( int32 ItemIndex = 0; ItemIndex < ItemsToBeCleanedUp.Num(); ++ItemIndex )
+			{
+				ItemType ItemToBeCleanedUp = ItemsToBeCleanedUp[ItemIndex];
+				const TSharedRef<ITableRow>* FindResult = ItemToWidgetMap.Find( ItemToBeCleanedUp );
+				if ( FindResult != NULL )
+				{
+					const TSharedRef<ITableRow> WidgetToCleanUp = *FindResult;
+					ItemToWidgetMap.Remove( ItemToBeCleanedUp );
+					WidgetMapToItem.Remove( &WidgetToCleanUp.Get() );
+				}				
+			}
+
+			ItemsToBeCleanedUp.Reset();
+		}
+
+		/** Map of DataItems to corresponding SWidgets */
+		TMap< ItemType, TSharedRef<ITableRow> > ItemToWidgetMap;
+		/** Map of SWidgets to DataItems from which they were generated */
+		TMap< const ITableRow*, ItemType > WidgetMapToItem;
+		/** A set of Items that currently have a generated widget */
+		TArray< ItemType > ItemsWithGeneratedWidgets;
+		/** Total number of DataItems the last time we performed a generation pass. */
+		int32 TotalItemsLastGeneration;
+		/** Items that need their widgets destroyed because they are no longer on screen. */
+		TArray<ItemType> ItemsToBeCleanedUp;
+	};
+
+
+public:
+	//
+	// Private Interface
+	//
+	// A low-level interface for use various widgets generated by ItemsWidgets(Lists, Trees, etc).
+	// These handle selection, expansion, and other such properties common to ItemsWidgets.
+	//
+
+	virtual void Private_SetItemSelection( ItemType TheItem, bool bShouldBeSelected, bool bWasUserDirected = false ) OVERRIDE
+	{
+		if ( SelectionMode == ESelectionMode::None )
+		{
+			return;
+		}
+
+		if ( bShouldBeSelected )
+		{
+			SelectedItems.Add( TheItem );
+		}
+		else
+		{
+			SelectedItems.Remove( TheItem );
+		}
+
+		// Only move the selector item and range selection start if the user directed this change in selection.
+		if( bWasUserDirected )
+		{
+			SelectorItem = TheItem;
+			RangeSelectionStart = TheItem;
+		}
+
+		this->InertialScrollManager.ClearScrollVelocity();
+	}
+
+	virtual void Private_ClearSelection() OVERRIDE
+	{
+		SelectedItems.Empty();
+
+		this->InertialScrollManager.ClearScrollVelocity();
+	}
+
+	virtual void Private_SelectRangeFromCurrentTo( ItemType InRangeSelectionEnd ) OVERRIDE
+	{
+		if ( SelectionMode == ESelectionMode::None )
+		{
+			return;
+		}
+
+		const TArray<ItemType>& ItemsSourceRef = (*ItemsSource);
+
+		int32 RangeStartIndex = 0;
+		if( TListTypeTraits<ItemType>::IsPtrValid(RangeSelectionStart) )
+		{
+			RangeStartIndex = ItemsSourceRef.Find( TListTypeTraits<ItemType>::NullableItemTypeConvertToItemType( RangeSelectionStart ) );
+		}
+
+		int32 RangeEndIndex = ItemsSourceRef.Find( InRangeSelectionEnd );
+
+		RangeStartIndex = FMath::Clamp(RangeStartIndex, 0, ItemsSourceRef.Num());
+		RangeEndIndex = FMath::Clamp(RangeEndIndex, 0, ItemsSourceRef.Num());
+
+		if (RangeEndIndex < RangeStartIndex)
+		{
+			Swap( RangeStartIndex, RangeEndIndex );
+		}
+
+		for( int32 ItemIndex = RangeStartIndex; ItemIndex <= RangeEndIndex; ++ItemIndex )
+		{
+			SelectedItems.Add( ItemsSourceRef[ItemIndex] );
+		}
+
+		this->InertialScrollManager.ClearScrollVelocity();
+	}
+
+	virtual void Private_SignalSelectionChanged(ESelectInfo::Type SelectInfo) OVERRIDE
+	{
+		if ( SelectionMode == ESelectionMode::None )
+		{
+			return;
+		}
+
+		if( OnSelectionChanged.IsBound() )
+		{
+			NullableItemType SelectedItem = (SelectedItems.Num() > 0)
+				? (*typename TSet<ItemType>::TIterator(SelectedItems))
+				: TListTypeTraits< ItemType >::MakeNullPtr();
+
+			OnSelectionChanged.ExecuteIfBound(SelectedItem, SelectInfo );
+		}
+	}
+
+	virtual const ItemType* Private_ItemFromWidget( const ITableRow* TheWidget ) const OVERRIDE
+	{
+		ItemType const * LookupResult = WidgetGenerator.WidgetMapToItem.Find( TheWidget );
+		return LookupResult == NULL ? NULL : LookupResult;
+	}
+
+	virtual bool Private_UsesSelectorFocus() const OVERRIDE
+	{
+		return true;
+	}
+
+	virtual bool Private_HasSelectorFocus( const ItemType& TheItem ) const OVERRIDE
+	{
+		return SelectorItem == TheItem;
+	}
+
+	virtual bool Private_IsItemSelected( const ItemType& TheItem ) const OVERRIDE
+	{
+		return NULL != SelectedItems.Find(TheItem);
+	}
+
+	virtual bool Private_IsItemExpanded( const ItemType& TheItem ) const OVERRIDE
+	{
+		// List View does not support item expansion.
+		return false;	
+	}
+
+	virtual void Private_SetItemExpansion( ItemType TheItem, bool bShouldBeExpanded ) OVERRIDE
+	{
+		// Do nothing; you cannot expand an item in a list!
+	}
+
+	virtual void Private_OnExpanderArrowShiftClicked( ItemType TheItem, bool bShouldBeExpanded ) OVERRIDE
+	{
+		// Do nothing; you cannot expand an item in a list!
+	}
+
+	virtual bool Private_DoesItemHaveChildren( int32 ItemIndexInList ) const OVERRIDE
+	{
+		// List View items cannot have children
+		return false;
+	}
+
+	virtual int32 Private_GetNumSelectedItems() const OVERRIDE
+	{
+		return SelectedItems.Num();
+	}
+
+	virtual int32 Private_GetNestingDepth( int32 ItemIndexInList ) const OVERRIDE
+	{
+		// List View items are not indented
+		return 0;
+	}
+
+	virtual ESelectionMode::Type Private_GetSelectionMode() const OVERRIDE
+	{
+		return SelectionMode.Get();
+	}
+
+	virtual void Private_OnItemRightClicked( ItemType TheItem, const FPointerEvent& MouseEvent ) OVERRIDE
+	{
+		this->OnRightMouseButtonUp( MouseEvent.GetScreenSpacePosition() );
+	}
+
+	virtual bool Private_OnItemDoubleClicked( ItemType TheItem ) OVERRIDE
+	{
+		if( OnDoubleClick.ExecuteIfBound( TheItem ) )
+		{
+			return true;	// Handled
+		}
+
+		return false;	// Not handled
+	}
+
+	virtual ETableViewMode::Type GetTableViewMode() const OVERRIDE
+	{
+		return TableViewMode;
+	}
+
+	virtual TSharedRef<SWidget> AsWidget() OVERRIDE
+	{
+		return SharedThis(this);
+	}
+
+
+
+public:	
+
+	/**
+	 * Remove any items that are no longer in the list from the selection set.
+	 */
+	virtual void UpdateSelectionSet()
+	{
+		// Trees take care of this update in a different way.
+		if ( TableViewMode != ETableViewMode::Tree )
+		{
+			bool bSelectionChanged = false;
+			if ( ItemsSource == NULL )
+			{
+				// We are no longer observing items so there is no more selection.
+				this->Private_ClearSelection();
+				bSelectionChanged = true;
+			}
+			else
+			{
+				// We are observing some items; they are potentially different.
+				// Unselect any that are no longer being observed.
+				TSet< ItemType > NewSelectedItems;
+				for ( int32 ItemIndex = 0; ItemIndex < ItemsSource->Num(); ++ItemIndex )
+				{
+					ItemType CurItem = (*ItemsSource)[ItemIndex];
+					const bool bItemIsSelected = ( NULL != SelectedItems.Find( CurItem ) );
+					if ( bItemIsSelected )
+					{
+						NewSelectedItems.Add( CurItem );
+					}
+				}
+
+				// Look for items that were removed from the selection.
+				TSet< ItemType > SetDifference = SelectedItems.Difference( NewSelectedItems );
+				bSelectionChanged = (SetDifference.Num()) > 0;
+
+				// Update the selection to reflect the removal of any items from the ItemsSource.
+				SelectedItems = NewSelectedItems;
+			}
+
+			if (bSelectionChanged)
+			{
+				Private_SignalSelectionChanged(ESelectInfo::Direct);
+			}
+		}
+	}
+
+	/**
+	 * Update generate Widgets for Items as needed and clean up any Widgets that are no longer needed.
+	 * Re-arrange the visible widget order as necessary.
+	 */
+	virtual FReGenerateResults ReGenerateItems( const FGeometry& MyGeometry ) OVERRIDE
+	{
+		// Clear all the items from our panel. We will re-add them in the correct order momentarily.
+		this->ClearWidgets();
+
+		// Ensure that we always begin and clean up a generation pass.
+		
+		FGenerationPassGuard GenerationPassGuard(WidgetGenerator);
+
+		const TArray<ItemType>* SourceItems = ItemsSource;
+		if ( SourceItems != NULL && SourceItems->Num() > 0 )
+		{
+			// Items in view, including fractional items
+			float ItemsInView = 0.0f;
+
+			// Height of generated widgets that is landing in the bounds of the view.
+			float ViewHeightUsedSoFar = 0.0f;
+
+			// Total height of widgets generated so far.
+			float HeightGeneratedSoFar = 0.0f;
+
+			// Index of the item at which we start generating based on how far scrolled down we are
+			// Note that we must generate at LEAST one item.
+			int32 StartIndex = FMath::Clamp( FMath::Floor(ScrollOffset), 0, SourceItems->Num()-1 );
+
+			// Height of the first item that is generated. This item is at the location where the user requested we scroll
+			float FirstItemHeight = 0.0f;
+
+			// Generate widgets assuming scenario a.
+			bool bGeneratedEnoughForSmoothScrolling = false;
+			bool AtEndOfList = false;
+
+			for( int32 ItemIndex = StartIndex; !bGeneratedEnoughForSmoothScrolling && ItemIndex < SourceItems->Num(); ++ItemIndex )
+			{
+				const ItemType& CurItem = (*SourceItems)[ItemIndex];
+
+				const float ItemHeight = GenerateWidgetForItem( CurItem, ItemIndex, StartIndex );
+
+				const bool IsFirstItem = ItemIndex == StartIndex;
+
+				if (IsFirstItem)
+				{
+					FirstItemHeight = ItemHeight;
+				}
+
+				// Track the number of items in the view, including fractions.
+				if (IsFirstItem)
+				{
+					ItemsInView += 1.0f - FMath::Fractional(ScrollOffset);
+				}
+				else if (ViewHeightUsedSoFar + ItemHeight > MyGeometry.Size.Y)
+				{
+					ItemsInView += (MyGeometry.Size.Y - ViewHeightUsedSoFar) / ItemHeight;
+				}
+				else
+				{
+					ItemsInView += 1;
+				}
+
+				HeightGeneratedSoFar += ItemHeight;
+
+				ViewHeightUsedSoFar += (IsFirstItem)
+					? ItemHeight * (1.0f - FMath::Fractional(ScrollOffset))
+					: ItemHeight;
+
+				if (ItemIndex >= SourceItems->Num()-1)
+				{
+					AtEndOfList = true;
+				}
+
+				if (ViewHeightUsedSoFar > MyGeometry.Size.Y )
+				{
+					bGeneratedEnoughForSmoothScrolling = true;
+				}
+			}
+
+			// Handle scenario b.
+			// We may have stopped because we got to the end of the items.
+			// But we may still have space to fill!
+			if (AtEndOfList && ViewHeightUsedSoFar < MyGeometry.Size.Y)
+			{
+				float NewScrollOffsetForBackfill = StartIndex + (HeightGeneratedSoFar - MyGeometry.Size.Y) / FirstItemHeight;
+
+				for( int32 ItemIndex = StartIndex-1; HeightGeneratedSoFar < MyGeometry.Size.Y && ItemIndex >= 0; --ItemIndex )
+				{
+					const ItemType& CurItem = (*SourceItems)[ItemIndex];
+
+					const float ItemHeight = GenerateWidgetForItem( CurItem, ItemIndex, StartIndex );
+
+					if (HeightGeneratedSoFar + ItemHeight > MyGeometry.Size.Y)
+					{
+						// Generated the item that puts us over the top.
+						// Count the fraction of this item that will stick out above the list
+						NewScrollOffsetForBackfill = ItemIndex + (HeightGeneratedSoFar + ItemHeight - MyGeometry.Size.Y) / ItemHeight;
+					}
+
+					// The widget used up some of the available vertical space.
+					HeightGeneratedSoFar += ItemHeight;
+				}
+
+				return FReGenerateResults(NewScrollOffsetForBackfill, HeightGeneratedSoFar, ItemsSource->Num() - NewScrollOffsetForBackfill, AtEndOfList);
+			}
+
+			return FReGenerateResults(ScrollOffset, HeightGeneratedSoFar, ItemsInView, AtEndOfList);
+		}
+
+		return FReGenerateResults(0.0f, 0.0f, 0.0f, false);
+
+	}
+
+	float GenerateWidgetForItem( const ItemType& CurItem, int32 ItemIndex, int32 StartIndex )
+	{
+		// Find a previously generated Widget for this item, if one exists.
+		TSharedPtr<ITableRow> WidgetForItem = WidgetGenerator.GetWidgetForItem( CurItem );
+		if ( !WidgetForItem.IsValid() )
+		{
+			// We couldn't find an existing widgets, meaning that this data item was not visible before.
+			// Make a new widget for it.
+			WidgetForItem = this->GenerateNewWidget( CurItem );
+		}
+
+		// It is useful to know the item's index that the widget was generated from.
+		// Helps with even/odd coloring
+		WidgetForItem->SetIndexInList(ItemIndex);
+
+		// Let the item generator know that we encountered the current Item and associated Widget.
+		WidgetGenerator.OnItemSeen( CurItem, WidgetForItem.ToSharedRef() );
+
+		// We rely on the widgets desired size in order to determine how many will fit on screen.
+		const TSharedRef<SWidget> NewlyGeneratedWidget = WidgetForItem->AsWidget();
+		NewlyGeneratedWidget->SlatePrepass();		
+
+		const bool IsFirstWidgetOnScreen = (ItemIndex == StartIndex);
+		const float ItemHeight = NewlyGeneratedWidget->GetDesiredSize().Y;
+
+		// We have a widget for this item; add it to the panel so that it is part of the UI.
+		if (ItemIndex >= StartIndex)
+		{
+			// Generating widgets downward
+			this->AppendWidget( WidgetForItem.ToSharedRef() );
+		}
+		else
+		{
+			// Backfilling widgets; going upward
+			this->InsertWidget( WidgetForItem.ToSharedRef() );
+		}
+
+		return ItemHeight;
+	}
+
+	/** @return how many items there are in the TArray being observed */
+	virtual int32 GetNumItemsBeingObserved() const
+	{
+		return ItemsSource == NULL ? 0 : ItemsSource->Num();
+	}
+
+	/**
+	 * Given an data item, generate a widget corresponding to it.
+	 *
+	 * @param InItem  The data item which to visualize.
+	 *
+	 * @return A new Widget that represents the data item.
+	 */
+	virtual TSharedRef<ITableRow> GenerateNewWidget(ItemType InItem)
+	{
+		if ( OnGenerateRow.IsBound() )
+		{
+			return OnGenerateRow.Execute( InItem, SharedThis(this) );
+		}
+		else
+		{
+			// The programmer did not provide an OnGenerateRow() handler; let them know.
+			TSharedRef< STableRow<ItemType> > NewListItemWidget = 
+				SNew( STableRow<ItemType>, SharedThis(this) )
+				.Content()
+				[
+					SNew(STextBlock) .Text( NSLOCTEXT("SListView", "BrokenUIMessage", "OnGenerateWidget() not assigned.") )
+				];
+
+			return NewListItemWidget;
+		}
+
+	}
+
+	/**
+	 * Given a Widget, find the corresponding data item.
+	 * 
+	 * @param WidgetToFind  An widget generated by the list view for some data item.
+	 *
+	 * @return the data item from which the WidgetToFind was generated
+	 */
+	const ItemType* ItemFromWidget( const ITableRow* WidgetToFind ) const
+	{
+		return Private_ItemFromWidget( WidgetToFind );
+	}
+
+	/**
+	 * Test if the current item is selected.
+	 *
+	 * @param InItem The item to test.
+	 *
+	 * @return true if the item is selected in this list; false otherwise.
+	 */
+	bool IsItemSelected( const ItemType& InItem ) const
+	{
+		if ( SelectionMode == ESelectionMode::None )
+		{
+			return false;
+		}
+
+		return Private_IsItemSelected( InItem );
+	}
+
+	/**
+	 * Set the selection state of an item.
+	 *
+	 * @param InItem      The Item whose selection state to modify
+	 * @param bSelected   true to select the item; false to unselect
+	 * @param SelectInfo  Provides context on how the selection changed
+	 */
+	void SetItemSelection( const ItemType& InItem, bool bSelected, ESelectInfo::Type SelectInfo = ESelectInfo::Direct )
+	{
+		if ( SelectionMode == ESelectionMode::None )
+		{
+			return;
+		}
+
+		Private_SetItemSelection( InItem, bSelected, SelectInfo != ESelectInfo::Direct);
+		Private_SignalSelectionChanged(SelectInfo);
+	}
+
+	/**
+	 * Empty the selection set.
+	 */
+	void ClearSelection()
+	{
+		if ( SelectionMode == ESelectionMode::None )
+		{
+			return;
+		}
+
+		Private_ClearSelection();
+		Private_SignalSelectionChanged(ESelectInfo::Direct);
+	}
+
+	/**
+	 * Gets the number of selected items.
+	 *
+	 * @return Number of selected items.
+	 */
+	int32 GetNumItemsSelected()
+	{
+		return SelectedItems.Num();
+	}
+
+	/**
+	 * Returns a list of selected item indices, or an empty array if nothing is selected
+	 *
+	 * @return	List of selected item indices (in no particular order)
+	 */
+	TArray< ItemType > GetSelectedItems()
+	{
+		TArray< ItemType > SelectedItemArray;
+		SelectedItemArray.Empty( SelectedItems.Num() );
+		for( typename TSet< ItemType >::TConstIterator SelectedItemIt( SelectedItems ); SelectedItemIt; ++SelectedItemIt )
+		{
+			SelectedItemArray.Add( *SelectedItemIt );
+		}
+		return SelectedItemArray;
+	}
+
+	/**
+	 * Checks whether the specified item is currently visible in the list view.
+	 *
+	 * @param Item - The item to check.
+	 *
+	 * @return true if the item is visible, false otherwise.
+	 */
+	bool IsItemVisible( ItemType Item )
+	{
+		return WidgetGenerator.GetWidgetForItem(Item).IsValid();
+	}	
+
+	/**
+	 * Scroll an item into view. If the item is not found, fails silently.
+	 *
+	 * @param ItemToView  The item to scroll into view on next tick.
+	 */
+	void RequestScrollIntoView( ItemType ItemToView )
+	{
+		ItemToScrollIntoView = ItemToView;
+		RequestListRefresh();
+	}
+
+	/**
+	 * Set the currently selected Item.
+	 *
+	 * @param SoleSelectedItem   Sole item that should be selected.
+	 * @param SelectInfo Provides context on how the selection changed
+	 */
+	void SetSelection( ItemType SoleSelectedItem, ESelectInfo::Type SelectInfo = ESelectInfo::Direct  )
+	{
+		SelectedItems.Empty();
+		SetItemSelection( SoleSelectedItem, true, SelectInfo );
+	}
+
+	/**
+	 * Find a widget for this item if it has already been constructed.
+	 *
+	 * @param InItem  The item for which to find the widget.
+	 *
+	 * @return A pointer to the corresponding widget if it exists; otherwise NULL.
+	*/
+	TSharedPtr<ITableRow> WidgetFromItem( const ItemType& InItem )
+	{
+		return WidgetGenerator.GetWidgetForItem(InItem);
+	}
+
+	/**
+	 * Lists and Trees serialize items that they observe because they rely on the property
+	 * that holding a reference means it will not be garbage collected.
+	 *
+	 * @param Ar The archive to serialize with
+	 */
+	virtual void AddReferencedObjects( FReferenceCollector& Collector )
+	{
+		TListTypeTraits<ItemType>::AddReferencedObjects( Collector, WidgetGenerator.ItemsWithGeneratedWidgets, SelectedItems );
+	}
+
+protected:
+
+	/**
+	 * If there is a pending request to scroll an item into view, do so.
+	 * 
+	 * @param ListViewGeometry  The geometry of the listView; can be useful for centering the item.
+	 */
+	virtual void ScrollIntoView( const FGeometry& ListViewGeometry ) OVERRIDE
+	{
+		if ( TListTypeTraits<ItemType>::IsPtrValid(ItemToScrollIntoView) && ItemsSource != NULL )
+		{
+			const int32 IndexOfItem = ItemsSource->Find( TListTypeTraits<ItemType>::NullableItemTypeConvertToItemType( ItemToScrollIntoView ) );
+			if (IndexOfItem != INDEX_NONE)
+			{
+				double NumLiveWidgets = GetNumLiveWidgets();
+				if (NumLiveWidgets == 0 && IsPendingRefresh())
+				{
+					// Use the last number of widgets on screen to estimate if we actually need to scroll.
+					NumLiveWidgets = LastGenerateResults.ExactNumWidgetsOnScreen;
+				}
+
+				// Only scroll the item into view if it's not already in the visible range
+				const double IndexPlusOne = IndexOfItem+1;
+				if (IndexOfItem < ScrollOffset || IndexPlusOne > (ScrollOffset + NumLiveWidgets))
+				{
+					// Scroll the top of the listview to the item in question
+					ScrollOffset = IndexOfItem;
+					// Center the list view on the item in question.
+					ScrollOffset -= (NumLiveWidgets / 2);
+					//we also don't want the widget being chopped off if it is at the end of the list
+					const double MoveBackBy = FMath::Clamp<double>(IndexPlusOne - (ScrollOffset + NumLiveWidgets), 0, FLT_MAX);
+					//Move to the correct center spot
+					ScrollOffset += MoveBackBy;
+				}
+
+				RequestListRefresh();
+
+				ItemToNotifyWhenInView = ItemToScrollIntoView;
+			}
+
+			TListTypeTraits<ItemType>::ResetPtr(ItemToScrollIntoView);
+		}
+	}
+
+	virtual void NotifyItemScrolledIntoView() OVERRIDE
+	{
+		// Notify that an item that came into view
+		if ( TListTypeTraits<ItemType>::IsPtrValid( ItemToNotifyWhenInView ) )
+		{
+			ItemType NonNullItemToNotifyWhenInView = TListTypeTraits<ItemType>::NullableItemTypeConvertToItemType( ItemToNotifyWhenInView );
+			TSharedPtr<ITableRow> Widget = WidgetGenerator.GetWidgetForItem(NonNullItemToNotifyWhenInView);
+			OnItemScrolledIntoView.ExecuteIfBound( NonNullItemToNotifyWhenInView, Widget );
+
+			TListTypeTraits<ItemType>::ResetPtr( ItemToNotifyWhenInView );
+		}
+	}
+
+	virtual float ScrollBy( const FGeometry& MyGeometry, float ScrollByAmountInSlateUnits ) OVERRIDE
+	{
+		float AbsScrollByAmount = FMath::Abs( ScrollByAmountInSlateUnits );
+		int32 StartingItemIndex = (int32)ScrollOffset;
+		double NewScrollOffset = ScrollOffset;
+
+		const TArray<ItemType>* SourceItems = ItemsSource;
+		if ( SourceItems != NULL && SourceItems->Num() > 0 )
+		{
+			int ItemIndex = StartingItemIndex;
+			while( AbsScrollByAmount != 0 && ItemIndex < SourceItems->Num() && ItemIndex >= 0 )
+			{
+				const ItemType CurItem = (*SourceItems)[ ItemIndex ];
+				TSharedPtr<ITableRow> RowWidget = WidgetGenerator.GetWidgetForItem( CurItem );
+				if ( !RowWidget.IsValid() )
+				{
+					// We couldn't find an existing widgets, meaning that this data item was not visible before.
+					// Make a new widget for it.
+					RowWidget = this->GenerateNewWidget( CurItem );
+
+					// It is useful to know the item's index that the widget was generated from.
+					// Helps with even/odd coloring
+					RowWidget->SetIndexInList(ItemIndex);
+
+					// Let the item generator know that we encountered the current Item and associated Widget.
+					WidgetGenerator.OnItemSeen( CurItem, RowWidget.ToSharedRef() );
+
+					RowWidget->AsWidget()->SlatePrepass();
+				}
+
+				if ( ScrollByAmountInSlateUnits > 0 )
+				{
+					FVector2D DesiredSize = RowWidget->AsWidget()->GetDesiredSize();
+					const float RemainingHeight = DesiredSize.Y * ( 1.0 - FMath::Fractional( NewScrollOffset ) );
+
+					if ( AbsScrollByAmount > RemainingHeight )
+					{
+						if ( ItemIndex != SourceItems->Num() )
+						{
+							AbsScrollByAmount -= RemainingHeight;
+							NewScrollOffset = 1.0f + (int32)NewScrollOffset;
+							++ItemIndex;
+						}
+						else
+						{
+							NewScrollOffset = SourceItems->Num();
+							break;
+						}
+					} 
+					else if ( AbsScrollByAmount == RemainingHeight )
+					{
+						NewScrollOffset = 1.0f + (int32)NewScrollOffset;
+						break;
+					}
+					else
+					{
+						NewScrollOffset = (int32)NewScrollOffset + ( 1.0f - ( ( RemainingHeight - AbsScrollByAmount ) / DesiredSize.Y ) );
+						break;
+					}
+				}
+				else
+				{
+					FVector2D DesiredSize = RowWidget->AsWidget()->GetDesiredSize();
+
+					float Fractional = FMath::Fractional( NewScrollOffset );
+					if ( Fractional == 0 )
+					{
+						Fractional = 1.0f;
+						--NewScrollOffset;
+					}
+
+					const float PrecedingHeight = DesiredSize.Y * Fractional;
+
+					if ( AbsScrollByAmount > PrecedingHeight )
+					{
+						if ( ItemIndex != 0 )
+						{
+							AbsScrollByAmount -= PrecedingHeight;
+							NewScrollOffset -= FMath::Fractional( NewScrollOffset );
+							--ItemIndex;
+						}
+						else
+						{
+							NewScrollOffset = 0;
+							break;
+						}
+					} 
+					else if ( AbsScrollByAmount == PrecedingHeight )
+					{
+						NewScrollOffset -= FMath::Fractional( NewScrollOffset );
+						break;
+					}
+					else
+					{
+						NewScrollOffset = (int32)NewScrollOffset + ( ( PrecedingHeight - AbsScrollByAmount ) / DesiredSize.Y );
+						break;
+					}
+				}
+			}
+		}
+
+		return ScrollTo( NewScrollOffset );
+	}
+
+protected:
+
+	/**
+	 * Selects the specified item and scrolls it into view. If shift is held, it will be a range select.
+	 * 
+	 * @param ItemToSelect		The item that was selected by a keystroke
+	 * @param InKeyboardEvent	The keyboard event that caused this selection
+	 */
+	virtual void KeyboardSelect(const ItemType& ItemToSelect, const FKeyboardEvent& InKeyboardEvent, bool bCausedByNavigation=false)
+	{
+		if ( SelectionMode != ESelectionMode::None )
+		{
+			// Must be set before signaling selection changes because sometimes new items will be selected that need to stomp this value
+			SelectorItem = ItemToSelect;
+
+			if ( SelectionMode == ESelectionMode::Multi && ( InKeyboardEvent.IsShiftDown() || InKeyboardEvent.IsControlDown() ) )
+			{
+				// Range select.
+				if ( InKeyboardEvent.IsShiftDown() )
+				{
+					// Holding control makes the range select bidirectional, where as it is normally unidirectional.
+					if( !( InKeyboardEvent.IsControlDown() ) )
+					{
+						this->Private_ClearSelection();
+					}
+
+					this->Private_SelectRangeFromCurrentTo( ItemToSelect );
+				}
+
+				this->Private_SignalSelectionChanged( ESelectInfo::OnNavigation );
+			}
+			else
+			{
+				// Single select.
+				this->SetSelection( ItemToSelect, ESelectInfo::OnNavigation );
+			}
+
+			// If the selector is not in the view, scroll it into view.
+			TSharedPtr<ITableRow> WidgetForItem = this->WidgetGenerator.GetWidgetForItem( ItemToSelect );
+			if ( !WidgetForItem.IsValid() )
+			{
+				this->RequestScrollIntoView( ItemToSelect );
+			}
+		}
+	}
+
+protected:
+	/** A widget generator component */
+	FWidgetGenerator WidgetGenerator;
+
+	/** Delegate to be invoked when the list needs to generate a new widget from a data item. */
+	FOnGenerateRow OnGenerateRow;
+
+	/** Delegate to be invoked when an item has come into view after it was requested to come into view. */
+	FOnItemScrolledIntoView OnItemScrolledIntoView;
+
+	/** A set of selected data items */
+	TSet< ItemType > SelectedItems;
+
+	/** The item to manipulate selection for */
+	NullableItemType SelectorItem;
+
+	/** The item which was last manipulated; used as a start for shift-click selection */
+	NullableItemType RangeSelectionStart;
+
+	/** Pointer to the array of data items that we are observing */
+	const TArray<ItemType>* ItemsSource;
+
+	/** When not null, the list will try to scroll to this item on tick. */
+	NullableItemType ItemToScrollIntoView;
+
+	/** When set, the list will notify this item when it has been scrolled into view */
+	NullableItemType ItemToNotifyWhenInView;
+
+	/** Delegate to invoke when selection changes. */
+	FOnSelectionChanged OnSelectionChanged;
+
+	/** Called when the user double-clicks on an element in the list view with the left mouse button */
+	FOnMouseButtonDoubleClick OnDoubleClick;
+
+	/** If true, the selection will be cleared if the user clicks in empty space (not on an item) */
+	bool bClearSelectionOnClick;
+
+private:
+
+	struct FGenerationPassGuard
+	{
+		FWidgetGenerator& Generator;
+		FGenerationPassGuard( FWidgetGenerator& InGenerator )
+			: Generator(InGenerator)
+		{
+			// Let the WidgetGenerator that we are starting a pass so that it can keep track of data items and widgets.
+			Generator.OnBeginGenerationPass();
+		}
+
+		~FGenerationPassGuard()
+		{
+			// We have completed the generation pass. The WidgetGenerator will clean up unused Widgets when it goes out of scope.
+			Generator.OnEndGenerationPass();
+		}
+	};
+};
diff --git a/Engine/Source/Runtime/Slate/Public/Widgets/Views/STableRow.h b/Engine/Source/Runtime/Slate/Public/Widgets/Views/STableRow.h
new file mode 100644
--- /dev/null
+++ b/Engine/Source/Runtime/Slate/Public/Widgets/Views/STableRow.h
@@ -0,0 +1,831 @@
+// Copyright 1998-2014 Epic Games, Inc. All Rights Reserved.
+
+#pragma once
+
+/**
+ * Interface for table views to talk to their rows.
+ */
+class SLATE_API ITableRow
+{
+	public:
+		/**
+		 * @param InIndexInList  The index of the item for which this widget was generated
+		 */
+		virtual void SetIndexInList( int32 InIndexInList ) = 0;
+
+		/** @return true if the corresponding item is expanded; false otherwise*/
+		virtual bool IsItemExpanded() const = 0;
+
+		/** Toggle the expansion of the item associated with this row */
+		virtual void ToggleExpansion() = 0;
+
+		/** @return how nested the item associated with this row when it is in a TreeView */
+		virtual int32 GetIndentLevel() const = 0;
+
+		/** @return Does this item have children? */
+		virtual int32 DoesItemHaveChildren() const = 0;
+
+		/** @return this table row as a widget */
+		virtual TSharedRef<SWidget> AsWidget() = 0;
+
+		/** @return the content of this table row */
+		virtual TSharedPtr<SWidget> GetContent() = 0;
+
+		/** Called when the expander arrow for this row is shift+clicked */
+		virtual void Private_OnExpanderArrowShiftClicked() = 0;
+};
+
+template <typename ItemType> class SListView;
+
+DECLARE_DELEGATE_OneParam(FOnTableRowDragEnter, FDragDropEvent const&);
+
+DECLARE_DELEGATE_OneParam(FOnTableRowDragLeave, FDragDropEvent const&);
+
+DECLARE_DELEGATE_RetVal_OneParam(FReply, FOnTableRowDrop, FDragDropEvent const&);
+/**
+ * The ListView is populated by Selectable widgets.
+ * A Selectable widget is away of the ListView containing it (OwnerWidget) and holds arbitrary Content (Content).
+ * A Selectable works with its corresponding ListView to provide selection functionality.
+ */
+template<typename ItemType>
+class STableRow : public ITableRow, public SBorder
+{
+	checkAtCompileTime( TIsValidListItem<ItemType>::Value, ItemType_must_be_a_pointer_or_TSharedPtr );
+
+public:
+
+
+	SLATE_BEGIN_ARGS( STableRow< ItemType > )
+		: _Style( &FCoreStyle::Get().GetWidgetStyle<FTableRowStyle>("TableView.Row") )
+		, _Padding( FMargin(0) )
+		, _ShowSelection(true)
+		, _Content()
+		{}
+	
+		SLATE_STYLE_ARGUMENT( FTableRowStyle, Style )
+
+		SLATE_EVENT( FOnDragDetected, OnDragDetected )
+		SLATE_EVENT( FOnTableRowDragEnter, OnDragEnter )
+		SLATE_EVENT( FOnTableRowDragLeave, OnDragLeave )
+		SLATE_EVENT( FOnTableRowDrop,      OnDrop )
+
+		SLATE_ATTRIBUTE( FMargin, Padding )
+	
+		SLATE_ARGUMENT( bool, ShowSelection )
+		
+		SLATE_DEFAULT_SLOT( typename STableRow<ItemType>::FArguments, Content )
+
+	SLATE_END_ARGS()
+
+	/**
+	 * Construct this widget
+	 *
+	 * @param	InArgs	The declaration data for this widget
+	 */
+	void Construct(const FArguments& InArgs, const TSharedRef<STableViewBase>& InOwnerTableView)
+	{
+		ConstructInternal(InArgs, InOwnerTableView);
+
+		ConstructChildren(
+			InOwnerTableView->TableViewMode,
+			InArgs._Padding,
+			InArgs._Content.Widget
+		);
+	}
+
+	virtual void ConstructChildren( ETableViewMode::Type InOwnerTableMode, const TAttribute<FMargin>& InPadding, const TSharedRef<SWidget>& InContent )
+	{
+		this->Content = InContent;
+
+		if ( InOwnerTableMode == ETableViewMode::List || InOwnerTableMode == ETableViewMode::Tile )
+		{
+			// -- Row is in a ListView or the user --
+
+			// We just need to hold on to this row's content.
+			this->ChildSlot
+			.Padding( InPadding )
+			[
+				InContent
+			];
+		}
+		else
+		{
+			// -- Row is for TreeView --
+			
+			// Rows in a TreeView need an expander button and some indentation
+			this->ChildSlot
+			[
+				SNew(SHorizontalBox)
+				+SHorizontalBox::Slot()
+					.AutoWidth()
+					.HAlign(HAlign_Right) .VAlign(VAlign_Fill)
+				[
+					SNew(SExpanderArrow, SharedThis(this) )
+				]
+				+SHorizontalBox::Slot()
+					.FillWidth(1)
+					.Padding( InPadding )
+				[
+					InContent
+				]
+			];
+		}
+	}
+
+	virtual int32 OnPaint( const FGeometry& AllottedGeometry, const FSlateRect& MyClippingRect, FSlateWindowElementList& OutDrawElements, int32 LayerId, const FWidgetStyle& InWidgetStyle, bool bParentEnabled ) const OVERRIDE
+	{
+		int32 Result = SBorder::OnPaint(AllottedGeometry, MyClippingRect, OutDrawElements, LayerId, InWidgetStyle, bParentEnabled);
+
+		TSharedRef< ITypedTableView<ItemType> > OwnerWidget = OwnerTablePtr.Pin().ToSharedRef();
+		const bool bIsActive = OwnerWidget->AsWidget()->HasKeyboardFocus();
+		const ItemType* MyItem = OwnerWidget->Private_ItemFromWidget( this );
+		if ( bIsActive && OwnerWidget->Private_UsesSelectorFocus() && OwnerWidget->Private_HasSelectorFocus( *MyItem ) )
+		{
+			FSlateDrawElement::MakeBox(
+				OutDrawElements,
+				LayerId,
+				AllottedGeometry.ToPaintGeometry(),
+				&Style->SelectorFocusedBrush,
+				MyClippingRect,
+				ESlateDrawEffect::None,
+				Style->SelectorFocusedBrush.GetTint( InWidgetStyle ) * InWidgetStyle.GetColorAndOpacityTint()
+				);
+		}
+
+		return Result;
+	}
+
+	/**
+	 * Called when a mouse button is double clicked.  Override this in derived classes.
+	 *
+	 * @param  InMyGeometry  Widget geometry
+	 * @param  InMouseEvent  Mouse button event
+	 *
+	 * @return  Returns whether the event was handled, along with other possible actions
+	 */
+	virtual FReply OnMouseButtonDoubleClick( const FGeometry& InMyGeometry, const FPointerEvent& InMouseEvent )
+	{
+		if( InMouseEvent.GetEffectingButton() == EKeys::LeftMouseButton )
+		{
+			TSharedRef< ITypedTableView<ItemType> > OwnerWidget = OwnerTablePtr.Pin().ToSharedRef();
+
+			// Only one item can be double-clicked
+			const ItemType* MyItem = OwnerWidget->Private_ItemFromWidget( this );
+
+			// If we're configured to route double-click messages to the owner of the table, then
+			// do that here.  Otherwise, we'll toggle expansion.
+			const bool bWasHandled = OwnerWidget->Private_OnItemDoubleClicked( *MyItem );
+			if( !bWasHandled )
+			{
+				ToggleExpansion();
+			}
+
+			return FReply::Handled();
+		}
+
+		return FReply::Unhandled();
+	}
+
+	/**
+	 * See SWidget::OnMouseButtonDown
+	 *
+	 * @param MyGeometry The Geometry of the widget receiving the event
+	 * @param MouseEvent Information about the input event
+	 *
+	 * @return Whether the event was handled along with possible requests for the system to take action.
+	 */	
+	virtual FReply OnMouseButtonDown( const FGeometry& MyGeometry, const FPointerEvent& MouseEvent )
+	{
+		TSharedPtr< ITypedTableView<ItemType> > OwnerWidget = OwnerTablePtr.Pin();
+		ChangedSelectionOnMouseDown = false;
+
+		check(OwnerWidget.IsValid());
+
+		if ( MouseEvent.GetEffectingButton() == EKeys::LeftMouseButton )
+		{
+			switch( OwnerWidget->Private_GetSelectionMode() )
+			{
+			case ESelectionMode::Single:
+				{
+					const ItemType* MyItem = OwnerWidget->Private_ItemFromWidget( this );
+					const bool bIsSelected = OwnerWidget->Private_IsItemSelected( *MyItem );
+
+					// Select the item under the cursor
+					if( !bIsSelected )
+					{
+						OwnerWidget->Private_ClearSelection();
+						OwnerWidget->Private_SetItemSelection( *MyItem, true, true );
+						ChangedSelectionOnMouseDown = true;
+					}
+
+					return FReply::Handled()
+						.DetectDrag( SharedThis(this), EKeys::LeftMouseButton )
+						.SetKeyboardFocus( OwnerWidget->AsWidget(), EKeyboardFocusCause::Mouse )
+						.CaptureMouse( SharedThis(this) );
+				}
+				break;
+
+			case ESelectionMode::SingleToggle:
+				{
+					const ItemType* MyItem = OwnerWidget->Private_ItemFromWidget( this );
+					const bool bIsSelected = OwnerWidget->Private_IsItemSelected( *MyItem );
+
+					if( !bIsSelected )
+					{
+						OwnerWidget->Private_ClearSelection();
+						OwnerWidget->Private_SetItemSelection( *MyItem, true, true );
+						ChangedSelectionOnMouseDown = true;
+					}
+
+					return FReply::Handled()
+						.DetectDrag( SharedThis(this), EKeys::LeftMouseButton )
+						.SetKeyboardFocus( OwnerWidget->AsWidget(), EKeyboardFocusCause::Mouse )
+						.CaptureMouse( SharedThis(this) );
+				}
+				break;
+
+			case ESelectionMode::Multi:
+				{
+					const ItemType* MyItem = OwnerWidget->Private_ItemFromWidget( this );
+					const bool bIsSelected = OwnerWidget->Private_IsItemSelected( *MyItem );
+
+					check(MyItem != NULL);
+
+					if ( MouseEvent.IsControlDown() )
+					{
+						OwnerWidget->Private_SetItemSelection( *MyItem, !bIsSelected, true );
+						ChangedSelectionOnMouseDown = true;
+					}
+					else if ( MouseEvent.IsShiftDown() )
+					{
+						OwnerWidget->Private_SelectRangeFromCurrentTo( *MyItem );
+						ChangedSelectionOnMouseDown = true;
+					}
+					else if ( !bIsSelected )
+					{
+						OwnerWidget->Private_ClearSelection();
+						OwnerWidget->Private_SetItemSelection( *MyItem, true, true );
+						ChangedSelectionOnMouseDown = true;
+					}
+
+					return FReply::Handled()
+						.DetectDrag( SharedThis(this), EKeys::LeftMouseButton )
+						.SetKeyboardFocus( OwnerWidget->AsWidget(), EKeyboardFocusCause::Mouse )
+						.CaptureMouse( SharedThis(this) );
+				}
+				break;
+			}
+		}
+
+		return FReply::Unhandled();
+	}
+	
+	/**
+	 * See SWidget::OnMouseButtonUp
+	 *
+	 * @param MyGeometry The Geometry of the widget receiving the event
+	 * @param MouseEvent Information about the input event
+	 *
+	 * @return Whether the event was handled along with possible requests for the system to take action.
+	 */
+	virtual FReply OnMouseButtonUp( const FGeometry& MyGeometry, const FPointerEvent& MouseEvent )
+	{
+		TSharedPtr< ITypedTableView<ItemType> > OwnerWidget = OwnerTablePtr.Pin();
+		check(OwnerWidget.IsValid());
+
+		TSharedRef< STableViewBase > OwnerTableViewBase = StaticCastSharedPtr< SListView<ItemType> >( OwnerWidget ).ToSharedRef();
+
+		if ( MouseEvent.GetEffectingButton() == EKeys::LeftMouseButton )
+		{
+			FReply Reply = FReply::Unhandled().ReleaseMouseCapture();
+
+			if ( ChangedSelectionOnMouseDown )
+			{
+				Reply = FReply::Handled().ReleaseMouseCapture();
+			}
+
+			const bool bIsUnderMouse = MyGeometry.IsUnderLocation(MouseEvent.GetScreenSpacePosition());
+			if ( HasMouseCapture() )
+			{
+				if ( bIsUnderMouse )
+				{
+					switch( OwnerWidget->Private_GetSelectionMode() )
+					{
+					case ESelectionMode::SingleToggle:
+						{
+							if ( !ChangedSelectionOnMouseDown )
+							{
+								const ItemType* MyItem = OwnerWidget->Private_ItemFromWidget( this );
+								const bool bIsSelected = OwnerWidget->Private_IsItemSelected( *MyItem );
+
+								OwnerWidget->Private_ClearSelection();
+								OwnerWidget->Private_SignalSelectionChanged(ESelectInfo::OnMouseClick);
+							}
+
+							Reply = FReply::Handled().ReleaseMouseCapture();
+						}
+						break;
+
+					case ESelectionMode::Multi:
+						{
+							if ( !ChangedSelectionOnMouseDown && !MouseEvent.IsControlDown() && !MouseEvent.IsShiftDown() )
+							{
+								const ItemType* MyItem = OwnerWidget->Private_ItemFromWidget( this );
+								check(MyItem != NULL);
+
+								const bool bIsSelected = OwnerWidget->Private_IsItemSelected( *MyItem );
+								if ( bIsSelected && OwnerWidget->Private_GetNumSelectedItems() > 1 )
+								{
+									// We are mousing up on a previous selected item;
+									// deselect everything but this item.
+
+									OwnerWidget->Private_ClearSelection();
+									OwnerWidget->Private_SetItemSelection( *MyItem, true, true );
+									OwnerWidget->Private_SignalSelectionChanged(ESelectInfo::OnMouseClick);
+
+									Reply = FReply::Handled().ReleaseMouseCapture();
+								}
+							}
+						}
+						break;
+					}
+				}
+
+				if ( ChangedSelectionOnMouseDown )
+				{
+					OwnerWidget->Private_SignalSelectionChanged(ESelectInfo::OnMouseClick);
+				}
+
+				return Reply;
+			}
+		}
+		else if ( MouseEvent.GetEffectingButton() == EKeys::RightMouseButton && !OwnerTableViewBase->IsRightClickScrolling() )
+		{
+			// Handle selection of items when releasing the right mouse button, but only if the user isn't actively
+			// scrolling the view by holding down the right mouse button.
+
+			switch( OwnerWidget->Private_GetSelectionMode() )
+			{
+			case ESelectionMode::Single:
+			case ESelectionMode::SingleToggle:
+			case ESelectionMode::Multi:
+				{
+					// Only one item can be selected at a time
+					const ItemType* MyItem = OwnerWidget->Private_ItemFromWidget( this );
+					const bool bIsSelected = OwnerWidget->Private_IsItemSelected( *MyItem );
+
+					// Select the item under the cursor
+					if( !bIsSelected )
+					{
+						OwnerWidget->Private_ClearSelection();
+						OwnerWidget->Private_SetItemSelection( *MyItem, true, true );
+						OwnerWidget->Private_SignalSelectionChanged(ESelectInfo::OnMouseClick);
+					}
+
+					OwnerWidget->Private_OnItemRightClicked( *MyItem, MouseEvent );
+
+					return FReply::Handled();
+				}
+				break;
+			}
+		}
+
+		return FReply::Unhandled();
+	}
+
+	virtual FReply OnDragDetected( const FGeometry& MyGeometry, const FPointerEvent& MouseEvent ) OVERRIDE
+	{
+		if ( HasMouseCapture() && ChangedSelectionOnMouseDown )
+		{
+			TSharedPtr< ITypedTableView<ItemType> > OwnerWidget = OwnerTablePtr.Pin();
+			OwnerWidget->Private_SignalSelectionChanged(ESelectInfo::OnMouseClick);
+		}
+
+		if (OnDragDetected_Handler.IsBound())
+		{
+			return OnDragDetected_Handler.Execute( MyGeometry, MouseEvent );
+		}
+		else
+		{
+			return FReply::Unhandled();
+		}
+	}
+
+	virtual void OnDragEnter(FGeometry const& MyGeometry, FDragDropEvent const& DragDropEvent) OVERRIDE
+	{
+		if (OnDragEnter_Handler.IsBound())
+		{
+			OnDragEnter_Handler.Execute(DragDropEvent);
+		}
+		else
+		{
+			SBorder::OnDragEnter(MyGeometry, DragDropEvent);
+		}
+	}
+
+	virtual void OnDragLeave(FDragDropEvent const& DragDropEvent) OVERRIDE
+	{
+		if (OnDragLeave_Handler.IsBound())
+		{
+			OnDragLeave_Handler.Execute(DragDropEvent);
+		}
+		else
+		{
+			SBorder::OnDragLeave(DragDropEvent);
+		}
+	}
+
+	virtual FReply OnDrop(FGeometry const& MyGeometry, FDragDropEvent const& DragDropEvent) OVERRIDE
+	{
+		if (OnDrop_Handler.IsBound())
+		{
+			return OnDrop_Handler.Execute(DragDropEvent);
+		}
+		else
+		{
+			return SBorder::OnDrop(MyGeometry, DragDropEvent);
+		}
+	}
+
+	virtual void SetIndexInList( int32 InIndexInList ) OVERRIDE
+	{
+		IndexInList = InIndexInList;
+	}
+
+	virtual bool IsItemExpanded() const OVERRIDE
+	{
+		TSharedPtr< ITypedTableView<ItemType> > OwnerWidget = OwnerTablePtr.Pin();
+		const ItemType* MyItem = OwnerWidget->Private_ItemFromWidget( this );
+		const bool bIsItemExpanded = OwnerWidget->Private_IsItemExpanded( *MyItem );
+		return bIsItemExpanded;
+	}
+
+	virtual void ToggleExpansion() OVERRIDE
+	{
+		TSharedPtr< ITypedTableView<ItemType> > OwnerWidget = OwnerTablePtr.Pin();
+
+		const bool bItemHasChildren = OwnerWidget->Private_DoesItemHaveChildren( IndexInList );
+		// Nothing to expand if row being clicked on doesn't have children
+		if( bItemHasChildren )
+		{
+			ItemType MyItem = *(OwnerWidget->Private_ItemFromWidget( this ));
+			const bool IsItemExpanded = bItemHasChildren && OwnerWidget->Private_IsItemExpanded( MyItem );
+			OwnerWidget->Private_SetItemExpansion( MyItem, !IsItemExpanded );
+		}
+	}
+
+	virtual int32 GetIndentLevel() const OVERRIDE
+	{
+		return OwnerTablePtr.Pin()->Private_GetNestingDepth( IndexInList );
+	}
+
+	virtual int32 DoesItemHaveChildren() const OVERRIDE
+	{
+		return OwnerTablePtr.Pin()->Private_DoesItemHaveChildren( IndexInList );
+	}
+
+	virtual TSharedRef<SWidget> AsWidget() OVERRIDE
+	{
+		return SharedThis(this);
+	}
+
+	virtual TSharedPtr<SWidget> GetContent() OVERRIDE
+	{
+		if ( this->Content.IsValid() )
+		{
+			return this->Content.Pin();
+		}
+		else
+		{
+			return TSharedPtr<SWidget>();
+		}
+	}
+
+	virtual void Private_OnExpanderArrowShiftClicked() OVERRIDE
+	{
+		TSharedPtr< ITypedTableView<ItemType> > OwnerWidget = OwnerTablePtr.Pin();
+
+		const bool bItemHasChildren = OwnerWidget->Private_DoesItemHaveChildren( IndexInList );
+		// Nothing to expand if row being clicked on doesn't have children
+		if( bItemHasChildren )
+		{
+			ItemType MyItem = *(OwnerWidget->Private_ItemFromWidget( this ));
+			const bool IsItemExpanded = bItemHasChildren && OwnerWidget->Private_IsItemExpanded( MyItem );
+			OwnerWidget->Private_OnExpanderArrowShiftClicked( MyItem, !IsItemExpanded );
+		}
+	}
+
+	/** @return The border to be drawn around this list item */
+	const FSlateBrush* GetBorder() const 
+	{
+		TSharedPtr< ITypedTableView<ItemType> > OwnerWidget = OwnerTablePtr.Pin();
+
+		const bool bIsActive = OwnerWidget->AsWidget()->HasKeyboardFocus();
+
+		// @todo: Slate Style - make this part of the widget style
+		const FSlateBrush* WhiteBox = FCoreStyle::Get().GetBrush("GenericWhiteBox");
+
+		const ItemType* MyItem = OwnerWidget->Private_ItemFromWidget( this );
+
+		check(MyItem);
+
+		const bool bIsSelected = OwnerWidget->Private_IsItemSelected( *MyItem );
+
+		if (bIsSelected && bShowSelection)
+		{
+			if( bIsActive )
+			{
+				return IsHovered()
+					? &Style->ActiveHoveredBrush
+					: &Style->ActiveBrush;
+			}
+			else
+			{
+				return IsHovered()
+					? &Style->InactiveHoveredBrush
+					: &Style->InactiveBrush;
+			}
+		}
+		else
+		{
+			// Add a slightly lighter background for even rows
+			const bool bAllowSelection = OwnerWidget->Private_GetSelectionMode() != ESelectionMode::None;
+			if( IndexInList % 2 == 0 )
+			{
+				return ( IsHovered() && bAllowSelection )
+					? &Style->EvenRowBackgroundHoveredBrush
+					: &Style->EvenRowBackgroundBrush;
+
+			}
+			else
+			{
+				return ( IsHovered() && bAllowSelection )
+					? &Style->OddRowBackgroundHoveredBrush
+					: &Style->OddRowBackgroundBrush;
+
+			}
+		}
+	}
+
+	/** 
+	 * Callback to determine if the row is selected or not
+	 *
+	 * @return		true if selected by owning widget.
+	 */
+	bool IsSelectedExclusively() const
+	{
+		TSharedPtr< ITypedTableView< ItemType > > OwnerWidget = OwnerTablePtr.Pin();
+
+		if(!OwnerWidget->AsWidget()->HasKeyboardFocus() || OwnerWidget->Private_GetNumSelectedItems() > 1)
+		{
+			return false;
+		}
+
+		const ItemType* MyItem = OwnerWidget->Private_ItemFromWidget( this );
+		return OwnerWidget->Private_IsItemSelected( *MyItem );
+	}
+
+	/** Protected constructor; SWidgets should only be instantiated via declarative syntax. */
+	STableRow()
+	: IndexInList(0)
+	, bShowSelection(true)
+	{
+	}
+
+protected:	
+	/**
+	 * An internal method to construct and setup this row widget (purposely avoids child construction). 
+	 * Split out from Construct() so that sub-classes can invoke super construction without invoking 
+	 * ConstructChildren() (sub-classes may want to constuct their own children in their own special way).
+	 * 
+	 * @param  InArgs			Declaration data for this widget.
+	 * @param  InOwnerTableView	The table that this row belongs to.
+	 */
+	void ConstructInternal(FArguments const& InArgs, TSharedRef<STableViewBase> const& InOwnerTableView)
+	{
+		check(InArgs._Style);
+		Style = InArgs._Style;
+
+		this->BorderImage = TAttribute<const FSlateBrush*>( this, &STableRow::GetBorder );
+
+		this->ForegroundColor = TAttribute<FSlateColor>( this, &STableRow::GetForegroundBasedOnSelection );
+
+		this->OnDragDetected_Handler = InArgs._OnDragDetected;
+		this->OnDragEnter_Handler = InArgs._OnDragEnter;
+		this->OnDragLeave_Handler = InArgs._OnDragLeave;
+		this->OnDrop_Handler = InArgs._OnDrop;
+		
+		this->SetOwnerTableView( InOwnerTableView );
+
+		this->bShowSelection = InArgs._ShowSelection;
+	}
+
+	void SetOwnerTableView( TSharedPtr<STableViewBase> OwnerTableView )
+	{
+		// We want to cast to a ITypedTableView.
+		// We cast to a SListView<ItemType> because C++ doesn't know that
+		// being a STableView implies being a ITypedTableView.
+		// See SListView.
+		this->OwnerTablePtr = StaticCastSharedPtr< SListView<ItemType> >(OwnerTableView);
+	}
+
+	FSlateColor GetForegroundBasedOnSelection() const
+	{
+		const TSharedPtr< ITypedTableView<ItemType> > OwnerWidget = OwnerTablePtr.Pin();
+		const FSlateColor& NonSelectedForeground = Style->TextColor; 
+		const FSlateColor& SelectedForeground = Style->SelectedTextColor;
+
+		if ( !bShowSelection || !OwnerWidget.IsValid() )
+		{
+			return NonSelectedForeground;
+		}
+
+		const ItemType* MyItem = OwnerWidget->Private_ItemFromWidget( this );
+		const bool bIsSelected = OwnerWidget->Private_IsItemSelected( *MyItem );
+
+		return bIsSelected
+			? SelectedForeground
+			: NonSelectedForeground;
+	}
+
+protected:
+
+	/** The list that owns this Selectable */
+	TWeakPtr< ITypedTableView<ItemType> > OwnerTablePtr;
+
+	/** Index of the corresponding data item in the list */
+	int32 IndexInList;
+
+	/** Whether or not to visually show that this row is selected */
+	bool bShowSelection;
+
+	/** Style used to draw this table row */
+	const FTableRowStyle* Style;
+
+	/** Delegate triggered when a user starts to drag a list item */
+	FOnDragDetected OnDragDetected_Handler;
+
+	/** Delegate triggered when a user's drag enters the bounds of this list item */
+	FOnTableRowDragEnter OnDragEnter_Handler;
+
+	/** Delegate triggered when a user's drag leaves the bounds of this list item */
+	FOnTableRowDragLeave OnDragLeave_Handler;
+
+	/** Delegate triggered when a user's drag is dropped in the bounds of this list item */
+	FOnTableRowDrop OnDrop_Handler;
+
+	/** The widget in the content slot for this row */
+	TWeakPtr<SWidget> Content;
+
+	bool ChangedSelectionOnMouseDown;
+};
+
+
+
+template<typename ItemType>
+class SMultiColumnTableRow : public STableRow<ItemType>
+{
+public:
+
+	/**
+	 * Users of SMultiColumnTableRow would usually some piece of data associated with it.
+	 * The type of this data is ItemType; it's the stuff that your TableView (i.e. List or Tree) is visualizing.
+	 * The ColumnName tells you which column of the TableView we need to make a widget for.
+	 * Make a widget and return it.
+	 *
+	 * @param ColumnName    A unique ID for a column in this TableView; see SHeaderRow::FColumn for more info.
+	 *
+	 * @return a widget to represent the contents of a cell in this row of a TableView. 
+	 */
+	virtual TSharedRef<SWidget> GenerateWidgetForColumn( const FName& ColumnName ) = 0;
+
+	SLATE_BEGIN_ARGS( SMultiColumnTableRow< ItemType > )
+		: _Style( &FCoreStyle::Get().GetWidgetStyle<FTableRowStyle>("TableView.Row") )
+		, _Padding( FMargin(0) )
+		, _ShowSelection(true)
+		{}
+	
+		SLATE_STYLE_ARGUMENT( FTableRowStyle, Style )
+
+		SLATE_EVENT( FOnDragDetected, OnDragDetected )
+		SLATE_EVENT( FOnTableRowDragEnter, OnDragEnter )
+		SLATE_EVENT( FOnTableRowDragLeave, OnDragLeave )
+		SLATE_EVENT( FOnTableRowDrop,      OnDrop )
+
+		SLATE_ATTRIBUTE( FMargin, Padding )
+	
+		SLATE_ARGUMENT( bool, ShowSelection )
+		
+	SLATE_END_ARGS()
+
+	typedef SMultiColumnTableRow< ItemType > FSuperRowType;
+	typedef typename STableRow<ItemType>::FArguments FTableRowArgs;
+
+protected:
+	void Construct( const FArguments& InArgs, const TSharedRef<STableViewBase>& OwnerTableView )
+	{
+		STableRow<ItemType>::Construct(
+			FTableRowArgs()
+			.Style(InArgs._Style)
+			.Padding(InArgs._Padding)
+			.ShowSelection(InArgs._ShowSelection)
+			.OnDragDetected(InArgs._OnDragDetected)
+			.OnDragEnter(InArgs._OnDragEnter)
+			.OnDragLeave(InArgs._OnDragLeave)
+			.OnDrop(InArgs._OnDrop)
+			.Content()
+			[
+				SAssignNew( Box, SHorizontalBox )
+			]
+		
+			, OwnerTableView );
+
+		// Sign up for notifications about changes to the HeaderRow
+		TSharedPtr< SHeaderRow > HeaderRow = OwnerTableView->GetHeaderRow();
+		check( HeaderRow.IsValid() );
+		HeaderRow->OnColumnsChanged()->AddSP( this, &SMultiColumnTableRow<ItemType>::GenerateColumns );
+
+		// Populate the row with user-generated content
+		this->GenerateColumns( HeaderRow.ToSharedRef() );
+	}
+
+	virtual void ConstructChildren( ETableViewMode::Type InOwnerTableMode, const TAttribute<FMargin>& InPadding, const TSharedRef<SWidget>& InContent ) OVERRIDE
+	{
+		STableRow<ItemType>::Content = InContent;
+
+		// MultiColumnRows let the user decide which column should contain the expander/indenter item.
+		this->ChildSlot
+		.Padding( InPadding )
+		[
+			InContent
+		];
+	}
+
+	void GenerateColumns( const TSharedRef<SHeaderRow>& InColumnHeaders )
+	{
+		Box->ClearChildren();
+		const TIndirectArray<SHeaderRow::FColumn>& Columns = InColumnHeaders->GetColumns();
+		const int32 NumColumns = Columns.Num();
+		TMap< FName, TSharedRef< SWidget > > NewColumnIdToSlotContents;
+
+		for( int32 ColumnIndex = 0; ColumnIndex < NumColumns; ++ColumnIndex )
+		{
+			const SHeaderRow::FColumn& Column = Columns[ColumnIndex];
+			const TSharedRef< SWidget >* const ExistingWidget = ColumnIdToSlotContents.Find( Column.ColumnId );
+
+			TSharedRef< SWidget > CellContents = SNullWidget::NullWidget;
+			if ( ExistingWidget != NULL )
+			{
+				CellContents = *ExistingWidget;
+			}
+			else
+			{
+				CellContents = GenerateWidgetForColumn( Column.ColumnId );
+			}
+
+			if ( Column.SizeRule == EColumnSizeMode::Fixed )
+			{
+				Box->AddSlot()
+				.AutoWidth()
+				[
+					SNew( SBox )
+					.WidthOverride( Column.Width.Get() )
+					.HAlign(Column.CellHAlignment)
+					.VAlign(Column.CellVAlignment)
+					.Content()
+					[
+						CellContents
+					]
+				];
+			}
+			else
+			{
+				TAttribute<float> WidthBinding;
+				WidthBinding.BindRaw( &Column, &SHeaderRow::FColumn::GetWidth );
+
+				SHorizontalBox::FSlot& NewSlot = Box->AddSlot()
+				.HAlign(Column.CellHAlignment)
+				.VAlign(Column.CellVAlignment)
+				.FillWidth( WidthBinding )
+				[
+					CellContents
+				];
+			}
+
+			NewColumnIdToSlotContents.Add( Column.ColumnId, CellContents );
+		}
+
+		ColumnIdToSlotContents = NewColumnIdToSlotContents;
+	}
+
+	void ClearCellCache()
+	{
+		ColumnIdToSlotContents.Empty();
+	}
+
+
+private:
+	
+	TSharedPtr<SHorizontalBox> Box;
+	TMap< FName, TSharedRef< SWidget > > ColumnIdToSlotContents;
+};
\ No newline at end of file
diff --git a/Engine/Source/Runtime/Slate/Public/Widgets/Views/STableViewBase.h b/Engine/Source/Runtime/Slate/Public/Widgets/Views/STableViewBase.h
new file mode 100644
--- /dev/null
+++ b/Engine/Source/Runtime/Slate/Public/Widgets/Views/STableViewBase.h
@@ -0,0 +1,235 @@
+// Copyright 1998-2014 Epic Games, Inc. All Rights Reserved.
+
+#pragma once
+
+#include "IScrollableWidget.h"
+
+class SListPanel;
+class SHeaderRow;
+
+/**
+ * Contains ListView functionality that does not depend on the type of data being observed by the ListView.
+ */
+class SLATE_API STableViewBase : public SCompoundWidget, public IScrollableWidget
+{
+public:
+
+	/** Create the child widgets that comprise the list */
+	void ConstructChildren( const TAttribute<float>& InItemWidth, const TAttribute<float>& InItemHeight, const TSharedPtr<SHeaderRow>& InColumnHeaders, const TSharedPtr<SScrollBar>& InScrollBar  );
+
+	/**
+	 * Invoked by the scrollbar when the user scrolls.
+	 *
+	 * @param InScrollOffsetFraction  The location to which the user scrolled as a fraction (between 0 and 1) of total height of the content.
+	 */
+	void ScrollBar_OnUserScrolled( float InScrollOffsetFraction );
+
+	// SWidget interface
+	virtual void OnKeyboardFocusLost( const FKeyboardFocusEvent& InKeyboardFocusEvent ) OVERRIDE;
+	virtual void OnMouseCaptureLost() OVERRIDE;
+	virtual bool SupportsKeyboardFocus() const OVERRIDE;
+	virtual void Tick( const FGeometry& AllottedGeometry, const double InCurrentTime, const float InDeltaTime ) OVERRIDE;
+	virtual FReply OnMouseButtonDown( const FGeometry& MyGeometry, const FPointerEvent& MouseEvent ) OVERRIDE;
+	virtual FReply OnMouseButtonUp( const FGeometry& MyGeometry, const FPointerEvent& MouseEvent ) OVERRIDE;
+	virtual void OnMouseLeave( const FPointerEvent& MouseEvent ) OVERRIDE;
+	virtual FReply OnMouseMove( const FGeometry& MyGeometry, const FPointerEvent& MouseEvent ) OVERRIDE;
+	virtual FReply OnMouseWheel( const FGeometry& MyGeometry, const FPointerEvent& MouseEvent ) OVERRIDE;
+	virtual FReply OnKeyDown( const FGeometry& MyGeometry, const FKeyboardEvent& InKeyboardEvent ) OVERRIDE;
+	virtual FCursorReply OnCursorQuery( const FGeometry& MyGeometry, const FPointerEvent& CursorEvent ) const OVERRIDE;
+	// End of SWidget interface
+
+	/** @return The number of Widgets we currently have generated. */
+	int32 GetNumGeneratedChildren() const;
+
+	TSharedPtr<SHeaderRow> GetHeaderRow() const;
+
+	/** @return Returns true if the user is currently interactively scrolling the view by holding
+		        the right mouse button and dragging. */
+	bool IsRightClickScrolling() const;
+
+	/** @return Returns true if the user is currently interactively scrolling the view by holding
+		        either mouse button and dragging. */
+	bool IsUserScrolling() const;
+
+	/**
+	 * Mark the list as dirty, so that it will regenerate its widgets on next tick.
+	 */
+	void RequestListRefresh();
+
+	/** Return true if there is currently a refresh pending, false otherwise */
+	bool IsPendingRefresh() const;
+
+	// IScrollableWidget interface
+	virtual FVector2D GetScrollDistance() OVERRIDE;
+	virtual FVector2D GetScrollDistanceRemaining() OVERRIDE;
+	virtual TSharedRef<class SWidget> GetScrollWidget() OVERRIDE;
+	// End of IScrollableWidget interface
+
+	virtual int32 OnPaint(
+		const FGeometry& AllottedGeometry,
+		const FSlateRect& MyClippingRect,
+		FSlateWindowElementList& OutDrawElements,
+		int32 LayerId,
+		const FWidgetStyle& InWidgetStyle,
+		bool bParentEnabled) const OVERRIDE;
+
+	/** Is this list backing a tree or just a standalone list */
+	const ETableViewMode::Type TableViewMode;
+
+protected:
+
+	STableViewBase( ETableViewMode::Type InTableViewMode );
+	
+	/**
+	 * Scroll the list view by some number of screen units.
+	 *
+	 * @param MyGeometry      The geometry of the ListView at the time
+	 * @param ScrollByAmount  The amount to scroll by in Slate Screen Units.
+	 *
+	 * @return The amount actually scrolled in items
+	 */
+	virtual float ScrollBy( const FGeometry& MyGeometry, float ScrollByAmount );
+
+	/**
+	 * Scroll the view to an offset
+	 *
+	 * @param InScrollOffset  Offset into the total list length to scroll down.
+	 *
+	 * @return The amount actually scrolled
+	 */
+	virtual float ScrollTo( float InScrollOffset );
+
+	/** Insert WidgetToInsert at the top of the view. */
+	void InsertWidget( const TSharedRef<ITableRow> & WidgetToInset );
+
+	/** Add a WidgetToAppend to the bottom of the view. */
+	void AppendWidget( const TSharedRef<ITableRow>& WidgetToAppend );
+	
+	/**
+	 * Remove all the widgets from the view.
+	 */
+	void ClearWidgets();
+
+	/**
+	 * Get the uniform item width.
+	 */
+	float GetItemWidth() const;
+
+	/**
+	 * Get the uniform item height that is enforced by ListViews.
+	 */
+	float GetItemHeight() const;
+
+	/** @return the number of items that can fit on the screen */
+	virtual float GetNumLiveWidgets() const;
+
+	/**
+	 * Get the number of items that can fit in the view horizontally before creating a new row.
+	 * Default is 1, but may be more in subclasses (like STileView)
+	 */
+	virtual int32 GetNumItemsWide() const;
+
+	/**
+	 * Opens a context menu as the result of a right click if OnContextMenuOpening is bound and we are not right click scrolling.
+	 */
+	virtual void OnRightMouseButtonUp(const FVector2D& SummonLocation);
+	
+	/**
+	 * Get the scroll rate in items that best approximates a constant physical scroll rate.
+	 */
+	float GetScrollRateInItems() const;
+
+	/**
+	 * Remove any items that are no longer in the list from the selection set.
+	 */
+	virtual void UpdateSelectionSet() = 0;
+
+
+	/** Information about the outcome of the WidgetRegeneratePass */
+	struct SLATE_API FReGenerateResults
+	{
+		FReGenerateResults( double InNewScrollOffset, double InHeightGenerated, double InItemsOnScreen, bool AtEndOfList )
+		: NewScrollOffset( InNewScrollOffset )
+		, HeightOfGeneratedItems( InHeightGenerated )
+		, ExactNumWidgetsOnScreen( InItemsOnScreen )
+		, bGeneratedPastLastItem( AtEndOfList )
+		{
+
+		}
+
+		/** The scroll offset that we actually use might not be what the user asked for */
+		double NewScrollOffset;
+
+		/** The total height of the widgets that we have generated to represent the visible subset of the items*/
+		double HeightOfGeneratedItems;
+
+		/** How many items are fitting on the screen, including fractions */
+		double ExactNumWidgetsOnScreen;
+
+		/** True when we have generated  */
+		bool bGeneratedPastLastItem;
+	};
+	/**
+	 * Update generate Widgets for Items as needed and clean up any Widgets that are no longer needed.
+	 * Re-arrange the visible widget order as necessary.
+	 */
+	virtual FReGenerateResults ReGenerateItems( const FGeometry& MyGeometry ) = 0;
+
+	/** @return how many items there are in the TArray being observed */
+	virtual int32 GetNumItemsBeingObserved() const = 0;
+
+	/**
+	 * If there is a pending request to scroll an item into view, do so.
+	 * 
+	 * @param ListViewGeometry  The geometry of the listView; can be useful for centering the item.
+	 */
+	virtual void ScrollIntoView(const FGeometry& ListViewGeometry) = 0;
+
+	/**
+	 * Called when an item has entered the visible geometry to check to see if the ItemScrolledIntoView delegate should be fired.
+	 */
+	virtual void NotifyItemScrolledIntoView() = 0;
+
+	/** The panel which holds the visible widgets in this list */
+	TSharedPtr< SListPanel > ItemsPanel;
+	/** The scroll bar widget */
+	TSharedPtr< SScrollBar > ScrollBar;
+
+	/** Scroll offset from the beginning of the list in items */
+	double ScrollOffset;
+
+	/** How much we scrolled while the rmb has been held */
+	float AmountScrolledWhileRightMouseDown;
+
+	/** Information about the widgets we generated during the last regenerate pass */
+	FReGenerateResults LastGenerateResults;
+
+	/** Last time we scrolled, did we end up at the end of the list. */
+	bool bWasAtEndOfList;
+
+	/** What the list's geometry was the last time a refresh occurred. */
+	FGeometry PanelGeometryLastTick;
+
+	/** Delegate to invoke when the context menu should be opening. If it is NULL, a context menu will not be summoned. */
+	FOnContextMenuOpening OnContextMenuOpening;
+
+	/** The selection mode that this tree/list is in. Note that it is up to the generated ITableRows to respect this setting. */
+	TAttribute<ESelectionMode::Type> SelectionMode;
+
+	/** Column headers that describe which columns this list shows */
+	TSharedPtr< SHeaderRow > HeaderRow;
+
+	/** Helper object to manage inertial scrolling */
+	FInertialScrollManager InertialScrollManager;
+
+	/**	The current position of the software cursor */
+	FVector2D SoftwareCursorPosition;
+
+	/**	Whether the software cursor should be drawn in the viewport */
+	bool bShowSoftwareCursor;
+
+private:
+	
+	/** When true, a refresh should occur the next tick */
+	bool bItemsNeedRefresh;
+};
diff --git a/Engine/Source/Runtime/Slate/Public/Widgets/Views/STileView.h b/Engine/Source/Runtime/Slate/Public/Widgets/Views/STileView.h
new file mode 100644
--- /dev/null
+++ b/Engine/Source/Runtime/Slate/Public/Widgets/Views/STileView.h
@@ -0,0 +1,257 @@
+// Copyright 1998-2014 Epic Games, Inc. All Rights Reserved.
+ 
+#pragma once
+
+/**
+ * A TileView widget is a list which arranges its items horizontally until there is no more space then creates a new row.
+ * Items are spaced evenly horizontally.
+ */
+template <typename ItemType>
+class STileView : public SListView<ItemType>
+{
+public:
+	typedef typename TListTypeTraits< ItemType >::NullableType NullableItemType;
+
+	typedef typename TSlateDelegates< ItemType >::FOnGenerateRow FOnGenerateRow;
+	typedef typename TSlateDelegates< ItemType >::FOnItemScrolledIntoView FOnItemScrolledIntoView;
+	typedef typename TSlateDelegates< ItemType >::FOnMouseButtonDoubleClick FOnMouseButtonDoubleClick;
+	typedef typename TSlateDelegates< NullableItemType >::FOnSelectionChanged FOnSelectionChanged;
+
+public:
+
+	SLATE_BEGIN_ARGS( STileView<ItemType> )
+		: _OnGenerateTile()
+		, _ListItemsSource( static_cast<TArray<ItemType>*>(NULL) ) //@todo Slate Syntax: Initializing from NULL without a cast
+		, _ItemHeight(128)
+		, _ItemWidth(128)
+		, _OnContextMenuOpening()
+		, _OnMouseButtonDoubleClick()
+		, _OnSelectionChanged()
+		, _SelectionMode(ESelectionMode::Multi)
+		, _ClearSelectionOnClick(true)
+		, _ExternalScrollbar()
+		{}
+
+		SLATE_EVENT( FOnGenerateRow, OnGenerateTile )
+
+		SLATE_EVENT( FOnItemScrolledIntoView, OnItemScrolledIntoView )
+
+		SLATE_ARGUMENT( const TArray<ItemType>* , ListItemsSource )
+
+		SLATE_ATTRIBUTE( float, ItemHeight )
+
+		SLATE_ATTRIBUTE( float, ItemWidth )
+
+		SLATE_EVENT( FOnContextMenuOpening, OnContextMenuOpening )
+
+		SLATE_EVENT( FOnMouseButtonDoubleClick, OnMouseButtonDoubleClick )
+
+		SLATE_EVENT( FOnSelectionChanged, OnSelectionChanged )
+
+		SLATE_ATTRIBUTE( ESelectionMode::Type, SelectionMode )
+
+		SLATE_ARGUMENT ( bool, ClearSelectionOnClick )
+
+		SLATE_ARGUMENT( TSharedPtr<SScrollBar>, ExternalScrollbar )
+
+	SLATE_END_ARGS()
+
+	/**
+	 * Construct this widget
+	 *
+	 * @param	InArgs	The declaration data for this widget
+	 */
+	void Construct( const typename STileView<ItemType>::FArguments& InArgs )
+	{
+		this->OnGenerateRow = InArgs._OnGenerateTile;
+		this->OnItemScrolledIntoView = InArgs._OnItemScrolledIntoView;
+		
+		this->ItemsSource = InArgs._ListItemsSource;
+		this->OnContextMenuOpening = InArgs._OnContextMenuOpening;
+		this->OnDoubleClick = InArgs._OnMouseButtonDoubleClick;
+		this->OnSelectionChanged = InArgs._OnSelectionChanged;
+		this->SelectionMode = InArgs._SelectionMode;
+
+		this->bClearSelectionOnClick = InArgs._ClearSelectionOnClick;
+
+		// Check for any parameters that the coder forgot to specify.
+		FString ErrorString;
+		{
+			if ( !this->OnGenerateRow.IsBound() )
+			{
+				ErrorString += TEXT("Please specify an OnGenerateTile. \n");
+			}
+
+			if ( this->ItemsSource == NULL )
+			{
+				ErrorString += TEXT("Please specify a ListItemsSource. \n");
+			}
+		}
+
+		if (ErrorString.Len() > 0)
+		{
+			// Let the coder know what they forgot
+			this->ChildSlot
+			.HAlign(HAlign_Center)
+			.VAlign(VAlign_Center)
+			[
+				SNew(STextBlock)
+				.Text(ErrorString)
+			];
+		}
+		else
+		{
+			// Make the TableView
+			this->ConstructChildren(InArgs._ItemWidth, InArgs._ItemHeight, TSharedPtr<SHeaderRow>(), InArgs._ExternalScrollbar);
+		}
+	}
+
+	STileView( ETableViewMode::Type InListMode = ETableViewMode::Tile )
+	: SListView<ItemType>( InListMode )
+	{
+	}
+
+public:
+	// Inherited from SWidget
+	virtual FReply OnKeyDown( const FGeometry& MyGeometry, const FKeyboardEvent& InKeyboardEvent ) OVERRIDE
+	{
+		const TArray<ItemType>& ItemsSourceRef = (*this->ItemsSource);
+
+		// Don't respond to key-presses containing "Alt" as a modifier
+		if ( ItemsSourceRef.Num() > 0 && !InKeyboardEvent.IsAltDown() )
+		{
+			// Check for selection manipulation keys that differ from SListView (Left, Right)
+			if ( InKeyboardEvent.GetKey() == EKeys::Left )
+			{
+				int32 SelectionIndex = ( !TListTypeTraits<ItemType>::IsPtrValid(this->SelectorItem) ) ? 0 : ItemsSourceRef.Find( TListTypeTraits<ItemType>::NullableItemTypeConvertToItemType( this->SelectorItem ) );
+
+				if ( SelectionIndex > 0 )
+				{
+					// Select the previous item
+					this->KeyboardSelect(ItemsSourceRef[SelectionIndex - 1], InKeyboardEvent);
+				}
+
+				return FReply::Handled();
+			}
+			else if ( InKeyboardEvent.GetKey() == EKeys::Right )
+			{
+				int32 SelectionIndex = ( !TListTypeTraits<ItemType>::IsPtrValid(this->SelectorItem) ) ? 0 : ItemsSourceRef.Find( TListTypeTraits<ItemType>::NullableItemTypeConvertToItemType( this->SelectorItem ) );
+
+				if ( SelectionIndex < ItemsSourceRef.Num() - 1 )
+				{
+					// Select the next item
+					this->KeyboardSelect(ItemsSourceRef[SelectionIndex + 1], InKeyboardEvent);
+				}
+
+				return FReply::Handled();
+			}
+		}
+
+		return SListView<ItemType>::OnKeyDown(MyGeometry, InKeyboardEvent);
+	}
+
+public:	
+
+	virtual STableViewBase::FReGenerateResults ReGenerateItems( const FGeometry& MyGeometry ) OVERRIDE
+	{
+		// Clear all the items from our panel. We will re-add them in the correct order momentarily.
+		this->ClearWidgets();
+		
+		const TArray<ItemType>* SourceItems = this->ItemsSource;
+		if ( SourceItems != NULL )
+		{
+			// Item width and height is constant by design.
+			const float AllottedWidth = MyGeometry.Size.X;
+			const float ItemWidth = this->GetItemWidth();
+			const float ItemHeight = this->GetItemHeight();
+
+			const int32 NumItemsWide = GetNumItemsWide();
+			const int32 NumItemsPaddedToFillLastRow = (SourceItems->Num() % NumItemsWide != 0)
+				? SourceItems->Num() + NumItemsWide - SourceItems->Num() % NumItemsWide
+				: SourceItems->Num();
+			const double RowsPerScreen = MyGeometry.Size.Y / ItemHeight;
+			const double EndOfListOffset = NumItemsPaddedToFillLastRow - NumItemsWide * RowsPerScreen;
+			const double ClampedScrollOffset = FMath::Clamp(STableViewBase::ScrollOffset, 0.0, EndOfListOffset);
+			const bool bAtEndOfList = (STableViewBase::ScrollOffset >= ClampedScrollOffset);
+			
+			// Once we run out of vertical and horizontal space, we stop generating widgets.
+			float WidthUsedSoFar = 0.0f;
+			float HeightUsedSoFar = 0.0f;
+			// Index of the item at which we start generating based on how far scrolled down we are
+			int32 StartIndex = FMath::Max( 0, FMath::Floor(ClampedScrollOffset / NumItemsWide) * NumItemsWide );
+
+			// Let the WidgetGenerator that we are starting a pass so that it can keep track of data items and widgets.
+			this->WidgetGenerator.OnBeginGenerationPass();
+
+			// Actually generate the widgets.
+			bool bKeepGenerating = true;
+			for( int32 ItemIndex = StartIndex; bKeepGenerating && ItemIndex < SourceItems->Num(); ++ItemIndex )
+			{
+				const ItemType& CurItem = (*SourceItems)[ItemIndex];
+
+				const float GeneratedItemHeight = SListView<ItemType>::GenerateWidgetForItem( CurItem, ItemIndex, StartIndex );
+
+				// The widget used up some of the available horizontal space.
+				WidthUsedSoFar += ItemWidth;
+
+				// Check to see if we have completed a row of widgets
+				if ( WidthUsedSoFar + ItemWidth > AllottedWidth )
+				{
+					// A new row of widgets was added.
+					WidthUsedSoFar = 0;
+					HeightUsedSoFar += ItemHeight;	
+					// Stop when we've generated a widget that's partially clipped by the bottom of the list.
+					if ( HeightUsedSoFar > MyGeometry.Size.Y + ItemHeight )
+					{
+						bKeepGenerating = false;
+					}
+				}
+			}
+
+			// We have completed the generation pass. The WidgetGenerator will clean up unused Widgets.
+			this->WidgetGenerator.OnEndGenerationPass();
+
+			const double ExactNumItemsOnScreen = MyGeometry.Size.Y / ItemHeight * NumItemsWide;
+
+			return STableViewBase::FReGenerateResults(ClampedScrollOffset, ItemHeight * ExactNumItemsOnScreen, ExactNumItemsOnScreen, bAtEndOfList);
+		}
+
+		return STableViewBase::FReGenerateResults(0, 0, 0, false);
+
+	}
+
+	virtual int32 GetNumItemsBeingObserved() const OVERRIDE
+	{
+		const int32 NumItemsBeingObserved = this->ItemsSource == NULL ? 0 : this->ItemsSource->Num();
+		const int32 NumItemsWide = GetNumItemsWide();
+		
+		int32 NumEmptySpacesAtEnd = 0;
+		if ( NumItemsWide > 0 )
+		{
+			NumEmptySpacesAtEnd = NumItemsWide - (NumItemsBeingObserved % NumItemsWide);
+			if ( NumEmptySpacesAtEnd >= NumItemsWide )
+			{
+				NumEmptySpacesAtEnd = 0;
+			}
+		}
+
+		return NumItemsBeingObserved + NumEmptySpacesAtEnd;
+	}
+
+protected:
+
+	virtual float ScrollBy( const FGeometry& MyGeometry, float ScrollByAmountInSlateUnits ) OVERRIDE
+	{
+		const float ItemHeight = this->GetItemHeight();
+		const double NewScrollOffset = this->ScrollOffset + ( ( ScrollByAmountInSlateUnits * GetNumItemsWide() ) / ItemHeight );
+
+		return this->ScrollTo( NewScrollOffset );
+	}
+
+	virtual int32 GetNumItemsWide() const OVERRIDE
+	{
+		const float ItemWidth = this->GetItemWidth();
+		const int32 NumItemsWide = ItemWidth > 0 ? FMath::Floor(this->PanelGeometryLastTick.Size.X / ItemWidth) : 1;
+		return FMath::Max(1, NumItemsWide);
+	}
+};
