Slate: Reorganized private Slate files; proper pre-compiled header for all files.

Notes:
- LandscapeEditor pulls in private MultiBox include files; temporarily made those public

diff --git a/Engine/Source/Editor/LandscapeEditor/Private/LandscapeEditorDetailWidgets.cpp b/Engine/Source/Editor/LandscapeEditor/Private/LandscapeEditorDetailWidgets.cpp
--- a/Engine/Source/Editor/LandscapeEditor/Private/LandscapeEditorDetailWidgets.cpp
+++ b/Engine/Source/Editor/LandscapeEditor/Private/LandscapeEditorDetailWidgets.cpp
@@ -2,9 +2,7 @@
 
 #include "LandscapeEditorPrivatePCH.h"
 #include "LandscapeEditorDetailWidgets.h"
-#include "../Private/MultiBox/MultiBox.h"
-#include "../Private/MultiBox/SToolBarButtonBlock.h"
-#include "../Private/MultiBox/SToolBarComboButtonBlock.h"
+
 
 // Based on a stripped-down FToolBarComboButtonBlock
 class FToolSelector : public FMultiBlock
diff --git a/Engine/Source/Editor/MainFrame/Private/Menus/MainMenu.cpp b/Engine/Source/Editor/MainFrame/Private/Menus/MainMenu.cpp
--- a/Engine/Source/Editor/MainFrame/Private/Menus/MainMenu.cpp
+++ b/Engine/Source/Editor/MainFrame/Private/Menus/MainMenu.cpp
@@ -10,7 +10,6 @@
 #include "GlobalEditorCommonCommands.h"
 #include "Editor/UnrealEd/Public/Features/EditorFeatures.h"
 #include "Runtime/Core/Public/Features/IModularFeatures.h"
-#include "MultiBox/MultiBoxDefs.h"
 #include "SourceCodeNavigation.h"
 #include "Toolkits/AssetEditorToolkit.h"
 
diff --git a/Engine/Source/Runtime/Slate/Private/Framework/Application/MenuStack.cpp b/Engine/Source/Runtime/Slate/Private/Framework/Application/MenuStack.cpp
--- a/Engine/Source/Runtime/Slate/Private/Framework/Application/MenuStack.cpp
+++ b/Engine/Source/Runtime/Slate/Private/Framework/Application/MenuStack.cpp
@@ -4,7 +4,7 @@
 	MenuStack.cpp: Implements the FMenuStack class.
 =============================================================================*/
 
-#include "Slate.h"
+#include "SlatePrivatePCH.h"
 
 
 namespace FMenuStackDefs
diff --git a/Engine/Source/Runtime/Slate/Private/Framework/Application/SlateApplication.cpp b/Engine/Source/Runtime/Slate/Private/Framework/Application/SlateApplication.cpp
--- a/Engine/Source/Runtime/Slate/Private/Framework/Application/SlateApplication.cpp
+++ b/Engine/Source/Runtime/Slate/Private/Framework/Application/SlateApplication.cpp
@@ -1,8 +1,9 @@
 // Copyright 1998-2014 Epic Games, Inc. All Rights Reserved.
 
-#include "Slate.h"
+#include "SlatePrivatePCH.h"
 #include "SWindowTitleBar.h"
 
+
 // Statics
 //// @todo checkin: Should this go in an almost empty Events.cpp?
 FTouchKeySet FPointerEvent::StandardTouchKeySet(EKeys::LeftMouseButton);
diff --git a/Engine/Source/Runtime/Slate/Private/Framework/Application/SlateIcon.cpp b/Engine/Source/Runtime/Slate/Private/Framework/Application/SlateIcon.cpp
--- a/Engine/Source/Runtime/Slate/Private/Framework/Application/SlateIcon.cpp
+++ b/Engine/Source/Runtime/Slate/Private/Framework/Application/SlateIcon.cpp
@@ -4,7 +4,7 @@
 	SlateIcon.cpp: Implements the FSlateIcon structure.
 =============================================================================*/
 
-#include "Slate.h"
+#include "SlatePrivatePCH.h"
 
 
 /* FSlateIcon structors
diff --git a/Engine/Source/Runtime/Slate/Private/Framework/Commands/GenericCommands.cpp b/Engine/Source/Runtime/Slate/Private/Framework/Commands/GenericCommands.cpp
new file mode 100644
--- /dev/null
+++ b/Engine/Source/Runtime/Slate/Private/Framework/Commands/GenericCommands.cpp
@@ -0,0 +1,18 @@
+// Copyright 1998-2014 Epic Games, Inc. All Rights Reserved.
+
+#include "SlatePrivatePCH.h"
+
+void FGenericCommands::RegisterCommands()
+{
+	UI_COMMAND( Cut, "Cut", "Cut selection", EUserInterfaceActionType::Button, FInputGesture( EModifierKey::Control, EKeys::X ) )
+	UI_COMMAND( Copy, "Copy", "Copy selection", EUserInterfaceActionType::Button, FInputGesture( EModifierKey::Control, EKeys::C ) )
+	UI_COMMAND( Paste, "Paste", "Paste clipboard contents", EUserInterfaceActionType::Button, FInputGesture( EModifierKey::Control, EKeys::V ) )
+	UI_COMMAND( Duplicate, "Duplicate", "Duplicate selection", EUserInterfaceActionType::Button, FInputGesture( EModifierKey::Control, EKeys::W ) )
+	UI_COMMAND( Undo, "Undo", "Undo last action", EUserInterfaceActionType::Button, FInputGesture( EModifierKey::Control, EKeys::Z ) )
+	UI_COMMAND( Redo, "Redo", "Redo last action", EUserInterfaceActionType::Button, FInputGesture( EModifierKey::Control, EKeys::Y ) )
+
+	UI_COMMAND(Delete, "Delete", "Delete current selection", EUserInterfaceActionType::Button, FInputGesture(EKeys::Platform_Delete))
+	
+	UI_COMMAND( Rename, "Rename", "Rename current selection", EUserInterfaceActionType::Button, FInputGesture( EKeys::F2 ) )
+	UI_COMMAND( SelectAll, "Select All", "Select everything in the current scope", EUserInterfaceActionType::Button, FInputGesture( EModifierKey::Control, EKeys::A ) )
+}
diff --git a/Engine/Source/Runtime/Slate/Private/Framework/Commands/InputBindingManager.cpp b/Engine/Source/Runtime/Slate/Private/Framework/Commands/InputBindingManager.cpp
--- a/Engine/Source/Runtime/Slate/Private/Framework/Commands/InputBindingManager.cpp
+++ b/Engine/Source/Runtime/Slate/Private/Framework/Commands/InputBindingManager.cpp
@@ -4,8 +4,7 @@
 	InputBindingManager.cpp: Implements the FInputBindingManager class.
 =============================================================================*/
 
-#include "Slate.h"
-#include "RemoteConfigIni.h"
+#include "SlatePrivatePCH.h"
 
 
 /* FUserDefinedGestures helper class
diff --git a/Engine/Source/Runtime/Slate/Private/Framework/Commands/UICommandDragDropOp.cpp b/Engine/Source/Runtime/Slate/Private/Framework/Commands/UICommandDragDropOp.cpp
new file mode 100644
--- /dev/null
+++ b/Engine/Source/Runtime/Slate/Private/Framework/Commands/UICommandDragDropOp.cpp
@@ -0,0 +1,54 @@
+// Copyright 1998-2014 Epic Games, Inc. All Rights Reserved.
+
+/*=============================================================================
+	UICommandDragDropOp.cpp: Implements the FUICommandDragDropOp class.
+=============================================================================*/
+
+#include "SlatePrivatePCH.h"
+
+
+TSharedRef<FUICommandDragDropOp> FUICommandDragDropOp::New( TSharedRef<const FUICommandInfo> InCommandInfo, FName InOriginMultiBox, TSharedPtr<SWidget> CustomDectorator, FVector2D DecoratorOffset )
+{
+	TSharedRef<FUICommandDragDropOp> Operation = MakeShareable(new FUICommandDragDropOp(InCommandInfo, InOriginMultiBox, CustomDectorator, DecoratorOffset));
+	Operation->Construct();
+
+	return Operation;
+}
+
+
+void FUICommandDragDropOp::OnDragged( const class FDragDropEvent& DragDropEvent )
+{
+	CursorDecoratorWindow->SetOpacity(0.85f);
+	CursorDecoratorWindow->MoveWindowTo(DragDropEvent.GetScreenSpacePosition() + Offset);
+}
+
+
+void FUICommandDragDropOp::OnDrop( bool bDropWasHandled, const FPointerEvent& MouseEvent )
+{
+	FDragDropOperation::OnDrop(bDropWasHandled, MouseEvent);
+	OnDropNotification.ExecuteIfBound();
+}
+
+
+TSharedPtr<SWidget> FUICommandDragDropOp::GetDefaultDecorator() const
+{
+	TSharedPtr<SWidget> Content;
+
+	if (CustomDecorator.IsValid())
+	{
+		Content = CustomDecorator;
+	}
+	else
+	{
+		Content = SNew(STextBlock)
+			.Text(UICommand->GetLabel());
+	}
+
+	// Create hover widget
+	return SNew(SBorder)
+		.BorderImage(FCoreStyle::Get().GetBrush("ToolPanel.GroupBorder"))
+		.Content()
+		[	
+			Content.ToSharedRef()
+		];
+}
diff --git a/Engine/Source/Runtime/Slate/Private/Framework/Docking/DockingPrivate.h b/Engine/Source/Runtime/Slate/Private/Framework/Docking/DockingPrivate.h
new file mode 100644
--- /dev/null
+++ b/Engine/Source/Runtime/Slate/Private/Framework/Docking/DockingPrivate.h
@@ -0,0 +1,17 @@
+// Copyright 1998-2014 Epic Games, Inc. All Rights Reserved.
+
+#pragma once
+
+
+#include "Slate.h"
+
+#define DEBUG_TAB_MANAGEMENT 0
+
+#include "SDockingNode.h"
+#include "SDockingSplitter.h"
+#include "SDockingTabStack.h"
+#include "SDockingArea.h"
+#include "SDockingTabWell.h"
+#include "SDockingTarget.h"
+#include "SDockingCross.h"
+#include "FDockingDragOperation.h"
diff --git a/Engine/Source/Runtime/Slate/Private/Framework/Docking/LayoutService.cpp b/Engine/Source/Runtime/Slate/Private/Framework/Docking/LayoutService.cpp
new file mode 100644
--- /dev/null
+++ b/Engine/Source/Runtime/Slate/Private/Framework/Docking/LayoutService.cpp
@@ -0,0 +1,56 @@
+// Copyright 1998-2014 Epic Games, Inc. All Rights Reserved.
+
+#include "SlatePrivatePCH.h"
+
+
+const FString& FLayoutSaveRestore::GetAdditionalLayoutConfigIni()
+{
+	static const FString IniSectionAdditionalConfig = TEXT("SlateAdditionalLayoutConfig");
+	return IniSectionAdditionalConfig;
+}
+
+
+void FLayoutSaveRestore::SaveTheLayout( const TSharedRef<FTabManager::FLayout>& LayoutToSave )
+{
+	SaveToConfig(LayoutToSave, GEditorUserSettingsIni);
+}
+
+void FLayoutSaveRestore::SaveToConfig( const TSharedRef<FTabManager::FLayout>& LayoutToSave, FString ConfigFileName )
+{
+	const FString LayoutAsString = FLayoutSaveRestore::PrepareLayoutStringForIni( LayoutToSave->ToString() );
+
+	GConfig->SetString( TEXT("EditorLayouts"), *LayoutToSave->GetLayoutName().ToString(), *LayoutAsString, ConfigFileName );
+}
+
+TSharedRef<FTabManager::FLayout> FLayoutSaveRestore::LoadUserConfigVersionOf( const TSharedRef<FTabManager::FLayout>& DefaultLayout )
+{
+	return LoadFromConfig(DefaultLayout,GEditorUserSettingsIni);
+}
+
+TSharedRef<FTabManager::FLayout> FLayoutSaveRestore::LoadFromConfig( const TSharedRef<FTabManager::FLayout>& DefaultLayout, FString ConfigFileName )
+{
+	const FName LayoutName = DefaultLayout->GetLayoutName();
+	FString UserLayoutString;
+	if ( GConfig->GetString( TEXT("EditorLayouts"), *LayoutName.ToString(), UserLayoutString, ConfigFileName ) && !UserLayoutString.IsEmpty() )
+	{
+		TSharedPtr<FTabManager::FLayout> UserLayout = FTabManager::FLayout::NewFromString( FLayoutSaveRestore::GetLayoutStringFromIni( UserLayoutString ));
+		if ( UserLayout.IsValid() && UserLayout->GetPrimaryArea().IsValid() )
+		{
+			return UserLayout.ToSharedRef();
+		}
+	}
+
+	return DefaultLayout;
+}
+
+FString FLayoutSaveRestore::PrepareLayoutStringForIni(const FString& LayoutString)
+{
+	// Have to store braces as parentheses due to braces causing ini issues
+	return LayoutString.Replace(TEXT("{"), TEXT("(")).Replace(TEXT("}"), TEXT(")")).Replace(LINE_TERMINATOR, TEXT("\\")LINE_TERMINATOR);
+}
+
+FString FLayoutSaveRestore::GetLayoutStringFromIni(const FString& LayoutString)
+{
+	// Revert parenthesis to braces, from ini readable to Json readable
+	return LayoutString.Replace(TEXT("("), TEXT("{")).Replace(TEXT(")"), TEXT("}")).Replace(TEXT("\\")LINE_TERMINATOR, LINE_TERMINATOR);
+}
diff --git a/Engine/Source/Runtime/Slate/Private/Framework/Docking/SDockTabStack.cpp b/Engine/Source/Runtime/Slate/Private/Framework/Docking/SDockTabStack.cpp
new file mode 100644
--- /dev/null
+++ b/Engine/Source/Runtime/Slate/Private/Framework/Docking/SDockTabStack.cpp
@@ -0,0 +1,14 @@
+// Copyright 1998-2014 Epic Games, Inc. All Rights Reserved.
+
+#include "SlatePrivatePCH.h"
+
+
+void SDockTabStack::Construct( const FArguments& InArgs )
+{
+
+}
+
+void SDockTabStack::AddTab( const TSharedRef<SDockableTab>& InTab, int32 AtLocation )
+{
+
+}
diff --git a/Engine/Source/Runtime/Slate/Private/Framework/Docking/SDockingArea.h b/Engine/Source/Runtime/Slate/Private/Framework/Docking/SDockingArea.h
new file mode 100644
--- /dev/null
+++ b/Engine/Source/Runtime/Slate/Private/Framework/Docking/SDockingArea.h
@@ -0,0 +1,112 @@
+// Copyright 1998-2014 Epic Games, Inc. All Rights Reserved.
+
+#pragma once
+
+
+
+/**
+ * Represents the root node in a hierarchy of DockNodes.
+ */
+class SLATE_API SDockingArea : public SDockingSplitter
+{
+public:
+
+	SLATE_BEGIN_ARGS(SDockingArea)
+		: _ParentWindow()
+		, _ShouldManageParentWindow(true)
+		, _InitialContent()
+		{
+			// Visible by default, but don't absorb clicks
+			_Visibility = EVisibility::SelfHitTestInvisible;
+		}
+		
+		/* The window whose content area this dock area is directly embedded within.  By default, ShouldManageParentWindow is
+		   set to true, which means the dock area will also destroy the window when the last tab goes away.  Assigning a
+		   parent window also allows the docking area to embed title area widgets (minimize, maximize, etc) into its content area */
+		SLATE_ARGUMENT( TSharedPtr<SWindow>, ParentWindow )
+
+		/** True if this docking area should close the parent window when the last tab in this docking area goes away */
+		SLATE_ARGUMENT( bool, ShouldManageParentWindow )
+
+		 /**
+		 * What to put into the DockArea initially. Usually a TabStack, so that some tabs can be added to it.
+		 */
+		SLATE_ARGUMENT( TSharedPtr<SDockingNode>, InitialContent )
+
+	SLATE_END_ARGS()
+
+	void Construct( const FArguments& InArgs, const TSharedRef<FTabManager>& InTabManager, const TSharedRef<FTabManager::FArea>& PersistentNode );
+
+	virtual Type GetNodeType() const
+	{
+		return SDockingNode::DockArea;
+	}
+	
+	/** Returns this dock area */
+	virtual TSharedPtr<SDockingArea> GetDockArea() OVERRIDE;
+	virtual TSharedPtr<const SDockingArea> GetDockArea() const OVERRIDE;
+
+	/** Returns the window that this dock area resides in directly and also manages */
+	TSharedPtr<SWindow> GetParentWindow() const;
+
+	virtual void OnDragEnter( const FGeometry& MyGeometry, const FDragDropEvent& DragDropEvent ) OVERRIDE;
+	virtual void OnDragLeave( const FDragDropEvent& DragDropEvent ) OVERRIDE;
+	
+	virtual FReply OnDrop( const FGeometry& MyGeometry, const FDragDropEvent& DragDropEvent ) OVERRIDE;
+
+	virtual FReply OnUserAttemptingDock( SDockingNode::RelativeDirection Direction, const FDragDropEvent& DragDropEvent ) OVERRIDE;
+
+	void OnTabFoundNewHome( const TSharedRef<SDockTab>& RelocatedTab, const TSharedRef<SWindow>& NewOwnerWindow );
+
+	// Show the dock-from-outside dock targets
+	void ShowCross();
+
+	// Hide the dock-from-outside dock targets
+	void HideCross();
+
+	/**
+	 * Removes redundant stack and splitters. Collapses any widgets that any no longer showing live content.
+	 */
+	void CleanUp( ELayoutModification RemovalMethod );
+
+	void SetParentWindow( const TSharedRef<SWindow>& NewParentWindow );
+
+	virtual TSharedPtr<FTabManager::FLayoutNode> GatherPersistentLayout() const OVERRIDE;
+
+	TSharedRef<FTabManager> GetTabManager() const;
+
+protected:
+	virtual SDockingNode::ECleanupRetVal CleanUpNodes() OVERRIDE;
+
+private:
+	EVisibility TargetCrossVisibility() const;
+	EVisibility TargetCrossCenterVisibility() const;
+	/** Dock a tab along the outer edge of this DockArea */
+	void DockFromOutside(SDockingNode::RelativeDirection Direction, const FDragDropEvent& DragDropEvent);
+
+	/** We were placed in a window, and it is being destroyed */
+	void OnOwningWindowBeingDestroyed(const TSharedRef<SWindow>& WindowBeingDestroyed);
+
+	virtual void OnLiveTabAdded() OVERRIDE;
+
+private:	
+	
+	/** The window this dock area is embedded within.  If bIsManagingParentWindow is true, the dock area will also
+	    destroy the window when the last tab goes away. */
+	TWeakPtr<SWindow> ParentWindowPtr;
+
+	/** True if this docking area should close the parent window when the last tab in this docking area goes away */
+	bool bManageParentWindow;
+
+	/** The tab manager that controls this DockArea */
+	TWeakPtr<FTabManager> MyTabManager;
+
+	/** The overlay is visible when the user is dragging a tab over the dock area */
+	bool bIsOverlayVisible;
+
+	/** The center target is visible when the overlay is visible and there are no live tabs */
+	bool bIsCenterTargetVisible;
+
+	/** True when the last tab has been pulled from this area, meaning that this DockArea will not be necessary once that tab finds a new home. */
+	bool bCleanUpUponTabRelocation;
+};
diff --git a/Engine/Source/Runtime/Slate/Private/Framework/Docking/SDockingCross.h b/Engine/Source/Runtime/Slate/Private/Framework/Docking/SDockingCross.h
new file mode 100644
--- /dev/null
+++ b/Engine/Source/Runtime/Slate/Private/Framework/Docking/SDockingCross.h
@@ -0,0 +1,32 @@
+// Copyright 1998-2014 Epic Games, Inc. All Rights Reserved.
+
+#pragma once
+
+
+/**
+ * Targets used by docking code. When re-arranging layout, hovering over targets
+ * gives the user a preview of what will happen if they drop on that target.
+ * Dropping actually commits the layout-restructuring.
+ */
+class SLATE_API SDockingCross : public SLeafWidget
+{
+public:
+	SLATE_BEGIN_ARGS(SDockingCross){}
+	SLATE_END_ARGS()
+
+	void Construct( const FArguments& InArgs, const TSharedPtr<class SDockingNode>& InOwnerNode );
+
+	// SWidget interface
+	virtual int32 OnPaint( const FGeometry& AllottedGeometry, const FSlateRect& MyClippingRect, FSlateWindowElementList& OutDrawElements, int32 LayerId, const FWidgetStyle& InWidgetStyle, bool bParentEnabled ) const OVERRIDE;
+	virtual FVector2D ComputeDesiredSize() const OVERRIDE;
+	virtual void OnDragLeave( const FDragDropEvent& DragDropEvent ) OVERRIDE;
+	virtual FReply OnDragOver( const FGeometry& MyGeometry, const FDragDropEvent& DragDropEvent ) OVERRIDE;
+	virtual FReply OnDrop( const FGeometry& MyGeometry, const FDragDropEvent& DragDropEvent ) OVERRIDE;
+	// End of SWidget interface
+	
+private:
+
+	/** The DockNode relative to which we want to dock */
+	TWeakPtr<class SDockingNode> OwnerNode;
+};
+
diff --git a/Engine/Source/Runtime/Slate/Private/Framework/Docking/SDockingSplitter.h b/Engine/Source/Runtime/Slate/Private/Framework/Docking/SDockingSplitter.h
new file mode 100644
--- /dev/null
+++ b/Engine/Source/Runtime/Slate/Private/Framework/Docking/SDockingSplitter.h
@@ -0,0 +1,106 @@
+// Copyright 1998-2014 Epic Games, Inc. All Rights Reserved.
+
+#pragma once
+
+
+/** Dynamic N-way splitter that provides the resizing functionality in the docking framework. */
+class SLATE_API SDockingSplitter : public SDockingNode
+{
+public:
+	SLATE_BEGIN_ARGS(SDockingSplitter){}
+	SLATE_END_ARGS()
+
+	void Construct( const FArguments& InArgs, const TSharedRef<FTabManager::FSplitter>& PersistentNode );
+
+	virtual Type GetNodeType() const
+	{
+		return SDockingNode::DockSplitter;
+	}
+
+	/**
+	 * Add a new child dock node at the desired location.
+	 * Assumes this DockNode is a splitter.
+	 *
+	 * @param InChild      The DockNode child to add.
+	 * @param InLocation   Index at which to add; INDEX_NONE (default) adds to the end of the list
+	 */
+	void AddChildNode( const TSharedRef<SDockingNode>& InChild, int32 InLocation = INDEX_NONE);
+
+	/**
+	 * Replace the child InChildToReplace with Replacement
+	 *
+	 * @param InChildToReplace     The child of this node to replace
+	 * @param Replacement          What to replace with
+	 */
+	void ReplaceChild( const TSharedRef<SDockingNode>& InChildToReplace, const TSharedRef<SDockingNode>& Replacement );
+
+	/**
+	 * Remove a child node from the splitter
+	 * 
+	 * @param ChildToRemove    The child node to remove from this splitter
+	 */
+	void RemoveChild( const TSharedRef<SDockingNode>& ChildToRemove );
+	
+	/**
+	 * Remove a child at this index
+	 *
+	 * @param IndexOfChildToRemove    Remove a child at this index
+	 */
+	void RemoveChildAt( int32 IndexOfChildToRemove );
+
+	/**
+	 * Inserts NodeToPlace next to RelativeToMe.
+	 * Next to means above, below, left or right of RelativeToMe.
+	 *
+	 * @param NodeToPlace        The node to insert.
+	 * @param Direction          Where to insert relative to the other node
+	 * @param RelativeToMove     Insert relative to this node.
+	 */
+	void PlaceNode( const TSharedRef<SDockingNode>& NodeToPlace, SDockingNode::RelativeDirection Direction, const TSharedRef<SDockingNode>& RelativeToMe );
+
+	/** Sets the orientation of this splitter */
+	void SetOrientation(EOrientation NewOrientation);
+
+	/** Gets an array of all child dock nodes */
+	const TArray< TSharedRef<SDockingNode> >& GetChildNodes() const;
+
+	/** Gets an array of all child dock nodes and all of their child dock nodes and so on */
+	TArray< TSharedRef<SDockingNode> > GetChildNodesRecursively() const;
+	
+	/** Recursively searches through all children looking for child tabs */
+	virtual TArray< TSharedRef<SDockTab> > GetAllChildTabs() const OVERRIDE;
+
+	/** Gets the size coefficient of a given child dock node */
+	float GetSizeCoefficientForSlot(int32 Index) const;
+
+	/** Gets the orientation of this splitter */
+	EOrientation GetOrientation() const;
+
+	virtual TSharedPtr<FTabManager::FLayoutNode> GatherPersistentLayout() const OVERRIDE;
+
+protected:
+
+	/** Is the docking direction (left, right, above, below) match the orientation (horizontal vs. vertical) */
+	static bool DoesDirectionMatchOrientation( SDockingNode::RelativeDirection InDirection, EOrientation InOrientation );
+
+	static SDockingNode::ECleanupRetVal MostResponsibility( SDockingNode::ECleanupRetVal A, SDockingNode::ECleanupRetVal B );
+
+	virtual SDockingNode::ECleanupRetVal CleanUpNodes() OVERRIDE;
+
+	float ComputeChildCoefficientTotal() const;
+
+	/**
+	 * @param ParentNode    DockSplitter node whose children might be redundant
+	 *
+	 * @return true if we found some redundant splitters and need to do another pass on this list; false if we found no redundancies.
+	 */
+	static bool ClearRedundantNodes( const TSharedRef<SDockingSplitter>& ParentNode );
+
+	/** The SSplitter widget that SDockingSplitter wraps. */
+	TSharedPtr<SSplitter> Splitter;
+	
+	/** The DockNode children. All these children are kept in sync with the SSplitter's children via the use of the public interface for adding, removing and replacing children. */
+	TArray< TSharedRef<class SDockingNode> > Children;
+};
+
+
diff --git a/Engine/Source/Runtime/Slate/Private/Framework/Docking/SDockingTabStack.h b/Engine/Source/Runtime/Slate/Private/Framework/Docking/SDockingTabStack.h
new file mode 100644
--- /dev/null
+++ b/Engine/Source/Runtime/Slate/Private/Framework/Docking/SDockingTabStack.h
@@ -0,0 +1,235 @@
+// Copyright 1998-2014 Epic Games, Inc. All Rights Reserved.
+
+#pragma once
+
+
+
+/**
+ * A node in the Docking/Tabbing hierarchy.
+ * A DockTabStack shows a row of tabs and the content of one selected tab.
+ * It also supports re-arranging tabs and dragging them out for the stack.
+ */
+class SLATE_API SDockingTabStack : public SDockingNode
+{
+public:
+	SLATE_BEGIN_ARGS(SDockingTabStack)
+			: _IsDocumentArea(false) {}
+		SLATE_ARGUMENT( bool, IsDocumentArea )
+	SLATE_END_ARGS()
+
+	virtual Type GetNodeType() const
+	{
+		return SDockingNode::DockTabStack;
+	}
+
+	void OnLastTabRemoved();
+
+	void OnTabClosed( const TSharedRef<SDockTab>& ClosedTab );
+	void OnTabRemoved( const FTabId& TabId );
+
+	void Construct( const FArguments& InArgs, const TSharedRef<FTabManager::FStack>& PersistentNode );
+
+	// TabStack methods
+
+	void OpenTab( const TSharedRef<SDockTab>& InTab, int32 InsertAtLocation = INDEX_NONE );
+
+	void AddTabWidget( const TSharedRef<SDockTab>& InTab, int32 AtLocation = INDEX_NONE);
+
+	/** @return All child tabs in this node */
+	const TArray< TSharedRef<SDockTab> >& GetTabs() const;
+	
+	/** @return How many tabs are in this node */
+	int32 GetNumTabs() const;
+
+	bool HasTab( const struct FTabMatcher& TabMatcher ) const;
+
+	/** @return the last known geometry of this TabStack */
+	FGeometry GetTabStackGeometry() const;
+
+	void RemoveClosedTabsWithName( FName InName );
+
+	bool IsShowingLiveTabs() const;
+
+	void BringToFront( const TSharedRef<SDockTab>& TabToBringToFront );
+
+	/** Set the content that the DockNode is presenting. */
+	void SetNodeContent( const TSharedRef<SWidget>& InContent, const TSharedRef<SWidget>& InContentLeft, const TSharedRef<SWidget>& InContentRight );
+
+	virtual FReply OnUserAttemptingDock( SDockingNode::RelativeDirection Direction, const FDragDropEvent& DragDropEvent ) OVERRIDE;
+	
+	/** Recursively searches through all children looking for child tabs */
+	virtual TArray< TSharedRef<SDockTab> > GetAllChildTabs() const OVERRIDE;
+
+	virtual SSplitter::ESizeRule GetSizeRule() const OVERRIDE;
+
+	void SetTabWellHidden( bool bShouldHideTabWell );
+	bool IsTabWellHidden() const;
+
+	virtual TSharedPtr<FTabManager::FLayoutNode> GatherPersistentLayout() const OVERRIDE;
+
+public:
+
+	virtual void Tick( const FGeometry& AllottedGeometry, const double InCurrentTime, const float InDeltaTime ) OVERRIDE;
+
+	virtual FReply OnDragOver( const FGeometry& MyGeometry, const FDragDropEvent& DragDropEvent ) OVERRIDE;
+
+	virtual void OnDragLeave( const FDragDropEvent& DragDropEvent ) OVERRIDE;
+	
+	virtual FReply OnDrop( const FGeometry& MyGeometry, const FDragDropEvent& DragDropEvent ) OVERRIDE;
+
+	virtual void OnKeyboardFocusChanging( const FWeakWidgetPath& PreviousFocusPath, const FWidgetPath& NewWidgetPath ) OVERRIDE;
+
+	virtual FReply OnMouseButtonDown( const FGeometry& MyGeometry, const FPointerEvent& MouseEvent ) OVERRIDE;
+
+	virtual bool SupportsKeyboardFocus() const OVERRIDE { return true; }
+
+protected:
+
+	virtual EWindowZone::Type GetWindowZoneOverride() const OVERRIDE
+	{
+		// Pretend we are a title bar so the user can grab the area to move the window around
+		return EWindowZone::TitleBar;
+	}
+
+	void CloseForegroundTab();
+
+	enum ETabsToClose
+	{
+		CloseDocumentTabs,
+		CloseDocumentsAndTools
+	};
+	/**
+	 * Close all the background tabs.
+	 *
+	 * @param  TabsToClose   Close just document tabs or both document and tool tabs.
+	 */
+	void CloseAllButForegroundTab(ETabsToClose TabsToClose);
+
+	FReply TabWellRightClicked( const FGeometry& TabWellGeometry, const FPointerEvent& MouseEvent );
+
+	virtual SDockingNode::ECleanupRetVal CleanUpNodes() OVERRIDE;
+
+	int32 OpenPersistentTab( const FTabId& TabId, int32 OpenLocationAmongActiveTabs = INDEX_NONE );
+
+	int32 ClosePersistentTab( const FTabId& TabId );
+
+	void RemovePersistentTab( const FTabId& TabId );
+
+	/** Overridden from SDockingNode */
+	virtual void SetParentNode( TSharedRef<class SDockingSplitter> InParent ) OVERRIDE;
+
+private:
+	
+	/** Data that persists across sessions and when the widget associated with this node is removed. */	
+	TArray<FTabManager::FTab> Tabs;
+
+
+	/**
+	 * Creates a SDockingTabStack by adding a new split to this stack's parent splitter and attaching the new SDockingTabStack.
+	 * 
+	 * @param	Direction	The relative direction to split the parent splitter
+	 *
+	 * @return	The newly created empty SDockingTabStack, ready for a tab to be added to it
+	 */
+	TSharedRef< SDockingTabStack > CreateNewTabStackBySplitting( const SDockingNode::RelativeDirection Direction );
+
+	/** What should the content area look like for the current tab? */
+	const FSlateBrush* GetContentAreaBrush() const;
+
+	/** How much padding to show around the content currently being presented */
+	FMargin GetContentPadding() const;
+
+	/** Depending on the tabs we put into the tab well, we want a different background brush. */
+	const FSlateBrush* GetTabWellBrush() const;
+
+	/** Show the tab well? */
+	EVisibility GetTabWellVisibility() const;
+
+	/** Show the stuff needed to unhide the tab well? */
+	EVisibility GetUnhideButtonVisibility() const;
+	
+	/** Show/Hide the tab well; do it smoothly with an animation */
+	void ToggleTabWellVisibility();
+
+	FReply UnhideTabWell();
+
+	/** Only allow hiding the tab well when there is a single tab in it. */
+	bool CanHideTabWell() const;
+
+	/** Keep around our geometry from the last frame so that we can resize the preview windows correctly */
+	FGeometry TabStackGeometry;
+
+	/** The tab well widget shows all tabs, keeps track of the selected tab, allows tab rearranging, etc. */
+	TSharedPtr<class SDockingTabWell> TabWell;
+
+	/** The borders that hold any potential inline content areas. */
+	SHorizontalBox::FSlot* InlineContentAreaLeft;
+	SHorizontalBox::FSlot* InlineContentAreaRight;
+	SVerticalBox::FSlot* TitleBarSlot;
+	TSharedPtr<SWidget> TitleBarContent;
+
+	TSharedPtr<SBorder> ContentSlot;
+
+	FOverlayManagement OverlayManagement;
+
+	TSharedRef<SWidget> MakeContextMenu();
+
+	/** Show the docking cross */
+	void ShowCross();
+
+	/** Hide the docking cross */
+	void HideCross();
+
+	/** Document Areas do not disappear when out of tabs, and instead say 'Document Area' */
+	bool bIsDocumentArea;
+
+	/** Animation that shows/hides the tab well; also used as a state machine to determine whether tab well is shown/hidden */
+	FCurveSequence ShowHideTabWell;
+
+	/** Grabs the scaling factor for the tab well size from the tab well animation. */
+	FVector2D GetTabWellScale() const;
+
+	/** Get the scale for the button that unhides the tab well */
+	FVector2D GetUnhideTabWellButtonScale() const;
+	/** Get the opacity for the button that unhides the tab well */
+	FSlateColor GetUnhideTabWellButtonOpacity() const;
+
+	/** @return Gets the visibility state for spacers that pad out the tab well to make room for title bar widgets */
+	EVisibility GetTitleAreaSpacerVisibility();
+
+	/** The tab in this dock stack that is active */
+	TSharedPtr<SDockTab> ActiveTab;
+
+	/** Visibility of TitleBar spacer based on maximize/restore status of the window.
+	 ** This gives us a little more space to grab the title bar when the window is not maximized
+	*/
+	EVisibility GetMaximizeSpacerVisibility() const;
+
+
+#if DEBUG_TAB_MANAGEMENT
+	FString ShowPersistentTabs() const;
+#endif
+
+};
+
+
+struct FTabMatcher
+{
+	FTabMatcher( const FTabId& InTabId, ETabState::Type InTabState = static_cast<ETabState::Type>(ETabState::ClosedTab | ETabState::OpenedTab) )
+		: TabIdToMatch( InTabId )
+		, RequiredTabState( InTabState )
+	{
+	}
+
+	bool Matches( const FTabManager::FTab& Candidate ) const
+	{
+		return
+			( (Candidate.TabState & RequiredTabState) != 0 ) &&
+			( Candidate.TabId.TabType == TabIdToMatch.TabType ) &&
+			// NAME_None is treated as a wildcard
+			( TabIdToMatch.InstanceId == INDEX_NONE || TabIdToMatch.InstanceId == Candidate.TabId.InstanceId );;
+	}
+
+	FTabId TabIdToMatch;
+	ETabState::Type RequiredTabState;
+};
diff --git a/Engine/Source/Runtime/Slate/Private/Framework/Docking/SDockingTabWell.cpp b/Engine/Source/Runtime/Slate/Private/Framework/Docking/SDockingTabWell.cpp
new file mode 100644
--- /dev/null
+++ b/Engine/Source/Runtime/Slate/Private/Framework/Docking/SDockingTabWell.cpp
@@ -0,0 +1,521 @@
+// Copyright 1998-2014 Epic Games, Inc. All Rights Reserved.
+
+#include "SlatePrivatePCH.h"
+#include "DockingPrivate.h"
+
+
+const FVector2D FDockingConstants::MaxMinorTabSize(150.f, 50.0f);
+const FVector2D FDockingConstants::MaxMajorTabSize(210.f, 50.f);
+
+const FVector2D FDockingConstants::GetMaxTabSizeFor( ETabRole::Type TabRole )
+{
+	return (TabRole == ETabRole::MajorTab)
+		? MaxMajorTabSize
+		: MaxMinorTabSize;
+}
+
+
+void SDockingTabWell::Construct( const FArguments& InArgs )
+{
+	ForegroundTabIndex = INDEX_NONE;
+	TabBeingDraggedPtr = TWeakPtr<SDockTab>(NULL);
+	ChildBeingDraggedOffset = 0.0f;
+	TabGrabOffsetFraction = FVector2D::ZeroVector;
+		
+	// We need a valid parent here. TabPanels must exist in a SDockingNode
+	check( InArgs._ParentStackNode.Get().IsValid() );
+	ParentTabStackPtr = InArgs._ParentStackNode.Get();
+}
+
+const TArray< TSharedRef<SDockTab> >& SDockingTabWell::GetTabs() const
+{
+	return Tabs;
+}
+
+int32 SDockingTabWell::GetNumTabs() const
+{
+	return Tabs.Num();
+}
+
+void SDockingTabWell::AddTab( const TSharedRef<SDockTab>& InTab, int32 AtIndex )
+{
+	// Add the tab and implicitly activate it.
+	if (AtIndex == INDEX_NONE)
+	{
+		this->Tabs.Add( InTab );
+		BringTabToFront( Tabs.Num()-1 );
+	}
+	else
+	{
+		AtIndex = FMath::Clamp( AtIndex, 0, Tabs.Num() );
+		this->Tabs.Insert( InTab, AtIndex );
+		BringTabToFront( AtIndex );
+	}
+
+	InTab->SetParent(SharedThis(this));
+
+	const TSharedPtr<SDockingTabStack> ParentTabStack = ParentTabStackPtr.Pin();
+	if (ParentTabStack.IsValid() && ParentTabStack->GetDockArea().IsValid())
+	{
+		ParentTabStack->GetDockArea()->GetTabManager()->GetPrivateApi().OnTabOpening( InTab );
+	}
+}
+
+
+void SDockingTabWell::ArrangeChildren( const FGeometry& AllottedGeometry, FArrangedChildren& ArrangedChildren ) const
+{
+	// The specialized TabWell is dedicated to arranging tabs.
+	// Tabs have uniform sizing (all tabs the same size).
+	// TabWell also ignores widget visibilit, as it is not really
+	// relevant.
+
+
+	// The tab that is being dragged by the user, if any.
+	TSharedPtr<SDockTab> TabBeingDragged = TabBeingDraggedPtr.Pin();
+		
+	const int32 NumChildren = Tabs.Num();
+
+	// Tabs have a uniform size.
+	const FVector2D ChildSize = ComputeChildSize(AllottedGeometry);
+
+
+	const float DraggedChildCenter = ChildBeingDraggedOffset + ChildSize.X / 2;
+
+	// Arrange all the tabs left to right.
+	float XOffset = 0;
+	for( int32 TabIndex=0; TabIndex < NumChildren; ++TabIndex )
+	{
+		const TSharedRef<SDockTab> CurTab = Tabs[TabIndex];
+		const float ChildWidthWithOverlap = ChildSize.X - CurTab->GetOverlapWidth();
+
+		// Is this spot reserved from the tab that is being dragged?
+		if ( TabBeingDragged.IsValid() && XOffset <= DraggedChildCenter && DraggedChildCenter < (XOffset + ChildWidthWithOverlap) )
+		{
+			// if so, leave some room to signify that this is where the dragged tab would end up
+			XOffset += ChildWidthWithOverlap;
+		}
+
+		ArrangedChildren.AddWidget( AllottedGeometry.MakeChild(CurTab, FVector2D(XOffset, 0), ChildSize) );
+
+		XOffset += ChildWidthWithOverlap;
+	}
+		
+	// Arrange the tab currently being dragged by the user, if any
+	if ( TabBeingDragged.IsValid() )
+	{
+		ArrangedChildren.AddWidget( AllottedGeometry.MakeChild( TabBeingDragged.ToSharedRef(), FVector2D(ChildBeingDraggedOffset,0), ChildSize) );
+	}
+}
+	
+
+int32 SDockingTabWell::OnPaint( const FGeometry& AllottedGeometry, const FSlateRect& MyClippingRect, FSlateWindowElementList& OutDrawElements, int32 LayerId, const FWidgetStyle& InWidgetStyle, bool bParentEnabled ) const
+{
+	// When we are dragging a tab, it must be painted on top of the other tabs, so we cannot
+	// just reuse the Panel's default OnPaint.
+
+
+	// The TabWell has no visualization of its own; it just visualizes its child tabs.
+	FArrangedChildren ArrangedChildren(EVisibility::Visible);
+	this->ArrangeChildren(AllottedGeometry, ArrangedChildren);
+
+	// Because we paint multiple children, we must track the maximum layer id that they produced in case one of our parents
+	// wants to an overlay for all of its contents.
+	int32 MaxLayerId = LayerId;
+
+	TSharedPtr<SDockTab> ForegroundTab = GetForegroundTab();
+	FArrangedWidget* ForegroundTabGeometry = NULL;
+	
+	// Draw all inactive tabs first
+	for (int32 ChildIndex = 0; ChildIndex < ArrangedChildren.Num(); ++ChildIndex)
+	{
+		FArrangedWidget& CurWidget = ArrangedChildren(ChildIndex);
+		if (CurWidget.Widget == ForegroundTab)
+		{
+			ForegroundTabGeometry = &CurWidget;
+		}
+		else
+		{
+			FSlateRect ChildClipRect = MyClippingRect.IntersectionWith( CurWidget.Geometry.GetClippingRect() );
+			const int32 CurWidgetsMaxLayerId = CurWidget.Widget->OnPaint( CurWidget.Geometry, ChildClipRect, OutDrawElements, MaxLayerId, InWidgetStyle, ShouldBeEnabled( bParentEnabled ) );
+			MaxLayerId = FMath::Max( MaxLayerId, CurWidgetsMaxLayerId );
+		}
+	}
+
+	// Draw active tab in front
+	if (ForegroundTab != TSharedPtr<SDockTab>())
+	{
+		FSlateRect ChildClipRect = MyClippingRect.IntersectionWith( ForegroundTabGeometry->Geometry.GetClippingRect() );
+		const int32 CurWidgetsMaxLayerId = ForegroundTabGeometry->Widget->OnPaint( ForegroundTabGeometry->Geometry, ChildClipRect, OutDrawElements, MaxLayerId, InWidgetStyle, ShouldBeEnabled( bParentEnabled ) );
+		MaxLayerId = FMath::Max( MaxLayerId, CurWidgetsMaxLayerId );
+	}
+
+	return MaxLayerId;
+}
+
+
+FVector2D SDockingTabWell::ComputeDesiredSize() const
+{
+	FVector2D DesiredSizeResult(0,0);
+
+	for ( int32 TabIndex=0; TabIndex < Tabs.Num(); ++TabIndex )
+	{
+		// Currently not respecting Visibility because tabs cannot be invisible.
+		const TSharedRef<SDockTab>& SomeTab = Tabs[TabIndex];
+		const FVector2D SomeTabDesiredSize = SomeTab->GetDesiredSize();
+		DesiredSizeResult.X += SomeTabDesiredSize.X;
+		DesiredSizeResult.Y = FMath::Max(SomeTabDesiredSize.Y, DesiredSizeResult.Y);
+	}
+
+	TSharedPtr<SDockTab> TabBeingDragged = TabBeingDraggedPtr.Pin();
+	if ( TabBeingDragged.IsValid() )
+	{
+		const FVector2D SomeTabDesiredSize = TabBeingDragged->GetDesiredSize();
+		DesiredSizeResult.X += SomeTabDesiredSize.X;
+		DesiredSizeResult.Y = FMath::Max(SomeTabDesiredSize.Y, DesiredSizeResult.Y);
+	}
+
+	return DesiredSizeResult;
+}
+
+
+FChildren* SDockingTabWell::GetChildren()
+{
+	return &Tabs;
+}
+
+
+FVector2D SDockingTabWell::ComputeChildSize( const FGeometry& AllottedGeometry ) const
+{
+	const int32 NumChildren = Tabs.Num();
+
+	/** Assume all tabs overlap the same amount. */
+	const float OverlapWidth = (NumChildren > 0)
+		? Tabs[0]->GetOverlapWidth()
+		: 0.0f;
+
+	// All children shall be the same size: evenly divide the alloted area.
+	// If we are dragging a tab, don't forget to take it into account when dividing.
+	const FVector2D ChildSize = TabBeingDraggedPtr.IsValid()
+		? FVector2D( (AllottedGeometry.Size.X - OverlapWidth) / ( NumChildren + 1 ) + OverlapWidth, AllottedGeometry.Size.Y )
+		: FVector2D( (AllottedGeometry.Size.X - OverlapWidth) / NumChildren + OverlapWidth, AllottedGeometry.Size.Y );
+
+	// Major vs. Minor tabs have different tab sizes.
+	// We will make our choice based on the first tab we encounter.
+	TSharedPtr<SDockTab> FirstTab = (NumChildren > 0)
+		? Tabs[0]
+		: TabBeingDraggedPtr.Pin();
+
+	// If there are no tabs in this tabwell, assume minor tabs.
+	FVector2D MaxTabSize(0,0);
+	if ( FirstTab.IsValid() )
+	{
+		const ETabRole::Type RoleToUse = FirstTab->IsNomadTabWithMajorTabStyle() ? ETabRole::MajorTab : FirstTab->GetTabRole();
+		MaxTabSize = FDockingConstants::GetMaxTabSizeFor(RoleToUse);
+	}
+	else
+	{
+		MaxTabSize = FDockingConstants::MaxMinorTabSize;
+	}
+
+	// Don't let the tabs get too big, or they'll look ugly.
+	return FVector2D (
+		FMath::Min( ChildSize.X, MaxTabSize.X ),
+		FMath::Min( ChildSize.Y, MaxTabSize.Y )
+	);
+}
+
+
+float SDockingTabWell::ComputeDraggedTabOffset( const FGeometry& MyGeometry, const FPointerEvent& MouseEvent, const FVector2D& InTabGrabOffsetFraction ) const
+{
+	const FVector2D ComputedChildSize = ComputeChildSize(MyGeometry);
+	return MyGeometry.AbsoluteToLocal( MouseEvent.GetScreenSpacePosition() ).X - InTabGrabOffsetFraction .X * ComputedChildSize.X;
+}
+
+
+
+
+FReply SDockingTabWell::StartDraggingTab( TSharedRef<SDockTab> TabToStartDragging, FVector2D InTabGrabOffsetFraction, const FPointerEvent& MouseEvent )
+{
+	Tabs.Remove(TabToStartDragging);
+	// We just removed the foreground tab.
+	ForegroundTabIndex = INDEX_NONE;
+	ParentTabStackPtr.Pin()->OnTabRemoved(TabToStartDragging->GetLayoutIdentifier());
+
+	// Tha tab well keeps track of which tab we are dragging; we treat is specially during rendering and layout.
+	TabBeingDraggedPtr = TabToStartDragging;
+		
+	TabGrabOffsetFraction = InTabGrabOffsetFraction;
+
+	// We are about to start dragging a tab, so make sure its offset is correct
+	this->ChildBeingDraggedOffset = ComputeDraggedTabOffset( MouseEvent.FindGeometry(SharedThis(this)), MouseEvent, InTabGrabOffsetFraction );
+
+	// Start dragging.
+	TSharedRef<FDockingDragOperation> DragDropOperation =
+		FDockingDragOperation::New(
+			TabToStartDragging,
+			InTabGrabOffsetFraction,
+			GetDockArea().ToSharedRef(),
+			ParentTabStackPtr.Pin()->GetTabStackGeometry().Size
+		);
+
+	return FReply::Handled().BeginDragDrop( DragDropOperation );
+}
+
+void SDockingTabWell::OnDragEnter( const FGeometry& MyGeometry, const FDragDropEvent& DragDropEvent )
+{
+	TSharedPtr<FDockingDragOperation> DragDropOperation = DragDropEvent.GetOperationAs<FDockingDragOperation>();
+	if ( DragDropOperation.IsValid() )
+	{
+		if (DragDropOperation->GetTabBeingDragged()->CanDockInNode(ParentTabStackPtr.Pin().ToSharedRef(), SDockTab::DockingViaTabWell))
+		{
+			// The user dragged a tab into this TabWell.
+
+			// Update the state of the DragDropOperation to reflect this change.
+			DragDropOperation->OnTabWellEntered( SharedThis(this) );
+
+			// Preview the position of the tab in the TabWell
+			this->TabBeingDraggedPtr = DragDropOperation->GetTabBeingDragged();
+			this->TabGrabOffsetFraction = DragDropOperation->GetTabGrabOffsetFraction();
+			
+			// The user should see the contents of the tab that we're dragging.
+			ParentTabStackPtr.Pin()->SetNodeContent(DragDropOperation->GetTabBeingDragged()->GetContent(), SNullWidget::NullWidget, SNullWidget::NullWidget);
+		}
+	}
+}
+
+void SDockingTabWell::OnDragLeave( const FDragDropEvent& DragDropEvent )
+{
+	TSharedPtr<FDockingDragOperation> DragDropOperation = DragDropEvent.GetOperationAs<FDockingDragOperation>();
+	if ( DragDropOperation.IsValid() )
+	{
+		// Check for TabBeingDraggedPtr validity as it may no longer be valid when dragging tabs in game
+		if (this->TabBeingDraggedPtr.IsValid() && DragDropOperation->GetTabBeingDragged()->CanDockInNode(ParentTabStackPtr.Pin().ToSharedRef(), SDockTab::DockingViaTabWell))
+		{
+			// Update the DragAndDrop operation based on this change.
+			const int32 LastForegroundTabIndex = Tabs.Find(this->TabBeingDraggedPtr.Pin().ToSharedRef());
+
+			// The user is pulling a tab out of this TabWell.
+			this->TabBeingDraggedPtr.Pin()->SetParent();
+
+			TSharedPtr<SDockingTabStack> ParentDockNode = ParentTabStackPtr.Pin();
+
+			// We are no longer dragging a tab in this tab well, so stop
+			// showing it in the TabWell.
+			this->TabBeingDraggedPtr.Reset();
+
+			// Also stop showing its content; switch to the last tab that was active.
+			BringTabToFront( FMath::Max(LastForegroundTabIndex-1, 0) );
+
+			// We may have removed the last tab that this DockNode had.
+			if (Tabs.Num() == 0 )
+			{
+				// Let the DockNode know that it is no longer needed.
+				ParentDockNode->OnLastTabRemoved();
+			}
+
+			GetDockArea()->CleanUp( SDockingNode::TabRemoval_DraggedOut );
+
+			const FGeometry& DockNodeGeometry = ParentDockNode->GetTabStackGeometry();
+			DragDropOperation->OnTabWellLeft( SharedThis(this), DockNodeGeometry );
+		}
+	}
+}
+
+
+FReply SDockingTabWell::OnDragOver( const FGeometry& MyGeometry, const FDragDropEvent& DragDropEvent )
+{
+	TSharedPtr<FDockingDragOperation> DragDropOperation = DragDropEvent.GetOperationAs<FDockingDragOperation>();
+	if ( DragDropOperation.IsValid() )
+	{
+		if (DragDropOperation->GetTabBeingDragged()->CanDockInNode(ParentTabStackPtr.Pin().ToSharedRef(), SDockTab::DockingViaTabWell))
+		{
+			// We are dragging the tab through a TabWell.
+			// Update the position of the Tab that we are dragging in the panel.
+			this->ChildBeingDraggedOffset = ComputeDraggedTabOffset(MyGeometry, DragDropEvent, TabGrabOffsetFraction);
+			return FReply::Handled();
+		}
+	}
+	return FReply::Unhandled();	
+}
+
+FReply SDockingTabWell::OnDrop( const FGeometry& MyGeometry, const FDragDropEvent& DragDropEvent )
+{
+	TSharedPtr<FDockingDragOperation> DragDropOperation = DragDropEvent.GetOperationAs<FDockingDragOperation>();
+	if (DragDropOperation.IsValid())
+	{
+		if (DragDropOperation->GetTabBeingDragged()->CanDockInNode(ParentTabStackPtr.Pin().ToSharedRef(), SDockTab::DockingViaTabWell))
+		{
+			if ( ensure( TabBeingDraggedPtr.IsValid() ) )
+			{
+				// We dropped a Tab into this TabWell.
+
+				// Figure out where in this TabWell to drop the Tab.
+				const float ChildWidth = ComputeChildSize(MyGeometry).X;
+				const TSharedRef<SDockTab> TabBeingDragged = TabBeingDraggedPtr.Pin().ToSharedRef();
+				const float ChildWidthWithOverlap = ChildWidth - TabBeingDragged->GetOverlapWidth();
+				const float DraggedChildCenter = ChildBeingDraggedOffset + ChildWidth / 2;
+				const int32 DropLocationIndex = FMath::Clamp( static_cast<int32>(DraggedChildCenter/ChildWidthWithOverlap), 0, Tabs.Num() );
+
+				ensure( DragDropOperation->GetTabBeingDragged().ToSharedRef() == TabBeingDragged );
+
+				// Actually insert the new tab.
+				ParentTabStackPtr.Pin()->OpenTab(TabBeingDragged, DropLocationIndex);
+
+				// We are no longer dragging a tab.
+				TabBeingDraggedPtr.Reset();
+
+				// We knew how to handled this drop operation!
+				return FReply::Handled();
+			}
+		}
+	}
+
+	// Someone just dropped something in here, but we have no idea what to do with it.
+	return FReply::Unhandled();
+}
+
+EWindowZone::Type SDockingTabWell::GetWindowZoneOverride() const
+{
+	// Pretend we are a title bar so the user can grab the area to move the window around
+	return EWindowZone::TitleBar;
+}
+
+
+void SDockingTabWell::BringTabToFront( int32 TabIndexToActivate )
+{
+	const bool bActiveIndexChanging = TabIndexToActivate != ForegroundTabIndex;
+	if ( bActiveIndexChanging )
+	{
+		const int32 LastForegroundTabIndex = FMath::Min(ForegroundTabIndex, Tabs.Num()-1);
+
+		// For positive indexes, don't go out of bounds on the array.
+		ForegroundTabIndex = FMath::Min(TabIndexToActivate, Tabs.Num()-1);
+
+		TSharedPtr<SDockingArea> MyDockArea = GetDockArea();
+		if ( Tabs.Num() > 0 && MyDockArea.IsValid() )
+		{
+			const TSharedPtr<SDockTab> PreviousForegroundTab = (LastForegroundTabIndex == INDEX_NONE)
+				? TSharedPtr<SDockTab>()
+				: Tabs[LastForegroundTabIndex];
+
+			const TSharedPtr<SDockTab> NewForegroundTab = (ForegroundTabIndex == INDEX_NONE)
+				? TSharedPtr<SDockTab>()
+				: Tabs[ForegroundTabIndex];
+			
+			MyDockArea->GetTabManager()->GetPrivateApi().OnTabForegrounded(NewForegroundTab, PreviousForegroundTab);
+		}
+	}
+
+	// Always force a refresh, even if we don't think the active index changed.
+	RefreshParentContent();
+}
+
+/** Activate the tab specified by TabToActivate SDockTab. */
+void SDockingTabWell::BringTabToFront( TSharedPtr<SDockTab> TabToActivate )
+{
+	if (Tabs.Num() > 0)
+	{
+		for (int32 TabIndex=0; TabIndex < Tabs.Num(); ++TabIndex )
+		{
+			if (Tabs[TabIndex] == TabToActivate)
+			{
+				BringTabToFront( TabIndex );
+				return;
+			}
+		}
+	}
+}
+
+/** Gets the currently active tab (or the currently dragged tab), or a null pointer if no tab is active. */
+TSharedPtr<SDockTab> SDockingTabWell::GetForegroundTab() const
+{
+	if (TabBeingDraggedPtr.IsValid())
+	{
+		return TabBeingDraggedPtr.Pin();
+	}
+	return (Tabs.Num() > 0 && ForegroundTabIndex > INDEX_NONE) ? Tabs[ForegroundTabIndex] : TSharedPtr<SDockTab>();
+}
+
+/** Gets the index of the currently active tab, or INDEX_NONE if no tab is active or a tab is being dragged. */
+int32 SDockingTabWell::GetForegroundTabIndex() const
+{
+	return (Tabs.Num() > 0) ? ForegroundTabIndex : INDEX_NONE;
+}
+
+void SDockingTabWell::RemoveAndDestroyTab(const TSharedRef<SDockTab>& TabToRemove, SDockingNode::ELayoutModification RemovalMethod)
+{
+	int32 TabIndex = Tabs.Find(TabToRemove);
+
+	if (TabIndex != INDEX_NONE)
+	{
+		const TSharedPtr<SDockingTabStack> ParentTabStack = ParentTabStackPtr.Pin();
+
+		// Remove the old tab from the list of tabs and activate the new tab.
+		{
+			BringTabToFront(TabIndex);
+			Tabs.RemoveAt(TabIndex);
+			// We no longer have a tab in the foreground.
+			// This is important because BringTabToFront triggers notifications based on the difference in active tab indexes.
+			ForegroundTabIndex = INDEX_NONE;
+
+			// Now bring the last tab that we were on to the foreground
+			BringTabToFront(FMath::Max(TabIndex-1, 0));
+		}
+		
+		if ( ensure(ParentTabStack.IsValid()) )
+		{
+			ParentTabStack->OnTabClosed( TabToRemove );
+			
+			// We might be closing down an entire dock area, if this is a major tab.
+			// Use this opportunity to save its layout
+			if (RemovalMethod == SDockingNode::TabRemoval_Closed)
+			{
+				ParentTabStack->GetDockArea()->GetTabManager()->GetPrivateApi().OnTabClosing( TabToRemove );
+			}
+
+			if (Tabs.Num() == 0)
+			{
+				ParentTabStack->OnLastTabRemoved();
+			}
+			else
+			{
+				RefreshParentContent();
+			}
+		}
+
+
+		GetDockArea()->CleanUp( RemovalMethod );
+	}
+}
+
+void SDockingTabWell::RefreshParentContent()
+{
+	if (Tabs.Num() > 0 && ForegroundTabIndex != INDEX_NONE)
+	{
+		const TSharedRef<SDockTab>& ForegroundTab = Tabs[ForegroundTabIndex];
+		FGlobalTabmanager::Get()->SetActiveTab( ForegroundTab );
+
+		TSharedPtr<SWindow> ParentWindowPtr = ForegroundTab->GetParentWindow();
+		if (ParentWindowPtr.IsValid() && ParentWindowPtr != FGlobalTabmanager::Get()->GetRootWindow())
+		{
+			ParentWindowPtr->SetTitle( ForegroundTab->GetTabLabel() );
+		}
+
+		ParentTabStackPtr.Pin()->SetNodeContent( ForegroundTab->GetContent(), ForegroundTab->GetLeftContent(), ForegroundTab->GetRightContent() );
+	}
+	else
+	{
+		ParentTabStackPtr.Pin()->SetNodeContent( SNullWidget::NullWidget, SNullWidget::NullWidget, SNullWidget::NullWidget );
+	}
+}
+
+TSharedPtr<SDockingArea> SDockingTabWell::GetDockArea()
+{
+	return ParentTabStackPtr.IsValid() ? ParentTabStackPtr.Pin()->GetDockArea() : TSharedPtr<SDockingArea>();
+}
+
+
+TSharedPtr<SDockingTabStack> SDockingTabWell::GetParentDockTabStack()
+{
+	return ParentTabStackPtr.IsValid() ? ParentTabStackPtr.Pin() : TSharedPtr<SDockingTabStack>();
+}
diff --git a/Engine/Source/Runtime/Slate/Private/Framework/Docking/SDockingTarget.h b/Engine/Source/Runtime/Slate/Private/Framework/Docking/SDockingTarget.h
new file mode 100644
--- /dev/null
+++ b/Engine/Source/Runtime/Slate/Private/Framework/Docking/SDockingTarget.h
@@ -0,0 +1,51 @@
+// Copyright 1998-2014 Epic Games, Inc. All Rights Reserved.
+
+#pragma once
+
+
+/**
+ * Targets used by docking code. When re-arranging layout, hovering over targets
+ * gives the user a preview of what will happen if they drop on that target.
+ * Dropping actually commits the layout-restructuring.
+ */
+class SLATE_API SDockingTarget : public SBorder
+{
+public:
+	SLATE_BEGIN_ARGS(SDockingTarget)
+		: _OwnerNode()
+		, _DockDirection(SDockingNode::LeftOf)
+		{}
+		/** The DockNode that this target is representing. Docking will occur relative to this DockNode */
+		SLATE_ARGUMENT( TSharedPtr<class SDockingNode>, OwnerNode )
+		/** In which direction relative to the TabStack to dock */
+		SLATE_ARGUMENT( SDockingNode::RelativeDirection, DockDirection )
+	SLATE_END_ARGS()
+
+	void Construct( const FArguments& InArgs );
+
+	// SWidget interface
+	virtual void OnDragEnter( const FGeometry& MyGeometry, const FDragDropEvent& DragDropEvent ) OVERRIDE;
+	virtual void OnDragLeave( const FDragDropEvent& DragDropEvent ) OVERRIDE;
+	virtual FReply OnDrop( const FGeometry& MyGeometry, const FDragDropEvent& DragDropEvent ) OVERRIDE;
+	// End of SWidget interface
+
+	/** @return the DockNode which this target represent */
+	TSharedPtr<SDockingNode> GetOwner() const;
+	
+	/** @return the direction which we represent (relative to the OwnerTabStack) */
+	SDockingNode::RelativeDirection GetDockDirection() const;
+
+private:
+	const FSlateBrush* GetTargetImage() const;
+
+	const FSlateBrush* TargetImage;
+	const FSlateBrush* TargetImage_Hovered;
+
+	/** The DockNode relative to which we want to dock */
+	TWeakPtr<class SDockingNode> OwnerNode;
+	/** The direction in which we want to dock relative to the tab stack */
+	SDockingNode::RelativeDirection DockDirection;
+
+	bool bIsDragHovered;
+};
+
