Slate: Reorganized private Slate files; proper pre-compiled header for all files.

diff --git a/Engine/Source/Runtime/Slate/Private/Framework/Docking/TabManager.cpp b/Engine/Source/Runtime/Slate/Private/Framework/Docking/TabManager.cpp
new file mode 100644
--- /dev/null
+++ b/Engine/Source/Runtime/Slate/Private/Framework/Docking/TabManager.cpp
@@ -0,0 +1,1742 @@
+// Copyright 1998-2014 Epic Games, Inc. All Rights Reserved.
+
+#include "SlatePrivatePCH.h"
+#include "DockingPrivate.h"
+
+
+const FVector2D FTabManager::FallbackWindowSize( 1000, 600 );
+
+DEFINE_LOG_CATEGORY_STATIC(LogTabManager, Display, All);
+
+#define LOCTEXT_NAMESPACE "TabManager"
+
+//////////////////////////////////////////////////////////////////////////
+// 
+//////////////////////////////////////////////////////////////////////////
+
+TSharedRef<FTabManager::FLayoutNode> FTabManager::FLayout::NewFromString_Helper( TSharedPtr<FJsonObject> JsonObject )
+{
+	struct local
+	{
+		static FTabManager::FArea::EWindowPlacement PlacementFromString( const FString& AsString )
+		{
+			static const FString Placement_NoWindow_Str = TEXT("Placement_NoWindow");
+			static const FString Placement_Automatic_Str = TEXT("Placement_Automatic");
+			static const FString Placement_Specified_Str = TEXT("Placement_Specified");
+
+			if (AsString == Placement_NoWindow_Str)
+			{
+				return FTabManager::FArea::Placement_NoWindow;
+			}
+			else if (AsString == Placement_Automatic_Str)
+			{
+				return FTabManager::FArea::Placement_Automatic;
+			}
+			else if (AsString == Placement_Specified_Str)
+			{
+				return FTabManager::FArea::Placement_Specified;
+			}
+			else
+			{
+				ensureMsg(false, TEXT("Invalid placement mode."));
+				return FTabManager::FArea::Placement_Automatic;
+			}
+		}
+
+		static EOrientation OrientationFromString( const FString& AsString )
+		{
+			static const FString Orient_Horizontal_Str = TEXT("Orient_Horizontal");
+			static const FString Orient_Vertical_Str = TEXT("Orient_Vertical");
+
+			if (AsString == Orient_Horizontal_Str)
+			{
+				return Orient_Horizontal;
+			}
+			else if (AsString == Orient_Vertical_Str)
+			{
+				return Orient_Vertical;
+			}
+			else
+			{
+				ensureMsg(false, TEXT("Invalid orientation."));
+				return Orient_Horizontal;
+			}
+			
+		}
+
+		static ETabState::Type TabStateFromString( const FString& AsString )
+		{
+			static const FString OpenedTab_Str = TEXT("OpenedTab");
+			static const FString ClosedTab_Str = TEXT("ClosedTab");
+
+			if (AsString == OpenedTab_Str)
+			{
+				return ETabState::OpenedTab;
+			}
+			else if(AsString == ClosedTab_Str)
+			{
+				return ETabState::ClosedTab;
+			}
+			else
+			{
+				ensureMsg(false, TEXT("Invalid tab state."));
+				return ETabState::OpenedTab;
+			}
+		}
+	};
+
+	const FString NodeType = JsonObject->GetStringField(TEXT("Type"));
+	if (NodeType == TEXT("Area"))
+	{
+		TSharedPtr<FTabManager::FArea> NewArea;
+
+		FTabManager::FArea::EWindowPlacement WindowPlacement = local::PlacementFromString( JsonObject->GetStringField(TEXT("WindowPlacement")) );
+		switch( WindowPlacement )
+		{
+			default:
+			case FTabManager::FArea::Placement_NoWindow:
+			{
+				NewArea = FTabManager::NewPrimaryArea();
+			}
+			break;
+
+			case FTabManager::FArea::Placement_Automatic:
+			{
+				FVector2D WindowSize;
+
+				WindowSize.X = JsonObject->GetNumberField( TEXT("WindowSize_X") );
+				WindowSize.Y = JsonObject->GetNumberField( TEXT("WindowSize_Y") );
+
+				NewArea = FTabManager::NewArea( WindowSize );
+			}
+			break;
+
+			case FTabManager::FArea::Placement_Specified:
+			{
+				FVector2D WindowPosition = FVector2D::ZeroVector;
+				FVector2D WindowSize;
+
+				WindowPosition.X = JsonObject->GetNumberField( TEXT("WindowPosition_X") );
+				WindowPosition.Y = JsonObject->GetNumberField( TEXT("WindowPosition_Y") );
+
+				WindowSize.X = JsonObject->GetNumberField( TEXT("WindowSize_X") );
+				WindowSize.Y = JsonObject->GetNumberField( TEXT("WindowSize_Y") );
+
+				bool bIsMaximized = JsonObject->GetBoolField(TEXT("bIsMaximized"));
+
+				NewArea = FTabManager::NewArea( WindowSize );
+				NewArea->SetWindow( WindowPosition, bIsMaximized );
+			}
+			break;
+		}
+		
+		NewArea->SetSizeCoefficient( JsonObject->GetNumberField( TEXT("SizeCoefficient") ) );
+		NewArea->SetOrientation( local::OrientationFromString( JsonObject->GetStringField(TEXT("Orientation")) ) );
+
+		TArray< TSharedPtr<FJsonValue> > ChildNodeValues = JsonObject->GetArrayField(TEXT("nodes"));
+		for( int32 ChildIndex=0; ChildIndex < ChildNodeValues.Num(); ++ChildIndex )
+		{
+			NewArea->Split( NewFromString_Helper( ChildNodeValues[ChildIndex]->AsObject() ) );
+		}
+
+		return NewArea.ToSharedRef();
+	}
+	else if ( NodeType == TEXT("Splitter") )
+	{
+		TSharedRef<FTabManager::FSplitter> NewSplitter =  FTabManager::NewSplitter();
+		NewSplitter->SetSizeCoefficient( JsonObject->GetNumberField(TEXT("SizeCoefficient")) );
+		NewSplitter->SetOrientation( local::OrientationFromString( JsonObject->GetStringField(TEXT("Orientation")) ) );
+		TArray< TSharedPtr<FJsonValue> > ChildNodeValues = JsonObject->GetArrayField(TEXT("nodes"));
+		for( int32 ChildIndex=0; ChildIndex < ChildNodeValues.Num(); ++ChildIndex )
+		{
+			NewSplitter->Split( NewFromString_Helper( ChildNodeValues[ChildIndex]->AsObject() ) );
+		}
+		return NewSplitter;
+	}
+	else if ( NodeType == TEXT("Stack") )
+	{
+		TSharedRef<FTabManager::FStack> NewStack = FTabManager::NewStack();
+		NewStack->SetSizeCoefficient( JsonObject->GetNumberField(TEXT("SizeCoefficient")) );
+		NewStack->SetHideTabWell( JsonObject->GetBoolField(TEXT("HideTabWell")) );
+
+		if(JsonObject->HasField(TEXT("ForegroundTab")))
+		{
+			FName TabId = FName( *JsonObject->GetStringField(TEXT("ForegroundTab")) );
+			TabId = FGlobalTabmanager::Get()->GetTabTypeForPotentiallyLegacyTab(TabId);
+			NewStack->SetForegroundTab( FTabId(TabId) );
+		}
+
+		TArray< TSharedPtr<FJsonValue> > TabsAsJson = JsonObject->GetArrayField( TEXT("Tabs") );
+		for (int32 TabIndex=0; TabIndex < TabsAsJson.Num(); ++TabIndex)
+		{
+			TSharedPtr<FJsonObject> TabAsJson = TabsAsJson[TabIndex]->AsObject();
+			FName TabId = FName( *TabAsJson->GetStringField(TEXT("TabId")) );
+			TabId = FGlobalTabmanager::Get()->GetTabTypeForPotentiallyLegacyTab(TabId);
+			NewStack->AddTab( TabId, local::TabStateFromString( TabAsJson->GetStringField(TEXT("TabState")) ) );			
+		}
+		return NewStack;
+	}
+	else
+	{
+		ensureMsg(false, TEXT("Unrecognized node type."));
+		return FTabManager::NewArea(FTabManager::FallbackWindowSize);
+	}	
+}
+
+TSharedPtr<FTabManager::FLayout> FTabManager::FLayout::NewFromString( const FString& LayoutAsText )
+{
+	TSharedPtr<FTabManager::FLayout> Layout;
+	TSharedPtr<FJsonObject> JsonObject;
+	FString OutError;
+
+	TSharedRef< TJsonReader<> > Reader = TJsonReaderFactory<>::Create( LayoutAsText );
+	
+	if (FJsonSerializer::Deserialize( Reader, JsonObject ))
+	{
+		const FString LayoutName = JsonObject->GetStringField(TEXT("Name"));
+		TSharedRef<FTabManager::FLayout> NewLayout = FTabManager::NewLayout( *LayoutName );
+		int32 PrimaryAreaIndex = FMath::Trunc( JsonObject->GetNumberField(TEXT("PrimaryAreaIndex")) );
+
+		TArray< TSharedPtr<FJsonValue> > Areas = JsonObject->GetArrayField(TEXT("Areas"));
+		for(int32 AreaIndex=0; AreaIndex < Areas.Num(); ++AreaIndex)
+		{
+			TSharedRef<FTabManager::FArea> NewArea = StaticCastSharedRef<FTabManager::FArea>( NewFromString_Helper( Areas[AreaIndex]->AsObject() ) );
+			NewLayout->AddArea( NewArea );
+			if (AreaIndex == PrimaryAreaIndex)
+			{
+				NewLayout->PrimaryArea = NewArea;
+			}
+		}
+		
+		Layout = NewLayout;
+	}
+	
+	return Layout;
+}
+
+FName FTabManager::FLayout::GetLayoutName() const
+{
+	return LayoutName;
+}
+
+FString FTabManager::FLayout::ToString() const
+{
+	TSharedRef<FJsonObject> LayoutJson = MakeShareable( new FJsonObject() );
+	LayoutJson->SetStringField( TEXT("Type"), TEXT("Layout") );
+	LayoutJson->SetStringField( TEXT("Name"), LayoutName.ToString() );
+
+	LayoutJson->SetNumberField( TEXT("PrimaryAreaIndex"), INDEX_NONE );
+
+	TArray< TSharedPtr<FJsonValue> > AreasAsJson;
+	for ( int32 AreaIndex=0; AreaIndex < Areas.Num(); ++AreaIndex )
+	{
+		if (PrimaryArea.Pin() == Areas[AreaIndex])
+		{
+			LayoutJson->SetNumberField( TEXT("PrimaryAreaIndex"), AreaIndex );
+		}
+		AreasAsJson.Add( MakeShareable( new FJsonValueObject( PersistToString_Helper( Areas[AreaIndex] ) ) ) );
+	}
+	LayoutJson->SetArrayField( TEXT("Areas"), AreasAsJson );
+
+	FString LayoutAsString;
+	TSharedRef< TJsonWriter<> > Writer = TJsonWriterFactory<>::Create( &LayoutAsString );
+	check( FJsonSerializer::Serialize( LayoutJson, Writer ) );
+
+	return LayoutAsString;
+}
+
+
+TSharedRef<FJsonObject> FTabManager::FLayout::PersistToString_Helper(const TSharedRef<FLayoutNode>& NodeToPersist)
+{
+	TSharedRef<FJsonObject> JsonObj = MakeShareable(new FJsonObject());
+
+	TSharedPtr<FTabManager::FStack> NodeAsStack = NodeToPersist->AsStack();
+	TSharedPtr<FTabManager::FSplitter> NodeAsSplitter = NodeToPersist->AsSplitter();
+	TSharedPtr<FTabManager::FArea> NodeAsArea = NodeToPersist->AsArea();
+
+
+	JsonObj->SetNumberField( TEXT("SizeCoefficient"), NodeToPersist->SizeCoefficient );
+
+	if ( NodeAsArea.IsValid() )
+	{
+		JsonObj->SetStringField( TEXT("Type"), TEXT("Area") );
+		JsonObj->SetStringField( TEXT("Orientation"), (NodeAsArea->GetOrientation() == Orient_Horizontal) ? TEXT("Orient_Horizontal") : TEXT("Orient_Vertical")  );
+
+		if ( NodeAsArea->WindowPlacement == FArea::Placement_Automatic )
+		{
+			JsonObj->SetStringField( TEXT("WindowPlacement"), TEXT("Placement_Automatic") );
+			JsonObj->SetNumberField( TEXT("WindowSize_X"), NodeAsArea->WindowSize.X );
+			JsonObj->SetNumberField( TEXT("WindowSize_Y"), NodeAsArea->WindowSize.Y );
+		}
+		else if (NodeAsArea->WindowPlacement == FArea::Placement_NoWindow)
+		{
+			JsonObj->SetStringField( TEXT("WindowPlacement"), TEXT("Placement_NoWindow") );
+		}
+		else if ( NodeAsArea->WindowPlacement == FArea::Placement_Specified )
+		{
+			JsonObj->SetStringField( TEXT("WindowPlacement"), TEXT("Placement_Specified") );
+			JsonObj->SetNumberField( TEXT("WindowPosition_X"), NodeAsArea->WindowPosition.X );
+			JsonObj->SetNumberField( TEXT("WindowPosition_Y"), NodeAsArea->WindowPosition.Y );
+			JsonObj->SetNumberField( TEXT("WindowSize_X"), NodeAsArea->WindowSize.X );
+			JsonObj->SetNumberField( TEXT("WindowSize_Y"), NodeAsArea->WindowSize.Y );
+			JsonObj->SetBoolField( TEXT("bIsMaximized"), NodeAsArea->bIsMaximized );
+		}
+				
+		TArray< TSharedPtr<FJsonValue> > Nodes;
+		for ( int32 ChildIndex=0; ChildIndex < NodeAsArea->ChildNodes.Num(); ++ChildIndex )
+		{
+			Nodes.Add( MakeShareable( new FJsonValueObject( PersistToString_Helper( NodeAsArea->ChildNodes[ChildIndex] ) ) ) );
+		}
+		JsonObj->SetArrayField( TEXT("Nodes"), Nodes );
+	}
+	else if ( NodeAsSplitter.IsValid() )
+	{
+		JsonObj->SetStringField( TEXT("Type"), TEXT("Splitter") );
+		JsonObj->SetStringField( TEXT("Orientation"), (NodeAsSplitter->GetOrientation() == Orient_Horizontal) ? TEXT("Orient_Horizontal") : TEXT("Orient_Vertical")  );
+
+		TArray< TSharedPtr<FJsonValue> > Nodes;
+		for ( int32 ChildIndex=0; ChildIndex < NodeAsSplitter->ChildNodes.Num(); ++ChildIndex )
+		{
+			Nodes.Add( MakeShareable( new FJsonValueObject( PersistToString_Helper( NodeAsSplitter->ChildNodes[ChildIndex] ) ) ) );
+		}
+		JsonObj->SetArrayField( TEXT("Nodes"), Nodes );
+	}
+	else if ( NodeAsStack.IsValid() )
+	{
+		JsonObj->SetStringField( TEXT("Type"), TEXT("Stack") );
+		JsonObj->SetBoolField( TEXT("HideTabWell"), NodeAsStack->bHideTabWell );
+		JsonObj->SetStringField( TEXT("ForegroundTab"), NodeAsStack->ForegroundTabId.ToString() );
+
+		TArray< TSharedPtr<FJsonValue> > TabsAsJson;
+		for ( int32 TabIndex=0; TabIndex < NodeAsStack->Tabs.Num(); ++TabIndex )
+		{
+			TSharedRef<FJsonObject> TabAsJson = MakeShareable( new FJsonObject() );
+			TabAsJson->SetStringField( TEXT("TabId"), NodeAsStack->Tabs[TabIndex].TabId.ToString() );
+			TabAsJson->SetStringField( TEXT("TabState"), (NodeAsStack->Tabs[TabIndex].TabState == ETabState::OpenedTab)
+				? TEXT("OpenedTab")
+				: TEXT("ClosedTab") );
+			TabsAsJson.Add( MakeShareable( new FJsonValueObject(TabAsJson) ) );
+		}
+		JsonObj->SetArrayField( TEXT("Tabs"), TabsAsJson );
+	}
+	else
+	{
+		ensureMsg( false, TEXT("Unable to persist layout node of unknown type.") );
+	}
+
+	return JsonObj;
+}
+
+
+//////////////////////////////////////////////////////////////////////////
+// FTabManager::PrivateApi
+//////////////////////////////////////////////////////////////////////////
+
+TSharedPtr<SWindow> FTabManager::FPrivateApi::GetParentWindow() const
+{
+	TSharedPtr<SDockTab> OwnerTab = TabManager.OwnerTabPtr.Pin();
+	if ( OwnerTab.IsValid() )
+	{
+		// The tab was dragged out of some context that is owned by a MajorTab.
+		// Whichever window possesses the MajorTab should be the parent of the newly created window.
+		FWidgetPath DummyPath;
+		TSharedPtr<SWindow> ParentWindow = FSlateApplication::Get().FindWidgetWindow( OwnerTab.ToSharedRef(), DummyPath );
+		return ParentWindow;
+	}
+	else
+	{
+		// This tab is not nested within a major tab, so it is a major tab itself.
+		// Ask the global tab manager for its root window.
+		return FGlobalTabmanager::Get()->GetRootWindow();
+	}
+}
+
+void FTabManager::FPrivateApi::OnDockAreaCreated( const TSharedRef<SDockingArea>& NewlyCreatedDockArea )
+{
+	CleanupPointerArray(TabManager.DockAreas);
+	TabManager.DockAreas.Add( NewlyCreatedDockArea );
+}
+
+void FTabManager::FPrivateApi::OnTabRelocated( const TSharedRef<SDockTab>& RelocatedTab, const TSharedPtr<SWindow>& NewOwnerWindow )
+{
+	TabManager.OnTabRelocated(RelocatedTab, NewOwnerWindow);
+}
+
+void FTabManager::FPrivateApi::OnTabOpening( const TSharedRef<SDockTab>& TabBeingOpened )
+{
+	TabManager.OnTabOpening(TabBeingOpened);
+}
+
+void FTabManager::FPrivateApi::OnTabClosing( const TSharedRef<SDockTab>& TabBeingClosed )
+{
+	TabManager.OnTabClosing(TabBeingClosed);
+}
+
+void FTabManager::FPrivateApi::OnDockAreaClosing( const TSharedRef<SDockingArea>& DockAreaThatIsClosing )
+{
+	TSharedPtr<FTabManager::FArea> PersistentDockAreaLayout = StaticCastSharedPtr<FTabManager::FArea>(DockAreaThatIsClosing->GatherPersistentLayout());
+
+	if ( PersistentDockAreaLayout.IsValid() )
+	{
+		TabManager.CollapsedDockAreas.Add( PersistentDockAreaLayout.ToSharedRef() );
+	}
+}
+
+void FTabManager::FPrivateApi::OnTabManagerClosing()
+{
+	TabManager.OnTabManagerClosing();
+}
+
+const TArray< TWeakPtr<SDockingArea> >& FTabManager::FPrivateApi::GetLiveDockAreas() const
+{
+	return TabManager.DockAreas;
+}
+
+void FTabManager::FPrivateApi::OnTabForegrounded( const TSharedPtr<SDockTab>& NewForegroundTab, const TSharedPtr<SDockTab>& BackgroundedTab )
+{
+	TabManager.OnTabForegrounded( NewForegroundTab, BackgroundedTab );
+}
+
+
+
+
+
+
+
+
+
+
+static void SetWindowVisibility( const TArray< TWeakPtr<SDockingArea> >& DockAreas, bool bWindowShouldBeVisible )
+{
+	for (int32 DockAreaIndex=0; DockAreaIndex < DockAreas.Num(); ++DockAreaIndex)
+	{
+		TSharedPtr<SWindow> DockAreaWindow = DockAreas[DockAreaIndex].Pin()->GetParentWindow();
+		if ( DockAreaWindow.IsValid() )
+		{
+			if (bWindowShouldBeVisible)
+			{
+				DockAreaWindow->ShowWindow();
+			}
+			else
+			{
+				DockAreaWindow->HideWindow();
+			}
+		}
+	}
+}
+
+void FTabManager::FPrivateApi::ShowWindows()
+{
+	CleanupPointerArray( TabManager.DockAreas );
+	SetWindowVisibility(TabManager.DockAreas, true);
+}
+
+void FTabManager::FPrivateApi::HideWindows()
+{
+	CleanupPointerArray( TabManager.DockAreas );
+	SetWindowVisibility(TabManager.DockAreas, false);
+}
+
+
+
+
+
+
+
+
+FTabManager::FPrivateApi& FTabManager::GetPrivateApi()
+{
+	return *PrivateApi;
+}
+
+TSharedPtr<FTabManager::FStack> FTabManager::FLayoutNode::AsStack()
+{
+	return TSharedPtr<FTabManager::FStack>();
+}
+
+TSharedPtr<FTabManager::FSplitter> FTabManager::FLayoutNode::AsSplitter()
+{
+	return TSharedPtr<FTabManager::FSplitter>();
+}
+
+TSharedPtr<FTabManager::FArea> FTabManager::FLayoutNode::AsArea()
+{
+	return TSharedPtr<FTabManager::FArea>();
+}
+
+void FTabManager::SetOnPersistLayout( const FOnPersistLayout& InHandler )
+{
+	OnPersistLayout_Handler = InHandler;
+}
+
+void FTabManager::CloseAllAreas()
+{
+	for ( int32 LiveAreaIndex=0; LiveAreaIndex < DockAreas.Num(); ++LiveAreaIndex )
+	{
+		const TSharedPtr<SDockingArea> SomeDockArea = DockAreas[LiveAreaIndex].Pin();
+		const TSharedPtr<SWindow> ParentWindow = (SomeDockArea.IsValid())
+			? SomeDockArea->GetParentWindow()
+			: TSharedPtr<SWindow>();
+
+		if (ParentWindow.IsValid())
+		{
+			ParentWindow->RequestDestroyWindow();
+		}
+	}
+	DockAreas.Empty();
+
+	CollapsedDockAreas.Empty();
+}
+
+
+TSharedRef<FTabManager::FLayout> FTabManager::PersistLayout() const
+{
+	TSharedRef<FLayout> PersistentLayout = FTabManager::NewLayout( this->ActiveLayoutName );
+	
+	// Persist layout for all LiveAreas
+	for ( int32 LiveAreaIndex=0; LiveAreaIndex < DockAreas.Num(); ++LiveAreaIndex )
+	{
+		TSharedPtr<FArea> PersistedNode;
+		TSharedPtr<SDockingArea> ChildDockingArea = DockAreas[LiveAreaIndex].Pin();
+
+		if ( ChildDockingArea.IsValid() )
+		{
+			TSharedPtr<FTabManager::FLayoutNode> LayoutNode = ChildDockingArea->GatherPersistentLayout();
+			if (LayoutNode.IsValid())
+			{
+				PersistedNode = LayoutNode->AsArea();
+			}
+		}
+
+		if ( PersistedNode.IsValid() )
+		{
+			PersistentLayout->AddArea( PersistedNode.ToSharedRef() );
+			if (PersistedNode->WindowPlacement == FArea::Placement_NoWindow)
+			{
+				ensure( !PersistentLayout->PrimaryArea.IsValid() );
+				PersistentLayout->PrimaryArea = PersistedNode;
+			}
+		}
+	}
+
+	// Gather existing persistent layouts for CollapsedAreas
+	for ( int32 CollapsedAreaIndex=0; CollapsedAreaIndex < CollapsedDockAreas.Num(); ++CollapsedAreaIndex )
+	{
+		PersistentLayout->AddArea( CollapsedDockAreas[CollapsedAreaIndex] );
+	}
+
+	return PersistentLayout;
+}
+
+void FTabManager::SavePersistentLayout()
+{
+	const TSharedRef<FLayout> MyLayout = this->PersistLayout();
+	OnPersistLayout_Handler.ExecuteIfBound( MyLayout );
+}
+
+FTabSpawnerEntry& FTabManager::RegisterTabSpawner( const FName TabId, const FOnSpawnTab& OnSpawnTab )
+{
+	ensure( !TabSpawner.Contains(TabId) );
+	ensure( !FGlobalTabmanager::Get()->IsLegacyTabType(TabId) );
+
+	TSharedRef<FTabSpawnerEntry> NewSpawnerEntry = MakeShareable( new FTabSpawnerEntry( TabId, OnSpawnTab ) );
+	TabSpawner.Add( TabId, NewSpawnerEntry );
+	return NewSpawnerEntry.Get();
+}
+
+bool FTabManager::UnregisterTabSpawner( const FName TabId )
+{
+	return TabSpawner.Remove( TabId ) > 0;
+}
+
+void FTabManager::UnregisterAllTabSpawners()
+{
+	TabSpawner.Empty();
+}
+
+TSharedPtr<SWidget> FTabManager::RestoreFrom( const TSharedRef<FLayout>& Layout, const TSharedPtr<SWindow>& ParentWindow, const bool bEmbedTitleAreaContent )
+{
+	ActiveLayoutName = Layout->LayoutName;
+
+	TSharedPtr<SDockingArea> PrimaryDockArea;
+	for (int32 AreaIndex=0; AreaIndex < Layout->Areas.Num(); ++AreaIndex )
+	{
+		const TSharedRef<FArea>& ThisArea = Layout->Areas[AreaIndex];
+
+		const bool bIsPrimaryArea = ThisArea->WindowPlacement == FArea::Placement_NoWindow;
+		const bool bShouldCreate = bIsPrimaryArea || HasValidTabs(ThisArea);
+
+		if ( bShouldCreate )
+		{
+			TSharedPtr<SDockingArea> RestoredDockArea;
+			const bool bHasOpenTabs = bIsPrimaryArea || HasOpenTabs(ThisArea);
+
+			if ( bHasOpenTabs )
+			{
+				RestoredDockArea = RestoreArea(ThisArea, ParentWindow, bEmbedTitleAreaContent);
+			}
+			else
+			{
+				CollapsedDockAreas.Add(ThisArea);
+			}
+
+			if ( bIsPrimaryArea )
+			{
+				check( !PrimaryDockArea.IsValid() );
+				PrimaryDockArea	= RestoredDockArea;
+			}
+		}
+	}
+
+	UpdateStats();
+	
+	return PrimaryDockArea;
+}
+
+struct FPopulateTabSpawnerMenu_Args
+{
+	FPopulateTabSpawnerMenu_Args( const TSharedRef< TArray< TWeakPtr<FTabSpawnerEntry> > >& InAllSpawners, const TSharedRef<FWorkspaceItem>& InMenuNode, int32 InLevel )
+		: AllSpawners( InAllSpawners )
+		, MenuNode( InMenuNode )
+		, Level( InLevel )	
+	{
+	}
+
+	TSharedRef< TArray< TWeakPtr<FTabSpawnerEntry> > > AllSpawners;
+	TSharedRef<FWorkspaceItem> MenuNode;
+	int32 Level;
+};
+
+/** Scoped guard to ensure that we reset the GuardedValue to false */
+struct FScopeGuard
+{
+	FScopeGuard( bool & InGuardedValue )
+	: GuardedValue(InGuardedValue)
+	{
+		GuardedValue = true;
+	}
+
+	~FScopeGuard()
+	{
+		GuardedValue = false;	
+	}
+
+	private:
+		bool& GuardedValue;
+
+};
+
+void FTabManager::PopulateTabSpawnerMenu_Helper( FMenuBuilder& PopulateMe, FPopulateTabSpawnerMenu_Args Args )
+{
+	const TArray< TSharedRef<FWorkspaceItem> >& ChildItems = Args.MenuNode->GetChildItems();
+
+	bool bFirstItemOnLevel = true;
+
+	for ( int32 ChildIndex=0; ChildIndex < ChildItems.Num(); ++ChildIndex )
+	{
+		const TSharedRef<FWorkspaceItem>& ChildItem = ChildItems[ChildIndex];
+		const TSharedPtr<FTabSpawnerEntry> SpawnerNode = ChildItem->AsSpawnerEntry();
+		if ( SpawnerNode.IsValid() )
+		{
+			// LEAF NODE.
+			// Make a menu item for summoning a tab.
+			if (Args.AllSpawners->Contains(SpawnerNode.ToSharedRef()))
+			{
+				MakeSpawnerMenuEntry(PopulateMe, SpawnerNode);
+			}
+		}
+		else
+		{
+			// GROUP NODE
+			// If it's not empty, populate it
+			if ( ChildItem->HasChildrenIn(*Args.AllSpawners) )
+			{
+				const FPopulateTabSpawnerMenu_Args Payload( Args.AllSpawners, ChildItem, Args.Level+1 );
+
+				if ( Args.Level % 2 == 0 )
+				{
+					if (!bFirstItemOnLevel)
+					{
+						PopulateMe.AddMenuSeparator();
+					}
+						
+					PopulateTabSpawnerMenu_Helper( PopulateMe, Payload );
+				}
+				else
+				{
+					PopulateMe.AddSubMenu(
+						ChildItem->GetDisplayName(),
+						ChildItem->GetTooltipText(),
+						FNewMenuDelegate::CreateRaw( this, &FTabManager::PopulateTabSpawnerMenu_Helper, Payload ),
+						false,
+						ChildItem->GetIcon()
+					);
+				}
+
+				bFirstItemOnLevel = false;
+			}
+		}
+	}
+}
+
+
+void FTabManager::MakeSpawnerMenuEntry( FMenuBuilder &PopulateMe, const TSharedPtr<FTabSpawnerEntry> &SpawnerNode ) 
+{
+	//PopulateMe.BeginSection( "TabSection", SpawnerNode->GetDisplayName() );
+	if ( SpawnerNode->MenuType != ETabSpawnerMenuType::Hide )
+	{
+		PopulateMe.AddMenuEntry(
+			SpawnerNode->GetDisplayName().IsEmpty() ? FText::FromName( SpawnerNode->TabType ) : SpawnerNode->GetDisplayName(),
+			SpawnerNode->GetTooltipText(),
+			SpawnerNode->GetIcon(),
+			FUIAction(
+			FExecuteAction::CreateSP(SharedThis(this), &FTabManager::InvokeTabForMenu, SpawnerNode->TabType),
+			FCanExecuteAction(),
+			FIsActionChecked::CreateSP(SpawnerNode.ToSharedRef(), &FTabSpawnerEntry::IsSoleTabInstanceSpawned)
+			),
+			NAME_None,
+			EUserInterfaceActionType::Check
+			);
+	}
+	//PopulateMe.EndSection();
+}
+
+void FTabManager::PopulateTabSpawnerMenu( FMenuBuilder& PopulateMe, TSharedRef<FWorkspaceItem> MenuStructure )
+{
+	TSharedRef< TArray< TWeakPtr<FTabSpawnerEntry> > > AllSpawners = MakeShareable( new TArray< TWeakPtr<FTabSpawnerEntry> >() );
+	{
+		for ( FTabSpawner::TIterator SpawnerIterator(TabSpawner); SpawnerIterator; ++SpawnerIterator  )	
+		{
+			const TSharedRef<FTabSpawnerEntry>& SpawnerEntry = SpawnerIterator.Value();
+			AllSpawners->AddUnique(SpawnerEntry);
+		}
+
+		for ( FTabSpawner::TIterator SpawnerIterator(*NomadTabSpawner); SpawnerIterator; ++SpawnerIterator  )	
+		{
+			const TSharedRef<FTabSpawnerEntry>& SpawnerEntry = SpawnerIterator.Value();
+			AllSpawners->AddUnique(SpawnerEntry);
+		}
+	}
+
+	// Put all orphaned spawners at the top of the menu so programmers go and find them a nice home.
+	for ( int32 ChildIndex=0; ChildIndex < AllSpawners->Num(); ++ChildIndex )
+	{
+		const TSharedPtr<FTabSpawnerEntry> Spawner = (*AllSpawners)[ChildIndex].Pin();
+		const bool bHasNoPlaceInMenuStructure = !Spawner->GetParent().IsValid();
+		if ( bHasNoPlaceInMenuStructure )
+		{
+			this->MakeSpawnerMenuEntry( PopulateMe, Spawner );
+		}
+	}
+
+	PopulateTabSpawnerMenu_Helper( PopulateMe, FPopulateTabSpawnerMenu_Args( AllSpawners, MenuStructure, 0 ) );
+}
+
+void FTabManager::PopulateTabSpawnerMenu( FMenuBuilder &PopulateMe, const FName& TabType )
+{
+	TSharedPtr<FTabSpawnerEntry> Spawner = FindTabSpawnerFor(TabType);
+	if (Spawner.IsValid())
+	{
+		MakeSpawnerMenuEntry(PopulateMe, Spawner);
+	}
+	else
+	{
+		UE_LOG(LogSlate, Warning, TEXT("PopulateTabSpawnerMenu failed to find entry for %s"), *(TabType.ToString()));
+	}
+}
+
+void FTabManager::DrawAttention( const TSharedRef<SDockTab>& TabToHighlight )
+{
+	// Bring the tab to front.
+	const TSharedPtr<SDockingArea> DockingArea = TabToHighlight->GetDockArea();
+	const TSharedRef<FTabManager> ManagerOfTabToHighlight = DockingArea->GetTabManager();
+	
+	if (ManagerOfTabToHighlight != FGlobalTabmanager::Get())
+	{
+		FGlobalTabmanager::Get()->DrawAttentionToTabManager( ManagerOfTabToHighlight );
+	}
+	
+	TSharedPtr<SWindow> OwnerWindow = DockingArea->GetParentWindow();
+	if (OwnerWindow.IsValid())
+	{
+		OwnerWindow->BringToFront();
+	}
+	TabToHighlight->GetParentDockTabStack()->BringToFront(TabToHighlight);
+	TabToHighlight->FlashTab();
+}
+
+
+void FTabManager::InsertNewDocumentTab( FName PlaceholderId, ESearchPreference::Type SearchPreference, const TSharedRef<SDockTab>& UnmanagedTab )
+{
+	InsertDocumentTab(PlaceholderId, SearchPreference, UnmanagedTab, true);
+}
+
+void FTabManager::RestoreDocumentTab( FName PlaceholderId, ESearchPreference::Type SearchPreference, const TSharedRef<SDockTab>& UnmanagedTab )
+{
+	InsertDocumentTab(PlaceholderId, SearchPreference, UnmanagedTab, false);
+}
+
+
+TSharedRef<SDockTab> FTabManager::InvokeTab( const FTabId& TabId )
+{
+	TSharedRef<SDockTab> NewTab = InvokeTab_Internal( TabId );
+	TSharedPtr<SWindow> ParentWindowPtr = NewTab->GetParentWindow();
+	if ((NewTab->GetTabRole() == ETabRole::MajorTab || NewTab->GetTabRole() == ETabRole::NomadTab) && ParentWindowPtr.IsValid() && ParentWindowPtr != FGlobalTabmanager::Get()->GetRootWindow())
+	{
+		ParentWindowPtr->SetTitle( NewTab->GetTabLabel() );
+	}
+	return NewTab;
+}
+
+TSharedRef<SDockTab> FTabManager::InvokeTab_Internal( const FTabId& TabId )
+{
+
+	// Tab Spawning Rules:
+	// 
+	//     * Find live instance --yes--> use it.
+	//         |no
+	//         v
+	//     * [non-Document only]
+	//       Find closed instance with matching TabId --yes--> restore it.	  
+	//         |no
+	//         v
+	//     * Find any tab of matching TabType (closed or open) --yes--> spawn next to it.
+	//         | no
+	//         v
+	//     * Is a nomad tab and we are NOT the global tab manager --yes--> try to invoke in the global tab manager
+	//         | no
+	//         v
+	//     * Spawn in a new window.
+
+	TSharedPtr<FTabSpawnerEntry> Spawner = FindTabSpawnerFor(TabId.TabType);
+
+	if ( !Spawner.IsValid() )
+	{
+		UE_LOG(LogTabManager, Warning, TEXT("Cannot spawn tab because no spawner is registered for '%s'"), *(TabId.ToString()));
+	}
+	else
+	{	
+		TSharedPtr<SDockTab> ExistingTab = Spawner->OnFindTabToReuse.IsBound()
+			? Spawner->OnFindTabToReuse.Execute( TabId )
+			: Spawner->SpawnedTabPtr.Pin();
+
+		if ( ExistingTab.IsValid() )
+		{
+			DrawAttention( ExistingTab.ToSharedRef() );
+			return ExistingTab.ToSharedRef();
+		}
+	}
+
+	// Tab is not live. Figure out where to spawn it.
+	TSharedPtr<SDockingTabStack> StackToSpawnIn = FindPotentiallyClosedTab( TabId );
+
+	if (StackToSpawnIn.IsValid())
+	{
+		const TSharedRef<SDockTab> NewTab = SpawnTab( TabId, TSharedPtr<SWindow>() );
+
+		StackToSpawnIn->OpenTab( NewTab );
+		NewTab->PlaySpawnAnim();
+		return NewTab;
+	}
+	else if ( FGlobalTabmanager::Get() != SharedThis(this) && NomadTabSpawner->Contains(TabId.TabType) )
+	{
+		// This tab could have been spawned in the global tab manager since it has a nomad tab spawner
+		return FGlobalTabmanager::Get()->InvokeTab_Internal(TabId);
+	}
+	else
+	{
+		// No layout info about this tab found; start 
+		TSharedRef<FArea> NewAreaForTab = FTabManager::NewArea(FTabManager::FallbackWindowSize)
+		->Split
+		(
+			FTabManager::NewStack()
+			->AddTab( TabId, ETabState::OpenedTab )
+		);
+
+		TSharedRef<SDockingArea> DockingArea = RestoreArea( NewAreaForTab, GetPrivateApi().GetParentWindow() );
+				
+		const TSharedPtr<SDockTab> NewlyOpenedTab = DockingArea->GetAllChildTabs()[0];
+		check(NewlyOpenedTab.IsValid());
+		return NewlyOpenedTab.ToSharedRef();
+	}
+}
+
+TSharedPtr<SDockingTabStack> FTabManager::FindPotentiallyClosedTab( const FTabId& ClosedTabId )
+{
+	TSharedPtr<SDockingTabStack> StackWithClosedTab;
+
+	FTabMatcher TabMatcher( ClosedTabId );
+
+	// Search among the COLLAPSED AREAS
+	const int32 CollapsedAreaWithMatchingTab = FindTabInCollapsedAreas( TabMatcher );
+	if ( CollapsedAreaWithMatchingTab != INDEX_NONE )
+	{
+		TSharedRef<SDockingArea> RestoredArea = RestoreArea( CollapsedDockAreas[CollapsedAreaWithMatchingTab], GetPrivateApi().GetParentWindow() );
+		// We have just un-collapsed this dock area
+		CollapsedDockAreas.RemoveAt(CollapsedAreaWithMatchingTab);
+		StackWithClosedTab = FindTabInLiveArea(TabMatcher, RestoredArea);
+	}
+
+	if ( !StackWithClosedTab.IsValid() )
+	{
+		// Search among the LIVE AREAS
+		StackWithClosedTab = FindTabInLiveAreas( TabMatcher );
+	}
+
+	return StackWithClosedTab;
+}
+
+void FTabManager::InvokeTabForMenu( FName TabId )
+{
+	InvokeTab(TabId);
+}
+
+void FTabManager::InsertDocumentTab( FName PlaceholderId, ESearchPreference::Type SearchPreference, const TSharedRef<SDockTab>& UnmanagedTab, bool bPlaySpawnAnim )
+{
+	const bool bTabNotManaged = ensure( ! FindTabInLiveAreas( FTabMatcher(UnmanagedTab->GetLayoutIdentifier()) ).IsValid() );
+	UnmanagedTab->SetLayoutIdentifier( FTabId(PlaceholderId, LastDocumentUID++) );
+	
+	if (bTabNotManaged)
+	{
+		TSharedPtr<SDockTab> LiveTab = (SearchPreference == ESearchPreference::PreferLiveTab) ? FindExistingLiveTab(FTabId(PlaceholderId)) : TSharedPtr<SDockTab>();
+		if (LiveTab.IsValid())
+		{
+			LiveTab->GetParent()->GetParentDockTabStack()->OpenTab( UnmanagedTab );
+		}
+		else
+		{
+			TSharedPtr<SDockingTabStack> StackToSpawnIn = FindPotentiallyClosedTab( PlaceholderId );
+			if( StackToSpawnIn.IsValid() == false )
+			{
+				UE_LOG(LogTabManager, Warning, TEXT("Unable to insert tab '%s'."), *(PlaceholderId.ToString()));
+				LiveTab = InvokeTab_Internal( FTabId( PlaceholderId ) );
+				LiveTab->GetParent()->GetParentDockTabStack()->OpenTab( UnmanagedTab );
+			}
+			else
+			{
+				StackToSpawnIn->OpenTab( UnmanagedTab );
+			}
+		}
+	}
+
+	DrawAttention(UnmanagedTab);
+	if (bPlaySpawnAnim)
+	{
+		UnmanagedTab->PlaySpawnAnim();
+	}
+	
+}
+
+FTabManager::FTabManager( const TSharedPtr<SDockTab>& InOwnerTab, const TSharedRef<FTabSpawner> & InNomadTabSpawner )
+: NomadTabSpawner(InNomadTabSpawner)
+, OwnerTabPtr( InOwnerTab )
+, PrivateApi( MakeShareable(new FPrivateApi(*this)) )
+, LastDocumentUID( 0 )
+, bIsSavingVisualState( false )
+{
+}
+
+
+TSharedRef<SDockingArea> FTabManager::RestoreArea( const TSharedRef<FArea>& AreaToRestore, const TSharedPtr<SWindow>& InParentWindow, const bool bEmbedTitleAreaContent  )
+{
+	TSharedRef<SDockingNode> RestoredNode = RestoreArea_Helper( AreaToRestore, InParentWindow, bEmbedTitleAreaContent );	
+	TSharedRef<SDockingArea> RestoredArea = StaticCastSharedRef<SDockingArea>( RestoredNode );	
+	RestoredArea->CleanUp(SDockingNode::TabRemoval_None);
+	return RestoredArea;
+}
+
+TSharedRef<SDockingNode> FTabManager::RestoreArea_Helper( const TSharedRef<FLayoutNode>& LayoutNode, const TSharedPtr<SWindow>& ParentWindow, const bool bEmbedTitleAreaContent )
+{
+	TSharedPtr<FTabManager::FStack> NodeAsStack = LayoutNode->AsStack();
+	TSharedPtr<FTabManager::FSplitter> NodeAsSplitter = LayoutNode->AsSplitter();
+	TSharedPtr<FTabManager::FArea> NodeAsArea = LayoutNode->AsArea();
+
+	if (NodeAsStack.IsValid())
+	{
+		TSharedPtr<SDockTab> WidgetToActivate;
+
+		TSharedRef<SDockingTabStack> NewStackWidget = SNew(SDockingTabStack, NodeAsStack.ToSharedRef());
+		NewStackWidget->SetSizeCoefficient(LayoutNode->GetSizeCoefficient());
+		for (int32 TabIndex=0; TabIndex < NodeAsStack->Tabs.Num(); ++TabIndex )
+		{
+			const FTab& SomeTab = NodeAsStack->Tabs[ TabIndex ];
+
+			if (SomeTab.TabState == ETabState::OpenedTab && IsValidTabForSpawning(SomeTab))
+			{
+				const TSharedRef<SDockTab> NewTabWidget = SpawnTab( SomeTab.TabId, ParentWindow );
+
+				if(SomeTab.TabId == NodeAsStack->ForegroundTabId)
+				{
+					WidgetToActivate = NewTabWidget;
+				}
+
+				NewStackWidget->AddTabWidget( NewTabWidget );
+			}
+		}
+		
+		if(WidgetToActivate.IsValid())
+		{
+			WidgetToActivate->ActivateInParent(ETabActivationCause::SetDirectly);
+
+			if ((WidgetToActivate->GetTabRole() == ETabRole::MajorTab || WidgetToActivate->GetTabRole() == ETabRole::NomadTab) && ParentWindow.IsValid() && ParentWindow != FGlobalTabmanager::Get()->GetRootWindow())
+			{
+				ParentWindow->SetTitle(WidgetToActivate->GetTabLabel());
+			}
+		}
+		
+		return NewStackWidget;
+
+	}
+	else if ( NodeAsArea.IsValid() )
+	{
+		const bool bSplitterIsDockArea = NodeAsArea.IsValid();
+		const bool bDockNeedsNewWindow = NodeAsArea.IsValid() && (NodeAsArea->WindowPlacement != FArea::Placement_NoWindow);
+
+		TSharedPtr<SDockingArea> NewDockAreaWidget;
+
+		if ( bDockNeedsNewWindow )
+		{
+			// The layout node we are restoring is a dock area.
+			// It needs a new window into which it will land.
+
+			const bool bIsChildWindow = ParentWindow.IsValid();
+
+			const bool bAutoPlacement = (NodeAsArea->WindowPlacement == FArea::Placement_Automatic);
+			TSharedRef<SWindow> NewWindow = (bAutoPlacement)
+				? SNew(SWindow)
+					.AutoCenter( EAutoCenter::PreferredWorkArea )
+					.ClientSize( NodeAsArea->WindowSize )
+					.CreateTitleBar( false )
+					.IsInitiallyMaximized( NodeAsArea->bIsMaximized )
+				: SNew(SWindow)
+					.AutoCenter( EAutoCenter::None )
+					.ScreenPosition( NodeAsArea->WindowPosition )
+					.ClientSize( NodeAsArea->WindowSize )
+					.CreateTitleBar( false )
+					.IsInitiallyMaximized( NodeAsArea->bIsMaximized );
+
+			// Any windows that were "pulled out" of a dock area should be children of the window in which the parent dock area resides.
+			if (bIsChildWindow)
+			{
+				FSlateApplication::Get().AddWindowAsNativeChild( NewWindow, ParentWindow.ToSharedRef() )->SetContent(
+					SAssignNew( NewDockAreaWidget, SDockingArea, SharedThis(this), NodeAsArea.ToSharedRef() ) .ParentWindow( NewWindow )
+				);
+			}
+			else
+			{
+				FSlateApplication::Get().AddWindow( NewWindow )->SetContent(
+					SAssignNew( NewDockAreaWidget, SDockingArea, SharedThis(this), NodeAsArea.ToSharedRef() ) .ParentWindow( NewWindow )
+				);
+			}
+
+			// Set a default title; restoring the splitter content may override this if it activates a tab
+			NewWindow->SetTitle( FGlobalTabmanager::Get()->GetApplicationTitle() );
+
+			RestoreSplitterContent( NodeAsArea.ToSharedRef(), NewDockAreaWidget.ToSharedRef(), NewWindow );
+		}
+		else
+		{
+			SAssignNew( NewDockAreaWidget, SDockingArea, SharedThis(this), NodeAsArea.ToSharedRef() )
+
+				// We only want to set a parent window on this dock area, if we need to have title area content
+				// embedded within it.  SDockingArea assumes that if it has a parent window set, then it needs to have
+				// title area content 
+				.ParentWindow( bEmbedTitleAreaContent ? ParentWindow : TSharedPtr<SWindow>() )
+
+				// Never manage these windows, even if a parent window is set.  The owner will take care of
+				// destroying these windows.
+				.ShouldManageParentWindow( false );
+
+			RestoreSplitterContent( NodeAsArea.ToSharedRef(), NewDockAreaWidget.ToSharedRef(), ParentWindow );
+		}
+		
+		
+
+		return NewDockAreaWidget.ToSharedRef();
+	}
+	else if ( NodeAsSplitter.IsValid() ) 
+	{
+		TSharedRef<SDockingSplitter> NewSplitterWidget = SNew( SDockingSplitter, NodeAsSplitter.ToSharedRef() );
+		NewSplitterWidget->SetSizeCoefficient(LayoutNode->GetSizeCoefficient());
+		RestoreSplitterContent( NodeAsSplitter.ToSharedRef(), NewSplitterWidget, ParentWindow );
+
+		return NewSplitterWidget;
+	}
+	else
+	{
+		ensureMsgf( false, TEXT("Unexpected node type") );
+
+		TSharedRef<SDockingTabStack> NewStackWidget = SNew(SDockingTabStack, FTabManager::NewStack());
+		NewStackWidget->OpenTab( SpawnTab( FName(NAME_None), ParentWindow ) );
+		return NewStackWidget;
+	}
+}
+
+
+
+void FTabManager::RestoreSplitterContent( const TSharedRef<FSplitter>& SplitterNode, const TSharedRef<SDockingSplitter>& SplitterWidget, const TSharedPtr<SWindow>& ParentWindow )
+{
+	// Restore the contents of this splitter.
+	for ( int32 ChildNodeIndex = 0; ChildNodeIndex < SplitterNode->ChildNodes.Num(); ++ChildNodeIndex )
+	{
+		TSharedRef<FLayoutNode> ThisChildNode = SplitterNode->ChildNodes[ChildNodeIndex];
+
+		const bool bEmbedTitleAreaContent = false;
+		TSharedRef<SDockingNode> ThisChildNodeWidget = RestoreArea_Helper( ThisChildNode, ParentWindow, bEmbedTitleAreaContent );
+		SplitterWidget->AddChildNode( ThisChildNodeWidget, INDEX_NONE );
+	}
+}
+
+
+bool FTabManager::IsValidTabForSpawning( const FTab& SomeTab ) const
+{
+	// Nomad tabs being restored from layouts should not be spawned if the nomad tab is already spawned.
+	TSharedRef<FTabSpawnerEntry>* NomadSpawner = NomadTabSpawner->Find( SomeTab.TabId.TabType );
+	return ( !NomadSpawner || !NomadSpawner->Get().IsSoleTabInstanceSpawned() );
+}
+
+
+TSharedRef<SDockTab> FTabManager::SpawnTab( const FTabId& TabId, const TSharedPtr<SWindow>& ParentWindow )
+{
+	TSharedPtr<SDockTab> NewTabWidget;
+
+	// Do we know how to spawn such a tab?
+	TSharedPtr<FTabSpawnerEntry> Spawner = FindTabSpawnerFor(TabId.TabType);
+	if ( Spawner.IsValid() )
+	{
+		NewTabWidget = Spawner->OnSpawnTab.Execute( FSpawnTabArgs( ParentWindow, TabId ) );
+		NewTabWidget->SetLayoutIdentifier( TabId );
+		NewTabWidget->ProvideDefaultLabel( Spawner->GetDisplayName().IsEmpty() ? FText::FromName( Spawner->TabType ) : Spawner->GetDisplayName() );
+		NewTabWidget->ProvideDefaultIcon( Spawner->GetIcon().GetIcon() );
+
+		// The spawner tracks that last tab it spawned
+		Spawner->SpawnedTabPtr = NewTabWidget;
+	}
+
+	if ( !NewTabWidget.IsValid() )
+	{
+		// We don't know how to spawn this tab.
+		// Make a dummy tab so that things aren't entirely broken.
+
+		NewTabWidget = SNew(SDockTab)
+			.Label( TabId.ToText() )
+			.ShouldAutosize( false )
+			[
+				SNew(SBox)
+				.HAlign(HAlign_Center)
+				.VAlign(VAlign_Center)
+				[
+					SNew(STextBlock)
+					.Text( NSLOCTEXT("TabManagement", "Unrecognized", "unrecognized tab") )
+				]
+			];
+
+		NewTabWidget->SetLayoutIdentifier(TabId);
+	}
+
+	NewTabWidget->SetTabManager( SharedThis(this) );
+
+	return NewTabWidget.ToSharedRef();
+}
+
+
+
+TSharedPtr<SDockTab> FTabManager::FindExistingLiveTab( const FTabId& TabId ) const
+{
+	for ( int32 AreaIndex = 0; AreaIndex < DockAreas.Num(); ++AreaIndex )
+	{
+		const TSharedPtr<SDockingArea> SomeDockArea = DockAreas[ AreaIndex ].Pin();
+		if ( SomeDockArea.IsValid() )
+		{
+			TArray< TSharedRef<SDockTab> > ChildTabs = SomeDockArea->GetAllChildTabs();
+			for (int32 ChildTabIndex=0; ChildTabIndex < ChildTabs.Num(); ++ChildTabIndex)
+			{
+				if ( TabId == ChildTabs[ChildTabIndex]->GetLayoutIdentifier() )
+				{
+					return ChildTabs[ChildTabIndex];
+				}
+			}
+		}
+	}
+
+	return TSharedPtr<SDockTab>();
+}
+
+
+TSharedPtr<class SDockingTabStack> FTabManager::FindTabInLiveAreas( const FTabMatcher& TabMatcher ) const
+{
+	for ( int32 AreaIndex = 0; AreaIndex < DockAreas.Num(); ++AreaIndex )
+	{
+		const TSharedPtr<SDockingArea> SomeDockArea = DockAreas[ AreaIndex ].Pin();
+		if (SomeDockArea.IsValid())
+		{
+			TSharedPtr<SDockingTabStack> TabFoundHere = FindTabInLiveArea(TabMatcher, SomeDockArea.ToSharedRef());
+			if ( TabFoundHere.IsValid() )
+			{
+				return TabFoundHere;
+			}
+		}
+	}
+
+	return TSharedPtr<SDockingTabStack>();
+}
+
+TSharedPtr<class SDockingTabStack> FTabManager::FindTabInLiveArea( const FTabMatcher& TabMatcher, const TSharedRef<SDockingArea>& InArea )
+{
+	TArray< TSharedRef<SDockingTabStack> > AllTabStacks;
+	GetAllStacks(InArea, AllTabStacks);
+
+	for (int32 StackIndex = 0; StackIndex < AllTabStacks.Num(); ++StackIndex)
+	{
+		if (AllTabStacks[StackIndex]->HasTab(TabMatcher))
+		{
+			return AllTabStacks[StackIndex];
+		}		
+	}
+
+	return TSharedPtr<SDockingTabStack>();
+}
+
+template<typename MatchFunctorType>
+bool FTabManager::HasAnyMatchingTabs( const TSharedRef<FTabManager::FLayoutNode>& SomeNode, const MatchFunctorType& Matcher )
+{
+	TSharedPtr<FTabManager::FSplitter> AsSplitter = SomeNode->AsSplitter();
+	TSharedPtr<FTabManager::FStack> AsStack = SomeNode->AsStack();
+
+	if ( AsStack.IsValid() )
+	{
+		return INDEX_NONE != AsStack->Tabs.FindMatch(Matcher);
+	}
+	else
+	{
+		ensure( AsSplitter.IsValid() );
+		// Do any of the child nodes have open tabs?
+		for (int32 ChildIndex=0; ChildIndex < AsSplitter->ChildNodes.Num(); ++ChildIndex)
+		{
+			if ( HasAnyMatchingTabs(AsSplitter->ChildNodes[ChildIndex], Matcher) )
+			{
+				return true;
+			}
+		}
+		return false;
+	}
+}
+
+bool FTabManager::HasOpenTabs( const TSharedRef<FTabManager::FLayoutNode>& SomeNode ) const
+{
+	// Search for valid and open tabs
+	struct OpenTabMatcher
+	{
+		const FTabManager* TabManager;
+
+		bool Matches( const FTab& Candidate ) const
+		{
+			return TabManager->IsValidTabForSpawning(Candidate) && Candidate.TabState == ETabState::OpenedTab;
+		}
+	};
+	OpenTabMatcher FindOpenTab;
+	FindOpenTab.TabManager = this;
+
+	return HasAnyMatchingTabs(SomeNode, FindOpenTab);
+}
+
+bool FTabManager::HasValidTabs( const TSharedRef<FTabManager::FLayoutNode>& SomeNode ) const
+{
+	// Search for valid tabs that can be spawned
+	struct ValidTabMatcher
+	{
+		const FTabManager* TabManager;
+
+		bool Matches( const FTab& Candidate ) const
+		{
+			return TabManager->IsValidTabForSpawning(Candidate);
+		}
+	};
+	ValidTabMatcher FindValidTab;
+	FindValidTab.TabManager = this;
+
+	return HasAnyMatchingTabs(SomeNode, FindValidTab);
+}
+
+void FTabManager::OnTabForegrounded( const TSharedPtr<SDockTab>& NewForegroundTab, const TSharedPtr<SDockTab>& BackgroundedTab )
+{
+	// Do nothing.
+}
+
+void FTabManager::OnTabRelocated( const TSharedRef<SDockTab>& RelocatedTab, const TSharedPtr<SWindow>& NewOwnerWindow )
+{
+	CleanupPointerArray(DockAreas);
+	RemoveTabFromCollapsedAreas( FTabMatcher( RelocatedTab->GetLayoutIdentifier() ) );
+	for (int32 DockAreaIndex=0; DockAreaIndex < DockAreas.Num(); ++DockAreaIndex)
+	{
+		DockAreas[DockAreaIndex].Pin()->OnTabFoundNewHome( RelocatedTab, NewOwnerWindow.ToSharedRef() );
+	}
+
+	UpdateStats();
+}
+
+void FTabManager::OnTabOpening( const TSharedRef<SDockTab>& TabBeingOpened )
+{
+	UpdateStats();
+}
+
+void FTabManager::OnTabClosing( const TSharedRef<SDockTab>& TabBeingClosed )
+{
+
+}
+
+void FTabManager::OnTabManagerClosing()
+{
+	CleanupPointerArray(DockAreas);
+
+	{
+		// Gather the persistent layout and allow a custom handler to persist it
+		const TSharedRef<FLayout> LayoutState = this->PersistLayout();
+		OnPersistLayout_Handler.ExecuteIfBound( LayoutState );
+	}
+
+	for (int32 DockAreaIndex=0; DockAreaIndex < DockAreas.Num(); ++DockAreaIndex)
+	{
+		TSharedRef<SDockingArea> ChildDockArea = DockAreas[DockAreaIndex].Pin().ToSharedRef();
+		TSharedPtr<SWindow> DockAreaWindow = ChildDockArea->GetParentWindow();
+		if (DockAreaWindow.IsValid())
+		{
+			DockAreaWindow->RequestDestroyWindow();
+		}
+	}
+}
+
+bool FTabManager::CanCloseManager( const TSet< TSharedRef<SDockTab> >& TabsToIgnore )
+{
+	CleanupPointerArray(DockAreas);
+
+	bool bCanCloseManager = true;
+
+	for (int32 DockAreaIndex=0; bCanCloseManager && DockAreaIndex < DockAreas.Num(); ++DockAreaIndex)
+	{
+		TSharedPtr<SDockingArea> SomeArea = DockAreas[DockAreaIndex].Pin();
+		TArray< TSharedRef<SDockTab> > AreasTabs = SomeArea.IsValid() ? SomeArea->GetAllChildTabs() : TArray< TSharedRef<SDockTab> >();
+		
+		for (int32 TabIndex=0; bCanCloseManager && TabIndex < AreasTabs.Num(); ++TabIndex)	
+		{
+			bCanCloseManager = TabsToIgnore.Contains( AreasTabs[TabIndex] ) || AreasTabs[TabIndex]->CanCloseTab();
+		}		
+	}
+
+	return bCanCloseManager;
+}
+
+void FTabManager::GetAllStacks( const TSharedRef<SDockingArea>& InDockArea, TArray< TSharedRef<SDockingTabStack> >& OutTabStacks )
+{
+	TArray< TSharedRef<SDockingNode> > AllNodes = InDockArea->GetChildNodesRecursively();
+	for (int32 NodeIndex=0; NodeIndex < AllNodes.Num(); ++NodeIndex)
+	{	
+		if ( AllNodes[NodeIndex]->GetNodeType() == SDockingNode::DockTabStack )
+		{
+			OutTabStacks.Add( StaticCastSharedRef<SDockingTabStack>( AllNodes[NodeIndex] ) );
+		}
+	}
+}
+
+TSharedPtr<FTabManager::FStack> FTabManager::FindTabUnderNode( const FTabMatcher& Matcher, const TSharedRef<FTabManager::FLayoutNode>& NodeToSearchUnder )
+{
+	TSharedPtr<FTabManager::FStack> NodeAsStack = NodeToSearchUnder->AsStack();
+	TSharedPtr<FTabManager::FSplitter> NodeAsSplitter = NodeToSearchUnder->AsSplitter();
+
+	if (NodeAsStack.IsValid())
+	{
+		const int32 TabIndex = NodeAsStack->Tabs.FindMatch( Matcher );
+		if (TabIndex != INDEX_NONE)
+		{
+			return NodeAsStack;
+		}
+		else
+		{
+			return TSharedPtr<FTabManager::FStack>();
+		}
+	}
+	else
+	{
+		ensure( NodeAsSplitter.IsValid() );
+		TSharedPtr<FTabManager::FStack> StackWithTab;
+		for ( int32 ChildIndex=0; !StackWithTab.IsValid() && ChildIndex < NodeAsSplitter->ChildNodes.Num(); ++ChildIndex)
+		{
+			StackWithTab = FindTabUnderNode( Matcher, NodeAsSplitter->ChildNodes[ChildIndex] );
+		}
+
+		return StackWithTab;
+	}
+}
+
+
+TSharedPtr<FTabSpawnerEntry> FTabManager::FindTabSpawnerFor(FName TabId)
+{
+	// Look for a spawner in this tab manager.
+	TSharedRef<FTabSpawnerEntry>* Spawner = TabSpawner.Find( TabId );
+	if (Spawner == NULL)
+	{
+		Spawner = NomadTabSpawner->Find( TabId );
+	}
+
+	return (Spawner != NULL)
+		? TSharedPtr<FTabSpawnerEntry>(*Spawner)
+		: TSharedPtr<FTabSpawnerEntry>();
+}
+
+
+int32 FTabManager::FindTabInCollapsedAreas( const FTabMatcher& Matcher )
+{
+	for ( int32 CollapsedDockAreaIndex=0; CollapsedDockAreaIndex < CollapsedDockAreas.Num(); ++CollapsedDockAreaIndex )
+	{
+		TSharedPtr<FTabManager::FStack> StackWithMatchingTab = FindTabUnderNode(Matcher, CollapsedDockAreas[CollapsedDockAreaIndex]);
+		if (StackWithMatchingTab.IsValid())
+		{
+			return CollapsedDockAreaIndex;
+		}
+	}	
+
+	return INDEX_NONE;
+}
+
+void FTabManager::RemoveTabFromCollapsedAreas( const FTabMatcher& Matcher )
+{
+	for ( int32 CollapsedDockAreaIndex=0; CollapsedDockAreaIndex < CollapsedDockAreas.Num(); ++CollapsedDockAreaIndex )
+	{
+		const TSharedRef<FTabManager::FArea>& DockArea = CollapsedDockAreas[CollapsedDockAreaIndex];
+
+		TSharedPtr<FTabManager::FStack> StackWithMatchingTab;
+		do
+		{
+			StackWithMatchingTab = FindTabUnderNode(Matcher, DockArea);
+
+			if (StackWithMatchingTab.IsValid())
+			{
+				const int32 TabIndex = StackWithMatchingTab->Tabs.FindMatch( Matcher );
+				if ( ensure(TabIndex != INDEX_NONE) )
+				{
+					StackWithMatchingTab->Tabs.RemoveAt(TabIndex);
+				}
+			}
+		}
+		while ( StackWithMatchingTab.IsValid() );
+	}	
+}
+
+void FTabManager::UpdateStats()
+{
+	StaticCastSharedRef<FTabManager>(FGlobalTabmanager::Get())->UpdateStats();
+}
+
+void FTabManager::GetRecordableStats( int32& OutTabCount, TArray<TSharedPtr<SWindow>>& OutUniqueParentWindows ) const
+{
+	OutTabCount = 0;
+	for (auto AreaIter = DockAreas.CreateConstIterator(); AreaIter; ++AreaIter)
+	{
+		TSharedPtr<SDockingArea> DockingArea = AreaIter->Pin();
+		if (DockingArea.IsValid())
+		{
+			TSharedPtr<SWindow> ParentWindow = DockingArea->GetParentWindow();
+			if (ParentWindow.IsValid())
+			{
+				OutUniqueParentWindows.AddUnique(ParentWindow);
+			}
+
+			TArray< TSharedRef<SDockingTabStack> > OutTabStacks;
+			GetAllStacks(DockingArea.ToSharedRef(), OutTabStacks);
+			for (auto StackIter = OutTabStacks.CreateConstIterator(); StackIter; ++StackIter)
+			{
+				OutTabCount += (*StackIter)->GetNumTabs();
+			}
+		}
+	}
+}
+
+const TSharedRef<FGlobalTabmanager>& FGlobalTabmanager::Get()
+{
+	static const TSharedRef<FGlobalTabmanager> Instance = FGlobalTabmanager::New();
+	// @todo: Never Destroy the Global Tab Manager because it has hooks into a bunch of different modules.
+	//        All those modules are unloaded first, so unbinding the delegates will cause a problem.
+	static const TSharedRef<FGlobalTabmanager>* NeverDestroyGlobalTabManager = new TSharedRef<FGlobalTabmanager>( Instance );
+	return Instance;
+}
+
+void FGlobalTabmanager::OnActiveTabChanged_Subscribe( const FOnActiveTabChanged::FDelegate& InDelegate )
+{
+	OnActiveTabChanged.Add( InDelegate );
+}
+
+
+
+void FGlobalTabmanager::OnActiveTabChanged_Unsubscribe( const FOnActiveTabChanged::FDelegate& InDelegate )
+{
+	OnActiveTabChanged.Remove( InDelegate );
+}
+
+
+
+TSharedPtr<class SDockTab> FGlobalTabmanager::GetActiveTab() const
+{
+	return ActiveTabPtr.Pin();
+}
+
+
+
+void FGlobalTabmanager::SetActiveTab( const TSharedPtr<class SDockTab>& NewActiveTab )
+{
+	const bool bShouldApplyChange =
+		// Setting NULL wipes out the active tab; always apply that change.
+		!NewActiveTab.IsValid() ||
+		// Major tabs are ignored for the purposes of active-tab tracking. We do not care about their 
+		(NewActiveTab->GetTabRole() != ETabRole::MajorTab && !NewActiveTab->IsNomadTabWithMajorTabStyle());
+	
+	TSharedPtr<SDockTab> CurrentlyActiveTab = GetActiveTab();
+
+	if (bShouldApplyChange && CurrentlyActiveTab != NewActiveTab)
+	{
+		OnActiveTabChanged.Broadcast( CurrentlyActiveTab, NewActiveTab );
+		ActiveTabPtr = NewActiveTab;
+	}	
+}
+
+FTabSpawnerEntry& FGlobalTabmanager::RegisterNomadTabSpawner( const FName TabId, const FOnSpawnTab& OnSpawnTab )
+{
+	ensure( !NomadTabSpawner->Contains(TabId) );
+	ensure( !IsLegacyTabType(TabId) );
+
+	TSharedRef<FTabSpawnerEntry> NewSpawnerEntry = MakeShareable( new FTabSpawnerEntry( TabId, OnSpawnTab ) );
+	NomadTabSpawner->Add( TabId, NewSpawnerEntry );
+	return NewSpawnerEntry.Get();
+}
+
+void FGlobalTabmanager::UnregisterNomadTabSpawner( const FName TabId )
+{
+	const int32 NumRemoved = NomadTabSpawner->Remove(TabId);
+}
+
+void FGlobalTabmanager::SetApplicationTitle( const FText& InAppTitle )
+{
+	AppTitle = InAppTitle;
+
+	for (int32 DockAreaIndex=0; DockAreaIndex < DockAreas.Num(); ++DockAreaIndex)
+	{
+		if (DockAreas[DockAreaIndex].IsValid())
+		{
+			TSharedPtr<SWindow> ParentWindow = DockAreas[DockAreaIndex].Pin()->GetParentWindow();
+			if (ParentWindow.IsValid() && ParentWindow == FGlobalTabmanager::Get()->GetRootWindow())
+			{
+				ParentWindow->SetTitle( AppTitle );
+			}
+		}
+	}
+}
+
+const FText& FGlobalTabmanager::GetApplicationTitle() const
+{
+	return AppTitle;
+}
+
+bool FGlobalTabmanager::CanCloseManager( const TSet< TSharedRef<SDockTab> >& TabsToIgnore )
+{
+	bool bCanCloseManager = FTabManager::CanCloseManager(TabsToIgnore);
+
+	for( int32 ManagerIndex=0; bCanCloseManager && ManagerIndex < SubTabManagers.Num(); ++ManagerIndex )
+	{
+		TSharedPtr<FTabManager> SubManager = SubTabManagers[ManagerIndex].TabManager.Pin();
+		if (SubManager.IsValid())
+		{
+			bCanCloseManager = SubManager->CanCloseManager(TabsToIgnore);
+		}
+	}
+
+	return bCanCloseManager;
+
+}
+
+void FGlobalTabmanager::DrawAttentionToTabManager( const TSharedRef<FTabManager>& ChildManager )
+{
+	const int32 MajorTabIndex = SubTabManagers.FindMatch( FindByManager(ChildManager) );
+	if (MajorTabIndex != INDEX_NONE)
+	{
+		this->DrawAttention( SubTabManagers[MajorTabIndex].MajorTab.Pin().ToSharedRef() );
+	}
+	
+}
+
+TSharedRef<FTabManager> FGlobalTabmanager::NewTabManager( const TSharedRef<SDockTab>& InOwnerTab )
+{
+	struct {
+		bool operator()(const FSubTabManager& InItem) const
+		{
+			return !InItem.MajorTab.IsValid();
+		}
+	} ShouldRemove;
+
+	SubTabManagers.RemoveAll( ShouldRemove );
+
+	const TSharedRef<FTabManager> NewTabManager = FTabManager::New( InOwnerTab, NomadTabSpawner );
+	SubTabManagers.Add( FSubTabManager(InOwnerTab, NewTabManager) );
+	UpdateStats();
+	return NewTabManager;
+}
+
+void FGlobalTabmanager::SaveAllVisualState()
+{
+	this->SavePersistentLayout();
+
+	for( int32 ManagerIndex=0; ManagerIndex < SubTabManagers.Num(); ++ManagerIndex )
+	{
+		const TSharedPtr<FTabManager> SubManagerTab = SubTabManagers[ManagerIndex].TabManager.Pin();
+		if (SubManagerTab.IsValid())
+		{
+			SubManagerTab->SavePersistentLayout();
+		}
+	}
+}
+
+void FGlobalTabmanager::SetRootWindow( const TSharedRef<SWindow> InRootWindow )
+{
+	RootWindowPtr = InRootWindow;
+}
+
+TSharedPtr<SWindow> FGlobalTabmanager::GetRootWindow() const
+{
+	return RootWindowPtr.Pin();
+}
+
+void FGlobalTabmanager::AddLegacyTabType(FName InLegacyTabType, FName InNewTabType)
+{
+	ensure(!TabSpawner.Contains(InLegacyTabType));
+	ensure(!NomadTabSpawner->Contains(InLegacyTabType));
+
+	LegacyTabTypeRedirectionMap.Add(InLegacyTabType, InNewTabType);
+}
+
+bool FGlobalTabmanager::IsLegacyTabType(FName InTabType) const
+{
+	return LegacyTabTypeRedirectionMap.Contains(InTabType);
+}
+
+FName FGlobalTabmanager::GetTabTypeForPotentiallyLegacyTab(FName InTabType) const
+{
+	const FName* NewTabType = LegacyTabTypeRedirectionMap.Find(InTabType);
+	return NewTabType ? *NewTabType : InTabType;
+}
+
+void FGlobalTabmanager::OnTabForegrounded( const TSharedPtr<SDockTab>& NewForegroundTab, const TSharedPtr<SDockTab>& BackgroundedTab )
+{
+	if (NewForegroundTab.IsValid())
+	{
+		// Show any child windows associated with the Major Tab that got foregrounded.
+		const int32 ForegroundedTabIndex = SubTabManagers.FindMatch( FindByTab(NewForegroundTab.ToSharedRef()) );
+		if (ForegroundedTabIndex != INDEX_NONE)
+		{
+			TSharedPtr<FTabManager> ForegroundTabManager = SubTabManagers[ForegroundedTabIndex].TabManager.Pin();
+			ForegroundTabManager->GetPrivateApi().ShowWindows();
+		}
+	}
+	
+	if ( BackgroundedTab.IsValid() )
+	{
+		// Hide any child windows associated with the Major Tab that got backgrounded.
+		const int32 BackgroundedTabIndex = SubTabManagers.FindMatch( FindByTab(BackgroundedTab.ToSharedRef()) );
+		if (BackgroundedTabIndex != INDEX_NONE)
+		{
+			TSharedPtr<FTabManager> BackgroundedTabManager = SubTabManagers[BackgroundedTabIndex].TabManager.Pin();
+			BackgroundedTabManager->GetPrivateApi().HideWindows();
+		}
+	}
+}
+
+void FGlobalTabmanager::OnTabRelocated( const TSharedRef<SDockTab>& RelocatedTab, const TSharedPtr<SWindow>& NewOwnerWindow )
+{
+	if (NewOwnerWindow.IsValid())
+	{
+		const int32 RelocatedManagerIndex = SubTabManagers.FindMatch( FindByTab(RelocatedTab) );
+		if (RelocatedManagerIndex != INDEX_NONE)
+		{
+			const TSharedRef<FTabManager>& RelocatedManager = SubTabManagers[RelocatedManagerIndex].TabManager.Pin().ToSharedRef();
+
+			// Reparent any DockAreas hanging out in a child window.
+			// We do not support native window re-parenting, so destroy old windows and re-create new ones in their place that are properly parented.
+			// Move the old DockAreas into new windows.
+			const TArray< TWeakPtr<SDockingArea> >& LiveDockAreas = RelocatedManager->GetPrivateApi().GetLiveDockAreas();
+			for (int32 DockAreaIndex=0; DockAreaIndex < LiveDockAreas.Num(); ++DockAreaIndex)
+			{
+				const TSharedRef<SDockingArea>& ChildDockArea = LiveDockAreas[ DockAreaIndex ].Pin().ToSharedRef();
+				const TSharedPtr<SWindow>& OldChildWindow = ChildDockArea->GetParentWindow();
+				if ( OldChildWindow.IsValid() )
+				{
+					TSharedRef<SWindow> NewChildWindow = SNew(SWindow)
+					.AutoCenter(EAutoCenter::None)
+					.ScreenPosition(OldChildWindow->GetPositionInScreen() + FVector2D(50,50))
+					.ClientSize(OldChildWindow->GetSizeInScreen())
+					.SupportsMinimize(false) .SupportsMaximize(false)
+					.CreateTitleBar(false)
+					[
+						ChildDockArea
+					];
+
+					ChildDockArea->SetParentWindow(NewChildWindow);
+
+					FSlateApplication::Get().AddWindowAsNativeChild( NewChildWindow, NewOwnerWindow.ToSharedRef() );
+
+					FSlateApplication::Get().RequestDestroyWindow( OldChildWindow.ToSharedRef() );
+				}
+			}
+		}
+	}
+
+	FTabManager::OnTabRelocated( RelocatedTab, NewOwnerWindow );
+}
+
+void FGlobalTabmanager::OnTabClosing( const TSharedRef<SDockTab>& TabBeingClosed )
+{
+	// Is this a major tab that contained a Sub TabManager?
+	// If so, need to properly close the sub tab manager
+	const int32 TabManagerBeingClosedIndex = SubTabManagers.FindMatch( FindByTab(TabBeingClosed) );
+	if (TabManagerBeingClosedIndex != INDEX_NONE)
+	{
+		const TSharedRef<FTabManager>& TabManagerBeingClosed = SubTabManagers[TabManagerBeingClosedIndex].TabManager.Pin().ToSharedRef();
+		TabManagerBeingClosed->GetPrivateApi().OnTabManagerClosing();
+	}
+}
+
+void FGlobalTabmanager::OnTabManagerClosing()
+{
+	for( int32 ManagerIndex=0; ManagerIndex < SubTabManagers.Num(); ++ManagerIndex )
+	{
+		TSharedPtr<SDockTab> SubManagerTab = SubTabManagers[ManagerIndex].MajorTab.Pin();
+		if (SubManagerTab.IsValid())
+		{
+			SubManagerTab->RemoveTabFromParent();
+		}
+		
+	}
+}
+
+void FGlobalTabmanager::UpdateStats()
+{
+	// Get all the tabs and windows in the global manager's own areas
+	int32 AllTabsCount = 0;
+	TArray<TSharedPtr<SWindow>> ParentWindows;
+
+	GetRecordableStats(AllTabsCount, ParentWindows);
+
+	// Add in all the tabs and windows in the sub-managers
+	for (auto ManagerIter = SubTabManagers.CreateConstIterator(); ManagerIter; ++ManagerIter)
+	{
+		if (ManagerIter->TabManager.IsValid())
+		{
+			int32 TabsCount = 0;
+			ManagerIter->TabManager.Pin()->GetRecordableStats(TabsCount, ParentWindows);
+
+			AllTabsCount += TabsCount;
+		}
+	}
+
+	// Keep a running maximum of the tab and window counts
+	AllTabsMaxCount = FMath::Max(AllTabsMaxCount, AllTabsCount);
+	AllAreasWindowMaxCount = FMath::Max(AllAreasWindowMaxCount, ParentWindows.Num());
+}
+
+#undef LOCTEXT_NAMESPACE
diff --git a/Engine/Source/Runtime/Slate/Private/Framework/Layout/InertialScrollManager.cpp b/Engine/Source/Runtime/Slate/Private/Framework/Layout/InertialScrollManager.cpp
new file mode 100644
--- /dev/null
+++ b/Engine/Source/Runtime/Slate/Private/Framework/Layout/InertialScrollManager.cpp
@@ -0,0 +1,67 @@
+// Copyright 1998-2014 Epic Games, Inc. All Rights Reserved.
+
+#include "SlatePrivatePCH.h"
+
+
+FInertialScrollManager::FInertialScrollManager(float InScrollDecceleration, double InSampleTimeout)
+	: ScrollVelocity(0.f)
+	, ScrollDecceleration(InScrollDecceleration)
+	, SampleTimeout(InSampleTimeout)
+{
+}
+
+void FInertialScrollManager::AddScrollSample(float Delta, double CurrentTime)
+{
+	new( this->ScrollSamples ) FScrollSample(Delta, CurrentTime);
+
+	float Total = 0;
+	double OldestTime = 0;
+	for ( int32 VelIdx = this->ScrollSamples.Num() - 1; VelIdx >= 0; --VelIdx )
+	{
+		const double SampleTime = this->ScrollSamples[VelIdx].Time;
+		const float SampleDelta = this->ScrollSamples[VelIdx].Delta;
+		if ( CurrentTime - SampleTime > SampleTimeout )
+		{
+			this->ScrollSamples.RemoveAt(VelIdx);
+		}
+		else
+		{
+			if ( SampleTime < OldestTime || OldestTime == 0 )
+			{
+				OldestTime = SampleTime;
+			}
+
+			Total += SampleDelta;
+		}
+	}
+
+	// Set the current velocity to the average of the previous recent samples
+	const double Duration = (OldestTime > 0) ? (CurrentTime - OldestTime) : 0;
+	if ( Duration > 0 )
+	{
+		this->ScrollVelocity = Total / Duration;
+	}
+	else
+	{
+		this->ScrollVelocity = 0;
+	}
+}	
+
+void FInertialScrollManager::UpdateScrollVelocity(const float InDeltaTime)
+{
+	if ( this->ScrollVelocity > 0 )
+	{
+		this->ScrollVelocity -= ScrollDecceleration * InDeltaTime;
+		this->ScrollVelocity = FMath::Max<float>(0.f, this->ScrollVelocity);
+	}
+	else if ( this->ScrollVelocity < 0 )
+	{
+		this->ScrollVelocity += ScrollDecceleration * InDeltaTime;
+		this->ScrollVelocity = FMath::Min<float>(0.f, this->ScrollVelocity);
+	}
+}
+
+void FInertialScrollManager::ClearScrollVelocity()
+{
+	this->ScrollVelocity = 0;
+}
\ No newline at end of file
diff --git a/Engine/Source/Runtime/Slate/Private/Framework/Layout/ScrollyZoomy.cpp b/Engine/Source/Runtime/Slate/Private/Framework/Layout/ScrollyZoomy.cpp
new file mode 100644
--- /dev/null
+++ b/Engine/Source/Runtime/Slate/Private/Framework/Layout/ScrollyZoomy.cpp
@@ -0,0 +1,216 @@
+// Copyright 1998-2014 Epic Games, Inc. All Rights Reserved.
+
+/*=============================================================================
+	ScrollyZoomy.h: Implements the FScrollyZoomy class.
+=============================================================================*/
+
+#include "SlatePrivatePCH.h"
+
+
+/* FScrollyZoomy structors
+ *****************************************************************************/
+
+FScrollyZoomy::FScrollyZoomy( )
+	: AmountScrolledWhileRightMouseDown( 0.0f ),
+	  bShowSoftwareCursor( false ),
+	  SoftwareCursorPosition( FVector2D::ZeroVector )
+{ }
+
+
+/* FScrollyZoomy interface
+ *****************************************************************************/
+
+void FScrollyZoomy::Tick( const float DeltaTime, IScrollableZoomable& ScrollableZoomable )
+{
+	if (!IsRightClickScrolling())
+	{
+		FVector2D ScrollBy = FVector2D::ZeroVector;
+
+		this->HorizontalIntertia.UpdateScrollVelocity( DeltaTime );
+		const float HorizontalScrollVelocity = this->HorizontalIntertia.GetScrollVelocity();
+
+		if (HorizontalScrollVelocity != 0.f)
+		{
+			ScrollBy.X = HorizontalScrollVelocity * DeltaTime;
+		}
+
+		this->VerticalIntertia.UpdateScrollVelocity( DeltaTime );
+		const float VerticalScrollVelocity = this->VerticalIntertia.GetScrollVelocity();
+
+		if (VerticalScrollVelocity != 0.f)
+		{
+			ScrollBy.Y = VerticalScrollVelocity * DeltaTime;
+		}
+
+		if (ScrollBy != FVector2D::ZeroVector)
+		{
+			ScrollableZoomable.ScrollBy(ScrollBy);
+		}
+	}
+}
+
+
+FReply FScrollyZoomy::OnMouseButtonDown( const FPointerEvent& MouseEvent )
+{
+	// @todo: Should we only clear on RMB down?  Also see other uses of this in the code base
+	this->HorizontalIntertia.ClearScrollVelocity();
+	this->VerticalIntertia.ClearScrollVelocity();
+
+	if (MouseEvent.GetEffectingButton() == EKeys::RightMouseButton)
+	{
+		AmountScrolledWhileRightMouseDown = 0;
+
+		// NOTE: We don't bother capturing the mouse, unless the user starts dragging a few pixels (see the
+		// mouse move handling here.)  This is important so that the item row has a chance to select
+		// items when the right mouse button is released.  Just keep in mind that you might not get
+		// an OnMouseButtonUp event for the right mouse button if the user moves off of the table before
+		// they reach our scroll threshold
+
+		return FReply::Handled();
+	}
+
+	return FReply::Unhandled();
+}
+
+
+FReply FScrollyZoomy::OnMouseButtonUp( const TSharedRef<SWidget> MyWidget, const FGeometry& MyGeometry, const FPointerEvent& MouseEvent )
+{
+	if (MouseEvent.GetEffectingButton() == EKeys::RightMouseButton)
+	{
+		AmountScrolledWhileRightMouseDown = 0;
+
+		FReply Reply = FReply::Handled().ReleaseMouseCapture();
+		bShowSoftwareCursor = false;
+
+		// If we have mouse capture, snap the mouse back to the closest location that is within the panel's bounds
+		if (MyWidget->HasMouseCapture())
+		{
+			FSlateRect PanelScreenSpaceRect = MyGeometry.GetClippingRect();
+			FVector2D CursorPosition = MyGeometry.LocalToAbsolute( SoftwareCursorPosition );
+
+			FIntPoint BestPositionInPanel(
+				FMath::Round( FMath::Clamp( CursorPosition.X, PanelScreenSpaceRect.Left, PanelScreenSpaceRect.Right ) ),
+				FMath::Round( FMath::Clamp( CursorPosition.Y, PanelScreenSpaceRect.Top, PanelScreenSpaceRect.Bottom ) )
+				);
+
+			Reply.SetMousePos( BestPositionInPanel );
+		}
+
+		return Reply;
+	}
+
+	return FReply::Unhandled();
+}
+
+
+FReply FScrollyZoomy::OnMouseMove( const TSharedRef<SWidget> MyWidget, IScrollableZoomable& ScrollableZoomable, const FGeometry& MyGeometry, const FPointerEvent& MouseEvent )
+{
+	if (MouseEvent.IsMouseButtonDown(EKeys::RightMouseButton))
+	{
+		// If scrolling with the right mouse button, we need to remember how much we scrolled.
+		// If we did not scroll at all, we will bring up the context menu when the mouse is released.
+		AmountScrolledWhileRightMouseDown += FMath::Abs( MouseEvent.GetCursorDelta().X ) + FMath::Abs( MouseEvent.GetCursorDelta().Y );
+
+		// Has the mouse moved far enough with the right mouse button held down to start capturing
+		// the mouse and dragging the view?
+		if (IsRightClickScrolling())
+		{
+			this->HorizontalIntertia.AddScrollSample( MouseEvent.GetCursorDelta().X, FPlatformTime::Seconds() );
+			this->VerticalIntertia.AddScrollSample( MouseEvent.GetCursorDelta().Y, FPlatformTime::Seconds() );
+			const bool DidScroll = ScrollableZoomable.ScrollBy( MouseEvent.GetCursorDelta() );
+
+			FReply Reply = FReply::Handled();
+
+			// Capture the mouse if we need to
+			if (FSlateApplication::Get().GetMouseCaptor() != MyWidget)
+			{
+				Reply.CaptureMouse( MyWidget ).UseHighPrecisionMouseMovement( MyWidget );
+				SoftwareCursorPosition = MyGeometry.AbsoluteToLocal( MouseEvent.GetScreenSpacePosition() );
+				bShowSoftwareCursor = true;
+			}
+
+			// Check if the mouse has moved.
+			if (DidScroll)
+			{
+				SoftwareCursorPosition += MouseEvent.GetCursorDelta();
+			}
+
+			return Reply;
+		}
+	}
+
+	return FReply::Unhandled();
+}
+
+
+void FScrollyZoomy::OnMouseLeave( const TSharedRef<SWidget> MyWidget, const FPointerEvent& MouseEvent )
+{
+	if (FSlateApplication::Get().GetMouseCaptor() != MyWidget)
+	{
+		// No longer scrolling (unless we have mouse capture)
+		AmountScrolledWhileRightMouseDown = 0;
+	}
+}
+
+
+FReply FScrollyZoomy::OnMouseWheel( const FPointerEvent& MouseEvent, IScrollableZoomable& ScrollableZoomable )
+{
+	// @todo: Inertial zoom support!
+	const bool DidZoom = ScrollableZoomable.ZoomBy(MouseEvent.GetWheelDelta());
+
+	if (DidZoom)
+	{
+		return FReply::Handled();
+	}
+
+	return FReply::Unhandled();
+}
+
+
+FCursorReply FScrollyZoomy::OnCursorQuery( ) const
+{
+	if (IsRightClickScrolling())
+	{
+		// We hide the native cursor as we'll be drawing the software EMouseCursor::GrabHandClosed cursor
+		return FCursorReply::Cursor(EMouseCursor::None);
+	}
+
+	return FCursorReply::Unhandled();
+}
+
+
+bool FScrollyZoomy::IsRightClickScrolling( ) const
+{
+	return (AmountScrolledWhileRightMouseDown >= SlatePanTriggerDistance);
+}
+
+
+bool FScrollyZoomy::NeedsSoftwareCursor( ) const
+{
+	return bShowSoftwareCursor;
+}
+
+
+const FVector2D& FScrollyZoomy::GetSoftwareCursorPosition( ) const
+{
+	return SoftwareCursorPosition;
+}
+
+
+int32 FScrollyZoomy::PaintSoftwareCursorIfNeeded( const FGeometry& AllottedGeometry, const FSlateRect& MyClippingRect, FSlateWindowElementList& OutDrawElements, int32 LayerId ) const
+{
+	if (bShowSoftwareCursor)
+	{
+		const FSlateBrush* Brush = FCoreStyle::Get().GetBrush(TEXT("SoftwareCursor_Grab"));
+
+		FSlateDrawElement::MakeBox(
+			OutDrawElements,
+			++LayerId,
+			AllottedGeometry.ToPaintGeometry(SoftwareCursorPosition - (Brush->ImageSize / 2), Brush->ImageSize),
+			Brush,
+			MyClippingRect
+		);
+	}
+
+	return LayerId;
+}
diff --git a/Engine/Source/Runtime/Slate/Private/Framework/MultiBox/MultiBox.cpp b/Engine/Source/Runtime/Slate/Private/Framework/MultiBox/MultiBox.cpp
new file mode 100644
--- /dev/null
+++ b/Engine/Source/Runtime/Slate/Private/Framework/MultiBox/MultiBox.cpp
@@ -0,0 +1,1032 @@
+// Copyright 1998-2014 Epic Games, Inc. All Rights Reserved.
+
+#include "SlatePrivatePCH.h"
+#include "MultiBox.h"
+#include "SButtonRowBlock.h"
+#include "SToolBarButtonBlock.h"
+#include "SMenuEntryBlock.h"
+#include "SWidgetBlock.h"
+#include "MultiBoxCustomization.h"
+#include "SClippingHorizontalBox.h"
+
+
+TAttribute<bool> FMultiBoxSettings::UseSmallToolBarIcons;
+TAttribute<bool> FMultiBoxSettings::DisplayMultiboxHooks;
+FMultiBoxSettings::FConstructToolTip FMultiBoxSettings::ToolTipConstructor = FConstructToolTip::CreateStatic( &FMultiBoxSettings::ConstructDefaultToolTip );
+
+bool FMultiBoxSettings::bInToolbarEditMode;
+
+
+FMultiBoxSettings::FMultiBoxSettings()
+{
+	ResetToolTipConstructor();
+}
+
+TSharedRef< SToolTip > FMultiBoxSettings::ConstructDefaultToolTip( const TAttribute<FText>& ToolTipText, const TSharedPtr<SWidget>& OverrideContent, const TSharedPtr<const FUICommandInfo>& Action )
+{
+	if ( OverrideContent.IsValid() )
+	{
+		return SNew( SToolTip )
+		[
+			OverrideContent.ToSharedRef()
+		];
+	}
+
+	return SNew( SToolTip ).Text( ToolTipText );
+}
+
+void FMultiBoxSettings::ResetToolTipConstructor()
+{
+	ToolTipConstructor = FConstructToolTip::CreateStatic( &FMultiBoxSettings::ConstructDefaultToolTip );
+}
+
+void FMultiBoxSettings::ToggleToolbarEditing()
+{
+	bool bCanBeEnabled = false;
+	if( GIsEditor )
+	{
+		GConfig->GetBool(TEXT("/Script/UnrealEd.EditorExperimentalSettings"), TEXT("bToolbarCustomization"), bCanBeEnabled, GEditorUserSettingsIni);
+	}
+
+	bInToolbarEditMode = !bInToolbarEditMode && bCanBeEnabled;
+}
+
+const FMultiBoxCustomization FMultiBoxCustomization::None( NAME_None );
+
+
+void SMultiBlockBaseWidget::OnDragEnter( const FGeometry& MyGeometry, const FDragDropEvent& DragDropEvent )
+{
+	if ( DragDropEvent.GetOperationAs<FUICommandDragDropOp>().IsValid() )
+	{
+		OwnerMultiBoxWidget.Pin()->OnCustomCommandDragEnter( MultiBlock.ToSharedRef(), MyGeometry, DragDropEvent );
+	}
+}
+
+FReply SMultiBlockBaseWidget::OnDragOver( const FGeometry& MyGeometry, const FDragDropEvent& DragDropEvent )
+{
+	if ( DragDropEvent.GetOperationAs<FUICommandDragDropOp>().IsValid() )
+	{
+		OwnerMultiBoxWidget.Pin()->OnCustomCommandDragged( MultiBlock.ToSharedRef(), MyGeometry, DragDropEvent );
+		return FReply::Handled();
+	}
+
+	return FReply::Unhandled();
+}
+
+FReply SMultiBlockBaseWidget::OnDrop( const FGeometry& MyGeometry, const FDragDropEvent& DragDropEvent )
+{
+	if ( DragDropEvent.GetOperationAs<FUICommandDragDropOp>().IsValid() )
+	{
+		OwnerMultiBoxWidget.Pin()->OnCustomCommandDropped();
+		return FReply::Handled();
+	}
+
+	return FReply::Unhandled();
+}
+
+/**
+ * Creates a MultiBlock widget for this MultiBlock
+ *
+ * @param	InOwnerMultiBoxWidget	The widget that will own the new MultiBlock widget
+ * @param	InLocation				The location information for the MultiBlock widget
+ *
+ * @return  MultiBlock widget object
+ */
+TSharedRef< IMultiBlockBaseWidget > FMultiBlock::MakeWidget( TSharedRef< SMultiBoxWidget > InOwnerMultiBoxWidget, EMultiBlockLocation::Type InLocation ) const
+{
+	TSharedRef< IMultiBlockBaseWidget > NewMultiBlockWidget = ConstructWidget();
+
+	// Tell the widget about its parent MultiBox widget
+	NewMultiBlockWidget->SetOwnerMultiBoxWidget( InOwnerMultiBoxWidget );
+
+	// Assign ourselves to the MultiBlock widget
+	NewMultiBlockWidget->SetMultiBlock( AsShared() );
+
+	// Pass location information to widget.
+	NewMultiBlockWidget->SetMultiBlockLocation(InLocation);
+
+	// Work out what style the widget should be using
+	const ISlateStyle* const StyleSet = InOwnerMultiBoxWidget->GetStyleSet();
+	const FName& StyleName = InOwnerMultiBoxWidget->GetStyleName();
+
+	// Build up the widget
+	NewMultiBlockWidget->BuildMultiBlockWidget(StyleSet, StyleName);
+
+	return NewMultiBlockWidget;
+}
+
+/**
+ * Constructor
+ *
+ * @param	InType	Type of MultiBox
+ * @param	bInShouldCloseWindowAfterMenuSelection	Sets whether or not the window that contains this multibox should be destroyed after the user clicks on a menu item in this box
+ */
+FMultiBox::FMultiBox( const EMultiBoxType::Type InType, FMultiBoxCustomization InCustomization, const bool bInShouldCloseWindowAfterMenuSelection )
+	: CustomizationData( new FMultiBoxCustomizationData( InCustomization.GetCustomizationName() ) )
+	, CommandLists()
+	, Blocks()
+	, StyleSet( &FCoreStyle::Get() )
+	, StyleName( "ToolBar" )
+	, Type( InType )
+	, bShouldCloseWindowAfterMenuSelection( bInShouldCloseWindowAfterMenuSelection )
+{
+
+}
+
+FMultiBox::~FMultiBox()
+{
+}
+
+TSharedRef<FMultiBox> FMultiBox::Create( const EMultiBoxType::Type InType, FMultiBoxCustomization InCustomization, const bool bInShouldCloseWindowAfterMenuSelection )
+{
+	TSharedRef<FMultiBox> NewBox = MakeShareable( new FMultiBox( InType, InCustomization, bInShouldCloseWindowAfterMenuSelection ) );
+
+	return NewBox;
+}
+
+/**
+ * Adds a MultiBlock to this MultiBox, to the end of the list
+ */
+void FMultiBox::AddMultiBlock( TSharedRef< const FMultiBlock > InBlock )
+{
+#if UE_BUILD_DEBUG
+	check( !Blocks.Contains( InBlock ) );
+#endif
+
+	if( InBlock->GetActionList().IsValid() )
+	{
+		CommandLists.AddUnique( InBlock->GetActionList() );
+	}
+
+	Blocks.Add( InBlock );
+}
+
+void FMultiBox::RemoveCustomMultiBlock( TSharedRef< const FMultiBlock> InBlock )
+{
+	if( IsCustomizable() )
+	{
+		int32 Index = Blocks.Find( InBlock );
+
+		// Remove the block from the visual list
+		if( Index != INDEX_NONE )
+		{
+			Blocks.RemoveAt( Index );
+
+			// Remove the block from the customization data
+			CustomizationData->BlockRemoved( InBlock, Index, Blocks );
+		}
+
+	}
+}
+
+void FMultiBox::InsertCustomMultiBlock( TSharedRef<const FMultiBlock> InBlock, int32 Index )
+{
+	if( IsCustomizable() && ensure( InBlock->GetAction().IsValid() ) )
+	{
+		int32 ExistingIndex = Blocks.Find( InBlock );
+		if( ExistingIndex != INDEX_NONE )
+		{
+			Blocks.RemoveAt( ExistingIndex );
+
+			CustomizationData->BlockRemoved( InBlock, ExistingIndex, Blocks );
+
+			if( ExistingIndex < Index )
+			{
+				--Index;
+			}
+		}
+
+		Blocks.Insert( InBlock, Index );
+
+		CustomizationData->BlockAdded( InBlock, Index, Blocks );
+	}
+}
+
+/**
+ * Creates a MultiBox widget for this MultiBox
+ *
+ * @return  MultiBox widget object
+ */
+TSharedRef< SMultiBoxWidget > FMultiBox::MakeWidget()
+{	
+	ApplyCustomizedBlocks();
+
+	TSharedRef< SMultiBoxWidget > NewMultiBoxWidget =
+		SNew( SMultiBoxWidget );
+
+	// Assign ourselves to the MultiBox widget
+	NewMultiBoxWidget->SetMultiBox( AsShared() );
+
+	// Build up the widget
+	NewMultiBoxWidget->BuildMultiBoxWidget();
+
+	return NewMultiBoxWidget;
+}
+
+bool FMultiBox::IsCustomizable() const
+{
+	bool bIsCustomizable = false;
+	if( CustomizationData->GetCustomizationName() != NAME_None )
+	{
+		GConfig->GetBool(TEXT("/Script/UnrealEd.EditorExperimentalSettings"), TEXT("bToolbarCustomization"), bIsCustomizable, GEditorUserSettingsIni);
+	}
+
+	return bIsCustomizable;
+}
+
+void FMultiBox::ApplyCustomizedBlocks()
+{
+	if( IsCustomizable() )
+	{
+		CustomizationData->LoadCustomizedBlocks();
+
+		// Build a map of commands to existing blocks,  we'll try to use existing blocks before creating new ones
+		TMap< TSharedPtr<const FUICommandInfo>, TSharedPtr<const FMultiBlock> > CommandToBlockMap;
+
+		for( int32 BlockIndex = 0; BlockIndex < Blocks.Num(); ++BlockIndex )
+		{
+			TSharedPtr<const FMultiBlock> Block = Blocks[BlockIndex];
+			if( Block->GetAction().IsValid() )
+			{
+				CommandToBlockMap.Add( Block->GetAction(), Block );
+			}
+		}
+
+		// Rebuild the users customized box by executing the transactions the user made to get the 
+		// box to its customized state
+		for( uint32 TransIndex = 0; TransIndex < CustomizationData->GetNumTransactions(); ++TransIndex )
+		{
+			const FCustomBlockTransaction& Transaction = CustomizationData->GetTransaction(TransIndex);
+		
+			// Try and find the block in the default map;
+			TSharedPtr<const FMultiBlock> Block = CommandToBlockMap.FindRef( Transaction.Command.Pin() );
+
+			if( Transaction.TransactionType == FCustomBlockTransaction::Add )
+			{
+			
+				if( !Block.IsValid() )
+				{
+					Block = MakeMultiBlockFromCommand( Transaction.Command.Pin(), false );
+				}
+
+				if( Block.IsValid() )
+				{
+					Blocks.Insert( Block.ToSharedRef(), FMath::Clamp( Transaction.BlockIndex, 0, Blocks.Num() ) );
+				}
+			}
+			else
+			{
+				if( Block.IsValid() )
+				{
+					Blocks.Remove( Block.ToSharedRef() );
+				}
+			}
+		}
+	}
+}
+
+FName FMultiBox::GetCustomizationName() const
+{
+	return CustomizationData->GetCustomizationName(); 
+}
+
+TSharedPtr<FMultiBlock> FMultiBox::MakeMultiBlockFromCommand( TSharedPtr<const FUICommandInfo> CommandInfo, bool bCommandMustBeBound ) const
+{
+	TSharedPtr<FMultiBlock> NewBlock;
+
+	// Find the command list that processes this command
+	TSharedPtr<const FUICommandList> CommandList;
+
+	for (int32 CommandListIndex = 0; CommandListIndex < CommandLists.Num(); ++CommandListIndex )
+	{
+		TSharedPtr<const FUICommandList> TestCommandList = CommandLists[CommandListIndex];
+		if( TestCommandList->GetActionForCommand( CommandInfo.ToSharedRef() ) != NULL )
+		{
+			CommandList = TestCommandList;
+			break;
+		}
+	}
+
+	
+	if( !bCommandMustBeBound && !CommandList.IsValid() && CommandLists.Num() > 0 )
+	{
+		// The first command list is the main command list and other are commandlists added from extension points
+		// Use the main command list if one was not found
+		CommandList = CommandLists[0];
+	}
+
+	if( !bCommandMustBeBound || CommandList.IsValid() )
+	{
+		// Only toolbars and menu buttons are supported currently
+		switch ( Type )
+		{
+		case EMultiBoxType::ToolBar:
+			{
+				NewBlock = MakeShareable( new FToolBarButtonBlock( CommandInfo, CommandList ) );
+			}
+			break;
+		case EMultiBoxType::Menu:
+			{
+				NewBlock = MakeShareable( new FMenuEntryBlock( NAME_None, CommandInfo, CommandList ) );
+			}
+			break;
+		}
+	}
+
+	return NewBlock;
+
+}
+
+TSharedPtr<const FMultiBlock> FMultiBox::FindBlockFromCommand( TSharedPtr<const FUICommandInfo> Command ) const
+{
+	for (TArray< TSharedRef< const FMultiBlock > >::TConstIterator It(Blocks); It; ++It)
+	{
+		const TSharedRef< const FMultiBlock >& Block = *It;
+		if( Block->GetAction() == Command )
+		{
+			return Block;
+		}
+	}
+
+	return NULL;
+}
+
+
+void SMultiBoxWidget::Construct( const FArguments& InArgs )
+{
+	ContentScale = InArgs._ContentScale;
+}
+
+TSharedRef<ITableRow> SMultiBoxWidget::GenerateTiles(TSharedPtr<SWidget> Item, const TSharedRef<STableViewBase>& OwnerTable)
+{
+	return SNew(STableRow< TSharedPtr<SWidget> >, OwnerTable)
+		[
+			Item.ToSharedRef()
+		];
+}
+
+float SMultiBoxWidget::GetItemWidth() const
+{
+	float MaxWidth = 0;
+	for (int32 i = 0; i < TileViewWidgets.Num(); ++i)
+	{
+		MaxWidth = FMath::Max(TileViewWidgets[i]->GetDesiredSize().X, MaxWidth);
+	}
+	return MaxWidth;
+}
+
+float SMultiBoxWidget::GetItemHeight() const
+{
+	float MaxHeight = 0;
+	for (int32 i = 0; i < TileViewWidgets.Num(); ++i)
+	{
+		MaxHeight = FMath::Max(TileViewWidgets[i]->GetDesiredSize().Y, MaxHeight);
+	}
+	return MaxHeight;
+}
+
+bool SMultiBoxWidget::IsBlockBeingDragged( TSharedPtr<const FMultiBlock> Block ) const
+{
+	if( DragPreview.PreviewBlock.IsValid() )
+	{
+		return DragPreview.PreviewBlock->GetActualBlock() == Block;
+	}
+
+	return false;
+}
+
+void SMultiBoxWidget::AddBlockWidget( const FMultiBlock& Block, TSharedPtr<SHorizontalBox> HorizontalBox, TSharedPtr<SVerticalBox> VerticalBox, EMultiBlockLocation::Type InLocation )
+{
+	check( MultiBox.IsValid() );
+
+	bool bDisplayExtensionHooks = FMultiBoxSettings::DisplayMultiboxHooks.Get() && Block.GetExtensionHook() != NAME_None;
+
+	TSharedRef<SWidget> BlockWidget = Block.MakeWidget( SharedThis( this ), InLocation)->AsWidget();
+
+	TWeakPtr<SWidget> BlockWidgetWeakPtr = BlockWidget;
+	TWeakPtr<const FMultiBlock> BlockWeakPtr = Block.AsShared();
+
+	const ISlateStyle* const StyleSet = MultiBox->GetStyleSet();
+
+	TSharedRef<SWidget> FinalWidget =
+	SNew( SOverlay )
+	+ SOverlay::Slot()
+	[
+		BlockWidget
+	]
+	+ SOverlay::Slot()
+	[
+		// This overlay prevents users from clicking on the actual block when in edit mode and also allows new blocks
+		// to be dropped on disabled blocks
+		SNew( SMultiBlockDragHandle, SharedThis( this ), Block.AsShared(), MultiBox->GetCustomizationName() )
+		.Visibility( this, &SMultiBoxWidget::GetCustomizationVisibility, BlockWeakPtr, BlockWidgetWeakPtr )
+	]
+	+ SOverlay::Slot()
+	.HAlign( HAlign_Right )
+	.VAlign( VAlign_Top )
+	.Padding( FMargin(0.0f, 2.0f, 1.0f, 0.0f ) )
+	[
+		// The delete button for removing blocks is only visible when in edit mode
+		SNew( SButton )
+		.Visibility( this, &SMultiBoxWidget::GetCustomizationVisibility, BlockWeakPtr, BlockWidgetWeakPtr )
+		.ContentPadding(0)
+		.OnClicked( this, &SMultiBoxWidget::OnDeleteBlockClicked, BlockWeakPtr )
+		.ButtonStyle( StyleSet, "MultiBox.DeleteButton" )
+	];
+
+	switch (MultiBox->GetType())
+	{
+	case EMultiBoxType::MenuBar:
+	case EMultiBoxType::ToolBar:
+		{
+			HorizontalBox->AddSlot()
+			.AutoWidth()
+			.Padding( 0 )
+			[
+				SNew(SVerticalBox)
+				+SVerticalBox::Slot()
+				.HAlign(HAlign_Center)
+				.AutoHeight()
+				[
+					SNew(STextBlock)
+					.Visibility(bDisplayExtensionHooks ? EVisibility::Visible : EVisibility::Collapsed)
+					.ColorAndOpacity(StyleSet->GetColor("MultiboxHookColor"))
+					.Text(Block.GetExtensionHook().ToString())
+				]
+				+SVerticalBox::Slot()
+				[
+					FinalWidget
+				]
+			];
+		}
+		break;
+	case EMultiBoxType::VerticalToolBar:
+		{
+			VerticalBox->AddSlot()
+				.AutoHeight()
+				.Padding( 0.0f, 1.0f, 0.0f, 1.0f )
+				[
+					SNew(SVerticalBox)
+					+SVerticalBox::Slot()
+					.HAlign(HAlign_Center)
+					.AutoHeight()
+					[
+						SNew(STextBlock)
+						.Visibility(bDisplayExtensionHooks ? EVisibility::Visible : EVisibility::Collapsed)
+						.ColorAndOpacity(StyleSet->GetColor("MultiboxHookColor"))
+						.Text(Block.GetExtensionHook().ToString())
+					]
+					+SVerticalBox::Slot()
+					[
+						FinalWidget
+					]
+				];
+		}
+		break;
+	case EMultiBoxType::ButtonRow:
+		{
+			TileViewWidgets.Add( FinalWidget );
+		}
+		break;
+	case EMultiBoxType::Menu:
+		{
+			VerticalBox->AddSlot()
+			.AutoHeight()
+			.Padding( 1.0f, 0.0f, 1.0f, 0.0f )
+			[
+				SNew(SHorizontalBox)
+				+SHorizontalBox::Slot()
+				.AutoWidth()
+				.VAlign(VAlign_Center)
+				[
+					SNew(STextBlock)
+					.Visibility(bDisplayExtensionHooks ? EVisibility::Visible : EVisibility::Collapsed)
+					.ColorAndOpacity(StyleSet->GetColor("MultiboxHookColor"))
+					.Text(Block.GetExtensionHook().ToString())
+				]
+				+SHorizontalBox::Slot()
+				[
+					FinalWidget
+				]
+			];
+		}
+		break;
+	}
+}
+
+
+/**
+ * Builds this MultiBox widget up from the MultiBox associated with it
+ */
+void SMultiBoxWidget::BuildMultiBoxWidget()
+{
+	check( MultiBox.IsValid() );
+
+	// Grab the list of blocks, early out if there's nothing to fill the widget with
+	const TArray< TSharedRef< const FMultiBlock > >& Blocks = MultiBox->GetBlocks();
+	if ( Blocks.Num() == 0 )
+	{
+		return;
+	}
+
+	// Select background brush based on the type of multibox.
+	const ISlateStyle* const StyleSet = MultiBox->GetStyleSet();
+	const FName& StyleName = MultiBox->GetStyleName();
+	const FSlateBrush* BackgroundBrush = StyleSet->GetBrush( StyleName, ".Background" );
+
+	// Create a box panel that the various multiblocks will resides within
+	// @todo Slate MultiBox: Expose margins and other useful bits
+	TSharedPtr< SHorizontalBox > HorizontalBox;
+	TSharedPtr< SVerticalBox > VerticalBox;
+	TSharedPtr< SWidget > MainWidget;
+
+	/** The current row of buttons for if the multibox type is a button row */
+	TSharedPtr<SHorizontalBox> ButtonRow;
+
+	TSharedPtr< STileView< TSharedPtr<SWidget> > > TileView;
+
+	switch (MultiBox->GetType())
+	{
+	case EMultiBoxType::MenuBar:
+	case EMultiBoxType::ToolBar:
+		{
+				MainWidget = HorizontalBox = ClippedHorizontalBox = SNew( SClippingHorizontalBox )
+					.BackgroundBrush(BackgroundBrush)
+					.OnWrapButtonClicked(FOnGetContent::CreateSP(this, &SMultiBoxWidget::OnWrapButtonClicked))
+					.StyleSet( StyleSet )
+					.StyleName( StyleName );
+		}
+		break;
+	case EMultiBoxType::VerticalToolBar:
+		{
+			MainWidget = VerticalBox = SNew( SVerticalBox );
+		}
+		break;
+	case EMultiBoxType::ButtonRow:
+		{
+			MainWidget = TileView = SNew(STileView< TSharedPtr<SWidget> >)
+				.OnGenerateTile(this, &SMultiBoxWidget::GenerateTiles)
+				.ListItemsSource(&TileViewWidgets)
+				.ItemWidth(this, &SMultiBoxWidget::GetItemWidth)
+				.ItemHeight(this, &SMultiBoxWidget::GetItemHeight)
+				.SelectionMode(ESelectionMode::None);
+		}
+		break;
+	case EMultiBoxType::Menu:
+		{
+			// wrap menu content in a scrollbox to support vertical scrolling if needed
+			MainWidget = SNew(SScrollBox)
+				+SScrollBox::Slot()
+				[
+					SAssignNew( VerticalBox, SVerticalBox )
+				];
+		}
+		break;
+	}
+	
+	bool bInsideGroup = false;
+
+	// Start building up the actual UI from each block in this MultiBox
+	bool bIsNextToStartBlock = false;
+	bool bIsNextToEndBlock = false;
+	for( int32 Index = 0; Index < Blocks.Num(); Index++ )
+	{
+		const FMultiBlock& Block = *Blocks[Index];
+		EMultiBlockLocation::Type Location = EMultiBlockLocation::None;
+
+		// Determine the location of the current block, used for group styling information
+		{
+			// Check if we are a start or end block
+			if (Block.IsGroupStartBlock())
+			{
+				bInsideGroup = true;
+			}
+			else if (Block.IsGroupEndBlock())
+			{
+				bInsideGroup = false;
+			}
+
+			// Check if we are next to a start or end block
+			bIsNextToStartBlock = false;
+			bIsNextToEndBlock = false;
+			if (Index + 1 < Blocks.Num())
+			{
+				const FMultiBlock& NextBlock = *Blocks[Index + 1];
+				if ( NextBlock.IsGroupEndBlock() )
+				{
+					bIsNextToEndBlock = true;
+				}
+			}
+			if (Index > 0)
+			{
+				const FMultiBlock& PrevBlock = *Blocks[Index - 1];
+				if ( PrevBlock.IsGroupStartBlock() )
+				{
+					bIsNextToStartBlock = true;
+				}
+			}
+
+			// determine location
+			if (bInsideGroup)
+			{
+				// assume we are in the middle of a group
+				Location = EMultiBlockLocation::Middle;
+
+				// We are the start of a group
+				if (bIsNextToStartBlock && !bIsNextToEndBlock)
+				{
+					Location = EMultiBlockLocation::Start;
+				}
+				// we are the end of a group
+				else if (!bIsNextToStartBlock && bIsNextToEndBlock)
+				{
+					Location = EMultiBlockLocation::End;
+				}
+				// we are the only block in a group
+				else if (bIsNextToStartBlock && bIsNextToEndBlock)
+				{
+					Location = EMultiBlockLocation::None;
+				}
+			}
+		}
+	
+
+		if( DragPreview.IsValid() && DragPreview.InsertIndex == Index )
+		{
+			// Add the drag preview before if we have it. This block shows where the custom block will be 
+			// added if the user drops it
+			AddBlockWidget( *DragPreview.PreviewBlock, HorizontalBox, VerticalBox, EMultiBlockLocation::None );
+		}
+		
+		// Do not add a block if it is being dragged
+		if( !IsBlockBeingDragged( Blocks[Index] ) )
+		{
+			AddBlockWidget( Block, HorizontalBox, VerticalBox, Location );
+		}
+	}
+
+	// Add the wrap button as the final block
+	if (ClippedHorizontalBox.IsValid())
+	{
+		ClippedHorizontalBox->AddWrapButton();
+	}
+
+	// Setup the root border widget
+	TSharedPtr< SBorder > RootBorder;
+	switch (MultiBox->GetType())
+	{
+	case EMultiBoxType::MenuBar:
+	case EMultiBoxType::ToolBar:
+		{
+			RootBorder =
+				SNew( SBorder )
+				.Padding(0)
+				.BorderImage( FCoreStyle::Get().GetBrush("NoBorder") )
+				// Assign the box panel as the child
+				[
+					MainWidget.ToSharedRef()
+				];
+		}
+		break;
+	default:
+		{
+			RootBorder =
+				SNew( SBorder )
+				.Padding(0)
+				.BorderImage( BackgroundBrush )
+				.ForegroundColor( FCoreStyle::Get().GetSlateColor("DefaultForeground") )
+				// Assign the box panel as the child
+				[
+					MainWidget.ToSharedRef()
+				];
+		}
+		break;
+	}
+
+	// Prevent tool-tips spawned by child widgets from drawing on top of our main widget
+	RootBorder->EnableToolTipForceField( true );
+
+	ChildSlot
+	[
+		RootBorder.ToSharedRef()
+	];
+
+}
+
+
+void SMultiBoxWidget::SetSummonedMenu( TSharedRef< SMenuAnchor > InMenuAnchor )
+{
+	SummonedMenuAnchor = InMenuAnchor;
+}
+
+	
+/**
+ * For menu bar multibox widgets, returns the currently open pull-down menu, if there is one open
+ *
+ * @return	Pull-down menu anchor, or null pointer
+ */
+TSharedPtr< const SMenuAnchor > SMultiBoxWidget::GetOpenMenu() const
+{
+	if( SummonedMenuAnchor.IsValid() && SummonedMenuAnchor.Pin()->IsOpen() )
+	{
+		return SummonedMenuAnchor.Pin();
+	}
+
+	// No open menus
+	return TSharedPtr< SMenuAnchor >();
+}
+
+
+
+/**
+ * For menu bar multibox widget, closes any open pull-down menus
+ */
+void SMultiBoxWidget::CloseSummonedMenus()
+{
+	if( GetOpenMenu().IsValid() )
+	{
+		SummonedMenuAnchor.Pin()->SetIsOpen( false );
+
+		// Menu was closed, so we no longer need a weak reference to it
+		SummonedMenuAnchor = NULL;
+	}
+}
+
+
+TSharedRef<SWidget> SMultiBoxWidget::OnWrapButtonClicked()
+{
+	FMenuBuilder MenuBuilder(true, NULL, TSharedPtr<FExtender>(), false, GetStyleSet());
+	{ 
+		// Iterate through the array of blocks telling each one to add itself to the menu
+		const TArray< TSharedRef< const FMultiBlock > >& Blocks = MultiBox->GetBlocks();
+		for (int32 BlockIdx = ClippedHorizontalBox->GetClippedIndex(); BlockIdx < Blocks.Num(); ++BlockIdx)
+		{
+			Blocks[BlockIdx]->CreateMenuEntry(MenuBuilder);
+		}
+	}
+
+	return MenuBuilder.MakeWidget();
+}
+
+void SMultiBoxWidget::UpdateDropAreaPreviewBlock( TSharedRef<const FMultiBlock> MultiBlock, TSharedPtr<FUICommandDragDropOp> DragDropContent, const FGeometry& DragAreaGeometry, const FVector2D& DragPos )
+{
+	TSharedPtr<const FUICommandInfo> UICommand = DragDropContent->UICommand;
+	FName OriginMultiBox = DragDropContent->OriginMultiBox;
+
+	FVector2D LocalDragPos = DragAreaGeometry.AbsoluteToLocal( DragPos );
+
+	FVector2D DrawSize = DragAreaGeometry.GetDrawSize();
+
+	bool bAddedNewBlock = false;
+	bool bValidCommand = true;
+	if( DragPreview.UICommand != UICommand )
+	{	
+		TSharedPtr<const FMultiBlock> ExistingBlock = MultiBox->FindBlockFromCommand( UICommand );
+
+		// Check that the command does not already exist and that we can create it or that we are dragging an exisiting block in this box
+		if( !ExistingBlock.IsValid() || ( ExistingBlock.IsValid() && OriginMultiBox == MultiBox->GetCustomizationName() ) )
+		{
+
+			TSharedPtr<const FMultiBlock> NewBlock = ExistingBlock;
+
+			if( !ExistingBlock.IsValid() )
+			{
+				NewBlock = MultiBox->MakeMultiBlockFromCommand( UICommand, true );
+			}
+
+			if( NewBlock.IsValid() )
+			{
+				DragPreview.Reset();
+				DragPreview.UICommand = UICommand;
+				DragPreview.PreviewBlock = 
+					MakeShareable(
+						new FDropPreviewBlock( 
+							NewBlock.ToSharedRef(), 
+							NewBlock->MakeWidget( SharedThis(this), EMultiBlockLocation::None ) )
+					);
+
+				bAddedNewBlock = true;
+			}
+
+
+		}
+		else
+		{
+			// this command cannot be dropped here
+			bValidCommand = false;
+		}
+	}
+
+	if( bValidCommand )
+	{
+		// determine whether or not to insert before or after
+		bool bInsertBefore = false;
+		if( MultiBox->GetType() == EMultiBoxType::ToolBar )
+		{
+			DragPreview.InsertOrientation  = EOrientation::Orient_Horizontal;
+			if( LocalDragPos.X < DrawSize.X / 2 )
+			{
+				// Insert before horizontally
+				bInsertBefore = true;
+			}
+			else
+			{
+				// Insert after horizontally
+				bInsertBefore = false;
+			}
+		}
+		else 
+		{
+			DragPreview.InsertOrientation  = EOrientation::Orient_Vertical;
+			if( LocalDragPos.Y < DrawSize.Y / 2 )
+			{
+				// Insert before vertically
+				bInsertBefore = true;
+			}
+			else
+			{
+				// Insert after vertically
+				bInsertBefore = false;
+			}
+		}
+
+
+		int32 CurrentIndex = DragPreview.InsertIndex;
+		DragPreview.InsertIndex = INDEX_NONE;
+		// Find the index of the multiblock being dragged over. This is where we will insert the new block
+		if( DragPreview.PreviewBlock.IsValid() )
+		{
+			const TArray< TSharedRef< const FMultiBlock > >& Blocks = MultiBox->GetBlocks();
+			for (int32 BlockIdx = 0; BlockIdx < Blocks.Num(); ++BlockIdx)
+			{
+				TSharedRef<const FMultiBlock> Block = Blocks[BlockIdx];
+				if( Block == MultiBlock )
+				{
+					if( bInsertBefore)
+					{
+						DragPreview.InsertIndex = BlockIdx;
+					}
+					else
+					{
+						DragPreview.InsertIndex = FMath::Min(Blocks.Num()-1, BlockIdx+1);
+					}
+
+					break;
+				}
+			}
+		}
+
+		if( CurrentIndex != DragPreview.InsertIndex && DragPreview.InsertIndex != INDEX_NONE )
+		{
+			BuildMultiBoxWidget();
+		}
+	}
+
+}
+
+EVisibility SMultiBoxWidget::GetCustomizationVisibility( TWeakPtr<const FMultiBlock> BlockWeakPtr, TWeakPtr<SWidget> BlockWidgetWeakPtr ) const
+{
+	if( MultiBox->IsInEditMode() && BlockWidgetWeakPtr.IsValid() && BlockWeakPtr.IsValid() && (!DragPreview.PreviewBlock.IsValid() || BlockWeakPtr.Pin() != DragPreview.PreviewBlock->GetActualBlock() ) )
+	{
+		// If in edit mode and this is not the block being dragged, the customization widget should be visible if the default block beging customized would have been visible
+		return BlockWeakPtr.Pin()->GetAction().IsValid() && BlockWidgetWeakPtr.Pin()->GetVisibility() == EVisibility::Visible ? EVisibility::Visible : EVisibility::Collapsed;
+	}
+	else
+	{
+		return EVisibility::Collapsed;
+	}
+}
+
+FReply SMultiBoxWidget::OnDeleteBlockClicked( TWeakPtr<const FMultiBlock> BlockWeakPtr )
+{
+	if( BlockWeakPtr.IsValid() )
+	{
+		MultiBox->RemoveCustomMultiBlock( BlockWeakPtr.Pin().ToSharedRef() );
+		BuildMultiBoxWidget();
+	}
+
+	return FReply::Handled();
+}
+
+void SMultiBoxWidget::OnCustomCommandDragEnter( TSharedRef<const FMultiBlock> MultiBlock, const FGeometry& MyGeometry, const FDragDropEvent& DragDropEvent )
+{
+	if( MultiBlock != DragPreview.PreviewBlock && MultiBox->IsInEditMode() )
+	{
+		TSharedPtr<FUICommandDragDropOp> DragDropContent = StaticCastSharedPtr<FUICommandDragDropOp>( DragDropEvent.GetOperation() );
+
+		UpdateDropAreaPreviewBlock( MultiBlock, DragDropContent, MyGeometry, DragDropEvent.GetScreenSpacePosition() );
+	}
+}
+
+
+void SMultiBoxWidget::OnCustomCommandDragged( TSharedRef<const FMultiBlock> MultiBlock, const FGeometry& MyGeometry, const FDragDropEvent& DragDropEvent )
+{
+	if( MultiBlock != DragPreview.PreviewBlock && MultiBox->IsInEditMode() )
+	{
+		TSharedPtr<FUICommandDragDropOp> DragDropContent = StaticCastSharedPtr<FUICommandDragDropOp>( DragDropEvent.GetOperation() );
+
+		UpdateDropAreaPreviewBlock( MultiBlock, DragDropContent, MyGeometry, DragDropEvent.GetScreenSpacePosition() );
+	}
+}
+
+void SMultiBoxWidget::OnCustomCommandDropped()
+{
+	if( DragPreview.IsValid() )
+	{	
+
+		// Check that the command does not already exist and that we can create it or that we are dragging an exisiting block in this box
+		TSharedPtr<const FMultiBlock> Block = MultiBox->FindBlockFromCommand( DragPreview.UICommand );
+		if( !Block.IsValid() )
+		{
+			Block = MultiBox->MakeMultiBlockFromCommand( DragPreview.UICommand, true );
+		}
+
+		if( Block.IsValid() )
+		{
+			MultiBox->InsertCustomMultiBlock( Block.ToSharedRef(), DragPreview.InsertIndex );
+		}
+
+		DragPreview.Reset();
+
+		BuildMultiBoxWidget();
+	}
+}
+
+void SMultiBoxWidget::OnDropExternal()
+{
+	// The command was not dropped in this widget
+	if( DragPreview.IsValid() )
+	{
+		DragPreview.Reset();
+
+		BuildMultiBoxWidget();
+	}
+}
+
+FReply SMultiBoxWidget::OnDragOver( const FGeometry& MyGeometry, const FDragDropEvent& DragDropEvent )
+{
+	if ( DragDropEvent.GetOperationAs<FUICommandDragDropOp>().IsValid() && MultiBox->IsInEditMode() )
+	{
+		return FReply::Handled();
+	}
+
+	return FReply::Unhandled();
+}
+
+FReply SMultiBoxWidget::OnDrop( const FGeometry& MyGeometry, const FDragDropEvent& DragDropEvent )
+{
+	if ( DragDropEvent.GetOperationAs<FUICommandDragDropOp>().IsValid() )
+	{
+		OnCustomCommandDropped();
+		return FReply::Handled();
+	}
+
+	return FReply::Unhandled();
+}
+
+bool SMultiBoxWidget::SupportsKeyboardFocus() const
+{
+	return true;
+}
+
+FReply SMultiBoxWidget::FocusNextWidget( EFocusMoveDirection::Type MoveDirection )
+{
+	TSharedPtr<SWidget> FocusWidget = FSlateApplication::Get().GetKeyboardFocusedWidget();
+	if(FocusWidget.IsValid())
+	{
+		FWidgetPath FocusPath;
+		FSlateApplication::Get().GeneratePathToWidgetUnchecked( FocusWidget.ToSharedRef(), FocusPath );
+		FWeakWidgetPath WeakFocusPath = FocusPath;
+		FWidgetPath NextFocusPath = WeakFocusPath.ToNextFocusedPath( MoveDirection );
+		if ( NextFocusPath.Widgets.Num() > 0 )
+		{
+			return FReply::Handled().SetKeyboardFocus( NextFocusPath.Widgets.Last().Widget, EKeyboardFocusCause::Keyboard );
+		}
+	}
+
+	return FReply::Unhandled();
+}
+
+FReply SMultiBoxWidget::OnKeyboardFocusReceived( const FGeometry& MyGeometry, const FKeyboardFocusEvent& InKeyboardFocusEvent )
+{
+	if(InKeyboardFocusEvent.GetCause() == EKeyboardFocusCause::Keyboard)
+	{
+		// forward focus to children
+		return FocusNextWidget( EFocusMoveDirection::Next );
+	}
+
+	return FReply::Unhandled();
+}
+
+FReply SMultiBoxWidget::OnKeyDown( const FGeometry& MyGeometry, const FKeyboardEvent& KeyboardEvent )
+{
+	SCompoundWidget::OnKeyDown(MyGeometry, KeyboardEvent);
+
+	// allow use of up and down keys to transfer focus/hover state
+	if(KeyboardEvent.GetKey() == EKeys::Up || KeyboardEvent.GetKey() == EKeys::Down)
+	{
+		return FocusNextWidget( EFocusMoveDirection::Next );
+	}
+
+	return FReply::Unhandled();
+}
