Slate: Reorganized private Slate files; proper pre-compiled header for all files.

diff --git a/Engine/Source/Runtime/Slate/Private/Framework/MultiBox/MultiBoxBuilder.cpp b/Engine/Source/Runtime/Slate/Private/Framework/MultiBox/MultiBoxBuilder.cpp
new file mode 100644
--- /dev/null
+++ b/Engine/Source/Runtime/Slate/Private/Framework/MultiBox/MultiBoxBuilder.cpp
@@ -0,0 +1,485 @@
+// Copyright 1998-2014 Epic Games, Inc. All Rights Reserved.
+
+#include "SlatePrivatePCH.h"
+#include "MultiBox.h"
+#include "SHeadingBlock.h"
+#include "SMenuEntryBlock.h"
+#include "SMenuSeparatorBlock.h"
+#include "SToolBarSeparatorBlock.h"
+#include "SToolBarButtonBlock.h"
+#include "SToolBarComboButtonBlock.h"
+#include "SEditableTextBlock.h"
+#include "SButtonRowBlock.h"
+#include "SWidgetBlock.h"
+#include "SGroupMarkerBlock.h"
+#include "STutorialWrapper.h"
+
+FMultiBoxBuilder::FMultiBoxBuilder( const EMultiBoxType::Type InType, FMultiBoxCustomization InCustomization, const bool bInShouldCloseWindowAfterMenuSelection, const TSharedPtr< const FUICommandList >& InCommandList, TSharedPtr<FExtender> InExtender )
+	: MultiBox( FMultiBox::Create( InType, InCustomization, bInShouldCloseWindowAfterMenuSelection ) )
+	, CommandListStack()
+{
+	CommandListStack.Push( InCommandList );
+	ExtenderStack.Push(InExtender);
+}
+
+void FMultiBoxBuilder::AddEditableText( const FText& InLabel, const FText& InToolTip, const FSlateIcon& InIcon, const TAttribute< FText >& InTextAttribute, const FOnTextCommitted& InOnTextCommitted, const FOnTextChanged& InOnTextChanged, bool bInReadOnly )
+{
+	MultiBox->AddMultiBlock( MakeShareable( new FEditableTextBlock( InLabel, InToolTip, InIcon, InTextAttribute, bInReadOnly, InOnTextCommitted, InOnTextChanged ) ) );
+}
+
+void FMultiBoxBuilder::PushCommandList( const TSharedRef< const FUICommandList > CommandList )
+{
+	CommandListStack.Push( CommandList );
+}
+
+void FMultiBoxBuilder::PopCommandList()
+{
+	// Never allowed to pop the last command-list!  This command-list was set when the multibox was first created and is canonical.
+	if( ensure( CommandListStack.Num() > 1 ) )
+	{
+		CommandListStack.Pop();
+	}
+}
+
+TSharedPtr<const FUICommandList> FMultiBoxBuilder::GetTopCommandList()
+{
+	return (CommandListStack.Num() > 0) ? CommandListStack.Top() : TSharedPtr<const FUICommandList>(NULL);
+}
+
+void FMultiBoxBuilder::PushExtender( TSharedRef< FExtender > InExtender )
+{
+	ExtenderStack.Push( InExtender );
+}
+
+void FMultiBoxBuilder::PopExtender()
+{
+	// Never allowed to pop the last extender! This extender was set when the multibox was first created and is canonical.
+	if( ensure( ExtenderStack.Num() > 1 ) )
+	{
+		ExtenderStack.Pop();
+	}
+}
+
+const ISlateStyle* FMultiBoxBuilder::GetStyleSet() const 
+{ 
+	return MultiBox->GetStyleSet();
+}
+
+const FName& FMultiBoxBuilder::GetStyleName() const 
+{ 
+	return MultiBox->GetStyleName();
+}
+
+void FMultiBoxBuilder::SetStyle( const ISlateStyle* InStyleSet, const FName& InStyleName ) 
+{ 
+	MultiBox->SetStyle( InStyleSet, InStyleName ); 
+}
+
+FMultiBoxCustomization FMultiBoxBuilder::GetCustomization() const
+{
+	return FMultiBoxCustomization( MultiBox->GetCustomizationName() ); 
+}
+
+TSharedRef< class SWidget > FMultiBoxBuilder::MakeWidget()
+{
+	return MultiBox->MakeWidget();
+}
+
+
+FBaseMenuBuilder::FBaseMenuBuilder( const EMultiBoxType::Type InType, const bool bInShouldCloseWindowAfterMenuSelection, TSharedPtr< const FUICommandList > InCommandList, bool bInCloseSelfOnly, TSharedPtr<FExtender> InExtender, const ISlateStyle* InStyleSet )
+	: FMultiBoxBuilder( InType, FMultiBoxCustomization::None, bInShouldCloseWindowAfterMenuSelection, InCommandList, InExtender )
+	, bCloseSelfOnly( bInCloseSelfOnly )
+{
+	MultiBox->SetStyle(InStyleSet, "Menu");
+}
+
+void FBaseMenuBuilder::AddMenuEntry( const TSharedPtr< const FUICommandInfo > InCommand, FName InExtensionHook, const TAttribute<FText>& InLabelOverride, const TAttribute<FText>& InToolTipOverride, const FSlateIcon& InIconOverride )
+{
+	ApplySectionBeginning();
+
+	ApplyHook(InExtensionHook, EExtensionHook::Before);
+	
+	// The command must be valid
+	check( InCommand.IsValid() );
+	TSharedRef< FMenuEntryBlock > NewMenuEntryBlock( new FMenuEntryBlock( InExtensionHook, InCommand, CommandListStack.Last(), InLabelOverride, InToolTipOverride, InIconOverride, bCloseSelfOnly ) );
+	MultiBox->AddMultiBlock( NewMenuEntryBlock );
+	
+	ApplyHook(InExtensionHook, EExtensionHook::After);
+}
+
+void FBaseMenuBuilder::AddMenuEntry( const TAttribute<FText>& InLabel, const TAttribute<FText>& InToolTip, const FSlateIcon& InIcon, const FUIAction& InAction, FName InExtensionHook, const EUserInterfaceActionType::Type UserInterfaceActionType )
+{
+	ApplySectionBeginning();
+
+	ApplyHook(InExtensionHook, EExtensionHook::Before);
+	
+	TSharedRef< FMenuEntryBlock > NewMenuEntryBlock( new FMenuEntryBlock( InExtensionHook, InLabel, InToolTip, InIcon, InAction, UserInterfaceActionType, bCloseSelfOnly ) );
+	MultiBox->AddMultiBlock( NewMenuEntryBlock );
+	
+	ApplyHook(InExtensionHook, EExtensionHook::After);
+}
+
+void FBaseMenuBuilder::AddMenuEntry( const FUIAction& UIAction, const TSharedRef< SWidget > Contents, const FName& InExtensionHook, const EUserInterfaceActionType::Type UserInterfaceActionType )
+{
+	ApplySectionBeginning();
+
+	ApplyHook(InExtensionHook, EExtensionHook::Before);
+
+	TSharedRef< FMenuEntryBlock > NewMenuEntryBlock( new FMenuEntryBlock( InExtensionHook, UIAction, Contents, UserInterfaceActionType, bCloseSelfOnly ) );
+	MultiBox->AddMultiBlock( NewMenuEntryBlock );
+
+	ApplyHook(InExtensionHook, EExtensionHook::After);
+}
+
+void FMenuBuilder::BeginSection( FName InExtensionHook, const TAttribute< FText >& InHeadingText )
+{
+	check(CurrentSectionExtensionHook == NAME_None && !bSectionNeedsToBeApplied);
+
+	ApplyHook(InExtensionHook, EExtensionHook::Before);
+	
+	// Do not actually apply the section header, because if this section is ended immediately
+	// then nothing ever gets created, preventing empty sections from ever appearing
+	bSectionNeedsToBeApplied = true;
+	CurrentSectionExtensionHook = InExtensionHook;
+	CurrentSectionHeadingText = InHeadingText.Get();
+
+	// Do apply the section beginning if we are in developer "show me all the hooks" mode
+	if (FMultiBoxSettings::DisplayMultiboxHooks.Get())
+	{
+		ApplySectionBeginning();
+	}
+	
+	ApplyHook(InExtensionHook, EExtensionHook::First);
+}
+
+void FMenuBuilder::EndSection()
+{
+	FName SectionExtensionHook = CurrentSectionExtensionHook;
+	CurrentSectionExtensionHook = NAME_None;
+	bSectionNeedsToBeApplied = false;
+	CurrentSectionHeadingText = FText::GetEmpty();
+
+	ApplyHook(SectionExtensionHook, EExtensionHook::After);
+}
+
+void FMenuBuilder::AddMenuSeparator(FName InExtensionHook)
+{
+	ApplySectionBeginning();
+
+	ApplyHook(InExtensionHook, EExtensionHook::Before);
+
+	// Never add a menu separate as the first item, even if we were asked to
+	if( MultiBox->GetBlocks().Num() > 0 || FMultiBoxSettings::DisplayMultiboxHooks.Get() )
+	{
+		TSharedRef< FMenuSeparatorBlock > NewMenuSeparatorBlock( new FMenuSeparatorBlock(InExtensionHook) );
+		MultiBox->AddMultiBlock( NewMenuSeparatorBlock );
+	}
+
+	ApplyHook(InExtensionHook, EExtensionHook::After);
+}
+
+void FMenuBuilder::AddSubMenu( const FText& InMenuLabel, const FText& InToolTip, const FNewMenuDelegate& InSubMenu, const FUIAction& InUIAction, FName InExtensionHook, const EUserInterfaceActionType::Type InUserInterfaceActionType, const bool bInOpenSubMenuOnClick, const FSlateIcon& InIcon )
+{
+	ApplySectionBeginning();
+
+	const bool bIsSubMenu = true;
+	TSharedRef< FMenuEntryBlock > NewMenuEntryBlock( new FMenuEntryBlock( InExtensionHook, InMenuLabel, InToolTip, InSubMenu, ExtenderStack.Top(), bIsSubMenu, bInOpenSubMenuOnClick, InIcon, InUIAction, InUserInterfaceActionType, bCloseSelfOnly ) );
+
+	MultiBox->AddMultiBlock( NewMenuEntryBlock );
+}
+
+void FMenuBuilder::AddSubMenu( const FText& InMenuLabel, const FText& InToolTip, const FNewMenuDelegate& InSubMenu, const bool bInOpenSubMenuOnClick /*= false*/, const FSlateIcon& InIcon /*= FSlateIcon()*/ )
+{
+	ApplySectionBeginning();
+
+	const bool bIsSubMenu = true;
+	TSharedRef< FMenuEntryBlock > NewMenuEntryBlock( new FMenuEntryBlock( NAME_None, InMenuLabel, InToolTip, InSubMenu, ExtenderStack.Top(), bIsSubMenu, bInOpenSubMenuOnClick, CommandListStack.Last(), bCloseSelfOnly, InIcon ) );
+
+	MultiBox->AddMultiBlock( NewMenuEntryBlock );
+}
+
+void FMenuBuilder::AddSubMenu( const TSharedRef< SWidget > Contents, const FNewMenuDelegate& InSubMenu, const bool bInOpenSubMenuOnClick /*= false*/ )
+{
+	ApplySectionBeginning();
+
+	const bool bIsSubMenu = true;
+	TSharedRef< FMenuEntryBlock > NewMenuEntryBlock( new FMenuEntryBlock( NAME_None, Contents, InSubMenu, ExtenderStack.Top(), bIsSubMenu, bInOpenSubMenuOnClick, CommandListStack.Last(), bCloseSelfOnly ) );
+
+	MultiBox->AddMultiBlock( NewMenuEntryBlock );
+}
+
+void FMenuBuilder::AddSubMenu( const FUIAction& UIAction, const TSharedRef< SWidget > Contents, const FNewMenuDelegate& InSubMenu )
+{
+	ApplySectionBeginning();
+
+	const bool bIsSubMenu = true;
+	TSharedRef< FMenuEntryBlock > NewMenuEntryBlock( new FMenuEntryBlock( NAME_None, UIAction, Contents, InSubMenu, ExtenderStack.Top(), bIsSubMenu, CommandListStack.Last(), bCloseSelfOnly ) );
+
+	MultiBox->AddMultiBlock( NewMenuEntryBlock );
+}
+
+void FMenuBuilder::AddWrapperSubMenu( const FText& InMenuLabel, const FText& InToolTip, const FOnGetContent& InSubMenu, const FSlateIcon& InIcon )
+{
+	ApplySectionBeginning();
+
+	const bool bIsSubMenu = true;
+	TSharedRef< FMenuEntryBlock > NewMenuEntryBlock( new FMenuEntryBlock( NAME_None, InMenuLabel, InToolTip, InSubMenu, ExtenderStack.Top(), bIsSubMenu, false, CommandListStack.Last(), bCloseSelfOnly, InIcon ) );
+
+	MultiBox->AddMultiBlock( NewMenuEntryBlock );
+}
+
+void FMenuBuilder::AddWrapperSubMenu( const FText& InMenuLabel, const FText& InToolTip, const TSharedPtr<SWidget>& InSubMenu, const FSlateIcon& InIcon )
+{
+	ApplySectionBeginning();
+
+	const bool bIsSubMenu = true;
+	TSharedRef< FMenuEntryBlock > NewMenuEntryBlock( new FMenuEntryBlock( NAME_None, InMenuLabel, InToolTip, InSubMenu, ExtenderStack.Top(), bIsSubMenu, false, CommandListStack.Last(), bCloseSelfOnly, InIcon ) );
+
+	MultiBox->AddMultiBlock( NewMenuEntryBlock );
+}
+
+void FMenuBuilder::AddWidget( TSharedRef<SWidget> InWidget, const FText& Label, bool bNoIndent )
+{
+	ApplySectionBeginning();
+
+	TSharedRef< FWidgetBlock > NewWidgetBlock( new FWidgetBlock( InWidget, Label, bNoIndent ) );
+	MultiBox->AddMultiBlock( NewWidgetBlock );
+}
+
+void FMenuBuilder::ApplyHook(FName InExtensionHook, EExtensionHook::Position HookPosition)
+{
+	// this is a virtual function to get a properly typed "this" pointer
+	auto& Extender = ExtenderStack.Top();
+	if (InExtensionHook != NAME_None && Extender.IsValid())
+	{
+		Extender->Apply(InExtensionHook, HookPosition, *this);
+	}
+}
+
+void FMenuBuilder::ApplySectionBeginning()
+{
+	if (bSectionNeedsToBeApplied)
+	{
+		if( MultiBox->GetBlocks().Num() > 0 || FMultiBoxSettings::DisplayMultiboxHooks.Get() )
+		{
+			MultiBox->AddMultiBlock( MakeShareable( new FMenuSeparatorBlock(CurrentSectionExtensionHook) ) );
+		}
+		if (!CurrentSectionHeadingText.IsEmpty())
+		{
+			MultiBox->AddMultiBlock( MakeShareable( new FHeadingBlock(CurrentSectionExtensionHook, CurrentSectionHeadingText) ) );
+		}
+		bSectionNeedsToBeApplied = false;
+		CurrentSectionHeadingText = FText::GetEmpty();
+	}
+}
+
+
+
+void FMenuBarBuilder::AddPullDownMenu( const FText& InMenuLabel, const FText& InToolTip, const FNewMenuDelegate& InPullDownMenu, FName InExtensionHook, FName InTutorialHighlightName )
+{
+	ApplySectionBeginning();
+
+	ApplyHook(InExtensionHook, EExtensionHook::Before);
+
+	const bool bIsSubMenu = false;
+	const bool bOpenSubMenuOnClick = false;
+	// Pulldown menus always close all menus not just themselves
+	const bool bCloseSelfOnly = false;
+	TSharedRef< FMenuEntryBlock > NewMenuEntryBlock( new FMenuEntryBlock( InExtensionHook, InMenuLabel, InToolTip, InPullDownMenu, ExtenderStack.Top(), bIsSubMenu, bOpenSubMenuOnClick, CommandListStack.Last(), bCloseSelfOnly ) );
+	NewMenuEntryBlock->SetTutorialHightlightName(InTutorialHighlightName);
+	MultiBox->AddMultiBlock( NewMenuEntryBlock );
+	
+	ApplyHook(InExtensionHook, EExtensionHook::After);
+}
+
+void FMenuBarBuilder::ApplyHook(FName InExtensionHook, EExtensionHook::Position HookPosition)
+{
+	// this is a virtual function to get a properly typed "this" pointer
+	auto& Extender = ExtenderStack.Top();
+	if (InExtensionHook != NAME_None && Extender.IsValid())
+	{
+		Extender->Apply(InExtensionHook, HookPosition, *this);
+	}
+}
+
+void FToolBarBuilder::AddToolBarButton( const TSharedPtr< const FUICommandInfo > InCommand, FName InExtensionHook, const TAttribute<FText>& InLabelOverride, const TAttribute<FText>& InToolTipOverride, const TAttribute<FSlateIcon>& InIconOverride, FName TutorialHighlightName)
+{
+	ApplySectionBeginning();
+
+	ApplyHook(InExtensionHook, EExtensionHook::Before);
+
+	TSharedRef< FToolBarButtonBlock > NewToolBarButtonBlock( new FToolBarButtonBlock( InCommand.ToSharedRef(), CommandListStack.Last(), InLabelOverride, InToolTipOverride, InIconOverride ) );
+
+	if ( LabelVisibility.IsSet() )
+	{
+		NewToolBarButtonBlock->SetLabelVisibility( LabelVisibility.GetValue() );
+	}
+
+	NewToolBarButtonBlock->SetIsFocusable(bIsFocusable);
+	NewToolBarButtonBlock->SetForceSmallIcons(bForceSmallIcons);
+	NewToolBarButtonBlock->SetTutorialHightlightName(TutorialHighlightName);
+
+	MultiBox->AddMultiBlock( NewToolBarButtonBlock );
+
+	ApplyHook(InExtensionHook, EExtensionHook::After);
+}
+
+void FToolBarBuilder::AddToolBarButton( const FUIAction& InAction, FName InExtensionHook, const TAttribute<FText>& InLabelOverride, const TAttribute<FText>& InToolTipOverride, const TAttribute<FSlateIcon>& InIconOverride, const EUserInterfaceActionType::Type UserInterfaceActionType, FName TutorialHighlightName )
+{
+	ApplySectionBeginning();
+
+	ApplyHook(InExtensionHook, EExtensionHook::Before);
+
+	TSharedRef< FToolBarButtonBlock > NewToolBarButtonBlock( new FToolBarButtonBlock( InLabelOverride, InToolTipOverride, InIconOverride, InAction, UserInterfaceActionType ) );
+
+	if ( LabelVisibility.IsSet() )
+	{
+		NewToolBarButtonBlock->SetLabelVisibility( LabelVisibility.GetValue() );
+	}
+
+	NewToolBarButtonBlock->SetIsFocusable(bIsFocusable);
+	NewToolBarButtonBlock->SetForceSmallIcons(bForceSmallIcons);
+	NewToolBarButtonBlock->SetTutorialHightlightName(TutorialHighlightName);
+
+	MultiBox->AddMultiBlock( NewToolBarButtonBlock );
+
+	ApplyHook(InExtensionHook, EExtensionHook::After);
+}
+
+void FToolBarBuilder::AddComboButton( const FUIAction& InAction, const FOnGetContent& InMenuContentGenerator, const TAttribute<FText>& InLabelOverride, const TAttribute<FText>& InToolTipOverride, const TAttribute<FSlateIcon>& InIconOverride, bool bInSimpleComboBox )
+{
+	ApplySectionBeginning();
+
+	TSharedRef< FToolBarComboButtonBlock > NewToolBarComboButtonBlock( new FToolBarComboButtonBlock( InAction, InMenuContentGenerator, InLabelOverride, InToolTipOverride, InIconOverride, bInSimpleComboBox ) );
+
+	if ( LabelVisibility.IsSet() )
+	{
+		NewToolBarComboButtonBlock->SetLabelVisibility( LabelVisibility.GetValue() );
+	}
+
+	NewToolBarComboButtonBlock->SetForceSmallIcons(bForceSmallIcons);
+
+	MultiBox->AddMultiBlock( NewToolBarComboButtonBlock );
+}
+
+void FToolBarBuilder::AddWidget( TSharedRef<SWidget> InWidget, FName TutorialHighlightName )
+{
+	ApplySectionBeginning();
+
+	// If tutorial name specified, wrap in tutorial wrapper
+	if (TutorialHighlightName != NAME_None)
+	{
+		TSharedRef<SWidget> ChildWidget = InWidget;
+
+		InWidget = 
+			SNew( STutorialWrapper, TutorialHighlightName )
+			[
+				ChildWidget
+			];
+	}
+	
+	TSharedRef< FWidgetBlock > NewWidgetBlock( new FWidgetBlock( InWidget, FText::GetEmpty(), true ) );
+	MultiBox->AddMultiBlock( NewWidgetBlock );
+}
+
+void FToolBarBuilder::AddSeparator(FName InExtensionHook)
+{
+	ApplySectionBeginning();
+
+	ApplyHook(InExtensionHook, EExtensionHook::Before);
+
+	// Never add a menu separate as the first item, even if we were asked to
+	if( MultiBox->GetBlocks().Num() > 0 || FMultiBoxSettings::DisplayMultiboxHooks.Get() )
+	{
+		MultiBox->AddMultiBlock( MakeShareable( new FToolBarSeparatorBlock(InExtensionHook) ) );
+	}
+
+	ApplyHook(InExtensionHook, EExtensionHook::After);
+}
+
+void FToolBarBuilder::BeginSection( FName InExtensionHook )
+{
+	checkf(CurrentSectionExtensionHook == NAME_None && !bSectionNeedsToBeApplied, TEXT("Did you forget to call EndSection()?"));
+
+	ApplyHook(InExtensionHook, EExtensionHook::Before);
+	
+	// Do not actually apply the section header, because if this section is ended immediately
+	// then nothing ever gets created, preventing empty sections from ever appearing
+	bSectionNeedsToBeApplied = true;
+	CurrentSectionExtensionHook = InExtensionHook;
+	
+	// Do apply the section beginning if we are in developer "show me all the hooks" mode
+	if (FMultiBoxSettings::DisplayMultiboxHooks.Get())
+	{
+		ApplySectionBeginning();
+	}
+	
+	ApplyHook(InExtensionHook, EExtensionHook::First);
+}
+
+void FToolBarBuilder::EndSection()
+{
+	FName SectionExtensionHook = CurrentSectionExtensionHook;
+	CurrentSectionExtensionHook = NAME_None;
+	bSectionNeedsToBeApplied = false;
+
+	ApplyHook(SectionExtensionHook, EExtensionHook::After);
+}
+
+void FToolBarBuilder::ApplyHook(FName InExtensionHook, EExtensionHook::Position HookPosition)
+{
+	// this is a virtual function to get a properly typed "this" pointer
+	auto& Extender = ExtenderStack.Top();
+	if (InExtensionHook != NAME_None && Extender.IsValid())
+	{
+		Extender->Apply(InExtensionHook, HookPosition, *this);
+	}
+}
+
+void FToolBarBuilder::ApplySectionBeginning()
+{
+	if (bSectionNeedsToBeApplied)
+	{
+		if( MultiBox->GetBlocks().Num() > 0 || FMultiBoxSettings::DisplayMultiboxHooks.Get() )
+		{
+			MultiBox->AddMultiBlock( MakeShareable( new FToolBarSeparatorBlock(CurrentSectionExtensionHook) ) );
+		}
+		bSectionNeedsToBeApplied = false;
+	}
+}
+
+void FToolBarBuilder::EndBlockGroup()
+{
+	ApplySectionBeginning();
+
+	TSharedRef< FGroupEndBlock > NewGroupEndBlock( new FGroupEndBlock() );
+
+	MultiBox->AddMultiBlock( NewGroupEndBlock );
+}
+
+void FToolBarBuilder::BeginBlockGroup()
+{
+	ApplySectionBeginning();
+
+	TSharedRef< FGroupStartBlock > NewGroupStartBlock( new FGroupStartBlock() );
+
+	MultiBox->AddMultiBlock( NewGroupStartBlock );
+}
+
+void FButtonRowBuilder::AddButton( const TSharedPtr< const FUICommandInfo > InCommand, const TAttribute<FText>& InLabelOverride, const TAttribute<FText>& InToolTipOverride, const FSlateIcon& InIconOverride )
+{
+	ApplySectionBeginning();
+
+	TSharedRef< FButtonRowBlock > NewButtonRowBlock( new FButtonRowBlock( InCommand.ToSharedRef(), CommandListStack.Last(), InLabelOverride, InToolTipOverride, InIconOverride ) );
+
+	MultiBox->AddMultiBlock( NewButtonRowBlock );
+}
+
+void FButtonRowBuilder::AddButton( const FText& InLabel, const FText& InToolTip, const FUIAction& UIAction, const FSlateIcon& InIcon, const EUserInterfaceActionType::Type UserInterfaceActionType )
+{
+	ApplySectionBeginning();
+
+	TSharedRef< FButtonRowBlock > NewButtonRowBlock( new FButtonRowBlock( InLabel, InToolTip, InIcon, UIAction, UserInterfaceActionType ) );
+
+	MultiBox->AddMultiBlock( NewButtonRowBlock );
+}
diff --git a/Engine/Source/Runtime/Slate/Private/Framework/MultiBox/MultiBoxCustomization.cpp b/Engine/Source/Runtime/Slate/Private/Framework/MultiBox/MultiBoxCustomization.cpp
new file mode 100644
--- /dev/null
+++ b/Engine/Source/Runtime/Slate/Private/Framework/MultiBox/MultiBoxCustomization.cpp
@@ -0,0 +1,392 @@
+// Copyright 1998-2014 Epic Games, Inc. All Rights Reserved.
+
+#include "SlatePrivatePCH.h"
+#include "MultiBox.h"
+#include "MultiBoxCustomization.h"
+#include "Json.h"
+#include "RemoteConfigIni.h"
+
+void SCustomToolbarPreviewWidget::Construct( const FArguments& InArgs )
+{
+	Content = InArgs._Content.Widget;
+}
+
+void SCustomToolbarPreviewWidget::BuildMultiBlockWidget(const ISlateStyle* StyleSet, const FName& StyleName)
+{
+	ChildSlot
+	[
+		SNew( SBorder )
+		.Padding(0)
+		.BorderImage( FCoreStyle::Get().GetBrush("NoBorder") )
+		.HAlign(HAlign_Center)
+		.VAlign(VAlign_Center)
+		[
+			Content.ToSharedRef()
+		]
+	];
+}
+
+void SMultiBlockDragHandle::Construct( const FArguments& InArgs, TSharedRef<SMultiBoxWidget> InBaseWidget, TSharedRef<const FMultiBlock> InBlock, FName CustomizationName )
+{
+	BaseWidget = InBaseWidget;
+	Block = InBlock;
+	MultiBoxCustomizationName = CustomizationName;
+}
+
+FReply SMultiBlockDragHandle::OnMouseButtonDown( const FGeometry& MyGeometry, const FPointerEvent& MouseEvent )
+{
+	if( MouseEvent.GetEffectingButton() == EKeys::LeftMouseButton && Block->GetAction().IsValid() )
+	{
+		return FReply::Handled().DetectDrag( SharedThis(this), MouseEvent.GetEffectingButton() );
+	}
+
+	return FReply::Unhandled();
+}
+
+void SMultiBlockDragHandle::OnDragEnter( const FGeometry& MyGeometry, const FDragDropEvent& DragDropEvent )
+{
+	if ( DragDropEvent.GetOperationAs<FUICommandDragDropOp>().IsValid() )
+	{
+		BaseWidget.Pin()->OnCustomCommandDragEnter( Block.ToSharedRef(), MyGeometry, DragDropEvent );
+	}
+}
+
+FReply SMultiBlockDragHandle::OnDragOver( const FGeometry& MyGeometry, const FDragDropEvent& DragDropEvent )
+{
+	if ( DragDropEvent.GetOperationAs<FUICommandDragDropOp>().IsValid() )
+	{
+		BaseWidget.Pin()->OnCustomCommandDragged( Block.ToSharedRef(), MyGeometry, DragDropEvent );
+		return FReply::Handled();
+	}
+
+	return FReply::Unhandled();
+}
+
+FReply SMultiBlockDragHandle::OnDrop( const FGeometry& MyGeometry, const FDragDropEvent& DragDropEvent )
+{
+	if ( DragDropEvent.GetOperationAs<FUICommandDragDropOp>().IsValid() )
+	{
+		BaseWidget.Pin()->OnCustomCommandDropped();
+		return FReply::Handled();
+	}
+
+	return FReply::Unhandled();
+}
+
+FReply SMultiBlockDragHandle::OnDragDetected( const FGeometry& MyGeometry, const FPointerEvent& MouseEvent )
+{
+	TSharedRef<FUICommandDragDropOp> NewOp = FUICommandDragDropOp::New(
+			Block->GetAction().ToSharedRef(), 
+			MultiBoxCustomizationName, 
+			Block->MakeWidget( BaseWidget.Pin().ToSharedRef(), EMultiBlockLocation::None )->AsWidget(),
+			MyGeometry.AbsolutePosition-MouseEvent.GetScreenSpacePosition()
+		);
+			
+	NewOp->SetOnDropNotification( FSimpleDelegate::CreateSP( BaseWidget.Pin().ToSharedRef(), &SMultiBoxWidget::OnDropExternal ) );
+
+	TSharedRef<FDragDropOperation> DragDropOp = NewOp;
+	return FReply::Handled().BeginDragDrop( DragDropOp );
+}
+
+TSharedRef< class IMultiBlockBaseWidget > FDropPreviewBlock::ConstructWidget() const
+{
+	return
+		SNew( SCustomToolbarPreviewWidget )
+		.Visibility( EVisibility::Hidden )
+		.Content()
+		[
+			ActualWidget->AsWidget()
+		];
+
+}
+
+void FMultiBoxCustomizationData::LoadCustomizedBlocks()
+{
+	Transactions.Empty();
+
+	FString Content;
+	
+	GConfig->GetString(*GetConfigSectionName(), *CustomizationName.ToString(), Content, GEditorUserSettingsIni);
+
+	TSharedPtr<FJsonObject> SavedData;
+	TSharedRef< TJsonReader<> > Reader = TJsonReaderFactory<>::Create( FRemoteConfig::ReplaceIniSpecialCharWithChar(Content).ReplaceEscapedCharWithChar() );
+	bool bResult = FJsonSerializer::Deserialize( Reader, SavedData );
+
+	if( bResult && SavedData.IsValid() )
+	{
+		TArray<TSharedPtr<FJsonValue> > CustomBlockData = SavedData->GetArrayField( TEXT("CustomBlocks") );
+		for (TArray<TSharedPtr<FJsonValue> >::TConstIterator It(CustomBlockData); It; ++It)
+		{
+			TSharedPtr<FJsonObject> SavedCommand = (*It)->AsObject();
+
+			if( SavedCommand.IsValid() )
+			{
+				FString CommandName = SavedCommand->GetStringField( TEXT("CommandName") );
+				FString Context = SavedCommand->GetStringField( TEXT("Context") );
+				int32 Index = FMath::Trunc( SavedCommand->GetNumberField( TEXT("Index") ) );
+				FCustomBlockTransaction::ETransactionType TransType = (FCustomBlockTransaction::ETransactionType)FMath::Trunc( SavedCommand->GetNumberField("TransactionType") );
+
+				if( !CommandName.IsEmpty() && !Context.IsEmpty() && Index >= 0 )
+				{
+					const TSharedPtr< const FUICommandInfo> FoundCommand = FInputBindingManager::Get().FindCommandInContext( FName( *Context ), FName( *CommandName ) );
+					if( FoundCommand.IsValid() )
+					{
+						FCustomBlockTransaction NewTransaction;
+						if( TransType == FCustomBlockTransaction::Add )
+						{
+							NewTransaction = FCustomBlockTransaction::CreateAdd( FoundCommand.ToSharedRef(), Index );
+						}
+						else
+						{
+							NewTransaction = FCustomBlockTransaction::CreateRemove( FoundCommand.ToSharedRef(), Index );
+						}
+
+						Transactions.Add( NewTransaction );
+					}
+				}
+			}
+		}
+	}
+}
+
+FString FMultiBoxCustomizationData::GetConfigSectionName() const
+{
+	return TEXT("CustomMultiBoxes1_0");
+}
+
+void FMultiBoxCustomizationData::RemoveTransactionAt(int32 RemoveIndex)
+{
+	// Copy off the transaction we are going to remove.
+	FCustomBlockTransaction TransToRemove = Transactions[RemoveIndex];
+	const int32 TransToRemoveIndex = TransToRemove.BlockIndex;
+
+	// Remove the transaction
+	Transactions.RemoveAt( RemoveIndex );
+
+	// Iterate over all transactions after this one to fix up indices.
+	for( int32 StartIndex = RemoveIndex; StartIndex < Transactions.Num(); ++StartIndex )
+	{
+		int32& CheckIndex = Transactions[StartIndex].BlockIndex;
+		
+		// If the index is greater than or equal to the index we are removing, this transaction is affected by this removal.
+		if( CheckIndex >= TransToRemoveIndex )
+		{
+			// If the removed transaction was an add, subtract 1 from the index of this transaction. Otherwise, add one.
+			if( TransToRemove.TransactionType == FCustomBlockTransaction::Add )
+			{
+				--CheckIndex;
+			}
+			else
+			{
+				++CheckIndex;
+			}
+		}
+	}
+}
+
+bool FMultiBoxCustomizationData::RemoveDuplicateTransaction()
+{
+	if( Transactions.Num() > 0 )
+	{
+		FCustomBlockTransaction LastTrans = Transactions.Last();
+
+		int32 CheckIndex = LastTrans.BlockIndex;
+		
+		// Skip the last item
+		for( int32 TransIndex = Transactions.Num()-2; TransIndex >= 0; --TransIndex )
+		{
+			FCustomBlockTransaction CurrentTrans = Transactions[TransIndex];
+
+			//ensure( !(CurrentTrans.TransactionType == LastTrans.TransactionType && CurrentTrans.Command == LastTrans.Command ) );
+
+			if( CurrentTrans.Command == LastTrans.Command &&
+				CurrentTrans.BlockIndex == CheckIndex &&
+				CurrentTrans.TransactionType != LastTrans.TransactionType )
+			{
+				RemoveTransactionAt( Transactions.Num()-1 );
+				RemoveTransactionAt( TransIndex );
+
+				return true;
+			}
+
+			if( CheckIndex >= CurrentTrans.BlockIndex )
+			{
+				if( CurrentTrans.TransactionType == FCustomBlockTransaction::Add )
+				{
+					--CheckIndex;
+				}
+				else
+				{
+					++CheckIndex;
+				}
+			}
+		}
+	}
+
+	return false;
+}
+
+bool FMultiBoxCustomizationData::RemoveUnnecessaryTransactions(const TArray< TSharedRef< const FMultiBlock > >& AllBlocks)
+{
+	// A local struct describing the state of the menu
+	struct FCustomizationState
+	{
+		FCustomizationState(const TArray< TSharedRef< const FMultiBlock > >& InAllBlocks)
+		{
+			for ( int32 BlockIdx = 0; BlockIdx < InAllBlocks.Num(); ++BlockIdx )
+			{
+				StateData.Add(InAllBlocks[BlockIdx]->GetAction());
+			}
+		}
+
+		// Applies the inverse operation of the given transaction to mutate the state
+		void ApplyInverseOfTransaction(const FCustomBlockTransaction& Trans)
+		{
+			if ( Trans.TransactionType == FCustomBlockTransaction::Add )
+			{
+				// Remove on an add operation
+				if ( ensure(StateData.IsValidIndex(Trans.BlockIndex)) )
+				{
+					ensure(StateData[Trans.BlockIndex] == Trans.Command);
+
+					StateData.RemoveAt(Trans.BlockIndex);
+				}
+			}
+			else if ( Trans.TransactionType == FCustomBlockTransaction::Remove )
+			{
+				// Add on a remove transaction
+				StateData.Insert(Trans.Command, Trans.BlockIndex);
+			}
+		}
+
+		// Compares two states. If all elements of the states are the same then they are "equal"
+		bool operator==(const FCustomizationState& Other) const
+		{
+			if ( StateData.Num() != Other.StateData.Num() )
+			{
+				return false;
+			}
+
+			const int32 Length = StateData.Num();
+			for ( int32 StateIdx = 0; StateIdx < Length; ++StateIdx )
+			{
+				if ( StateData[StateIdx] != Other.StateData[StateIdx] )
+				{
+					return false;
+				}
+			}
+
+			return true;
+		}
+
+	private:
+		TArray<TWeakPtr<const FUICommandInfo>> StateData;
+	};
+
+	if( Transactions.Num() > 0 )
+	{
+		// Take the current state and start applying the inverse of each transaction to it to determine the states before the current.
+		// If we found any state that is identical to the current state, all transactions between that state and the current state are unnecessary, so remove them
+
+		// Record the initial (current) state and initialize the Test state.
+		FCustomizationState InitialState(AllBlocks);
+		FCustomizationState TestState = InitialState;
+
+		// Walk backwards through all transactions, applying the inverse of each transaction to the test state
+		for ( int32 TransIdx = Transactions.Num() - 1; TransIdx >= 0; --TransIdx )
+		{
+			TestState.ApplyInverseOfTransaction(Transactions[TransIdx]);
+
+			// If the test state is equal to the current state, all transactions between are unnecessary. Remove them.
+			if ( TestState == InitialState )
+			{
+				for ( int32 RemoveIdx = Transactions.Num() - 1; RemoveIdx >= TransIdx; --RemoveIdx )
+				{
+					RemoveTransactionAt(RemoveIdx);
+				}
+
+				// Return true to indicate we actually removed something.
+				return true;
+			}
+		}
+	}
+
+	// All transactions were necessary to form the final state. Return false.
+	return false;
+}
+
+void FMultiBoxCustomizationData::SaveCustomizedBlocks()
+{
+	FString SaveData;
+	TSharedRef< TJsonWriter<> > Writer = TJsonWriterFactory<>::Create( &SaveData );
+
+	Writer->WriteObjectStart();
+	Writer->WriteArrayStart(TEXT("CustomBlocks"));
+	for (TArray<FCustomBlockTransaction>::TConstIterator It(Transactions); It; ++It)
+	{
+		const FCustomBlockTransaction& Transaction = *It;
+
+		if( Transaction.Command.IsValid() )
+		{
+			const FUICommandInfo& Command = *Transaction.Command.Pin();
+
+			Writer->WriteObjectStart();
+			Writer->WriteValue(TEXT("CommandName"), Command.GetCommandName().ToString() );
+			Writer->WriteValue(TEXT("Context"), Command.GetBindingContext().ToString() );
+			Writer->WriteValue(TEXT("Index"), (float)Transaction.BlockIndex );
+			Writer->WriteValue(TEXT("TransactionType"), (float)Transaction.TransactionType );
+			Writer->WriteObjectEnd();
+		}
+	}
+
+	Writer->WriteArrayEnd();
+	Writer->WriteObjectEnd();
+	Writer->Close();
+
+	GConfig->SetString( *GetConfigSectionName(), *CustomizationName.ToString(), *FRemoteConfig::ReplaceIniCharWithSpecialChar(SaveData).ReplaceCharWithEscapedChar(), GEditorUserSettingsIni );
+}
+
+void FMultiBoxCustomizationData::BlockRemoved( TSharedRef< const FMultiBlock> RemovedBlock, int32 Index, const TArray< TSharedRef< const FMultiBlock > >& AllBlocks )
+{
+	FCustomBlockTransaction Remove = FCustomBlockTransaction::CreateRemove( RemovedBlock->GetAction().ToSharedRef(), Index );
+
+	SaveTransaction( Remove, AllBlocks );
+
+}
+
+void FMultiBoxCustomizationData::BlockAdded( TSharedRef< const FMultiBlock > AddedBlock, int32 Index, const TArray< TSharedRef< const FMultiBlock > >& AllBlocks )
+{
+	FCustomBlockTransaction Add = FCustomBlockTransaction::CreateAdd( AddedBlock->GetAction().ToSharedRef(), Index );
+
+	SaveTransaction( Add, AllBlocks );
+}
+
+#define DEBUG_TRANSACTIONS 0
+#if DEBUG_TRANSACTIONS
+static void PrintTransactions( const TArray<FCustomBlockTransaction>& Transactions )
+{
+	FPlatformMisc::LowLevelOutputDebugStringf(TEXT("==========BEGIN TRANSACTIONS=======\n"));
+	for( int32 i = 0; i < Transactions.Num(); ++i )
+	{
+		const FCustomBlockTransaction& Trans = Transactions[i];
+
+		FString TransType = Trans.TransactionType == FCustomBlockTransaction::Add ? TEXT("+") : TEXT("-");
+		FPlatformMisc::LowLevelOutputDebugStringf( TEXT("%s(%s,%d)\n"), *TransType, *Trans.Command.Pin()->GetCommandName().ToString(), Trans.BlockIndex );
+	}
+	FPlatformMisc::LowLevelOutputDebugStringf(TEXT("============END TRANSACTIONS=======\n"));
+}
+#endif
+
+void FMultiBoxCustomizationData::SaveTransaction( const struct FCustomBlockTransaction& Transaction, const TArray< TSharedRef< const FMultiBlock > >& AllBlocks )
+{
+	Transactions.Add( Transaction );
+
+	while( RemoveDuplicateTransaction() );
+	while( RemoveUnnecessaryTransactions(AllBlocks) );
+
+#if DEBUG_TRANSACTIONS
+	PrintTransactions( Transactions );
+#endif
+
+	SaveCustomizedBlocks();
+}
diff --git a/Engine/Source/Runtime/Slate/Private/Framework/MultiBox/MultiBoxCustomization.h b/Engine/Source/Runtime/Slate/Private/Framework/MultiBox/MultiBoxCustomization.h
new file mode 100644
--- /dev/null
+++ b/Engine/Source/Runtime/Slate/Private/Framework/MultiBox/MultiBoxCustomization.h
@@ -0,0 +1,132 @@
+// Copyright 1998-2014 Epic Games, Inc. All Rights Reserved.
+
+#pragma once
+
+#include "MultiBox.h"
+
+class SCustomToolbarPreviewWidget : public SMultiBlockBaseWidget
+{
+public:
+	SLATE_BEGIN_ARGS( SCustomToolbarPreviewWidget ) {}
+		SLATE_DEFAULT_SLOT( FArguments, Content )
+	SLATE_END_ARGS()
+
+	void Construct( const FArguments& InArgs );
+
+	virtual void BuildMultiBlockWidget(const ISlateStyle* StyleSet, const FName& StyleName) OVERRIDE;
+private:
+	TSharedPtr<SWidget> Content;
+};
+
+class SMultiBlockDragHandle : public SCompoundWidget
+{
+	SLATE_BEGIN_ARGS( SMultiBlockDragHandle ) {}
+	SLATE_END_ARGS()
+
+	void Construct( const FArguments& InArgs, TSharedRef<SMultiBoxWidget> InBaseWidget, TSharedRef<const FMultiBlock> InBlock, FName CustomizationName );
+
+	FReply OnMouseButtonDown( const FGeometry& MyGeometry, const FPointerEvent& MouseEvent );
+
+	virtual void OnDragEnter( const FGeometry& MyGeometry, const FDragDropEvent& DragDropEvent ) OVERRIDE;
+	virtual FReply OnDragOver( const FGeometry& MyGeometry, const FDragDropEvent& DragDropEvent ) OVERRIDE;
+	virtual FReply OnDrop( const FGeometry& MyGeometry, const FDragDropEvent& DragDropEvent ) OVERRIDE;
+	virtual FReply OnDragDetected( const FGeometry& MyGeometry, const FPointerEvent& MouseEvent ) OVERRIDE;
+
+private:
+	TSharedPtr<const FMultiBlock> Block;
+	TWeakPtr<SMultiBoxWidget> BaseWidget;
+	FName MultiBoxCustomizationName;
+};
+
+
+/**
+ * Arbitrary Widget MultiBlock
+ */
+class FDropPreviewBlock
+	: public FMultiBlock
+{
+
+public:
+
+	FDropPreviewBlock( TSharedRef<const FMultiBlock> InActualBlock, TSharedRef<IMultiBlockBaseWidget> InActualWidget ) 
+		: FMultiBlock( NULL, NULL )
+		, ActualBlock( InActualBlock )
+		, ActualWidget( InActualWidget )
+	{
+	}
+
+	/** FMultiBlock interface */
+	virtual TSharedRef< class IMultiBlockBaseWidget > ConstructWidget() const OVERRIDE;
+	TSharedRef<const FMultiBlock> GetActualBlock() const { return ActualBlock.ToSharedRef(); }
+private:
+	TSharedPtr<const FMultiBlock> ActualBlock;
+	TSharedPtr<IMultiBlockBaseWidget> ActualWidget;
+
+};
+
+struct FCustomBlockTransaction
+{
+	enum ETransactionType
+	{
+		Remove,
+		Add
+	};
+
+	TWeakPtr<const FUICommandInfo> Command;
+	int32 BlockIndex;
+	ETransactionType TransactionType;
+
+	bool operator==(const FCustomBlockTransaction& Other ) const
+	{
+		return Command == Other.Command && BlockIndex == Other.BlockIndex && TransactionType == Other.TransactionType;
+	}
+
+	static FCustomBlockTransaction CreateRemove( TSharedRef<const FUICommandInfo> InCommand, int32 Index )
+	{
+		FCustomBlockTransaction Transaction;
+		Transaction.Command = InCommand;
+		Transaction.BlockIndex = Index;
+		Transaction.TransactionType = ETransactionType::Remove;
+
+		return Transaction;
+	}
+
+	static FCustomBlockTransaction CreateAdd( TSharedRef<const FUICommandInfo> InCommand, int32 Index )
+	{
+		FCustomBlockTransaction Transaction;
+		Transaction.Command = InCommand;
+		Transaction.BlockIndex = Index;
+		Transaction.TransactionType = ETransactionType::Add;
+
+		return Transaction;
+	}
+};
+
+class FMultiBoxCustomizationData
+{
+public:
+	FMultiBoxCustomizationData( FName InCustomizationName )
+		: CustomizationName( InCustomizationName )
+	{}
+
+	FName GetCustomizationName() const { return CustomizationName; }
+
+	void LoadCustomizedBlocks();
+
+	void SaveCustomizedBlocks();
+	void BlockRemoved( TSharedRef< const FMultiBlock> RemovedBlock, int32 Index, const TArray< TSharedRef< const FMultiBlock > >& AllBlocks );
+	void BlockAdded( TSharedRef< const FMultiBlock > AddedBlock, int32 Index, const TArray< TSharedRef< const FMultiBlock > >& AllBlocks );
+	void SaveTransaction( const struct FCustomBlockTransaction& Transaction, const TArray< TSharedRef< const FMultiBlock > >& AllBlocks );
+
+	uint32 GetNumTransactions() const { return Transactions.Num(); }
+
+	const FCustomBlockTransaction& GetTransaction( uint32 TransactionIndex ) { return Transactions[TransactionIndex]; }
+private:
+	FString GetConfigSectionName() const;
+	void RemoveTransactionAt(int32 TransIndex);
+	bool RemoveDuplicateTransaction();
+	bool RemoveUnnecessaryTransactions(const TArray< TSharedRef< const FMultiBlock > >& AllBlocks);
+private:
+	TArray<FCustomBlockTransaction> Transactions;
+	FName CustomizationName;
+};
\ No newline at end of file
diff --git a/Engine/Source/Runtime/Slate/Private/Framework/MultiBox/SClippingHorizontalBox.cpp b/Engine/Source/Runtime/Slate/Private/Framework/MultiBox/SClippingHorizontalBox.cpp
new file mode 100644
--- /dev/null
+++ b/Engine/Source/Runtime/Slate/Private/Framework/MultiBox/SClippingHorizontalBox.cpp
@@ -0,0 +1,158 @@
+// Copyright 1998-2014 Epic Games, Inc. All Rights Reserved.
+
+#include "SlatePrivatePCH.h"
+#include "SClippingHorizontalBox.h"
+
+void SClippingHorizontalBox::Tick( const FGeometry& AllottedGeometry, const double InCurrentTime, const float InDeltaTime )
+{
+	// Call ArrangeChildren so we can cache off the index of the first clipped tab
+	FArrangedChildren ArrangedChildren(EVisibility::Visible);
+	this->ArrangeChildren(AllottedGeometry, ArrangedChildren);
+	ClippedIdx = ArrangedChildren.Num() - 1;
+}
+
+void SClippingHorizontalBox::ArrangeChildren( const FGeometry& AllottedGeometry, FArrangedChildren& ArrangedChildren ) const
+{
+	// If WrapButton hasn't been initialized, that means AddWrapButton() hasn't 
+	// been called and this method isn't going to behave properly
+	check(WrapButton.IsValid());
+
+	SHorizontalBox::ArrangeChildren(AllottedGeometry, ArrangedChildren);
+
+	// Remove children that are clipped by the allotted geometry
+	const int32 NumChildren = ArrangedChildren.Num(); 
+	int32 IndexClippedAt = NumChildren;
+	for (int32 ChildIdx = NumChildren - 2; ChildIdx >= 0; --ChildIdx)
+	{
+		const FArrangedWidget& CurWidget = ArrangedChildren(ChildIdx);
+		if (FMath::Trunc(CurWidget.Geometry.AbsolutePosition.X + CurWidget.Geometry.Size.X) > FMath::Trunc(AllottedGeometry.AbsolutePosition.X + AllottedGeometry.Size.X))
+		{
+			ArrangedChildren.Remove(ChildIdx);
+			IndexClippedAt = ChildIdx;
+		}
+	}
+
+	if (IndexClippedAt == NumChildren)
+	{
+		// None of the children are being clipped, so remove the wrap button
+		ArrangedChildren.Remove(ArrangedChildren.Num() - 1);
+	}
+	else
+	{
+		// Right align the wrap button
+		FArrangedWidget& ArrangedButton = ArrangedChildren(ArrangedChildren.Num() - 1);
+		ArrangedButton.Geometry.AbsolutePosition = AllottedGeometry.AbsolutePosition + AllottedGeometry.Size - ArrangedButton.Geometry.Size;
+
+		// Further remove any children that the wrap button overlaps with
+		for (int32 ChildIdx = IndexClippedAt - 1; ChildIdx >= 0; --ChildIdx)
+		{
+			const FArrangedWidget& CurWidget = ArrangedChildren(ChildIdx);
+			if (FMath::Trunc(CurWidget.Geometry.AbsolutePosition.X + CurWidget.Geometry.Size.X) > FMath::Trunc(ArrangedButton.Geometry.AbsolutePosition.X))
+			{
+				ArrangedChildren.Remove(ChildIdx);
+			}
+		}
+	}
+}
+
+int32 SClippingHorizontalBox::OnPaint( const FGeometry& AllottedGeometry, const FSlateRect& MyClippingRect, FSlateWindowElementList& OutDrawElements, int32 LayerId, const FWidgetStyle& InWidgetStyle, bool bParentEnabled ) const
+{
+	// Get the clipped children info
+	FArrangedChildren ClippedArrangedChildren(EVisibility::Visible);
+	ArrangeChildren(AllottedGeometry, ClippedArrangedChildren);
+	
+	// Get the non-clipped children info
+	FArrangedChildren ArrangedChildren(EVisibility::Visible);
+	SBoxPanel::ArrangeChildren(AllottedGeometry, ArrangedChildren);
+	
+	if ((ClippedArrangedChildren.Num() != 0) && (ArrangedChildren.Num() != 0))
+	{
+		int32 IndexClippedAt = ClippedArrangedChildren.Num() - 1;
+		const FArrangedWidget& LastCippedChild = ClippedArrangedChildren(IndexClippedAt);
+		const FArrangedWidget& FirstChild = ArrangedChildren(0);
+		const FArrangedWidget& LastChild = ArrangedChildren(ArrangedChildren.Num() - 1);
+		FGeometry ResizedGeometry = AllottedGeometry;
+	
+		// If only the last child/block, which is the wrap button, is being clipped
+		if (IndexClippedAt == ArrangedChildren.Num() - 2)
+		{
+			// Only recalculate the alloted geometry size if said size is fitted to the toolbar/menubar
+			if (FMath::Trunc(AllottedGeometry.AbsolutePosition.X + AllottedGeometry.Size.X) <= FMath::Trunc(LastChild.Geometry.AbsolutePosition.X + LastChild.Geometry.Size.X))
+			{
+				// Calculate the size of the custom border
+				ResizedGeometry.Size.X = LastCippedChild.Geometry.AbsolutePosition.X + LastCippedChild.Geometry.Size.X - FirstChild.Geometry.AbsolutePosition.X;
+			}
+		}
+		else
+		{
+			// Children/blocks are being clipped, calculate the size of the custom border
+			const FArrangedWidget& NextChild = (IndexClippedAt + 1 < ClippedArrangedChildren.Num())? ClippedArrangedChildren(IndexClippedAt + 1): LastCippedChild;
+			ResizedGeometry.Size.X = NextChild.Geometry.AbsolutePosition.X + NextChild.Geometry.Size.X - FirstChild.Geometry.AbsolutePosition.X;
+		}
+	
+		bool bEnabled = ShouldBeEnabled( bParentEnabled );
+		ESlateDrawEffect::Type DrawEffects = !bEnabled ? ESlateDrawEffect::DisabledEffect : ESlateDrawEffect::None;
+		FSlateColor BorderBackgroundColor = FLinearColor::White;
+
+		// Draw the custom border
+		FSlateDrawElement::MakeBox(
+			OutDrawElements,
+			LayerId,
+			ResizedGeometry.ToPaintGeometry(),
+			BackgroundBrush,
+			MyClippingRect,
+			DrawEffects,
+			BackgroundBrush->GetTint( InWidgetStyle ) * InWidgetStyle.GetColorAndOpacityTint() * BorderBackgroundColor.GetColor(InWidgetStyle)
+			);
+	}
+
+	return SHorizontalBox::OnPaint(AllottedGeometry, MyClippingRect, OutDrawElements, LayerId, InWidgetStyle, bParentEnabled);
+}
+
+FVector2D SClippingHorizontalBox::ComputeDesiredSize() const
+{
+	FVector2D Size = SBoxPanel::ComputeDesiredSize();
+	if (ClippedIdx == (Children.Num() - 2))
+	{
+		// If the wrap button isn't being shown, subtract it's size from the total desired size
+		const SBoxPanel::FSlot& Child = Children[Children.Num() - 1];
+		const FVector2D& ChildDesiredSize = Child.Widget->GetDesiredSize();
+		Size.X -= ChildDesiredSize.X;
+	}
+	return Size;
+}
+
+void SClippingHorizontalBox::Construct( const FArguments& InArgs )
+{
+	BackgroundBrush = InArgs._BackgroundBrush;
+	OnWrapButtonClicked = InArgs._OnWrapButtonClicked;
+	StyleSet = InArgs._StyleSet;
+	StyleName = InArgs._StyleName;
+}
+
+void SClippingHorizontalBox::AddWrapButton()
+{
+	// Construct the wrap button used in toolbars and menubars
+	WrapButton = 
+		SNew( SComboButton )
+		.HasDownArrow( false )
+		.ButtonStyle( StyleSet, ISlateStyle::Join(StyleName, ".Button") )
+		.ContentPadding( 0 )
+		.ToolTipText( NSLOCTEXT("Slate", "ExpandToolbar", "Click to expand toolbar").ToString() )
+		.OnGetMenuContent( OnWrapButtonClicked )
+		.Cursor( EMouseCursor::Default )
+		.ButtonContent()
+		[
+			SNew( SImage )
+			.Image( StyleSet->GetBrush(StyleName, ".Expand") )
+		];
+
+	// Add the wrap button
+	AddSlot()
+	.AutoWidth()
+	.Padding( 0 )
+	[
+		WrapButton.ToSharedRef()
+	];
+}
+
diff --git a/Engine/Source/Runtime/Slate/Private/Framework/MultiBox/SClippingHorizontalBox.h b/Engine/Source/Runtime/Slate/Private/Framework/MultiBox/SClippingHorizontalBox.h
new file mode 100644
--- /dev/null
+++ b/Engine/Source/Runtime/Slate/Private/Framework/MultiBox/SClippingHorizontalBox.h
@@ -0,0 +1,56 @@
+// Copyright 1998-2014 Epic Games, Inc. All Rights Reserved.
+
+#pragma once
+
+/** Specialized control for handling the clipping of toolbars and menubars */
+class SClippingHorizontalBox : public SHorizontalBox
+{
+public:
+	SLATE_BEGIN_ARGS(SClippingHorizontalBox) 
+		: _BackgroundBrush(NULL)
+		, _StyleSet(&FCoreStyle::Get())
+		, _StyleName(NAME_None)
+		{ }
+
+		SLATE_ARGUMENT(const FSlateBrush*, BackgroundBrush)
+		SLATE_ARGUMENT(FOnGetContent, OnWrapButtonClicked)
+		SLATE_ARGUMENT(const ISlateStyle*, StyleSet)
+		SLATE_ARGUMENT(FName, StyleName)
+	SLATE_END_ARGS()
+
+	FORCENOINLINE SClippingHorizontalBox() { }
+
+	/** SWidget interface */
+	virtual void Tick( const FGeometry& AllottedGeometry, const double InCurrentTime, const float InDeltaTime ) OVERRIDE;
+	virtual int32 OnPaint( const FGeometry& AllottedGeometry, const FSlateRect& MyClippingRect, FSlateWindowElementList& OutDrawElements, int32 LayerId, const FWidgetStyle& InWidgetStyle, bool bParentEnabled ) const OVERRIDE;
+	virtual FVector2D ComputeDesiredSize() const OVERRIDE;
+
+	/** SPanel interface */
+	virtual void ArrangeChildren( const FGeometry& AllottedGeometry, FArrangedChildren& ArrangedChildren ) const OVERRIDE;
+
+	/** Construct this widget */
+	void Construct( const FArguments& InArgs );
+
+	/** Adds the wrap button */
+	void AddWrapButton();
+
+	/** Returns to index of the first clipped child/block */
+	int32 GetClippedIndex() { return ClippedIdx; }
+
+private:
+	/** The button that is displayed when a toolbar or menubar is clipped */
+	TSharedPtr<SComboButton> WrapButton;
+
+	/** Brush used for drawing the custom border */
+	const FSlateBrush* BackgroundBrush;
+
+	/** Callback for when the wrap button is clicked */
+	FOnGetContent OnWrapButtonClicked;
+
+	/** Index of the first clipped child/block */
+	int32 ClippedIdx;
+
+	/** The style to use */
+	const ISlateStyle* StyleSet;
+	FName StyleName;
+};
diff --git a/Engine/Source/Runtime/Slate/Private/Framework/MultiBox/SEditableTextBlock.h b/Engine/Source/Runtime/Slate/Private/Framework/MultiBox/SEditableTextBlock.h
new file mode 100644
--- /dev/null
+++ b/Engine/Source/Runtime/Slate/Private/Framework/MultiBox/SEditableTextBlock.h
@@ -0,0 +1,106 @@
+// Copyright 1998-2014 Epic Games, Inc. All Rights Reserved.
+
+#pragma once
+
+/**
+ * Editable text block
+ */
+class FEditableTextBlock
+	: public FMultiBlock
+{
+
+public:
+
+	/**
+	 * Constructor
+	 *
+	 * @param	InLabel				The label to display in the menu
+ 	 * @param	InToolTip			The tool tip to display when the menu entry is hovered over
+ 	 * @param	InIcon				The icon to display to the left of the label
+	 * @param	InTextAttribute		The text string we're editing (often, a delegate will be bound to the attribute)
+	 * @param	bInReadOnly			Whether or not the text block should be read only
+	 * @param	InOnTextCommitted	Called when the user commits their change to the editable text control
+	 * @param	InOnTextChanged		Called when the text is changed interactively
+	 */
+	FEditableTextBlock( const FText& InLabel, const FText& InToolTip, const FSlateIcon& InIcon, const TAttribute< FText >& InTextAttribute, bool bInReadOnly, const FOnTextCommitted& InOnTextCommitted, const FOnTextChanged& InOnTextChanged );
+
+
+private:
+
+	/**
+	 * Allocates a widget for this type of MultiBlock.  Override this in derived classes.
+	 *
+	 * @return  MultiBlock widget object
+	 */
+	virtual TSharedRef< class IMultiBlockBaseWidget > ConstructWidget() const;
+
+
+private:
+
+	// Friend our corresponding widget class
+	friend class SEditableTextBlock;
+
+	/** Optional overridden text label for this menu entry.  If not set, then the action's label will be used instead. */
+	TAttribute<FText> LabelOverride;
+
+	/** Optional overridden tool tip for this menu entry.  If not set, then the action's tool tip will be used instead. */
+	TAttribute<FText> ToolTipOverride;
+
+	/** Optional overridden icon for this tool bar button.  IF not set, then the action's icon will be used instead. */
+	FSlateIcon IconOverride;
+
+	/** The text we're editing */
+	TAttribute< FText > TextAttribute;
+
+	/** Called when the user commits their change to the editable text control */
+	FOnTextCommitted OnTextCommitted;
+
+	/** Called when the text is changed interactively */
+	FOnTextChanged OnTextChanged;
+
+	/** If true the text box is read only */
+	bool bReadOnly;
+	
+};
+
+
+
+
+/**
+ * Editable text block widget
+ */
+class SLATE_API SEditableTextBlock
+	: public SMultiBlockBaseWidget
+{
+
+public:
+
+	SLATE_BEGIN_ARGS( SEditableTextBlock ){}
+	SLATE_END_ARGS()
+
+
+	/**
+	 * Builds this MultiBlock widget up from the MultiBlock associated with it
+	 */
+	virtual void BuildMultiBlockWidget(const ISlateStyle* StyleSet, const FName& StyleName) OVERRIDE;
+
+	/**
+	 * Construct this widget
+	 *
+	 * @param	InArgs	The declaration data for this widget
+	 */
+	void Construct( const FArguments& InArgs );
+
+protected:
+
+	/**
+	 * Called by Slate to determine if this widget is enabled
+	 * 
+	 * @return True if the widget is enabled, false otherwise
+	 */
+	bool IsEnabled() const;
+
+
+private:
+
+};
diff --git a/Engine/Source/Runtime/Slate/Private/Framework/MultiBox/SGroupMarkerBlock.h b/Engine/Source/Runtime/Slate/Private/Framework/MultiBox/SGroupMarkerBlock.h
new file mode 100644
--- /dev/null
+++ b/Engine/Source/Runtime/Slate/Private/Framework/MultiBox/SGroupMarkerBlock.h
@@ -0,0 +1,72 @@
+// Copyright 1998-2014 Epic Games, Inc. All Rights Reserved.
+
+#pragma once
+
+/**
+ * Group Start MultiBlock
+ */
+class FGroupStartBlock
+	: public FMultiBlock
+{
+public:
+	FGroupStartBlock();
+
+	virtual bool IsGroupStartBlock()	const	{ return true; };
+private:
+
+	/**
+	 * Allocates a widget for this type of MultiBlock.  Override this in derived classes.
+	 *
+	 * @return  MultiBlock widget object
+	 */
+	virtual TSharedRef< class IMultiBlockBaseWidget > ConstructWidget() const;
+};
+
+/**
+ * Group End MultiBlock
+ */
+class FGroupEndBlock
+	: public FMultiBlock
+{
+
+public:
+	FGroupEndBlock();
+
+	virtual bool IsGroupEndBlock()	const	{ return true; };
+private:
+
+	/**
+	 * Allocates a widget for this type of MultiBlock.  Override this in derived classes.
+	 *
+	 * @return  MultiBlock widget object
+	 */
+	virtual TSharedRef< class IMultiBlockBaseWidget > ConstructWidget() const;
+};
+
+/**
+ * Group Marker MultiBlock widget
+ */
+class SLATE_API SGroupMarkerBlock 
+	: public SMultiBlockBaseWidget
+{
+
+public:
+
+	SLATE_BEGIN_ARGS( SGroupMarkerBlock ){}
+	SLATE_END_ARGS()
+
+	/**
+	 * Builds this MultiBlock widget up from the MultiBlock associated with it
+	 */
+	virtual void BuildMultiBlockWidget(const ISlateStyle* StyleSet, const FName& StyleName) OVERRIDE;
+
+	/**
+	 * Construct this widget
+	 *
+	 * @param	InArgs	The declaration data for this widget
+	 */
+	void Construct( const FArguments& InArgs ) {};
+
+private:
+
+};
\ No newline at end of file
diff --git a/Engine/Source/Runtime/Slate/Private/Framework/MultiBox/SHeadingBlock.h b/Engine/Source/Runtime/Slate/Private/Framework/MultiBox/SHeadingBlock.h
new file mode 100644
--- /dev/null
+++ b/Engine/Source/Runtime/Slate/Private/Framework/MultiBox/SHeadingBlock.h
@@ -0,0 +1,73 @@
+// Copyright 1998-2014 Epic Games, Inc. All Rights Reserved.
+
+#pragma once
+
+/**
+ * Heading MultiBlock
+ */
+class FHeadingBlock
+	: public FMultiBlock
+{
+
+public:
+
+	/**
+	 * Constructor
+	 *
+	 * @param	InHeadingText	Heading text
+	 */
+	FHeadingBlock( const FName& InExtensionHook, const TAttribute< FText >& InHeadingText );
+
+
+private:
+
+	/**
+	 * Allocates a widget for this type of MultiBlock.  Override this in derived classes.
+	 *
+	 * @return  MultiBlock widget object
+	 */
+	virtual TSharedRef< class IMultiBlockBaseWidget > ConstructWidget() const;
+
+
+private:
+
+	// Friend our corresponding widget class
+	friend class SHeadingBlock;
+
+	/** Text for this heading */
+	TAttribute< FText > HeadingText;
+};
+
+
+
+
+/**
+ * Heading MultiBlock widget
+ */
+class SLATE_API SHeadingBlock
+	: public SMultiBlockBaseWidget
+{
+
+public:
+
+	SLATE_BEGIN_ARGS( SHeadingBlock ){}
+
+	SLATE_END_ARGS()
+
+
+	/**
+	 * Builds this MultiBlock widget up from the MultiBlock associated with it
+	 */
+	virtual void BuildMultiBlockWidget(const ISlateStyle* StyleSet, const FName& StyleName) OVERRIDE;
+
+
+	/**
+	 * Construct this widget
+	 *
+	 * @param	InArgs	The declaration data for this widget
+	 */
+	void Construct( const FArguments& InArgs );
+
+private:
+
+};
diff --git a/Engine/Source/Runtime/Slate/Private/Framework/MultiBox/SMenuEntryBlock.h b/Engine/Source/Runtime/Slate/Private/Framework/MultiBox/SMenuEntryBlock.h
new file mode 100644
--- /dev/null
+++ b/Engine/Source/Runtime/Slate/Private/Framework/MultiBox/SMenuEntryBlock.h
@@ -0,0 +1,342 @@
+// Copyright 1998-2014 Epic Games, Inc. All Rights Reserved.
+
+#pragma once
+
+/**
+ * Menu entry MultiBlock
+ */
+class FMenuEntryBlock
+	: public FMultiBlock
+{
+
+public:
+
+	/**
+	 * Constructor
+	 *
+	 * @param	InCommand			The command associated with this menu entry
+	 * @param	InCommandList		The list of commands that are mapped to delegates so that we know what to execute when this menu entry is activated
+	 * @param	InLabelOverride		Optional label override.  If omitted, then the action's label will be used instead.
+	 * @param	InToolTipOverride	Optional tool tip override.	 If omitted, then the action's label will be used instead.
+	 * @param	InIconOverride		Optional icon to use for the tool bar image.  If omitted, then the action's icon will be used instead.
+	 */
+	FMenuEntryBlock( const FName& InExtensionHook, const TSharedPtr< const FUICommandInfo > InCommand, TSharedPtr< const FUICommandList > InCommandList, const TAttribute<FText>& InLabelOverride = TAttribute<FText>(), const TAttribute<FText>& InToolTipOverride = TAttribute<FText>(), const FSlateIcon& InIconOverride = FSlateIcon(), bool bInCloseSelfOnly = false );
+
+
+	/**
+	 * Constructor
+	 *
+	 * @param	InLabel				The label to display in the menu
+	 * @param	InToolTip			The tool tip to display when the menu entry is hovered over
+ 	 * @param	InEntryBuilder		Menu builder object for the menu to add.
+	 * @param	InExtender			The menu extender class to pass down to child menus
+	 * @param	bInSubMenu			True if this menu entry spawns a sub-menu instead of activating a command
+	 * @param	bInSubMenuOnClick	True if this menu entry spawns a sub-menu only by clicking on it
+	 * @param	InCommandList		The list of commands bound to delegates that should be executed for menu entries
+	 * @param	InIcon				The icon to display to the left of the label
+	 */
+	FMenuEntryBlock( const FName& InExtensionHook, const TAttribute<FText>& InLabel, const TAttribute<FText>& InToolTip, const FNewMenuDelegate& InEntryBuilder, TSharedPtr<FExtender> InExtender, bool bInSubMenu, bool bInSubMenuOnClick, TSharedPtr< const FUICommandList > InCommandList, bool bInCloseSelfOnly, const FSlateIcon& InIcon = FSlateIcon());
+	
+
+	/**
+	 * Constructor
+	 *
+	 * @param	InLabel				The label to display in the menu
+	 * @param	InToolTip			The tool tip to display when the menu entry is hovered over
+ 	 * @param	InMenuBuilder		Menu widget object to add.
+	 * @param	InExtender			The menu extender class to pass down to child menus
+	 * @param	bInSubMenu			True if this menu entry spawns a sub-menu instead of activating a command
+	 * @param	bInSubMenuOnClick	True if this menu entry spawns a sub-menu only by clicking on it
+	 * @param	InCommandList		The list of commands bound to delegates that should be executed for menu entries
+	 * @param	InIcon				The icon to display to the left of the label
+	 */
+	FMenuEntryBlock( const FName& InExtensionHook, const TAttribute<FText>& InLabel, const TAttribute<FText>& InToolTip, const FOnGetContent& InMenuBuilder, TSharedPtr<FExtender> InExtender, bool bInSubMenu, bool bInSubMenuOnClick, TSharedPtr< const FUICommandList > InCommandList, bool bInCloseSelfOnly, const FSlateIcon& InIcon = FSlateIcon());
+
+	/**
+	 * Constructor
+	 *
+	 * @param	InLabel				The label to display in the menu
+	 * @param	InToolTip			The tool tip to display when the menu entry is hovered over
+ 	 * @param	InEntryWidget		Menu widget object to add.
+	 * @param	InExtender			The menu extender class to pass down to child menus
+	 * @param	bInSubMenu			True if this menu entry spawns a sub-menu instead of activating a command
+	 * @param	bInSubMenuOnClick	True if this menu entry spawns a sub-menu only by clicking on it
+	 * @param	InCommandList		The list of commands bound to delegates that should be executed for menu entries
+	 * @param	InIcon				The icon to display to the left of the label
+	 */
+	FMenuEntryBlock( const FName& InExtensionHook, const TAttribute<FText>& InLabel, const TAttribute<FText>& InToolTip, const TSharedPtr<SWidget>& InEntryWidget, TSharedPtr<FExtender> InExtender, bool bInSubMenu, bool bInSubMenuOnClick, TSharedPtr< const FUICommandList > InCommandList, bool bInCloseSelfOnly, const FSlateIcon& InIcon = FSlateIcon());
+
+
+	/**
+	 * Constructor for use with dynamic menu entries that do not have commands
+	 *
+	 * @param	InLabel				The label to display in the menu
+ 	 * @param	InToolTip			The tool tip to display when the menu entry is hovered over
+ 	 * @param	InIcon				The icon to display to the left of the label
+	 * @param	InUIAction			UI action to take when this menu item is clicked as well as to determine if the menu entry can be executed or appears "checked"
+	 * @param	InUserInterfaceActionType	Type of interface action
+	 */
+	FMenuEntryBlock( const FName& InExtensionHook, const TAttribute<FText>& InLabel, const TAttribute<FText>& InToolTip, const FNewMenuDelegate& InEntryBuilder, TSharedPtr<FExtender> InExtender, bool bInSubMenu, bool bInSubMenuOnClick, const FSlateIcon& InIcon, const FUIAction& InUIAction, const EUserInterfaceActionType::Type InUserInterfaceActionType, bool bInCloseSelfOnly );
+
+	FMenuEntryBlock( const FName& InExtensionHook, const TAttribute<FText>& InLabel, const TAttribute<FText>& InToolTip, const FSlateIcon& InIcon, const FUIAction& InUIAction, const EUserInterfaceActionType::Type InUserInterfaceActionType, bool bInCloseSelfOnly );
+
+	FMenuEntryBlock( const FName& InExtensionHook, const FUIAction& UIAction, const TSharedRef< SWidget > Contents, const EUserInterfaceActionType::Type InUserInterfaceActionType, bool bInCloseSelfOnly );
+
+	FMenuEntryBlock( const FName& InExtensionHook, const TSharedRef< SWidget > Contents, const FNewMenuDelegate& InEntryBuilder, TSharedPtr<FExtender> InExtender, bool bInSubMenu, bool bInSubMenuOnClick, TSharedPtr< const FUICommandList > InCommandList, bool bInCloseSelfOnly );
+
+	FMenuEntryBlock( const FName& InExtensionHook, const FUIAction& UIAction, const TSharedRef< SWidget > Contents, const FNewMenuDelegate& InEntryBuilder, TSharedPtr<FExtender> InExtender, bool bInSubMenu, TSharedPtr< const FUICommandList > InCommandList, bool bInCloseSelfOnly );
+
+	/** FMultiBlock interface */
+	virtual void CreateMenuEntry(class FMenuBuilder& MenuBuilder) const OVERRIDE;
+
+
+
+private:
+
+	/**
+	 * Allocates a widget for this type of MultiBlock.  Override this in derived classes.
+	 *
+	 * @return  MultiBlock widget object
+	 */
+	virtual TSharedRef< class IMultiBlockBaseWidget > ConstructWidget() const;
+
+
+private:
+
+	// Friend our corresponding widget class
+	friend class SMenuEntryBlock;
+
+	/** Optional overridden text label for this menu entry.  If not set, then the action's label will be used instead. */
+	TAttribute<FText> LabelOverride;
+
+	/** Optional overridden tool tip for this menu entry.  If not set, then the action's tool tip will be used instead. */
+	TAttribute<FText> ToolTipOverride;
+
+	/** Optional overridden icon for this tool bar button.  IF not set, then the action's icon will be used instead. */
+	FSlateIcon IconOverride;
+
+	/** Optional menu entry builder associated with this entry for building sub-menus and pull down menus */
+	FNewMenuDelegate EntryBuilder;
+
+	/** Delegate that returns an entire menu */
+	FOnGetContent MenuBuilder;
+
+	/** Widget to be added to the menu */
+	TSharedPtr<SWidget> EntryWidget;
+
+	/** True if this menu entry opens a sub-menu */
+	bool bIsSubMenu;
+
+	/** True if this menu entry opens a sub-menu by clicking on it only */
+	bool bOpenSubMenuOnClick;
+
+	/** In the case where a command is not bound, the user interface action type to use.  If a command is bound, we
+	    simply use the action type associated with that command. */
+	EUserInterfaceActionType::Type UserInterfaceActionType;
+
+	/** True if the menu should close itself and all its children or the entire open menu stack */
+	bool bCloseSelfOnly;
+
+	/** An extender that this menu entry should pass down to its children, so they get extended properly */
+	TSharedPtr<FExtender> Extender;
+};
+
+
+
+
+/**
+ * Menu entry MultiBlock widget
+ */
+class SLATE_API SMenuEntryBlock
+	: public SMultiBlockBaseWidget
+{
+
+public:
+
+	SLATE_BEGIN_ARGS( SMenuEntryBlock ){}
+
+	SLATE_END_ARGS()
+
+
+	/**
+	 * Builds this MultiBlock widget up from the MultiBlock associated with it
+	 */
+	virtual void BuildMultiBlockWidget(const ISlateStyle* StyleSet, const FName& StyleName) OVERRIDE;
+
+	/**
+	 * Construct this widget
+	 *
+	 * @param	InArgs	The declaration data for this widget
+	 */
+	void Construct( const FArguments& InArgs );
+
+protected:
+	/** Struct for creating menu entry widgets */
+	struct FMenuEntryBuildParams
+	{
+		/** The owning multibox */
+		TSharedPtr< const FMultiBox > MultiBox;
+		/** Our menu entry block */
+		TSharedPtr< const FMenuEntryBlock > MenuEntryBlock;
+		/** UI Command associated with the menu entry */
+		TSharedPtr< const FUICommandInfo > UICommand;
+		/** The text to display */
+		TAttribute<FText> Label;
+		/** The tooltip to display */
+		TAttribute<FText> ToolTip;
+		/** The style set to use */
+		const ISlateStyle* StyleSet;
+		/** The style name to use */
+		FName StyleName;
+	};
+
+
+	/**
+	 * Called by Slate when this menu entry's button is clicked
+	 */
+	FReply OnMenuItemButtonClicked();
+
+	/**
+	 * Called when a checkbox in the menu item or the menu item itself is clicked
+	 *
+	 * @param bCheckBoxClicked	true if a check box was clicked and not the menu item.  We dont close the menu when a check box is clicked
+	 */
+	void OnClicked( bool bCheckBoxClicked );
+
+	/**
+	 * Called by Slate to determine if this menu entry is enabled
+	 * 
+	 * @return True if the menu entry is enabled, false otherwise
+	 */
+	bool IsEnabled() const;
+
+	/**
+	 * Called by Slate when this check box button is toggled in a menu entry
+	 */
+	void OnCheckStateChanged( const ESlateCheckBoxState::Type NewCheckedState );
+
+	/**
+	 * Called by slate to determine if this menu entry should appear checked
+	 *
+	 * @return true if it should be checked, false if not.
+	 */
+	ESlateCheckBoxState::Type IsChecked() const;
+
+	/**
+	 * In the case that we have an icon to show.  This function is called to get the image that indicates the menu item should appear checked
+	 * If we can show an actual check box, this function is not called
+	 */
+	const FSlateBrush* OnGetCheckImage() const;
+
+	// SWidget interface
+	virtual void OnMouseEnter( const FGeometry& MyGeometry, const FPointerEvent& MouseEvent ) OVERRIDE;
+	virtual void OnMouseLeave( const FPointerEvent& MouseEvent ) OVERRIDE;
+	virtual FReply OnKeyDown( const FGeometry& MyGeometry, const FKeyboardEvent& KeyboardEvent ) OVERRIDE;
+	virtual void Tick( const FGeometry& AllottedGeometry, const double InCurrentTime, const float InDeltaTime ) OVERRIDE;
+	// End of SWidget interface
+
+	/**
+	 * Called to create content for a pull-down or sub-menu window when it's summoned by the user
+	 *
+	 * @return	The widget content for the new menu
+	 */
+	TSharedRef< SWidget > MakeNewMenuWidget() const;
+
+	/**
+	 * Called to get the appropriate border for Buttons on Menu Bars based on whether or not submenu is open
+	 *
+	 * @return	The appropriate border to use
+	 */
+	const FSlateBrush* GetMenuBarButtonBorder( ) const;
+
+	/**
+	 * Called to create content for a pull-down menu widget
+	 *
+	 * @param InBuildParams	Parameters for how to build the widget
+	 * @return	The widget content for the new menu
+	 */
+	TSharedRef< SWidget > BuildMenuBarWidget( const FMenuEntryBuildParams& InBuildParams );
+
+	/**
+	 * Called to create content for a menu entry inside a pull-down, context, or sub-menu
+	 *
+	 * @param InBuildParams	Parameters for how to build the widget
+	 * @return The widget content for the new menu
+	 */
+	TSharedRef< SWidget > BuildMenuEntryWidget( const FMenuEntryBuildParams& InBuildParams );
+
+	/**
+	 * Called to create content for a sub-menu widget (entry in a menu that opens another menu to the right of it)
+	 *
+	 * @param InBuildParams	Parameters for how to build the widget
+	 * @return The widget content for the new menu
+	 */
+	TSharedRef< SWidget > BuildSubMenuWidget( const FMenuEntryBuildParams& InBuildParams );
+
+	/**
+	 * Requests that the sub-menu associated with this widget be toggled on or off.
+	 * It does not immediately toggle the menu.  After a set amount of time is passed the menu will toggle
+	 *
+	 * @param bOpenMenu	true to open the menu, false to close the menu if one is currently open
+	 * @param bClobber true if we want to open a menu when another menu is already open
+	 */
+	void RequestSubMenuToggle( bool bOpenMenu, const bool bClobber );
+
+	/**
+	 * Cancels any open requests to toggle a sub-menu       
+	 */
+	void CancelPendingSubMenu();
+
+	/**
+	 * Returns whether or the sub-menu entry should appear hovered.  If the sub-menu is open we will always show the menu as hovered to indicate which sub-menu is open
+	 * In the case that the user is interacting in this menu we do not show the menu as hovered because we need to show what the user is actually selecting
+	 */
+	bool ShouldSubMenuAppearHovered() const;
+	
+	/**
+	 * Called to query the tool tip text for this widget, but will return an empty text for menu bar items
+	 * when a menu for that menu bar is already open
+	 *
+	 * @param	ToolTipText	Tool tip text to display, if possible
+	 *
+	 * @return	Tool tip text, or an empty text if filtered out
+	 */
+	FText GetFilteredToolTipText( TAttribute<FText> ToolTipText ) const;
+
+
+	// Gets the visibility of the menu item
+	EVisibility GetVisibility() const;
+
+	/** Get the selection color when the entry is hovered */
+	FSlateColor TintOnHover() const;
+
+	/** Get the inverted foreground color when the entry is hovered */
+	FSlateColor InvertOnHover() const;
+
+	/** Updates state machine for sub-menu opening logic.  Called in the widget's Tick as well as on demand in some cases. */
+	void UpdateSubMenuState();
+
+private:
+	/** The brush to use when an item should appear checked */
+	const FSlateBrush* CheckedImage;
+	/** The brush to use when an item should appear unchecked */
+	const FSlateBrush* UncheckedImage;
+	/** For pull-down or sub-menu entries, this stores a weak reference to the menu anchor widget that we'll use to summon the menu */
+	TWeakPtr< SMenuAnchor > MenuAnchor;
+	/** The time until a SubMenu action is taken (open or close).  <0 means no action taken */
+	float TimeToSubMenuOpen;
+
+	/** Style for menu bar button with sub menu opened */
+	const FSlateBrush* MenuBarButtonBorderSubmenuOpen;
+	/** Style for menu bar button with no sub menu opened */
+	const FSlateBrush* MenuBarButtonBorderSubmenuClosed;
+
+	/** The pending SubMenu request state */
+	enum
+	{
+		Idle,
+		WantOpen,
+		WantClose
+	} 
+	SubMenuRequestState;
+};
diff --git a/Engine/Source/Runtime/Slate/Private/Framework/MultiBox/SMenuSeparatorBlock.cpp b/Engine/Source/Runtime/Slate/Private/Framework/MultiBox/SMenuSeparatorBlock.cpp
new file mode 100644
--- /dev/null
+++ b/Engine/Source/Runtime/Slate/Private/Framework/MultiBox/SMenuSeparatorBlock.cpp
@@ -0,0 +1,63 @@
+// Copyright 1998-2014 Epic Games, Inc. All Rights Reserved.
+
+#include "SlatePrivatePCH.h"
+#include "MultiBox.h"
+#include "SMenuSeparatorBlock.h"
+
+
+/**
+ * Constructor
+ */
+FMenuSeparatorBlock::FMenuSeparatorBlock(const FName& InExtensionHook)
+	: FMultiBlock( NULL, NULL, InExtensionHook )
+{
+}
+
+
+/**
+ * Allocates a widget for this type of MultiBlock.  Override this in derived classes.
+ *
+ * @return  MultiBlock widget object
+ */
+TSharedRef< class IMultiBlockBaseWidget > FMenuSeparatorBlock::ConstructWidget() const
+{
+	return SNew( SMenuSeparatorBlock );
+}
+
+
+
+/**
+ * Construct this widget
+ *
+ * @param	InArgs	The declaration data for this widget
+ */
+void SMenuSeparatorBlock::Construct( const FArguments& InArgs )
+{
+}
+
+
+
+/**
+ * Builds this MultiBlock widget up from the MultiBlock associated with it
+ */
+void SMenuSeparatorBlock::BuildMultiBlockWidget(const ISlateStyle* StyleSet, const FName& StyleName)
+{
+	ChildSlot.Widget =
+		SNew( SVerticalBox )
+			+SVerticalBox::Slot()
+				.AutoHeight()
+
+				// Add some empty space before the line, and a tiny bit after it
+				.Padding( 0.0f, 4.0f, 0.0f, 2.0f )
+				[
+					SNew( SBorder )
+
+						// We'll use the border's padding to actually create the horizontal line
+						.Padding( StyleSet->GetMargin( StyleName, ".Separator.Padding" ) )
+
+						// Separator graphic
+						.BorderImage( StyleSet->GetBrush( StyleName, ".Separator" ) )
+				];
+}
+
+
diff --git a/Engine/Source/Runtime/Slate/Private/Framework/MultiBox/SMenuSeparatorBlock.h b/Engine/Source/Runtime/Slate/Private/Framework/MultiBox/SMenuSeparatorBlock.h
new file mode 100644
--- /dev/null
+++ b/Engine/Source/Runtime/Slate/Private/Framework/MultiBox/SMenuSeparatorBlock.h
@@ -0,0 +1,69 @@
+// Copyright 1998-2014 Epic Games, Inc. All Rights Reserved.
+
+#pragma once
+
+/**
+ * Menu separator MultiBlock
+ */
+class FMenuSeparatorBlock
+	: public FMultiBlock
+{
+
+public:
+
+	/**
+	 * Constructor
+	 */
+	FMenuSeparatorBlock(const FName& InExtensionHook);
+
+
+private:
+
+	/**
+	 * Allocates a widget for this type of MultiBlock.  Override this in derived classes.
+	 *
+	 * @return  MultiBlock widget object
+	 */
+	virtual TSharedRef< class IMultiBlockBaseWidget > ConstructWidget() const;
+
+
+private:
+
+	// Friend our corresponding widget class
+	friend class SMenuSeparatorBlock;
+
+};
+
+
+
+/**
+ * Menu separator MultiBlock widget
+ */
+class SLATE_API SMenuSeparatorBlock
+	: public SMultiBlockBaseWidget
+{
+
+public:
+
+	SLATE_BEGIN_ARGS( SMenuSeparatorBlock ){}
+
+	SLATE_END_ARGS()
+
+
+	/**
+	 * Builds this MultiBlock widget up from the MultiBlock associated with it
+	 */
+	virtual void BuildMultiBlockWidget(const ISlateStyle* StyleSet, const FName& StyleName) OVERRIDE;
+
+
+	/**
+	 * Construct this widget
+	 *
+	 * @param	InArgs	The declaration data for this widget
+	 */
+	void Construct( const FArguments& InArgs );
+
+
+private:
+
+};
diff --git a/Engine/Source/Runtime/Slate/Private/Framework/MultiBox/SToolBarButtonBlock.cpp b/Engine/Source/Runtime/Slate/Private/Framework/MultiBox/SToolBarButtonBlock.cpp
new file mode 100644
--- /dev/null
+++ b/Engine/Source/Runtime/Slate/Private/Framework/MultiBox/SToolBarButtonBlock.cpp
@@ -0,0 +1,434 @@
+// Copyright 1998-2014 Epic Games, Inc. All Rights Reserved.
+
+#include "SlatePrivatePCH.h"
+#include "MultiBox.h"
+#include "SToolBarButtonBlock.h"
+
+
+FToolBarButtonBlock::FToolBarButtonBlock( const TSharedPtr< const FUICommandInfo > InCommand, TSharedPtr< const FUICommandList > InCommandList, const TAttribute<FText>& InLabelOverride, const TAttribute<FText>& InToolTipOverride, const TAttribute<FSlateIcon>& InIconOverride )
+	: FMultiBlock( InCommand, InCommandList )
+	, LabelOverride( InLabelOverride )
+	, ToolTipOverride( InToolTipOverride )
+	, IconOverride( InIconOverride )
+	, LabelVisibility()
+	, UserInterfaceActionType(EUserInterfaceActionType::Button)
+	, bIsFocusable(false)
+	, bForceSmallIcons(false)
+{
+}
+
+FToolBarButtonBlock::FToolBarButtonBlock( const TAttribute<FText>& InLabel, const TAttribute<FText>& InToolTip, const TAttribute<FSlateIcon>& InIcon, const FUIAction& InUIAction, const EUserInterfaceActionType::Type InUserInterfaceActionType )
+	: FMultiBlock( InUIAction )
+	, LabelOverride( InLabel )
+	, ToolTipOverride( InToolTip )
+	, IconOverride( InIcon )
+	, LabelVisibility()
+	, UserInterfaceActionType(InUserInterfaceActionType)
+	, bIsFocusable(false)
+	, bForceSmallIcons(false)
+{
+}
+
+void FToolBarButtonBlock::CreateMenuEntry(FMenuBuilder& MenuBuilder) const
+{
+	TSharedPtr<const FUICommandInfo> Action = GetAction();
+	TSharedPtr<const FUICommandList> ActionList = GetActionList();
+	if (Action.IsValid() && ActionList.IsValid())
+	{
+		MenuBuilder.PushCommandList(ActionList.ToSharedRef());
+		MenuBuilder.AddMenuEntry(Action);				
+		MenuBuilder.PopCommandList();
+	}
+	else if ( LabelOverride.IsSet() )
+	{
+		const FUIAction& DirectAction = GetDirectActions();
+		MenuBuilder.AddMenuEntry( LabelOverride.Get(), ToolTipOverride.Get(), IconOverride.Get(), DirectAction );
+	}
+}
+
+/**
+ * Allocates a widget for this type of MultiBlock.  Override this in derived classes.
+ *
+ * @return  MultiBlock widget object
+ */
+TSharedRef< class IMultiBlockBaseWidget > FToolBarButtonBlock::ConstructWidget() const
+{
+	return SNew( SToolBarButtonBlock )
+		.LabelVisibility( LabelVisibility.IsSet() ? LabelVisibility.GetValue() : TOptional< EVisibility >() )
+		.IsFocusable( bIsFocusable )
+		.ForceSmallIcons( bForceSmallIcons )
+		.TutorialHighlightName(GetTutorialHighlightName())
+		.Cursor( EMouseCursor::Default );
+}
+
+
+/**
+ * Construct this widget
+ *
+ * @param	InArgs	The declaration data for this widget
+ */
+void SToolBarButtonBlock::Construct( const FArguments& InArgs )
+{
+	if ( InArgs._LabelVisibility.IsSet() )
+	{
+		LabelVisibility = InArgs._LabelVisibility.GetValue();
+	}
+	else
+	{
+		LabelVisibility = TAttribute< EVisibility >::Create( TAttribute< EVisibility >::FGetter::CreateSP( SharedThis( this ), &SToolBarButtonBlock::GetIconVisibility, false ) );
+	}
+
+	bIsFocusable = InArgs._IsFocusable;
+	bForceSmallIcons = InArgs._ForceSmallIcons;
+	TutorialHighlightName = InArgs._TutorialHighlightName;
+}
+
+
+/**
+ * Builds this MultiBlock widget up from the MultiBlock associated with it
+ */
+void SToolBarButtonBlock::BuildMultiBlockWidget(const ISlateStyle* StyleSet, const FName& StyleName)
+{
+	struct Local
+	{
+		/** Appends the key binding to the end of the provided ToolTip */
+		static FText AppendKeyBindingToToolTip( const TAttribute<FText> ToolTip, TWeakPtr< const FUICommandInfo> Command )
+		{
+			TSharedPtr<const FUICommandInfo> CommandPtr = Command.Pin();
+
+			if( CommandPtr.IsValid() && CommandPtr->GetActiveGesture()->IsValidGesture() )
+			{
+				FFormatNamedArguments Args;
+				Args.Add( TEXT("ToolTipDescription"), ToolTip.Get() );
+				Args.Add( TEXT("Keybinding"), CommandPtr->GetInputText() );
+				return FText::Format( NSLOCTEXT("ToolBar", "ToolTip + Keybinding", "{ToolTipDescription} ({Keybinding})"), Args );
+			}
+			else
+			{
+				return ToolTip.Get();
+			}
+		}
+	};
+
+
+	TSharedRef< const FMultiBox > MultiBox( OwnerMultiBoxWidget.Pin()->GetMultiBox() );
+	
+	TSharedRef< const FToolBarButtonBlock > ToolBarButtonBlock = StaticCastSharedRef< const FToolBarButtonBlock >( MultiBlock.ToSharedRef() );
+
+	// Allow the block to override the action's label and tool tip string, if desired
+	TAttribute<FText> ActualLabel;
+	if (ToolBarButtonBlock->LabelOverride.IsSet())
+	{
+		ActualLabel = ToolBarButtonBlock->LabelOverride;
+	}
+	else
+	{
+		ActualLabel = ToolBarButtonBlock->GetAction()->GetLabel();
+	}
+
+	TAttribute<FText> ActualToolTip;
+	if (ToolBarButtonBlock->ToolTipOverride.IsSet())
+	{
+		ActualToolTip = ToolBarButtonBlock->ToolTipOverride;
+	}
+	else
+	{
+		ActualToolTip = ToolBarButtonBlock->GetAction()->GetDescription();
+	}
+
+	// If a key is bound to the command, append it to the tooltip text.
+	TWeakPtr<const FUICommandInfo> Action = ToolBarButtonBlock->GetAction();
+	ActualToolTip = TAttribute< FText >::Create( TAttribute< FText >::FGetter::CreateStatic( &Local::AppendKeyBindingToToolTip, ActualToolTip, Action ) );
+	
+	// If we were supplied an image than go ahead and use that, otherwise we use a null widget
+	TSharedRef< SWidget > IconWidget =
+		SNew( SImage )
+		.Visibility( this, &SToolBarButtonBlock::GetIconVisibility, false )	
+		.Image( this, &SToolBarButtonBlock::GetIconBrush );
+	TSharedRef< SWidget > SmallIconWidget =
+		SNew( SImage )
+		.Visibility( this, &SToolBarButtonBlock::GetIconVisibility, true )	
+		.Image( this, &SToolBarButtonBlock::GetSmallIconBrush );
+
+	// Create the content for our button
+	TSharedRef< SWidget > ButtonContent =
+
+		SNew(SHorizontalBox)
+		+ SHorizontalBox::Slot()
+		.FillWidth(1)
+		.VAlign(VAlign_Center)
+		[
+			SNew( SVerticalBox )
+
+			// Icon image
+			+ SVerticalBox::Slot()
+			.AutoHeight()
+			.HAlign( HAlign_Center )	// Center the icon horizontally, so that large labels don't stretch out the artwork
+			[
+				IconWidget
+			]
+			+ SVerticalBox::Slot().AutoHeight()
+			.HAlign( HAlign_Center )
+			[
+				SmallIconWidget
+			]
+
+			// Label text
+			+ SVerticalBox::Slot().AutoHeight()
+			.HAlign( HAlign_Center )	// Center the label text horizontally
+			[
+				SNew( STextBlock )
+					.Visibility( LabelVisibility )
+					.Text( ActualLabel )
+					.TextStyle( StyleSet, ISlateStyle::Join( StyleName, ".Label" ) )	// Smaller font for tool tip labels
+					.ShadowOffset( FVector2D::UnitVector )
+			]
+		];
+
+	// Add a tutorial highlight wrapper if desired
+	if(TutorialHighlightName != NAME_None)
+	{
+		TSharedRef<SWidget> Wrapper = 
+			SNew(STutorialWrapper, TutorialHighlightName)
+			[
+				ButtonContent
+			];
+
+		ButtonContent = Wrapper;
+	}
+
+	EMultiBlockLocation::Type BlockLocation = GetMultiBlockLocation();
+	
+	// What type of UI should we create for this block?
+	EUserInterfaceActionType::Type UserInterfaceType = ToolBarButtonBlock->UserInterfaceActionType;
+	if ( Action.IsValid() )
+	{
+		// If we have a UICommand, then this is specified in the command.
+		UserInterfaceType = Action.Pin()->GetUserInterfaceType();
+	}
+	
+	if( UserInterfaceType == EUserInterfaceActionType::Button )
+	{
+		FName BlockStyle = EMultiBlockLocation::ToName(ISlateStyle::Join( StyleName, ".Button" ), BlockLocation);
+		ChildSlot.Widget =
+
+			// Create a button
+			SNew( SButton )
+				.ContentPadding(0)
+
+				// Use the tool bar item style for this button
+				.ButtonStyle( StyleSet, BlockStyle )
+
+				.ForegroundColor( FSlateColor::UseForeground() )
+
+				.IsFocusable(bIsFocusable)
+				[
+					ButtonContent
+				]
+
+				// Bind the button's "on clicked" event to our object's method for this
+				.OnClicked( this, &SToolBarButtonBlock::OnClicked )
+
+				// Pass along the block's tool-tip string
+				.ToolTip( FMultiBoxSettings::ToolTipConstructor.Execute( ActualToolTip, NULL, Action.Pin() ) )
+			;
+	}
+	else if( ensure( UserInterfaceType == EUserInterfaceActionType::ToggleButton || UserInterfaceType == EUserInterfaceActionType::RadioButton ) )
+	{
+		FName BlockStyle = EMultiBlockLocation::ToName(ISlateStyle::Join( StyleName, ".ToggleButton" ), BlockLocation);
+		FName CheckboxStyle = ISlateStyle::Join( StyleName, ".SToolBarButtonBlock.CheckBox.Padding" );
+
+		ChildSlot.Widget =
+
+			// Create a check box
+			SNew( SCheckBox )
+
+				// Use the tool bar style for this check box
+				.Style( StyleSet, BlockStyle )
+
+				// User will have set the focusable attribute for the block, honor it
+				.IsFocusable( bIsFocusable )
+
+				// Pass along the block's tool-tip string
+				.ToolTip( FMultiBoxSettings::ToolTipConstructor.Execute( ActualToolTip, NULL, Action.Pin() ) )
+				[
+					ButtonContent
+				]
+
+				// Bind the button's "on checked" event to our object's method for this
+				.OnCheckStateChanged( this, &SToolBarButtonBlock::OnCheckStateChanged )
+
+				// Bind the check box's "checked" state to our user interface action
+				.IsChecked( this, &SToolBarButtonBlock::OnIsChecked )
+
+				.Padding( StyleSet->GetMargin(CheckboxStyle) );
+			;
+	}
+	
+	ChildSlot.Padding(StyleSet->GetMargin(ISlateStyle::Join( StyleName, ".SToolBarButtonBlock.Padding" )));
+
+	// Bind our widget's enabled state to whether or not our action can execute
+	SetEnabled( TAttribute< bool >( this, &SToolBarButtonBlock::IsEnabled ) );
+
+	// Bind our widget's visible state to whether or not the button should be visible
+	SetVisibility( TAttribute<EVisibility>(this, &SToolBarButtonBlock::GetBlockVisibility) );
+}
+
+
+
+/**
+ * Called by Slate when this tool bar button's button is clicked
+ */
+FReply SToolBarButtonBlock::OnClicked()
+{
+	// Button was clicked, so trigger the action!
+	TSharedPtr< const FUICommandList > ActionList = MultiBlock->GetActionList();
+	TSharedPtr< const FUICommandInfo > Action = MultiBlock->GetAction();
+	const FUIAction& DirectActions = MultiBlock->GetDirectActions();
+	
+	if( ActionList.IsValid() && Action.IsValid() )
+	{
+		ActionList->ExecuteAction( Action.ToSharedRef() );
+	}
+	else
+	{
+		// There is no action list or action associated with this block via a UI command.  Execute any direct action we have
+		MultiBlock->GetDirectActions().Execute();
+	}
+
+	TSharedRef< const FMultiBox > MultiBox( OwnerMultiBoxWidget.Pin()->GetMultiBox() );
+
+	// If this is a context menu, then we'll also dismiss the window after the user clicked on the item
+	const bool ClosingMenu = MultiBox->ShouldCloseWindowAfterMenuSelection();
+	if( ClosingMenu )
+	{
+		TSharedRef<SWindow> ParentContextMenuWindow = FSlateApplication::Get().FindWidgetWindow( AsShared() ).ToSharedRef();
+		FSlateApplication::Get().RequestDestroyWindow( ParentContextMenuWindow );
+	}
+
+	return FReply::Handled();
+}
+
+
+
+/**
+ * Called by Slate when this tool bar check box button is toggled
+ */
+void SToolBarButtonBlock::OnCheckStateChanged( const ESlateCheckBoxState::Type NewCheckedState )
+{
+	OnClicked();
+}
+
+/**
+ * Called by slate to determine if this button should appear checked
+ *
+ * @return ESlateCheckBoxState::Checked if it should be checked, ESlateCheckBoxState::Unchecked if not.
+ */
+ESlateCheckBoxState::Type SToolBarButtonBlock::OnIsChecked() const
+{
+	TSharedPtr< const FUICommandList > ActionList = MultiBlock->GetActionList();
+	TSharedPtr< const FUICommandInfo > Action = MultiBlock->GetAction();
+	const FUIAction& DirectActions = MultiBlock->GetDirectActions();
+
+	bool bIsChecked = true;
+	if( ActionList.IsValid() && Action.IsValid() )
+	{
+		bIsChecked = ActionList->IsChecked( Action.ToSharedRef() );
+	}
+	else
+	{
+		// There is no action list or action associated with this block via a UI command.  Execute any direct action we have
+		bIsChecked = DirectActions.IsChecked();
+	}
+
+	return bIsChecked ? ESlateCheckBoxState::Checked : ESlateCheckBoxState::Unchecked;
+}
+
+/**
+ * Called by Slate to determine if this button is enabled
+ * 
+ * @return True if the menu entry is enabled, false otherwise
+ */
+bool SToolBarButtonBlock::IsEnabled() const
+{
+	TSharedPtr< const FUICommandList > ActionList = MultiBlock->GetActionList();
+	TSharedPtr< const FUICommandInfo > Action = MultiBlock->GetAction();
+	const FUIAction& DirectActions = MultiBlock->GetDirectActions();
+
+	bool bEnabled = true;
+	if( ActionList.IsValid() && Action.IsValid() )
+	{
+		bEnabled = ActionList->CanExecuteAction( Action.ToSharedRef() );
+	}
+	else
+	{
+		// There is no action list or action associated with this block via a UI command.  Execute any direct action we have
+		bEnabled = DirectActions.CanExecute();
+	}
+
+	return bEnabled;
+}
+
+
+/**
+ * Called by Slate to determine if this button is visible
+ *
+ * @return EVisibility::Visible or EVisibility::Collapsed, depending on if the button should be displayed
+ */
+EVisibility SToolBarButtonBlock::GetBlockVisibility() const
+{
+	TSharedPtr< const FUICommandList > ActionList = MultiBlock->GetActionList();
+	if( ActionList.IsValid() )
+	{
+		return ActionList->GetVisibility( MultiBlock->GetAction().ToSharedRef() );
+	}
+
+	return EVisibility::Visible;
+}
+
+EVisibility SToolBarButtonBlock::GetIconVisibility(bool bIsASmallIcon) const
+{
+	return ((bForceSmallIcons || FMultiBoxSettings::UseSmallToolBarIcons.Get()) ^ bIsASmallIcon) ? EVisibility::Collapsed : EVisibility::Visible;
+}
+
+const FSlateBrush* SToolBarButtonBlock::GetIconBrush() const
+{
+	TSharedRef< const FToolBarButtonBlock > ToolBarButtonBlock = StaticCastSharedRef< const FToolBarButtonBlock >( MultiBlock.ToSharedRef() );
+
+	const FSlateIcon ActionIcon = ToolBarButtonBlock->GetAction().IsValid() ? ToolBarButtonBlock->GetAction()->GetIcon() : FSlateIcon();
+	const FSlateIcon& ActualIcon = ToolBarButtonBlock->IconOverride.IsSet() ? ToolBarButtonBlock->IconOverride.Get() : ActionIcon;
+	
+	if( ActualIcon.IsSet() )
+	{
+		return ActualIcon.GetIcon();
+	}
+	else
+	{
+		check( OwnerMultiBoxWidget.IsValid() );
+
+		TSharedPtr<SMultiBoxWidget> MultiBoxWidget = OwnerMultiBoxWidget.Pin();
+		const ISlateStyle* const StyleSet = MultiBoxWidget->GetStyleSet();
+		return StyleSet->GetBrush( "MultiBox.GenericToolBarIcon" );
+	}
+}
+
+const FSlateBrush* SToolBarButtonBlock::GetSmallIconBrush() const
+{
+	TSharedRef< const FToolBarButtonBlock > ToolBarButtonBlock = StaticCastSharedRef< const FToolBarButtonBlock >( MultiBlock.ToSharedRef() );
+	
+	const FSlateIcon ActionIcon = ToolBarButtonBlock->GetAction().IsValid() ? ToolBarButtonBlock->GetAction()->GetIcon() : FSlateIcon();
+	const FSlateIcon& ActualIcon = ToolBarButtonBlock->IconOverride.IsSet() ? ToolBarButtonBlock->IconOverride.Get() : ActionIcon;
+	
+	if( ActualIcon.IsSet() )
+	{
+		return ActualIcon.GetSmallIcon();
+	}
+	else
+	{
+		check( OwnerMultiBoxWidget.IsValid() );
+
+		TSharedPtr<SMultiBoxWidget> MultiBoxWidget = OwnerMultiBoxWidget.Pin();
+		const ISlateStyle* const StyleSet = MultiBoxWidget->GetStyleSet();
+		return StyleSet->GetBrush( "MultiBox.GenericToolBarIcon.Small" );
+	}
+}
diff --git a/Engine/Source/Runtime/Slate/Private/Framework/MultiBox/SToolBarComboButtonBlock.cpp b/Engine/Source/Runtime/Slate/Private/Framework/MultiBox/SToolBarComboButtonBlock.cpp
new file mode 100644
--- /dev/null
+++ b/Engine/Source/Runtime/Slate/Private/Framework/MultiBox/SToolBarComboButtonBlock.cpp
@@ -0,0 +1,221 @@
+// Copyright 1998-2014 Epic Games, Inc. All Rights Reserved.
+
+#include "SlatePrivatePCH.h"
+#include "MultiBox.h"
+#include "SToolBarComboButtonBlock.h"
+
+
+FToolBarComboButtonBlock::FToolBarComboButtonBlock( const FUIAction& InAction, const FOnGetContent& InMenuContentGenerator, const TAttribute<FText>& InLabel, const TAttribute<FText>& InToolTip, const TAttribute<FSlateIcon>& InIcon, bool bInSimpleComboBox )
+	: FMultiBlock( InAction )
+	, MenuContentGenerator( InMenuContentGenerator )
+	, Label( InLabel )
+	, ToolTip( InToolTip )
+	, Icon( InIcon )
+	, LabelVisibility()
+	, bSimpleComboBox( bInSimpleComboBox )
+	, bForceSmallIcons( false )
+{
+}
+
+void FToolBarComboButtonBlock::CreateMenuEntry(FMenuBuilder& MenuBuilder) const
+{
+	FName IconName;
+	FText EntryLabel = Label.Get();
+	if ( EntryLabel.IsEmpty() )
+	{
+		EntryLabel = NSLOCTEXT("ToolBar", "CustomControlLabel", "Custom Control");
+	}
+
+	MenuBuilder.AddWrapperSubMenu(EntryLabel, FText::GetEmpty(), MenuContentGenerator, Icon.Get());
+}
+
+TSharedRef< class IMultiBlockBaseWidget > FToolBarComboButtonBlock::ConstructWidget() const
+{
+	return SNew( SToolBarComboButtonBlock )
+		.LabelVisibility( LabelVisibility.IsSet() ? LabelVisibility.GetValue() : TOptional< EVisibility >() )
+		.Icon(Icon)
+		.ForceSmallIcons( bForceSmallIcons )
+		.Cursor( EMouseCursor::Default );
+}
+
+void SToolBarComboButtonBlock::Construct( const FArguments& InArgs )
+{
+	if ( InArgs._LabelVisibility.IsSet() )
+	{
+		LabelVisibility = InArgs._LabelVisibility.GetValue();
+	}
+	else
+	{
+		LabelVisibility = TAttribute< EVisibility >::Create( TAttribute< EVisibility >::FGetter::CreateSP( SharedThis( this ), &SToolBarComboButtonBlock::GetIconVisibility, false ) );
+	}
+
+	Icon = InArgs._Icon;
+	bForceSmallIcons = InArgs._ForceSmallIcons;
+}
+
+void SToolBarComboButtonBlock::BuildMultiBlockWidget(const ISlateStyle* StyleSet, const FName& StyleName)
+{
+	TSharedRef< const FMultiBox > MultiBox( OwnerMultiBoxWidget.Pin()->GetMultiBox() );
+	
+	TSharedRef< const FToolBarComboButtonBlock > ToolBarComboButtonBlock = StaticCastSharedRef< const FToolBarComboButtonBlock >( MultiBlock.ToSharedRef() );
+
+	TAttribute<FText> Label;
+
+	// If we were supplied an image than go ahead and use that, otherwise we use a null widget
+	TSharedRef< SWidget > IconWidget = SNullWidget::NullWidget;
+	TSharedRef< SWidget > SmallIconWidget = SNullWidget::NullWidget;
+	if (!ToolBarComboButtonBlock->bSimpleComboBox)
+	{
+		if( !HasDynamicIcon() )
+		{
+			// Not dynamic, so resolve the icons now
+			const FSlateBrush* IconBrush = GetIcon();
+			const FSlateBrush* SmallIconBrush = GetSmallIcon();
+
+			if( IconBrush->GetResourceName() != NAME_None )
+			{
+				IconWidget =
+					SNew( SImage )
+					.Visibility( this, &SToolBarComboButtonBlock::GetIconVisibility, false )
+					.Image( IconBrush );
+			}
+			if( SmallIconBrush->GetResourceName() != NAME_None )
+			{
+				SmallIconWidget =
+					SNew( SImage )
+					.Visibility( this, &SToolBarComboButtonBlock::GetIconVisibility, true )
+					.Image( SmallIconBrush );
+			}
+		}
+		else 
+		{
+			// Dynamic, so we need to preserve the callback used
+			IconWidget =
+				SNew( SImage )
+				.Visibility( this, &SToolBarComboButtonBlock::GetIconVisibility, false )
+				.Image( this, &SToolBarComboButtonBlock::GetIcon );
+
+			SmallIconWidget =
+				SNew( SImage )
+				.Visibility( this, &SToolBarComboButtonBlock::GetIconVisibility, true )
+				.Image( this, &SToolBarComboButtonBlock::GetSmallIcon );
+		}
+
+		Label = ToolBarComboButtonBlock->Label;
+	}
+
+	// Create the content for our button
+	TSharedRef< SWidget > ButtonContent =
+		SNew( SVerticalBox )
+
+		// Icon image
+		+ SVerticalBox::Slot()
+		.AutoHeight()
+		.HAlign(HAlign_Center)
+		// Center the icon horizontally, so that large labels don't stretch out the artwork
+		[
+			IconWidget
+		]
+		+ SVerticalBox::Slot()
+		.AutoHeight()
+		[
+			SmallIconWidget
+		]
+
+		// Label text
+		+ SVerticalBox::Slot()
+		.AutoHeight()
+		.HAlign( HAlign_Center )	// Center the label text horizontally
+		[
+			SNew( STextBlock )
+				.Visibility( LabelVisibility )
+				.Text( Label )
+
+				// Smaller font for tool tip labels
+				.TextStyle( StyleSet, ISlateStyle::Join( StyleName, ".Label" ) )
+				.ShadowOffset(FVector2D::UnitVector)
+		]
+		;
+		
+	EMultiBlockLocation::Type BlockLocation = GetMultiBlockLocation();
+	FName BlockStyle = EMultiBlockLocation::ToName(ISlateStyle::Join( StyleName, ".Button" ), BlockLocation);
+	FName ColorStyle = ISlateStyle::Join( StyleName, ".SToolBarComboButtonBlock.ComboButton.Color" );
+
+	ChildSlot.Widget =
+		SNew( SComboButton )
+			.ContentPadding(0)
+
+			// Use the tool bar item style for this button
+			.ButtonStyle( StyleSet, BlockStyle )
+
+			// Pass along the block's tool-tip string
+			.ToolTipText( ToolBarComboButtonBlock->ToolTip )
+
+			.ForegroundColor( StyleSet->GetSlateColor( ColorStyle ) )
+
+			.ButtonContent()
+			[
+				ButtonContent
+			]
+
+			// Route the content generator event
+			.OnGetMenuContent( this, &SToolBarComboButtonBlock::OnGetMenuContent )
+		;
+
+	ChildSlot.Padding(StyleSet->GetMargin(ISlateStyle::Join( StyleName, ".SToolBarComboButtonBlock.Padding" )));
+	// Bind our widget's enabled state to whether or not our action can execute
+	SetEnabled( TAttribute< bool >( this, &SToolBarComboButtonBlock::IsEnabled ) );
+
+	// Bind our widget's visible state to whether or not the button should be visible
+	SetVisibility( TAttribute<EVisibility>(this, &SToolBarComboButtonBlock::GetVisibility) );
+}
+
+TSharedRef<SWidget> SToolBarComboButtonBlock::OnGetMenuContent()
+{
+	TSharedRef< const FToolBarComboButtonBlock > ToolBarButtonComboBlock = StaticCastSharedRef< const FToolBarComboButtonBlock >( MultiBlock.ToSharedRef() );
+	return ToolBarButtonComboBlock->MenuContentGenerator.Execute();
+}
+
+bool SToolBarComboButtonBlock::IsEnabled() const
+{
+	const FUIAction& UIAction = MultiBlock->GetDirectActions();
+	if( UIAction.CanExecuteAction.IsBound() )
+	{
+		return UIAction.CanExecuteAction.Execute();
+	}
+
+	return true;
+}
+
+EVisibility SToolBarComboButtonBlock::GetVisibility() const
+{
+	const FUIAction& UIAction = MultiBlock->GetDirectActions();
+	if (UIAction.IsActionVisibleDelegate.IsBound())
+	{
+		return UIAction.IsActionVisibleDelegate.Execute() ? EVisibility::Visible : EVisibility::Collapsed;
+	}
+
+	return EVisibility::Visible;
+}
+
+bool SToolBarComboButtonBlock::HasDynamicIcon() const
+{
+	return Icon.IsBound();
+}
+
+const FSlateBrush* SToolBarComboButtonBlock::GetIcon() const
+{
+	const FSlateIcon& ActualIcon = Icon.Get();
+	return ActualIcon.GetIcon();
+}
+
+const FSlateBrush* SToolBarComboButtonBlock::GetSmallIcon() const
+{
+	const FSlateIcon& ActualIcon = Icon.Get();
+	return ActualIcon.GetSmallIcon();
+}
+
+EVisibility SToolBarComboButtonBlock::GetIconVisibility(bool bIsASmallIcon) const
+{
+	return ((bForceSmallIcons || FMultiBoxSettings::UseSmallToolBarIcons.Get()) ^ bIsASmallIcon) ? EVisibility::Collapsed : EVisibility::Visible;
+}
diff --git a/Engine/Source/Runtime/Slate/Private/Framework/MultiBox/SToolBarSeparatorBlock.cpp b/Engine/Source/Runtime/Slate/Private/Framework/MultiBox/SToolBarSeparatorBlock.cpp
new file mode 100644
--- /dev/null
+++ b/Engine/Source/Runtime/Slate/Private/Framework/MultiBox/SToolBarSeparatorBlock.cpp
@@ -0,0 +1,65 @@
+// Copyright 1998-2014 Epic Games, Inc. All Rights Reserved.
+
+#include "SlatePrivatePCH.h"
+#include "MultiBox.h"
+#include "SToolBarSeparatorBlock.h"
+
+
+/**
+ * Constructor
+ */
+FToolBarSeparatorBlock::FToolBarSeparatorBlock(const FName& InExtensionHook)
+	: FMultiBlock( NULL, NULL, InExtensionHook )
+{
+}
+
+
+
+void FToolBarSeparatorBlock::CreateMenuEntry(FMenuBuilder& MenuBuilder) const
+{
+	MenuBuilder.AddMenuSeparator();
+}
+
+
+
+/**
+ * Allocates a widget for this type of MultiBlock.  Override this in derived classes.
+ *
+ * @return  MultiBlock widget object
+ */
+TSharedRef< class IMultiBlockBaseWidget > FToolBarSeparatorBlock::ConstructWidget() const
+{
+	return SNew( SToolBarSeparatorBlock );
+}
+
+
+
+/**
+ * Construct this widget
+ *
+ * @param	InArgs	The declaration data for this widget
+ */
+void SToolBarSeparatorBlock::Construct( const FArguments& InArgs )
+{
+}
+
+
+
+/**
+ * Builds this MultiBlock widget up from the MultiBlock associated with it
+ */
+void SToolBarSeparatorBlock::BuildMultiBlockWidget(const ISlateStyle* StyleSet, const FName& StyleName)
+{
+	ChildSlot.Widget =
+		SNew( SHorizontalBox )
+			+SHorizontalBox::Slot()
+				.AutoWidth()
+				.Padding( 0.0f, 0.0f, 0.0f, 0.0f )
+				[
+					SNew(SSeparator)
+						.Orientation(Orient_Vertical)
+						.SeparatorImage( StyleSet->GetBrush( ISlateStyle::Join( StyleName, ".Separator" ) ) )
+				];
+}
+
+
diff --git a/Engine/Source/Runtime/Slate/Private/Framework/MultiBox/SToolBarSeparatorBlock.h b/Engine/Source/Runtime/Slate/Private/Framework/MultiBox/SToolBarSeparatorBlock.h
new file mode 100644
--- /dev/null
+++ b/Engine/Source/Runtime/Slate/Private/Framework/MultiBox/SToolBarSeparatorBlock.h
@@ -0,0 +1,72 @@
+// Copyright 1998-2014 Epic Games, Inc. All Rights Reserved.
+
+#pragma once
+
+/**
+ * Toolbar separator MultiBlock
+ */
+class FToolBarSeparatorBlock
+	: public FMultiBlock
+{
+
+public:
+
+	/**
+	 * Constructor
+	 */
+	FToolBarSeparatorBlock(const FName& InExtensionHook);
+
+	/** FMultiBlock interface */
+	virtual void CreateMenuEntry(class FMenuBuilder& MenuBuilder) const OVERRIDE;
+
+
+private:
+
+	/**
+	 * Allocates a widget for this type of MultiBlock.  Override this in derived classes.
+	 *
+	 * @return  MultiBlock widget object
+	 */
+	virtual TSharedRef< class IMultiBlockBaseWidget > ConstructWidget() const;
+
+
+private:
+
+	// Friend our corresponding widget class
+	friend class SToolBarSeparatorBlock;
+
+};
+
+
+
+/**
+ * Toolbar separator MultiBlock widget
+ */
+class SLATE_API SToolBarSeparatorBlock
+	: public SMultiBlockBaseWidget
+{
+
+public:
+
+	SLATE_BEGIN_ARGS( SToolBarSeparatorBlock ){}
+
+	SLATE_END_ARGS()
+
+
+	/**
+	 * Builds this MultiBlock widget up from the MultiBlock associated with it
+	 */
+	virtual void BuildMultiBlockWidget(const ISlateStyle* StyleSet, const FName& StyleName) OVERRIDE;
+
+
+	/**
+	 * Construct this widget
+	 *
+	 * @param	InArgs	The declaration data for this widget
+	 */
+	void Construct( const FArguments& InArgs );
+
+
+private:
+
+};
diff --git a/Engine/Source/Runtime/Slate/Private/Framework/MultiBox/SWidgetBlock.cpp b/Engine/Source/Runtime/Slate/Private/Framework/MultiBox/SWidgetBlock.cpp
new file mode 100644
--- /dev/null
+++ b/Engine/Source/Runtime/Slate/Private/Framework/MultiBox/SWidgetBlock.cpp
@@ -0,0 +1,89 @@
+// Copyright 1998-2014 Epic Games, Inc. All Rights Reserved.
+
+#include "SlatePrivatePCH.h"
+#include "MultiBox.h"
+#include "SWidgetBlock.h"
+
+
+/**
+ * Constructor
+ *
+ * @param	InHeadingText	Heading text
+ */
+FWidgetBlock::FWidgetBlock( TSharedRef<SWidget> InContent, const FText& InLabel, bool bInNoIndent )
+	: FMultiBlock( NULL, NULL )
+	, ContentWidget( InContent )
+	, Label( InLabel )
+	, bNoIndent( bInNoIndent )
+{
+}
+
+
+void FWidgetBlock::CreateMenuEntry(FMenuBuilder& MenuBuilder) const
+{
+	FText EntryLabel = (!Label.IsEmpty()) ? Label : NSLOCTEXT("WidgetBlock", "CustomControl", "Custom Control");
+	MenuBuilder.AddWidget(ContentWidget, FText::GetEmpty(), true);
+}
+
+
+/**
+ * Allocates a widget for this type of MultiBlock.  Override this in derived classes.
+ *
+ * @return  MultiBlock widget object
+ */
+TSharedRef< class IMultiBlockBaseWidget > FWidgetBlock::ConstructWidget() const
+{
+	return SNew( SWidgetBlock )
+			.Cursor(EMouseCursor::Default);
+}
+
+
+/**
+ * Construct this widget
+ *
+ * @param	InArgs	The declaration data for this widget
+ */
+void SWidgetBlock::Construct( const FArguments& InArgs )
+{
+}
+
+
+
+/**
+ * Builds this MultiBlock widget up from the MultiBlock associated with it
+ */
+void SWidgetBlock::BuildMultiBlockWidget(const ISlateStyle* StyleSet, const FName& StyleName)
+{
+	TSharedRef< const FWidgetBlock > WidgetBlock = StaticCastSharedRef< const FWidgetBlock >( MultiBlock.ToSharedRef() );
+
+	bool bHasLabel = !WidgetBlock->Label.IsEmpty();
+	FMargin Padding = WidgetBlock->bNoIndent ? StyleSet->GetMargin( StyleName, ".Block.Padding" ) : StyleSet->GetMargin( StyleName, ".Block.IndentedPadding" );
+
+	ChildSlot
+	.Padding( Padding )	// Large left margin mimics the indent of normal menu items when bNoIndent is false
+	[
+		SNew(SHorizontalBox)
+		+ SHorizontalBox::Slot()
+		.AutoWidth()
+		[
+			SNew(SHorizontalBox)
+			.Visibility( bHasLabel ? EVisibility::Visible : EVisibility::Collapsed )
+			+ SHorizontalBox::Slot()
+			.AutoWidth()
+			.Padding(0.f, 0.f, 4.f, 0.f)
+			.VAlign( VAlign_Center )
+			[
+				SNew( STextBlock )
+				.TextStyle( StyleSet, ISlateStyle::Join( StyleName, ".Label" ) )
+				.Text( WidgetBlock->Label )
+			]
+		]
+		+SHorizontalBox::Slot()
+		.VAlign( VAlign_Bottom )
+		.FillWidth(1.f)
+		[
+			WidgetBlock->ContentWidget
+		]
+	];
+}
+
