Slate: Reorganized private Slate files; proper pre-compiled header for all files.

diff --git a/Engine/Source/Runtime/Slate/Private/Framework/MultiBox/SWidgetBlock.h b/Engine/Source/Runtime/Slate/Private/Framework/MultiBox/SWidgetBlock.h
new file mode 100644
--- /dev/null
+++ b/Engine/Source/Runtime/Slate/Private/Framework/MultiBox/SWidgetBlock.h
@@ -0,0 +1,79 @@
+// Copyright 1998-2014 Epic Games, Inc. All Rights Reserved.
+
+#pragma once
+
+/**
+ * Arbitrary Widget MultiBlock
+ */
+class FWidgetBlock
+	: public FMultiBlock
+{
+
+public:
+
+	/**
+	 * Constructor
+	 *
+	 * @param	InContent	The widget to place in the block
+	 * @param	InLabel		Optional label text to be added to the left of the content
+	 * @param	bInNoIndent	If true, removes the padding from the left of the widget that lines it up with other menu items
+	 */
+	FWidgetBlock( TSharedRef<SWidget> InContent, const FText& InLabel, bool bInNoIndent );
+
+	/** FMultiBlock interface */
+	virtual void CreateMenuEntry(class FMenuBuilder& MenuBuilder) const OVERRIDE;
+
+
+private:
+
+	/**
+	 * Allocates a widget for this type of MultiBlock.  Override this in derived classes.
+	 *
+	 * @return  MultiBlock widget object
+	 */
+	virtual TSharedRef< class IMultiBlockBaseWidget > ConstructWidget() const;
+
+
+private:
+
+	// Friend our corresponding widget class
+	friend class SWidgetBlock;
+
+	/** Content widget */
+	TSharedRef<SWidget> ContentWidget;
+
+	/** Optional label text */
+	FText Label;
+
+	/** Remove the padding from the left of the widget that lines it up with other menu items? */
+	bool bNoIndent;
+};
+
+
+
+
+/**
+ * Arbitrary Widget MultiBlock widget
+ */
+class SLATE_API SWidgetBlock
+	: public SMultiBlockBaseWidget
+{
+
+public:
+
+	SLATE_BEGIN_ARGS( SWidgetBlock ){}
+	SLATE_END_ARGS()
+
+
+	/**
+	 * Builds this MultiBlock widget up from the MultiBlock associated with it
+	 */
+	virtual void BuildMultiBlockWidget(const ISlateStyle* StyleSet, const FName& StyleName) OVERRIDE;
+
+	/**
+	 * Construct this widget
+	 *
+	 * @param	InArgs	The declaration data for this widget
+	 */
+	void Construct( const FArguments& InArgs );
+};
diff --git a/Engine/Source/Runtime/Slate/Private/Framework/Notifications/NotificationManager.cpp b/Engine/Source/Runtime/Slate/Private/Framework/Notifications/NotificationManager.cpp
new file mode 100644
--- /dev/null
+++ b/Engine/Source/Runtime/Slate/Private/Framework/Notifications/NotificationManager.cpp
@@ -0,0 +1,234 @@
+// Copyright 1998-2014 Epic Games, Inc. All Rights Reserved.
+
+#include "SlatePrivatePCH.h"
+
+
+namespace NotificationManagerConstants
+{
+	// Offsets from the bottom-left corner of the work area
+	const FVector2D NotificationOffset( 15.0f, 15.0f );
+}
+
+FSlateNotificationManager::FRegionalNotificationList::FRegionalNotificationList(const FSlateRect& InRectangle)
+	: Region(InRectangle)
+{
+
+}
+
+void FSlateNotificationManager::FRegionalNotificationList::RemoveDeadNotifications()
+{
+	// Iterate backwards and remove anything that has finished
+	for (int32 ListIndex = Notifications.Num() - 1; ListIndex >= 0; --ListIndex)
+	{
+		if (Notifications[ListIndex]->bDone)
+		{
+			TSharedPtr<SWindow> PinnedWindow = Notifications[ListIndex]->ParentWindowPtr.Pin();
+			if( PinnedWindow.IsValid() )
+			{
+				PinnedWindow->RequestDestroyWindow();
+			}
+
+			Notifications.RemoveAt(ListIndex);
+		}
+	}
+}
+
+void FSlateNotificationManager::FRegionalNotificationList::Arrange()
+{
+	FVector2D AnchorPoint(
+		Region.Right - NotificationManagerConstants::NotificationOffset.X,
+		Region.Bottom - NotificationManagerConstants::NotificationOffset.Y );
+
+	for (int32 ListIndex = Notifications.Num() - 1; ListIndex >= 0; --ListIndex)
+	{
+		TSharedPtr<SWindow> PinnedWindow = Notifications[ListIndex]->ParentWindowPtr.Pin();
+		if( PinnedWindow.IsValid() )
+		{
+			const FVector2D DesiredSize = PinnedWindow->GetDesiredSize();
+			const FVector2D NewPosition( AnchorPoint.X - DesiredSize.X, AnchorPoint.Y - DesiredSize.Y );
+			if( NewPosition != PinnedWindow->GetPositionInScreen() && DesiredSize != PinnedWindow->GetSizeInScreen() )
+			{
+				PinnedWindow->ReshapeWindow( NewPosition, DesiredSize );
+			}
+			else if( NewPosition != PinnedWindow->GetPositionInScreen() )
+			{
+				PinnedWindow->MoveWindowTo( NewPosition );
+			}
+			AnchorPoint.Y -= DesiredSize.Y;
+		}
+	}
+}
+
+FSlateNotificationManager::FSlateNotificationManager()
+{
+}
+
+FSlateNotificationManager& FSlateNotificationManager::Get()
+{
+	static FSlateNotificationManager* Instance = NULL;
+	if( Instance == NULL )
+	{
+		static FCriticalSection CriticalSection;
+		FScopeLock Lock(&CriticalSection);
+		if( Instance == NULL )
+		{
+			Instance = new FSlateNotificationManager;
+		}
+	}
+	return *Instance;
+}
+
+void FSlateNotificationManager::SetRootWindow( const TSharedRef<SWindow> InRootWindow )
+{
+	RootWindowPtr = InRootWindow;
+}
+
+
+
+TSharedRef<SNotificationList> FSlateNotificationManager::CreateStackForArea(const FSlateRect& InRectangle)
+{
+	TSharedRef<SNotificationList> NotificationList = SNew(SNotificationList);
+	TSharedRef<SWindow> NotificationWindow = SWindow::MakeNotificationWindow();
+	NotificationWindow->SetContent(NotificationList);
+	NotificationList->ParentWindowPtr = NotificationWindow;
+
+	if( RootWindowPtr.IsValid() )
+	{
+		FSlateApplication::Get().AddWindowAsNativeChild( NotificationWindow, RootWindowPtr.Pin().ToSharedRef() );
+	}
+	else
+	{
+		FSlateApplication::Get().AddWindow( NotificationWindow );
+	}
+
+	if( !FSlateApplication::Get().GetActiveModalWindow().IsValid() )
+	{
+		NotificationWindow->BringToFront();
+	}
+
+	bool bFound = false;
+	for (FRegionalNotificationList& List : RegionalLists)
+	{
+		if (FSlateRect::IsRectangleContained(List.Region, InRectangle))
+		{
+			List.Notifications.Add(NotificationList);
+			bFound = true;
+		}
+	}
+
+	if (!bFound)
+	{
+		FRegionalNotificationList NewList(FSlateApplication::Get().GetWorkArea(InRectangle));
+		NewList.Notifications.Add(NotificationList);
+		RegionalLists.Add(NewList);
+	}
+
+	return NotificationList;
+}
+
+TSharedPtr<SNotificationItem> FSlateNotificationManager::AddNotification(const FNotificationInfo& Info)
+{
+	check(IsInGameThread() || !"FSlateNotificationManager::AddNotification must be called on game thread. Use QueueNotification if necessary.");
+
+	// Early calls of this function can happen before Slate is initialized.
+	if( FSlateApplication::IsInitialized() )
+	{
+		const FSlateRect PreferredWorkArea = FSlateApplication::Get().GetPreferredWorkArea();
+
+		TSharedRef<SNotificationList> List = CreateStackForArea(PreferredWorkArea);
+
+		return List->AddNotification(Info);
+	}
+
+	return NULL;
+}
+
+void FSlateNotificationManager::QueueNotification(FNotificationInfo* Info)
+{
+	PendingNotifications.Push(Info);
+}
+
+void FSlateNotificationManager::GetWindows(TArray< TSharedRef<SWindow> >& OutWindows) const
+{
+	for (const FRegionalNotificationList& RegionList : RegionalLists)
+	{
+		for (const auto& NotificationList : RegionList.Notifications)
+		{
+			TSharedPtr<SWindow> PinnedWindow = NotificationList->ParentWindowPtr.Pin();
+			if( PinnedWindow.IsValid() )
+			{
+				OutWindows.Add(PinnedWindow.ToSharedRef());
+			}
+		}
+	}
+}
+
+void FSlateNotificationManager::Tick()
+{
+	// Ensure that the region rectangles still match the screen work areas.
+	// This is necessary if the desktop configuration has changed
+	for (auto& RegionList : RegionalLists)
+	{
+		RegionList.Region = FSlateApplication::Get().GetWorkArea(RegionList.Region);
+	}
+
+
+	while (auto* Notification = PendingNotifications.Pop())
+	{
+		AddNotification(*Notification);
+		delete Notification;
+	}
+
+	// Check notifications to see if any have timed out & need to be removed.
+	// We need to do this here as we cant remove their windows in the normal
+	// window-tick callstack (as the SlateWindows array gets corrupted)
+	
+	// We don't need to worry about duplicates here as there is always a unique list per-region
+	for (int32 RegionIndex = RegionalLists.Num() - 1; RegionIndex >= 0; --RegionIndex)
+	{
+		RegionalLists[RegionIndex].RemoveDeadNotifications();
+
+		if (RegionalLists[RegionIndex].Notifications.Num() == 0)
+		{
+			// It's empty, so remove it
+			RegionalLists.RemoveAt(RegionIndex);
+		}
+		else
+		{
+			// Arrange the notifications in the list
+			RegionalLists[RegionIndex].Arrange();
+		}
+	}
+}
+
+void FSlateNotificationManager::ForceNotificationsInFront( const TSharedRef<SWindow> InWindow )
+{
+	// check to see if this is a re-entrant call from one of our windows
+	for (const auto& RegionList : RegionalLists)
+	{
+		for (auto& Notification : RegionList.Notifications)
+		{
+			TSharedPtr<SWindow> PinnedWindow = Notification->ParentWindowPtr.Pin();
+			if( PinnedWindow.IsValid() )
+			{
+				if( InWindow == PinnedWindow )
+				{
+					return;
+				}
+			}
+		}
+	}
+
+	// now bring all of our windows back to the front
+	for (const auto& RegionList : RegionalLists)
+	{
+		for (auto& Notification : RegionList.Notifications)
+		{
+			TSharedPtr<SWindow> PinnedWindow = Notification->ParentWindowPtr.Pin();
+			if( PinnedWindow.IsValid() && !FSlateApplication::Get().GetActiveModalWindow().IsValid() )
+			{
+				PinnedWindow->BringToFront();
+			}
+		}
+	}
+}
diff --git a/Engine/Source/Runtime/Slate/Private/Framework/Styling/ButtonWidgetStyle.cpp b/Engine/Source/Runtime/Slate/Private/Framework/Styling/ButtonWidgetStyle.cpp
new file mode 100644
--- /dev/null
+++ b/Engine/Source/Runtime/Slate/Private/Framework/Styling/ButtonWidgetStyle.cpp
@@ -0,0 +1,10 @@
+// Copyright 1998-2014 Epic Games, Inc. All Rights Reserved.
+
+#include "SlatePrivatePCH.h"
+
+
+UButtonWidgetStyle::UButtonWidgetStyle( const class FPostConstructInitializeProperties& PCIP )
+	: Super(PCIP)
+{
+	
+}
\ No newline at end of file
diff --git a/Engine/Source/Runtime/Slate/Private/Framework/Styling/CheckBoxWidgetStyle.cpp b/Engine/Source/Runtime/Slate/Private/Framework/Styling/CheckBoxWidgetStyle.cpp
new file mode 100644
--- /dev/null
+++ b/Engine/Source/Runtime/Slate/Private/Framework/Styling/CheckBoxWidgetStyle.cpp
@@ -0,0 +1,10 @@
+// Copyright 1998-2014 Epic Games, Inc. All Rights Reserved.
+
+#include "SlatePrivatePCH.h"
+
+
+UCheckBoxWidgetStyle::UCheckBoxWidgetStyle( const class FPostConstructInitializeProperties& PCIP )
+	: Super(PCIP)
+{
+	
+}
\ No newline at end of file
diff --git a/Engine/Source/Runtime/Slate/Private/Framework/Styling/ComboBoxWidgetStyle.cpp b/Engine/Source/Runtime/Slate/Private/Framework/Styling/ComboBoxWidgetStyle.cpp
new file mode 100644
--- /dev/null
+++ b/Engine/Source/Runtime/Slate/Private/Framework/Styling/ComboBoxWidgetStyle.cpp
@@ -0,0 +1,10 @@
+// Copyright 1998-2014 Epic Games, Inc. All Rights Reserved.
+
+#include "SlatePrivatePCH.h"
+
+
+UComboBoxWidgetStyle::UComboBoxWidgetStyle( const class FPostConstructInitializeProperties& PCIP )
+	: Super(PCIP)
+{
+	
+}
diff --git a/Engine/Source/Runtime/Slate/Private/Framework/Styling/ComboButtonWidgetStyle.cpp b/Engine/Source/Runtime/Slate/Private/Framework/Styling/ComboButtonWidgetStyle.cpp
new file mode 100644
--- /dev/null
+++ b/Engine/Source/Runtime/Slate/Private/Framework/Styling/ComboButtonWidgetStyle.cpp
@@ -0,0 +1,10 @@
+// Copyright 1998-2014 Epic Games, Inc. All Rights Reserved.
+
+#include "SlatePrivatePCH.h"
+
+
+UComboButtonWidgetStyle::UComboButtonWidgetStyle( const class FPostConstructInitializeProperties& PCIP )
+	: Super(PCIP)
+{
+	
+}
diff --git a/Engine/Source/Runtime/Slate/Private/Framework/Styling/TextBlockWidgetStyle.cpp b/Engine/Source/Runtime/Slate/Private/Framework/Styling/TextBlockWidgetStyle.cpp
new file mode 100644
--- /dev/null
+++ b/Engine/Source/Runtime/Slate/Private/Framework/Styling/TextBlockWidgetStyle.cpp
@@ -0,0 +1,10 @@
+// Copyright 1998-2014 Epic Games, Inc. All Rights Reserved.
+
+#include "SlatePrivatePCH.h"
+
+
+UTextBlockWidgetStyle::UTextBlockWidgetStyle( const class FPostConstructInitializeProperties& PCIP )
+	: Super(PCIP)
+{
+	
+}
\ No newline at end of file
diff --git a/Engine/Source/Runtime/Slate/Private/Framework/Testing/RemoteSlateServer.cpp b/Engine/Source/Runtime/Slate/Private/Framework/Testing/RemoteSlateServer.cpp
new file mode 100644
--- /dev/null
+++ b/Engine/Source/Runtime/Slate/Private/Framework/Testing/RemoteSlateServer.cpp
@@ -0,0 +1,247 @@
+// Copyright 1998-2014 Epic Games, Inc. All Rights Reserved.
+
+#include "SlatePrivatePCH.h"
+#include "RemoteSlateServer.h"
+
+
+void FRemoteSlateServer::Tick()
+{
+	// calculate the time since last tick
+	static double LastTickTime = FPlatformTime::Seconds();
+	double Now = FPlatformTime::Seconds();
+	float DeltaTime = LastTickTime - Now;
+	LastTickTime = Now;
+
+
+	// attempt to open the socket if it needs it
+	if (ServerSocket == NULL)
+	{
+		ISocketSubsystem* SocketSubsystem = ISocketSubsystem::Get(PLATFORM_SOCKETSUBSYSTEM);
+		if (SocketSubsystem)
+		{
+			// create a UDP listening socket
+			ServerSocket = SocketSubsystem->CreateSocket(NAME_DGram, TEXT("UDKRemoteListener"));
+
+			if (ServerSocket)
+			{
+				// allow reuse (or there will be errors when running multiple game instances on a single machine)
+				ServerSocket->SetReuseAddr();
+
+				// bind the socket to listen on the port
+				TSharedRef<FInternetAddr> BindAddr = SocketSubsystem->CreateInternetAddr();
+				BindAddr->SetAnyAddress();
+				int32 Port;
+				if (GIsEditor)
+				{
+					Port = 41766;
+					//verify(GConfig->GetInt(TEXT("MobileSupport"), TEXT("UDKRemotePortPIE"), Port, GEngineIni));
+				}
+				else
+				{
+					Port = 41765;
+					//verify(GConfig->GetInt(TEXT("MobileSupport"), TEXT("UDKRemotePort"), Port, GEngineIni));
+				}
+
+				BindAddr->SetPort(Port);
+				ServerSocket->Bind(*BindAddr);
+
+				// make an async socket
+				ServerSocket->SetNonBlocking(true);
+
+				// create a UDP reply socket
+				ReplySocket = SocketSubsystem->CreateSocket(NAME_DGram, TEXT("UDKRemoteReply"));
+				if (ReplySocket)
+				{
+					// make an async socket
+					ReplySocket->SetNonBlocking(true);
+				}
+
+				// 					ImageSocket = GSocketSubsystem->CreateDGramSocket(TEXT("TiltImage"));
+				// 					if (ImageSocket)
+				// 					{
+				// 						// make an async socket
+				// 						ImageSocket->SetNonBlocking(TRUE);
+				// 					}
+			}
+		}
+	}
+
+	if (ServerSocket && ReplySocket)
+	{
+		FSlateApplication& CurrentApp = FSlateApplication::Get();
+
+		// try to read some data from the 
+		uint8 Msg[256];
+		int32 BytesRead;
+		while (ServerSocket->RecvFrom(Msg, sizeof(Msg), BytesRead, *ReplyAddr))
+		{
+			// make sure it's what we expected
+			if (BytesRead != sizeof(FMessage))
+			{
+				UE_LOG(LogSlate, Log, TEXT("Received %d bytes, expected %d"), BytesRead, sizeof(FMessage));
+				continue;
+			}
+
+			// get the message in a blob we can process
+			FMessage* Message = (FMessage*)Msg;
+
+			// make sure tag and version are okay
+			bool bIsValidMessageVersion = 
+				Message->MagicTag == MESSAGE_MAGIC_ID && 
+				// in the future, we can write code to handle older versions
+				Message->MessageVersion == CURRENT_MESSAGE_VERSION;
+
+			// make sure we got a later message than last time (handling wrapping
+			// around with plenty of slop)
+			bool bIsValidID = 
+				(Message->DataType == DT_Ping) ||
+				(Message->MessageID > HighestMessageReceived) || 
+				(Message->MessageID < 1000 && HighestMessageReceived > 65000);
+			HighestMessageReceived = Message->MessageID;
+
+			bool bIsValidMessage = bIsValidMessageVersion && bIsValidID;
+			if (!bIsValidMessage)
+			{
+				//					debugf(TEXT("Dropped message %d, %d"), Message->MessageID, Message->DataType);
+				continue;
+			}
+			else
+			{
+				//					debugf(TEXT("Received message %d, %d"), Message->MessageID, Message->DataType);
+			}
+
+			// handle tilt input
+			if (Message->DataType == DT_Tilt || Message->DataType == DT_Motion)
+			{
+				FVector Attitude;
+				FVector RotationRate;
+				FVector Gravity;
+				FVector Accel;
+
+				if (Message->DataType == DT_Tilt)
+				{
+					// get the raw and processed values from the other end
+					FVector CurrentAccelerometer(Message->Data[0], Message->Data[1], Message->Data[2]);
+					float Pitch = Message->Data[3], Roll = Message->Data[4];
+
+					// convert it into "Motion" data
+					static float LastPitch, LastRoll;
+
+					Attitude = FVector(Pitch,0,Roll);
+					RotationRate = FVector(LastPitch - Pitch,0,LastRoll - Roll);
+					Gravity = FVector(0,0,0);
+					Accel = CurrentAccelerometer;
+
+					// save the previous values to delta rotation
+					LastPitch = Pitch;
+					LastRoll = Roll;
+				}
+				else if (Message->DataType == DT_Motion)
+				{
+					// just use the values directly from UDK Remote
+					// Negate the yaw to match direction
+					Attitude = FVector(Message->Data[0], -Message->Data[1], Message->Data[2]);
+					RotationRate = FVector(Message->Data[3], -Message->Data[4], Message->Data[5]);
+					Gravity = FVector(Message->Data[6], Message->Data[7], Message->Data[8]);
+					Accel = FVector(Message->Data[9], Message->Data[10], Message->Data[11]);
+				}
+
+				// munge the vectors based on the orientation of the remote device
+				// 						EUIOrientation Orientation = (EUIOrientation)Message->UIOrientation;
+				// 						UMobilePlayerInput::ModifyVectorByOrientation(Attitude, Orientation, TRUE);
+				//  						UMobilePlayerInput::ModifyVectorByOrientation(RotationRate, Orientation, TRUE);
+				//  						UMobilePlayerInput::ModifyVectorByOrientation(Gravity, Orientation, FALSE);
+				//  						UMobilePlayerInput::ModifyVectorByOrientation(Accel, Orientation, FALSE);
+
+				// send the tilt to the input system
+				FMotionEvent Event(0, Attitude, RotationRate, Gravity, Accel);
+				CurrentApp.OnMotionDetectedMessage(Event);
+			}
+			else if (Message->DataType == DT_Gyro)
+			{
+				// ignore for now
+			}
+			else if (Message->DataType == DT_Ping)
+			{
+				TimeSinceLastPing = 0.0f;
+				// reply
+				ReplyAddr->SetPort(41764);
+
+				int32 BytesSent;
+				const ANSICHAR* HELO = "HELO";
+				ReplySocket->SendTo((uint8*)HELO, 5, BytesSent, *ReplyAddr);
+			}
+			// handle touches
+			else
+			{
+				// @todo This should be some global value or something in slate??
+				static FVector2D LastTouchPositions[EKeys::NUM_TOUCH_KEYS];
+					
+				if (Message->Handle < EKeys::NUM_TOUCH_KEYS)
+				{
+					// convert from 0..1 to windows coords to screen coords
+					// @todo: Put this into a function in SlateApplication
+					if (CurrentApp.GetActiveTopLevelWindow().IsValid())
+					{
+						FVector2D ScreenPosition;
+
+						// The remote is interested in finding the gameviewport for the user and mapping the input in to it
+						TSharedPtr<SViewport> GameViewport = CurrentApp.GetGameViewport();
+						if (GameViewport.IsValid())
+						{
+							FWidgetPath WidgetPath;
+
+							WidgetPath = GameViewportWidgetPath.ToWidgetPath();
+						
+							if (WidgetPath.Widgets.Num() == 0 || WidgetPath.Widgets.Last().Widget != GameViewport)
+							{
+								CurrentApp.FindPathToWidget(CurrentApp.GetInteractiveTopLevelWindows(), GameViewport.ToSharedRef(), WidgetPath);
+								GameViewportWidgetPath = WidgetPath;
+							}
+
+							const FGeometry& GameViewportGeometry = WidgetPath.Widgets.Last().Geometry;
+							ScreenPosition = GameViewportGeometry.LocalToAbsolute(FVector2D(Message->Data[0], Message->Data[1]) * GameViewportGeometry.Size);
+						}
+						else
+						{
+							const FSlateRect WindowScreenRect = CurrentApp.GetActiveTopLevelWindow()->GetRectInScreen();
+							const FVector2D WindowPosition = WindowScreenRect.GetSize() * FVector2D(Message->Data[0], Message->Data[1]);
+							ScreenPosition = FVector2D(WindowScreenRect.Left, WindowScreenRect.Top) + WindowPosition;
+						}
+
+						// for up/down messages, we need to let the cursor go in the same location (mouse up with
+						// a delta confuses things)
+						if (Message->DataType == DT_TouchBegan || Message->DataType == DT_TouchEnded)
+						{
+							LastTouchPositions[Message->Handle] = ScreenPosition;
+						}
+
+						// create the event struct
+						FPointerEvent Event(0, Message->Handle, ScreenPosition, LastTouchPositions[Message->Handle], Message->DataType != DT_TouchEnded);
+						LastTouchPositions[Message->Handle] = ScreenPosition;
+
+						// send input to handler
+						if (Message->DataType == DT_TouchBegan)
+						{
+							CurrentApp.OnTouchStartedMessage( NULL, Event);
+						}
+						else if (Message->DataType == DT_TouchEnded)
+						{
+							CurrentApp.OnTouchEndedMessage(Event);
+						}
+						else
+						{
+							CurrentApp.OnTouchMovedMessage(Event);
+						}
+					}
+				}
+				else
+				{
+					//checkf(Message->Handle < ARRAY_COUNT(LastTouchPositions), TEXT("Received handle %d, but it's too big for our array"), Message->Handle);
+				}
+			}
+		}
+
+		TimeSinceLastPing += DeltaTime;
+	}
+}
\ No newline at end of file
diff --git a/Engine/Source/Runtime/Slate/Private/Framework/Testing/RemoteSlateServer.h b/Engine/Source/Runtime/Slate/Private/Framework/Testing/RemoteSlateServer.h
new file mode 100644
--- /dev/null
+++ b/Engine/Source/Runtime/Slate/Private/Framework/Testing/RemoteSlateServer.h
@@ -0,0 +1,116 @@
+// Copyright 1998-2014 Epic Games, Inc. All Rights Reserved.
+
+#pragma once
+
+#include "SocketSubsystem.h"
+#include "Sockets.h"
+
+/** These data types must match up on the other side! */
+enum EDataType
+{
+	DT_TouchBegan=0,
+	DT_TouchMoved=1,
+	DT_TouchEnded=2,
+	DT_Tilt=3,
+	DT_Motion=4,
+	DT_Gyro=5,
+	DT_Ping=6,
+};
+
+/** Possible device orientations */
+enum EDeviceOrientation
+{
+	DO_Unknown,
+	DO_Portrait,
+	DO_PortraitUpsideDown,
+	DO_LandscapeLeft,
+	DO_LandscapeRight,
+	DO_FaceUp,
+	DO_FaceDown,
+};
+
+// magic number that must match UDKRemote
+#define MESSAGE_MAGIC_ID 0xAB
+
+// versioning information for future expansion
+#define CURRENT_MESSAGE_VERSION 1
+
+/** A single event message to send over the network */
+struct FMessage
+{
+	/** A byte that must match to what we expect */
+	uint8 MagicTag;
+
+	/** What version of message this is from UDK Remote */
+	uint8 MessageVersion;
+
+	/** Unique Id for the message, used for detecting lost/duplicate packets, etc (only duplicate currently handled) */
+	uint16 MessageID;
+
+	/** What type of message is this? */
+	uint8 DataType;
+
+	/** Unique identifier for the touch/finger */
+	uint8 Handle;
+
+	/** The current orientation of the device */
+	uint8 DeviceOrientation;
+
+	/** The current orientation of the UI */
+	uint8 UIOrientation;
+
+	/** X/Y or Pitch/Yaw data or CoreMotion data */
+	float Data[12];
+};
+
+
+class FRemoteSlateServer
+{
+protected:
+	/** The socket to listen on, will be initialized in first tick */
+	FSocket* ServerSocket;
+
+	/** The socket to send replies on, will be initialized in first tick */
+	FSocket* ReplySocket;
+
+	/** The socket to send image data on, will be initialized in first tick */
+	FSocket* ImageSocket;
+
+	/** The address of the most recent UDKRemote to talk to us, this is who we reply to */
+	TSharedRef<FInternetAddr> ReplyAddr;
+
+	/** Ever increasing timestamp to send to the input system */
+	double Timestamp;
+
+	/** Highest message ID (must handle wrapping around at 16 bits) */
+	uint16 HighestMessageReceived;
+
+	float TimeSinceLastPing;
+
+//	/** Event used to */
+//	FEvent* CompressionCompleteEvent;
+//	/** An event used to signal 
+//	FEvent* EncoderIsReadyEvent;
+
+public:
+	FRemoteSlateServer()
+		: ServerSocket(NULL)
+		, ReplySocket(NULL)
+		, ReplyAddr(ISocketSubsystem::Get()->CreateInternetAddr())
+		, Timestamp(0.0)
+		, HighestMessageReceived(0xFFFF)
+		, TimeSinceLastPing(200.0f)
+	{
+	}
+
+	/**
+	 * Pure virtual that must be overloaded by the inheriting class. It will
+	 * be called from within UnLevTic.cpp after ticking all actors or from
+	 * the rendering thread (depending on bIsRenderingThreadObject)
+	 *
+	 * @param DeltaTime	Game time passed since the last call.
+	 */
+	void Tick();
+
+	FWeakWidgetPath GameViewportWidgetPath;
+};
diff --git a/Engine/Source/Runtime/Slate/Private/Framework/Testing/SLayoutExample.h b/Engine/Source/Runtime/Slate/Private/Framework/Testing/SLayoutExample.h
new file mode 100644
--- /dev/null
+++ b/Engine/Source/Runtime/Slate/Private/Framework/Testing/SLayoutExample.h
@@ -0,0 +1,6 @@
+// Copyright 1998-2014 Epic Games, Inc. All Rights Reserved.
+
+#pragma once
+
+
+SLATE_API TSharedRef<SWidget> MakeLayoutExample();
diff --git a/Engine/Source/Runtime/Slate/Private/Framework/Testing/STableViewTesting.h b/Engine/Source/Runtime/Slate/Private/Framework/Testing/STableViewTesting.h
new file mode 100644
--- /dev/null
+++ b/Engine/Source/Runtime/Slate/Private/Framework/Testing/STableViewTesting.h
@@ -0,0 +1,6 @@
+// Copyright 1998-2014 Epic Games, Inc. All Rights Reserved.
+
+#pragma once
+
+class SWidget;
+SLATE_API TSharedRef<SWidget> MakeTableViewTesting();
diff --git a/Engine/Source/Runtime/Slate/Private/Framework/Testing/SWidgetGallery.h b/Engine/Source/Runtime/Slate/Private/Framework/Testing/SWidgetGallery.h
new file mode 100644
--- /dev/null
+++ b/Engine/Source/Runtime/Slate/Private/Framework/Testing/SWidgetGallery.h
@@ -0,0 +1,6 @@
+// Copyright 1998-2014 Epic Games, Inc. All Rights Reserved.
+
+#pragma once
+
+/** @return a new Drag and Drop test widget */
+TSharedRef<SWidget> MakeWidgetGallery();
diff --git a/Engine/Source/Runtime/Slate/Private/Framework/Text/IOS/IOSPlatformTextField.cpp b/Engine/Source/Runtime/Slate/Private/Framework/Text/IOS/IOSPlatformTextField.cpp
new file mode 100644
--- /dev/null
+++ b/Engine/Source/Runtime/Slate/Private/Framework/Text/IOS/IOSPlatformTextField.cpp
@@ -0,0 +1,107 @@
+// Copyright 1998-2014 Epic Games, Inc. All Rights Reserved.
+
+#include "SlatePrivatePCH.h"
+
+
+FIOSPlatformTextField::FIOSPlatformTextField()
+	: TextField( NULL )
+{
+}
+
+FIOSPlatformTextField::~FIOSPlatformTextField()
+{
+	if(TextField != NULL)
+	{
+		[TextField release];
+		TextField = NULL;
+	}
+}
+
+void FIOSPlatformTextField::ShowKeyboard(bool bShow, TSharedPtr<SVirtualKeyboardEntry> TextEntryWidget)
+{
+	if(TextField == NULL)
+	{
+		TextField = [SlateTextField alloc];
+	}
+
+	if(bShow)
+	{
+		// these functions must be run on the main thread
+		dispatch_async(dispatch_get_main_queue(),^ {
+			[TextField show: TextEntryWidget];
+		} );
+	}
+}
+
+
+@implementation SlateTextField
+
+-(void)show:(TSharedPtr<SVirtualKeyboardEntry>)InTextWidget
+{
+	TextWidget = InTextWidget;
+
+	UIAlertView* AlertView = [[UIAlertView alloc] initWithTitle:@""
+													message:@""
+													delegate:self
+													cancelButtonTitle:NSLocalizedString(@"Cancel", nil)
+													otherButtonTitles:NSLocalizedString(@"OK", nil), nil];
+
+	// give the UIAlertView a style so a UITextField is created
+	switch(TextWidget->GetKeyboardType())
+	{
+		case EKeyboardType::Keyboard_Password:
+			AlertView.alertViewStyle = UIAlertViewStyleSecureTextInput;
+			break;
+		case EKeyboardType::Keyboard_Default:
+		case EKeyboardType::Keyboard_Email:
+		case EKeyboardType::Keyboard_Number:
+		case EKeyboardType::Keyboard_Web:
+		default:
+			AlertView.alertViewStyle = UIAlertViewStylePlainTextInput;
+			break;
+	}
+
+	UITextField* TextField = [AlertView textFieldAtIndex:0];
+	TextField.clearsOnBeginEditing = NO;
+	TextField.clearsOnInsertion = NO;
+	TextField.autocorrectionType = UITextAutocorrectionTypeNo;
+	TextField.autocapitalizationType = UITextAutocapitalizationTypeNone;
+	TextField.text = [NSString stringWithFormat:@"%s" , TCHAR_TO_ANSI(*TextWidget->GetText().ToString())];
+	TextField.placeholder = [NSString stringWithFormat:@"%s" , TCHAR_TO_ANSI(*TextWidget->GetHintText().ToString())];
+
+	// set up the keyboard styles not supported in the AlertViewStyle styles
+	switch(TextWidget->GetKeyboardType())
+	{
+		case EKeyboardType::Keyboard_Email:
+			TextField.keyboardType = UIKeyboardTypeEmailAddress;
+			break;
+		case EKeyboardType::Keyboard_Number:
+			TextField.keyboardType = UIKeyboardTypeNumberPad;
+			break;
+		case EKeyboardType::Keyboard_Web:
+			TextField.keyboardType = UIKeyboardTypeURL;
+			break;
+		case EKeyboardType::Keyboard_Default:
+		case EKeyboardType::Keyboard_Password:
+		default:
+			// nothing to do, UIAlertView style handles these keyboard types
+			break;
+	}
+
+	[AlertView show];
+
+	[AlertView release];
+}
+
+- (void)alertView:(UIAlertView *)alertView clickedButtonAtIndex:(NSInteger)buttonIndex
+{
+	// index 1 is the OK button
+	if(buttonIndex == 1)
+	{
+		UITextField* TextField = [alertView textFieldAtIndex:0];
+		TextWidget->SetText(FText::FromString(TextField.text));
+	}
+}
+
+@end
+
diff --git a/Engine/Source/Runtime/Slate/Private/Framework/Text/RichTextMarkupProcessing.h b/Engine/Source/Runtime/Slate/Private/Framework/Text/RichTextMarkupProcessing.h
new file mode 100644
--- /dev/null
+++ b/Engine/Source/Runtime/Slate/Private/Framework/Text/RichTextMarkupProcessing.h
@@ -0,0 +1,21 @@
+// Copyright 1998-2014 Epic Games, Inc. All Rights Reserved.
+#pragma once
+
+#include "IRichTextMarkupParser.h"
+
+class FRichTextMarkupProcessing : public IRichTextMarkupParser
+{
+public:
+	static TSharedRef< FRichTextMarkupProcessing > Create();
+
+public:
+	virtual void Process(TArray<FTextLineParseResults>& Results, const FString& Input, FString& Output) OVERRIDE;
+
+private:
+	FRichTextMarkupProcessing() {}
+
+	void CalculateLineRanges(const FString& Input, TArray<FTextRange>& LineRanges) const;
+	void ParseLineRanges(const FString& Input, const TArray<FTextRange>& LineRanges, TArray<FTextLineParseResults>& LineParseResultsArray) const;
+	void HandleEscapeSequences(const FString& Input, TArray<FTextLineParseResults>& LineParseResultsArray, FString& ConcatenatedUnescapedLines) const;
+
+};
diff --git a/Engine/Source/Runtime/Slate/Private/Framework/Text/SlateHyperlinkRun.cpp b/Engine/Source/Runtime/Slate/Private/Framework/Text/SlateHyperlinkRun.cpp
new file mode 100644
--- /dev/null
+++ b/Engine/Source/Runtime/Slate/Private/Framework/Text/SlateHyperlinkRun.cpp
@@ -0,0 +1,120 @@
+// Copyright 1998-2014 Epic Games, Inc. All Rights Reserved.
+
+#include "SlatePrivatePCH.h"
+#include "SlateHyperlinkRun.h"
+
+TSharedRef< FSlateHyperlinkRun > FSlateHyperlinkRun::Create( const TSharedRef< const FString >& InText, const FHyperlinkStyle& InStyle, const TMap<FString,FString>& Metadata, FOnClick NavigateDelegate )
+{
+	return MakeShareable( new FSlateHyperlinkRun( InText, InStyle, Metadata, NavigateDelegate ) );
+}
+
+TSharedRef< FSlateHyperlinkRun > FSlateHyperlinkRun::Create( const TSharedRef< const FString >& InText, const FHyperlinkStyle& InStyle, const TMap<FString,FString>& Metadata, FOnClick NavigateDelegate, const FTextRange& InRange )
+{
+	return MakeShareable( new FSlateHyperlinkRun( InText, InStyle, Metadata, NavigateDelegate, InRange ) );
+}
+
+FTextRange FSlateHyperlinkRun::GetTextRange() const 
+{
+	return Range;
+}
+
+int16 FSlateHyperlinkRun::GetBaseLine( float Scale ) const 
+{
+	const TSharedRef< FSlateFontMeasure > FontMeasure = FSlateApplication::Get().GetRenderer()->GetFontMeasureService();
+	return FontMeasure->GetBaseline( Style.TextStyle.Font, Scale );
+}
+
+int16 FSlateHyperlinkRun::GetMaxHeight( float Scale ) const 
+{
+	const TSharedRef< FSlateFontMeasure > FontMeasure = FSlateApplication::Get().GetRenderer()->GetFontMeasureService();
+	return FontMeasure->GetMaxCharacterHeight( Style.TextStyle.Font, Scale );
+}
+
+FVector2D FSlateHyperlinkRun::Measure( int32 StartIndex, int32 EndIndex, float Scale ) const 
+{
+	if ( EndIndex - StartIndex == 0 )
+	{
+		return FVector2D( 0, GetMaxHeight( Scale ) );
+	}
+
+	const TSharedRef< FSlateFontMeasure > FontMeasure = FSlateApplication::Get().GetRenderer()->GetFontMeasureService();
+	return FontMeasure->Measure( **Text, StartIndex, EndIndex, Style.TextStyle.Font, true, Scale );
+}
+
+uint8 FSlateHyperlinkRun::GetKerning( int32 CurrentIndex, float Scale ) const 
+{
+	return 0;
+}
+
+TSharedRef< ILayoutBlock > FSlateHyperlinkRun::CreateBlock( int32 StartIndex, int32 EndIndex, FVector2D Size, const TSharedPtr< IRunHighlighter >& Highlighter )
+{
+	TSharedRef< SWidget > Widget = SNew( SRichTextHyperlink, ViewModel )
+		.Style( &Style )
+		.Text( FText::FromString( FString( EndIndex - StartIndex, **Text + StartIndex ) ) )
+		.OnNavigate( this, &FSlateHyperlinkRun::OnNavigate );
+
+	Children.Add( Widget );
+
+	return FWidgetLayoutBlock::Create( SharedThis( this ), Widget, FTextRange( StartIndex, EndIndex ), Size, Highlighter );
+}
+
+void FSlateHyperlinkRun::OnNavigate()
+{
+	NavigateDelegate.Execute( Metadata );
+}
+
+int32 FSlateHyperlinkRun::OnPaint( const FTextLayout::FLineView& Line, const TSharedRef< ILayoutBlock >& Block, const FTextBlockStyle& DefaultStyle, const FGeometry& AllottedGeometry, const FSlateRect& MyClippingRect, FSlateWindowElementList& OutDrawElements, int32 LayerId, const FWidgetStyle& InWidgetStyle, bool bParentEnabled ) const 
+{
+	const TSharedRef< FWidgetLayoutBlock > WidgetBlock = StaticCastSharedRef< FWidgetLayoutBlock >( Block );
+
+	FGeometry WidgetGeometry = AllottedGeometry.MakeChild( Block->GetLocationOffset() * ( 1 / AllottedGeometry.Scale ), Block->GetSize() * ( 1 / AllottedGeometry.Scale ), 1.0f );
+	return WidgetBlock->GetWidget()->OnPaint( WidgetGeometry, MyClippingRect, OutDrawElements, LayerId, InWidgetStyle, bParentEnabled );
+}
+
+FChildren* FSlateHyperlinkRun::GetChildren()
+{
+	return &Children;
+}
+
+void FSlateHyperlinkRun::ArrangeChildren( const TSharedRef< ILayoutBlock >& Block, const FGeometry& AllottedGeometry, FArrangedChildren& ArrangedChildren ) const 
+{
+	const TSharedRef< FWidgetLayoutBlock > WidgetBlock = StaticCastSharedRef< FWidgetLayoutBlock >( Block );
+	ArrangedChildren.AddWidget( AllottedGeometry.MakeChild( WidgetBlock->GetWidget(), Block->GetLocationOffset() * ( 1 / AllottedGeometry.Scale ), Block->GetSize() * ( 1 / AllottedGeometry.Scale ), 1.0f ) );
+}
+
+FSlateHyperlinkRun::FSlateHyperlinkRun( const TSharedRef< const FString >& InText, const FHyperlinkStyle& InStyle, const FMetadata& InMetadata, FOnClick InNavigateDelegate ) 
+	: Text( InText )
+	, Range( 0, Text->Len() )
+	, Style( InStyle )
+	, Metadata( InMetadata )
+	, NavigateDelegate( InNavigateDelegate )
+	, ViewModel( MakeShareable( new FSlateHyperlinkRun::FWidgetViewModel() ) )
+	, Children()
+{
+
+}
+
+FSlateHyperlinkRun::FSlateHyperlinkRun( const TSharedRef< const FString >& InText, const FHyperlinkStyle& InStyle, const FMetadata& InMetadata, FOnClick InNavigateDelegate, const FTextRange& InRange ) 
+	: Text( InText )
+	, Range( InRange )
+	, Style( InStyle )
+	, Metadata( InMetadata )
+	, NavigateDelegate( InNavigateDelegate )
+	, ViewModel( MakeShareable( new FSlateHyperlinkRun::FWidgetViewModel() ) )
+	, Children()
+{
+
+}
+
+FSlateHyperlinkRun::FSlateHyperlinkRun( const FSlateHyperlinkRun& Run ) 
+	: Text( Run.Text )
+	, Range( Run.Range )
+	, Style( Run.Style )
+	, Metadata( Run.Metadata )
+	, NavigateDelegate( Run.NavigateDelegate )
+	, ViewModel( MakeShareable( new FSlateHyperlinkRun::FWidgetViewModel() ) )
+	, Children()
+{
+
+}
+
diff --git a/Engine/Source/Runtime/Slate/Private/Framework/Text/SlateImageRun.cpp b/Engine/Source/Runtime/Slate/Private/Framework/Text/SlateImageRun.cpp
new file mode 100644
--- /dev/null
+++ b/Engine/Source/Runtime/Slate/Private/Framework/Text/SlateImageRun.cpp
@@ -0,0 +1,96 @@
+// Copyright 1998-2014 Epic Games, Inc. All Rights Reserved.
+
+#include "SlatePrivatePCH.h"
+#include "SlateImageRun.h"
+
+FNoChildren FSlateImageRun::NoChildrenInstance;
+
+TSharedRef< FSlateImageRun > FSlateImageRun::Create( const TSharedRef< const FString >& InText, const FSlateBrush* InImage, int16 InBaseline )
+{
+	if ( InImage == NULL )
+	{
+		InImage = FStyleDefaults::GetNoBrush();
+	}
+
+	return MakeShareable( new FSlateImageRun( InText, InImage, InBaseline ) );
+}
+
+TSharedRef< FSlateImageRun > FSlateImageRun::Create( const TSharedRef< const FString >& InText, const FSlateBrush* InImage, int16 InBaseline, const FTextRange& InRange )
+{
+	if ( InImage == NULL )
+	{
+		InImage = FStyleDefaults::GetNoBrush();
+	}
+
+	return MakeShareable( new FSlateImageRun( InText, InImage, InBaseline, InRange ) );
+}
+
+FSlateImageRun::FSlateImageRun( const TSharedRef< const FString >& InText, const FSlateBrush* InImage, int16 InBaseline, const FTextRange& InRange ) : Text( InText )
+	, Range( InRange )
+	, Image( InImage )
+	, Baseline( InBaseline )
+{
+	check( Image != NULL );
+}
+
+FSlateImageRun::FSlateImageRun( const TSharedRef< const FString >& InText, const FSlateBrush* InImage, int16 InBaseline ) : Text( InText )
+	, Range( 0, Text->Len() )
+	, Image( InImage )
+	, Baseline( InBaseline )
+{
+	check( Image != NULL );
+}
+
+FChildren* FSlateImageRun::GetChildren() 
+{
+	return &NoChildrenInstance;
+}
+
+void FSlateImageRun::ArrangeChildren( const TSharedRef< ILayoutBlock >& Block, const FGeometry& AllottedGeometry, FArrangedChildren& ArrangedChildren ) const 
+{
+	// no widgets
+}
+
+int32 FSlateImageRun::OnPaint( const FTextLayout::FLineView& Line, const TSharedRef< ILayoutBlock >& Block, const FTextBlockStyle& DefaultStyle, const FGeometry& AllottedGeometry, const FSlateRect& MyClippingRect, FSlateWindowElementList& OutDrawElements, int32 LayerId, const FWidgetStyle& InWidgetStyle, bool bParentEnabled ) const 
+{
+	if ( Image->DrawAs != ESlateBrushDrawType::NoDrawType )
+	{
+		const FColor FinalColorAndOpacity( InWidgetStyle.GetColorAndOpacityTint() * Image->GetTint( InWidgetStyle ) );
+		const uint32 DrawEffects = bParentEnabled ? ESlateDrawEffect::None : ESlateDrawEffect::DisabledEffect;
+		FSlateDrawElement::MakeBox(OutDrawElements, ++LayerId, FPaintGeometry( AllottedGeometry.AbsolutePosition + Block->GetLocationOffset(), Block->GetSize(), AllottedGeometry.Scale ), Image, MyClippingRect, DrawEffects, FinalColorAndOpacity );
+	}
+
+	return LayerId;
+}
+
+TSharedRef< ILayoutBlock > FSlateImageRun::CreateBlock( int32 BeginIndex, int32 EndIndex, FVector2D Size, const TSharedPtr< IRunHighlighter >& Highlighter )
+{
+	return FDefaultLayoutBlock::Create( SharedThis( this ), FTextRange( BeginIndex, EndIndex ), Size, Highlighter );
+}
+
+uint8 FSlateImageRun::GetKerning( int32 CurrentIndex, float Scale ) const 
+{
+	return 0;
+}
+
+FVector2D FSlateImageRun::Measure( int32 BeginIndex, int32 EndIndex, float Scale ) const 
+{
+	check( BeginIndex == Range.BeginIndex && EndIndex == Range.EndIndex );
+	return Image->ImageSize * Scale;
+}
+
+int16 FSlateImageRun::GetMaxHeight( float Scale ) const 
+{
+	return Image->ImageSize.Y * Scale;
+}
+
+int16 FSlateImageRun::GetBaseLine( float Scale ) const 
+{
+	return Baseline * Scale;
+}
+
+FTextRange FSlateImageRun::GetTextRange() const 
+{
+	return Range;
+}
+
diff --git a/Engine/Source/Runtime/Slate/Private/Framework/Text/SlateSimpleRunHighlighter.cpp b/Engine/Source/Runtime/Slate/Private/Framework/Text/SlateSimpleRunHighlighter.cpp
new file mode 100644
--- /dev/null
+++ b/Engine/Source/Runtime/Slate/Private/Framework/Text/SlateSimpleRunHighlighter.cpp
@@ -0,0 +1,52 @@
+// Copyright 1998-2014 Epic Games, Inc. All Rights Reserved.
+
+#include "SlatePrivatePCH.h"
+#include "SlateSimpleRunHighlighter.h"
+
+FNoChildren FSlateSimpleRunHighlighter::NoChildrenInstance;
+
+FSlateSimpleRunHighlighter::FSlateSimpleRunHighlighter()
+{
+
+}
+
+void FSlateSimpleRunHighlighter::ArrangeChildren( const TSharedRef< ILayoutBlock >& Block, const FGeometry& AllottedGeometry, FArrangedChildren& ArrangedChildren ) const 
+{
+	//No Widgets
+}
+
+FChildren* FSlateSimpleRunHighlighter::GetChildren()
+{
+	return &NoChildrenInstance;
+}
+
+int32 FSlateSimpleRunHighlighter::OnPaint( const FTextLayout::FLineView& Line, const TSharedRef< ISlateRun >& Run, const TSharedRef< ILayoutBlock >& Block, const FTextBlockStyle& DefaultStyle, const FGeometry& AllottedGeometry, const FSlateRect& MyClippingRect, FSlateWindowElementList& OutDrawElements, int32 LayerId, const FWidgetStyle& InWidgetStyle, bool bParentEnabled ) const 
+{
+	FVector2D Location( Block->GetLocationOffset() );
+	Location.Y = Line.Offset.Y;
+
+	// Draw the actual highlight rectangle
+	FSlateDrawElement::MakeBox(
+		OutDrawElements,
+		++LayerId,
+		FPaintGeometry( AllottedGeometry.AbsolutePosition + Location, FVector2D( Block->GetSize().X, Line.Size.Y ), AllottedGeometry.Scale ),
+		&DefaultStyle.HighlightShape,
+		MyClippingRect,
+		bParentEnabled ? ESlateDrawEffect::None : ESlateDrawEffect::DisabledEffect,
+		InWidgetStyle.GetColorAndOpacityTint() * DefaultStyle.HighlightColor
+		);
+
+	FLinearColor InvertedForeground = FLinearColor::White - InWidgetStyle.GetForegroundColor();
+	InvertedForeground.A = InWidgetStyle.GetForegroundColor().A;
+
+	FWidgetStyle WidgetStyle( InWidgetStyle );
+	WidgetStyle.SetForegroundColor( InvertedForeground );
+
+	return Run->OnPaint( Line, Block, DefaultStyle, AllottedGeometry, MyClippingRect, OutDrawElements, LayerId, WidgetStyle, bParentEnabled );
+}
+
+TSharedRef< FSlateSimpleRunHighlighter > FSlateSimpleRunHighlighter::Create()
+{
+	return MakeShareable( new FSlateSimpleRunHighlighter() );
+}
+
diff --git a/Engine/Source/Runtime/Slate/Private/Framework/Text/SlateTextLayout.cpp b/Engine/Source/Runtime/Slate/Private/Framework/Text/SlateTextLayout.cpp
new file mode 100644
--- /dev/null
+++ b/Engine/Source/Runtime/Slate/Private/Framework/Text/SlateTextLayout.cpp
@@ -0,0 +1,122 @@
+// Copyright 1998-2014 Epic Games, Inc. All Rights Reserved.
+
+#include "SlatePrivatePCH.h"
+#include "SlateTextLayout.h"
+
+TSharedRef< FSlateTextLayout > FSlateTextLayout::Create()
+{
+	TSharedRef< FSlateTextLayout > Layout = MakeShareable( new FSlateTextLayout() );
+	Layout->AggregateChildren();
+
+	return Layout;
+}
+
+FSlateTextLayout::FSlateTextLayout()
+	: Children()
+{
+
+}
+
+FChildren* FSlateTextLayout::GetChildren()
+{
+	return &Children;
+}
+
+void FSlateTextLayout::ArrangeChildren( const FGeometry& AllottedGeometry, FArrangedChildren& ArrangedChildren ) const
+{
+	for (int LineIndex = 0; LineIndex < LineViews.Num(); LineIndex++)
+	{
+		const FTextLayout::FLineView& LineView = LineViews[ LineIndex ];
+
+		for (int BlockIndex = 0; BlockIndex < LineView.Blocks.Num(); BlockIndex++)
+		{
+			const TSharedRef< ILayoutBlock > Block = LineView.Blocks[ BlockIndex ];
+			const TSharedRef< ISlateRun > Run = StaticCastSharedRef< ISlateRun >( Block->GetRun() );
+			Run->ArrangeChildren( Block, AllottedGeometry, ArrangedChildren );
+		}
+	}
+}
+
+int32 FSlateTextLayout::OnPaint( const FTextBlockStyle& DefaultTextStyle, const FGeometry& AllottedGeometry, const FSlateRect& MyClippingRect, FSlateWindowElementList& OutDrawElements, int32 LayerId, const FWidgetStyle& InWidgetStyle, bool bParentEnabled ) const
+{
+	const FSlateRect ClippingRect = AllottedGeometry.GetClippingRect().IntersectionWith(MyClippingRect);
+	const ESlateDrawEffect::Type DrawEffects = bParentEnabled ? ESlateDrawEffect::None : ESlateDrawEffect::DisabledEffect;
+
+	const int32 DebugLayer = LayerId;
+	const int32 TextLayer = LayerId + 1;
+
+	static bool ShowDebug = false;
+	FLinearColor BlockHue( 0, 1.0f, 1.0f, 0.5 );
+	int32 MaxLayerId = LayerId;
+
+	for (int LineIndex = 0; LineIndex < LineViews.Num(); LineIndex++)
+	{
+		const FTextLayout::FLineView& LineView = LineViews[ LineIndex ];
+
+		for (int BlockIndex = 0; BlockIndex < LineView.Blocks.Num(); BlockIndex++)
+		{
+			const TSharedRef< ILayoutBlock > Block = LineView.Blocks[ BlockIndex ];
+
+			if ( ShowDebug )
+			{
+				BlockHue.R += 50.0f;
+
+				FSlateDrawElement::MakeBox(
+					OutDrawElements, 
+					DebugLayer,
+					FPaintGeometry( AllottedGeometry.AbsolutePosition + Block->GetLocationOffset(), Block->GetSize(), AllottedGeometry.Scale ),
+					&DefaultTextStyle.HighlightShape,
+					ClippingRect,
+					DrawEffects,
+					InWidgetStyle.GetColorAndOpacityTint() * BlockHue.HSVToLinearRGB()
+					);
+			}
+
+			int32 BlockLayerId = TextLayer;
+			const TSharedRef< ISlateRun > Run = StaticCastSharedRef< ISlateRun >( Block->GetRun() );
+
+			int32 BlocksLayerId = MaxLayerId;
+			const TSharedPtr< ISlateRunHighlighter > Highlighter = StaticCastSharedPtr< ISlateRunHighlighter >( Block->GetHighlighter() );
+			if ( Highlighter.IsValid() )
+			{
+				BlocksLayerId = Highlighter->OnPaint( LineView, Run, Block, DefaultTextStyle, AllottedGeometry, ClippingRect, OutDrawElements, BlockLayerId, InWidgetStyle, bParentEnabled );
+			}
+			else
+			{
+				BlocksLayerId = Run->OnPaint( LineView, Block, DefaultTextStyle, AllottedGeometry, ClippingRect, OutDrawElements, BlockLayerId, InWidgetStyle, bParentEnabled );
+			}
+
+			MaxLayerId = FMath::Max( MaxLayerId, BlocksLayerId );
+		}
+	}
+
+	return MaxLayerId;
+}
+
+void FSlateTextLayout::EndLayout()
+{
+	FTextLayout::EndLayout();
+	AggregateChildren();
+}
+
+void FSlateTextLayout::AggregateChildren()
+{
+	Children.Empty();
+	const TArray< FLineModel >& LineModels = GetLineModels();
+	for (int LineModelIndex = 0; LineModelIndex < LineModels.Num(); LineModelIndex++)
+	{
+		const FLineModel& LineModel = LineModels[ LineModelIndex ];
+		for (int RunIndex = 0; RunIndex < LineModel.Runs.Num(); RunIndex++)
+		{
+			const FRunModel& LineRun = LineModel.Runs[ RunIndex ];
+			const TSharedRef< ISlateRun > SlateRun = StaticCastSharedRef< ISlateRun >( LineRun.GetRun() );
+
+			FChildren* RunChildren = SlateRun->GetChildren();
+			for (int ChildIndex = 0; ChildIndex < RunChildren->Num(); ChildIndex++)
+			{
+				const TSharedRef< SWidget >& Child = RunChildren->GetChildAt( ChildIndex );
+				Children.Add( Child );
+			}
+		}
+	}
+}
diff --git a/Engine/Source/Runtime/Slate/Private/Framework/Text/SlateTextRun.cpp b/Engine/Source/Runtime/Slate/Private/Framework/Text/SlateTextRun.cpp
new file mode 100644
--- /dev/null
+++ b/Engine/Source/Runtime/Slate/Private/Framework/Text/SlateTextRun.cpp
@@ -0,0 +1,145 @@
+// Copyright 1998-2014 Epic Games, Inc. All Rights Reserved.
+
+#include "SlatePrivatePCH.h"
+#include "SlateTextRun.h"
+
+FNoChildren FSlateTextRun::NoChildrenInstance;
+
+TSharedRef< FSlateTextRun > FSlateTextRun::Create( const TSharedRef< const FString >& InText, const FTextBlockStyle& Style )
+{
+	return MakeShareable( new FSlateTextRun( InText, Style ) );
+}
+
+TSharedRef< FSlateTextRun > FSlateTextRun::Create( const TSharedRef< const FString >& InText, const FTextBlockStyle& Style, const FTextRange& InRange )
+{
+	return MakeShareable( new FSlateTextRun( InText, Style, InRange ) );
+}
+
+FTextRange FSlateTextRun::GetTextRange() const 
+{
+	return Range;
+}
+
+int16 FSlateTextRun::GetBaseLine( float Scale ) const 
+{
+	const TSharedRef< FSlateFontMeasure > FontMeasure = FSlateApplication::Get().GetRenderer()->GetFontMeasureService();
+	return FontMeasure->GetBaseline( Style.Font, Scale );
+}
+
+int16 FSlateTextRun::GetMaxHeight( float Scale ) const 
+{
+	const TSharedRef< FSlateFontMeasure > FontMeasure = FSlateApplication::Get().GetRenderer()->GetFontMeasureService();
+	return FontMeasure->GetMaxCharacterHeight( Style.Font, Scale );
+}
+
+FVector2D FSlateTextRun::Measure( int32 BeginIndex, int32 EndIndex, float Scale ) const 
+{
+	if ( EndIndex - BeginIndex == 0 )
+	{
+		return FVector2D( 0, GetMaxHeight( Scale ) );
+	}
+
+	const TSharedRef< FSlateFontMeasure > FontMeasure = FSlateApplication::Get().GetRenderer()->GetFontMeasureService();
+	return FontMeasure->Measure( **Text, BeginIndex, EndIndex, Style.Font, true, Scale );
+}
+
+uint8 FSlateTextRun::GetKerning( int32 CurrentIndex, float Scale ) const 
+{
+	const int32 PreviousIndex = CurrentIndex - 1;
+	if ( PreviousIndex < 0 )
+	{
+		return 0;
+	}
+
+	const TSharedRef< FSlateFontMeasure > FontMeasure = FSlateApplication::Get().GetRenderer()->GetFontMeasureService();
+	return FontMeasure->GetKerning( Style.Font, Scale, (*Text)[ PreviousIndex ], (*Text)[ CurrentIndex ] );
+}
+
+TSharedRef< ILayoutBlock > FSlateTextRun::CreateBlock( int32 BeginIndex, int32 EndIndex, FVector2D Size, const TSharedPtr< IRunHighlighter >& Highlighter )
+{
+	return FDefaultLayoutBlock::Create( SharedThis( this ), FTextRange( BeginIndex, EndIndex ), Size, Highlighter );
+}
+
+int32 FSlateTextRun::OnPaint( const FTextLayout::FLineView& Line, const TSharedRef< ILayoutBlock >& Block, const FTextBlockStyle& DefaultStyle, const FGeometry& AllottedGeometry, const FSlateRect& MyClippingRect, FSlateWindowElementList& OutDrawElements, int32 LayerId, const FWidgetStyle& InWidgetStyle, bool bParentEnabled ) const 
+{
+	const FSlateRect ClippingRect = AllottedGeometry.GetClippingRect().IntersectionWith(MyClippingRect);
+	const ESlateDrawEffect::Type DrawEffects = bParentEnabled ? ESlateDrawEffect::None : ESlateDrawEffect::DisabledEffect;
+
+	const bool ShouldDropShadow = Style.ShadowOffset.Size() > 0;
+	const FVector2D BlockLocationOffset = Block->GetLocationOffset();
+	const FTextRange BlockRange = Block->GetTextRange();
+
+	// Draw the optional shadow
+	if ( ShouldDropShadow )
+	{
+		FSlateDrawElement::MakeText(
+			OutDrawElements,
+			++LayerId,
+			FPaintGeometry( AllottedGeometry.AbsolutePosition + Block->GetLocationOffset() + Style.ShadowOffset, Block->GetSize(), AllottedGeometry.Scale ),
+			Text.Get(),
+			BlockRange.BeginIndex,
+			BlockRange.EndIndex,
+			Style.Font,
+			ClippingRect,
+			DrawEffects,
+			InWidgetStyle.GetColorAndOpacityTint() * Style.ShadowColorAndOpacity
+			);
+	}
+
+	// Draw the text itself
+	FSlateDrawElement::MakeText(
+		OutDrawElements,
+		++LayerId,
+		FPaintGeometry( AllottedGeometry.AbsolutePosition + Block->GetLocationOffset(), Block->GetSize(), AllottedGeometry.Scale ),
+		Text.Get(),
+		BlockRange.BeginIndex,
+		BlockRange.EndIndex,
+		Style.Font,
+		ClippingRect,
+		DrawEffects,
+		InWidgetStyle.GetColorAndOpacityTint() * Style.ColorAndOpacity.GetColor(InWidgetStyle)
+		);
+
+	return LayerId;
+}
+
+FChildren* FSlateTextRun::GetChildren()
+{
+	return &NoChildrenInstance;
+}
+
+void FSlateTextRun::ArrangeChildren( const TSharedRef< ILayoutBlock >& Block, const FGeometry& AllottedGeometry, FArrangedChildren& ArrangedChildren ) const 
+{
+	// no widgets
+}
+
+FSlateTextRun::FSlateTextRun( const TSharedRef< const FString >& InText, const FTextBlockStyle& InStyle ) : Text( InText )
+	, Style( InStyle )
+	, Range( 0, Text->Len() )
+#if TEXT_LAYOUT_DEBUG
+	, DebugSlice( FString( Text->Len(), **Text ) )
+#endif
+{
+
+}
+
+FSlateTextRun::FSlateTextRun( const TSharedRef< const FString >& InText, const FTextBlockStyle& InStyle, const FTextRange& InRange ) : Text( InText )
+	, Style( InStyle )
+	, Range( InRange )
+#if TEXT_LAYOUT_DEBUG
+	, DebugSlice( FString( InRange.EndIndex - InRange.BeginIndex, (**Text) + InRange.BeginIndex ) )
+#endif
+{
+
+}
+
+FSlateTextRun::FSlateTextRun( const FSlateTextRun& Run ) : Text( Run.Text )
+	, Style( Run.Style )
+	, Range( Run.Range )
+#if TEXT_LAYOUT_DEBUG
+	, DebugSlice( Run.DebugSlice )
+#endif
+{
+
+}
+
diff --git a/Engine/Source/Runtime/Slate/Private/Framework/Text/SlateWidgetRun.cpp b/Engine/Source/Runtime/Slate/Private/Framework/Text/SlateWidgetRun.cpp
new file mode 100644
--- /dev/null
+++ b/Engine/Source/Runtime/Slate/Private/Framework/Text/SlateWidgetRun.cpp
@@ -0,0 +1,90 @@
+// Copyright 1998-2014 Epic Games, Inc. All Rights Reserved.
+
+#include "SlatePrivatePCH.h"
+#include "SlateWidgetRun.h"
+
+TSharedRef< FSlateWidgetRun > FSlateWidgetRun::Create( const TSharedRef< const FString >& InText, const FWidgetRunInfo& InWidgetInfo )
+{
+	return MakeShareable( new FSlateWidgetRun( InText, InWidgetInfo ) );
+}
+
+TSharedRef< FSlateWidgetRun > FSlateWidgetRun::Create( const TSharedRef< const FString >& InText, const FWidgetRunInfo& InWidgetInfo, const FTextRange& InRange )
+{
+	return MakeShareable( new FSlateWidgetRun( InText, InWidgetInfo, InRange ) );
+}
+
+FTextRange FSlateWidgetRun::GetTextRange() const 
+{
+	return Range;
+}
+
+int16 FSlateWidgetRun::GetBaseLine( float Scale ) const 
+{
+	return Info.Baseline * Scale;
+}
+
+int16 FSlateWidgetRun::GetMaxHeight( float Scale ) const 
+{
+	return Info.Size.Get( Info.Widget->GetDesiredSize() ).Y * Scale;
+}
+
+FVector2D FSlateWidgetRun::Measure( int32 StartIndex, int32 EndIndex, float Scale ) const 
+{
+	return Info.Size.Get( Info.Widget->GetDesiredSize() ) * Scale;
+}
+
+uint8 FSlateWidgetRun::GetKerning( int32 CurrentIndex, float Scale ) const 
+{
+	return 0;
+}
+
+TSharedRef< ILayoutBlock > FSlateWidgetRun::CreateBlock( int32 StartIndex, int32 EndIndex, FVector2D Size, const TSharedPtr< IRunHighlighter >& Highlighter )
+{
+	return FDefaultLayoutBlock::Create( SharedThis( this ), FTextRange( StartIndex, EndIndex ), Size, Highlighter );
+}
+
+int32 FSlateWidgetRun::OnPaint( const FTextLayout::FLineView& Line, const TSharedRef< ILayoutBlock >& Block, const FTextBlockStyle& DefaultStyle, const FGeometry& AllottedGeometry, const FSlateRect& MyClippingRect, FSlateWindowElementList& OutDrawElements, int32 LayerId, const FWidgetStyle& InWidgetStyle, bool bParentEnabled ) const 
+{
+	FGeometry WidgetGeometry = AllottedGeometry.MakeChild( Block->GetLocationOffset() * ( 1 / AllottedGeometry.Scale ), Block->GetSize() * ( 1 / AllottedGeometry.Scale ), 1.0f );
+	return Info.Widget->OnPaint( WidgetGeometry, MyClippingRect, OutDrawElements, LayerId, InWidgetStyle, bParentEnabled );
+}
+
+FChildren* FSlateWidgetRun::GetChildren()
+{
+	return &Children;
+}
+
+void FSlateWidgetRun::ArrangeChildren( const TSharedRef< ILayoutBlock >& Block, const FGeometry& AllottedGeometry, FArrangedChildren& ArrangedChildren ) const
+{
+	ArrangedChildren.AddWidget( AllottedGeometry.MakeChild( Info.Widget, Block->GetLocationOffset() * ( 1 / AllottedGeometry.Scale ), Block->GetSize() * ( 1 / AllottedGeometry.Scale ), 1.0f ) );
+}
+
+FSlateWidgetRun::FSlateWidgetRun( const TSharedRef< const FString >& InText, const FWidgetRunInfo& InWidgetInfo ) 
+	: Text( InText )
+	, Range( 0, Text->Len() )
+	, Info( InWidgetInfo )
+	, Children()
+{
+	Info.Widget->SlatePrepass();
+	Children.Add( Info.Widget );
+}
+
+FSlateWidgetRun::FSlateWidgetRun( const TSharedRef< const FString >& InText, const FWidgetRunInfo& InWidgetInfo, const FTextRange& InRange ) 
+	: Text( InText )
+	, Range( InRange )
+	, Info( InWidgetInfo )
+	, Children()
+{
+	Info.Widget->SlatePrepass();
+	Children.Add( Info.Widget );
+}
+
+FSlateWidgetRun::FSlateWidgetRun( const FSlateWidgetRun& Run ) 
+	: Text( Run.Text )
+	, Range( Run.Range )
+	, Info( Run.Info )
+	, Children( Run.Children )
+{
+
+}
+
diff --git a/Engine/Source/Runtime/Slate/Private/Framework/Text/SlateWordWrapper.cpp b/Engine/Source/Runtime/Slate/Private/Framework/Text/SlateWordWrapper.cpp
new file mode 100644
--- /dev/null
+++ b/Engine/Source/Runtime/Slate/Private/Framework/Text/SlateWordWrapper.cpp
@@ -0,0 +1,168 @@
+// Copyright 1998-2014 Epic Games, Inc. All Rights Reserved.
+
+#include "SlatePrivatePCH.h"
+#include "SlateWordWrapper.h"
+
+namespace
+{
+	class FSlateWordWrapperBase : public FWordWrapper
+	{
+	protected:
+		FSlateWordWrapperBase(const FString& InString, const FSlateFontInfo& InFontInfo, const int32 InWrapWidth, const float InFontScale, FWrappedLineData* const OutWrappedLineData);
+		bool DoesSubstringFit(const int32 StartIndex, const int32 EndIndex);
+		int32 FindIndexAtOrAfterWrapWidth(const int32 StartIndex);
+	protected:
+		const FSlateFontInfo& FontInfo;
+		const int32 WrapWidth;
+		const float FontScale;
+	};
+
+	bool FSlateWordWrapperBase::DoesSubstringFit(const int32 StartIndex, const int32 EndIndex)
+	{
+		const TSharedRef< FSlateFontMeasure > FontMeasureService = FSlateApplication::Get().GetRenderer()->GetFontMeasureService();
+		FVector2D Size( FontMeasureService->Measure( String, StartIndex, EndIndex, FontInfo, false, FontScale ) );
+		return Size.X <= WrapWidth;
+	}
+
+	int32 FSlateWordWrapperBase::FindIndexAtOrAfterWrapWidth(const int32 StartIndex)
+	{
+		FString StringToMeasure(String + StartIndex);
+		const TSharedRef< FSlateFontMeasure > FontMeasureService = FSlateApplication::Get().GetRenderer()->GetFontMeasureService();
+
+		int32 Index = FontMeasureService->FindLastWholeCharacterIndexBeforeOffset(StringToMeasure, FontInfo, WrapWidth, FontScale);
+
+		if (Index == INDEX_NONE)
+		{
+			return StartIndex;
+		}
+
+		//This index is inclusive but we need to return an exclusive index so we increment.
+		return StartIndex + Index + 1; 
+	}
+
+	FSlateWordWrapperBase::FSlateWordWrapperBase(const FString& InString, const FSlateFontInfo& InFontInfo, const int32 InWrapWidth, const float InFontScale, FWrappedLineData* const OutWrappedLineData)
+		: FWordWrapper(*InString, InString.Len(), OutWrappedLineData), FontInfo(InFontInfo), WrapWidth(InWrapWidth), FontScale(InFontScale)
+	{
+	}
+}
+
+namespace
+{
+	class FSlateWordWrapperSingle : public FSlateWordWrapperBase
+	{
+	public:
+		/**
+		* Used to generate multi-line/wrapped text.
+		*
+		* @param InString The unwrapped text.
+		* @param InFontInfo The font used to render the text.
+		* @param InWrapWidth The width available.
+		* @param OutWrappedLineData An optional array to fill with the indices from the source string marking the begin and end points of the wrapped lines
+		*
+		* @return The index of the character closest to the specified horizontal offset
+		*/
+		static FString Execute(const FString& InString, const FSlateFontInfo& InFontInfo, const int32 InWrapWidth, const float InFontScale, FWrappedLineData* const OutWrappedLineData);
+	protected:
+		void AddLine(const int32 StartIndex, const int32 EndIndex);
+	private:
+		FSlateWordWrapperSingle(const FString& InString, const FSlateFontInfo& InFontInfo, const int32 InWrapWidth, const float InFontScale, FWrappedLineData* const OutWrappedLineData);
+	private:
+		FString Result;
+		bool bPendingNewLine;
+	};
+
+	FString FSlateWordWrapperSingle::Execute(const FString& InString, const FSlateFontInfo& InFontInfo, const int32 InWrapWidth, const float InFontScale, FWrappedLineData* const OutWrappedLineData)
+	{
+		FSlateWordWrapperSingle WordWrapper(InString, InFontInfo, InWrapWidth, InFontScale, OutWrappedLineData);
+		for(int i = 0; i < InString.Len(); ++i) // Sanity check: Doesn't seem valid to have more lines than code units.
+		{
+			if( !WordWrapper.ProcessLine() )
+			{
+				break;
+			}
+		}
+		return WordWrapper.Result;
+	}
+
+	void FSlateWordWrapperSingle::AddLine(const int32 StartIndex, const int32 EndIndex)
+	{
+		if(bPendingNewLine)
+		{
+			Result += TEXT("\n");
+		}
+		FString NewString(EndIndex - StartIndex, String + StartIndex);
+		Result += NewString;
+		bPendingNewLine = NewString.EndsWith( FString( TEXT("\n") ) ) ? false : true; // Add a new line if the last new line didn't end with a line break.
+	}
+
+	FSlateWordWrapperSingle::FSlateWordWrapperSingle(const FString& InString, const FSlateFontInfo& InFontInfo, const int32 InWrapWidth, const float InFontScale, FWrappedLineData* const OutWrappedLineData)
+		: FSlateWordWrapperBase(InString, InFontInfo, InWrapWidth, InFontScale, OutWrappedLineData), bPendingNewLine(false)
+	{
+	}
+}
+
+FString SlateWordWrapper::WrapText( const FString& InString, const FSlateFontInfo& InFontInfo, const int32 InWrapWidth, const float InFontScale, FWordWrapper::FWrappedLineData* const OutWrappedLineData )
+{
+	return FSlateWordWrapperSingle::Execute(InString, InFontInfo, InWrapWidth, InFontScale, OutWrappedLineData);
+}
+
+namespace
+{
+	class FSlateWordWrapperMulti : public FSlateWordWrapperBase
+	{
+	public:
+		/**
+		* Used to generate multi-line/wrapped text.
+		*
+		* @param InString The unwrapped text.
+		* @param InFontInfo The font used to render the text.
+		* @param InWrapWidth The width available.
+		* @param OutString The resulting string with newlines inserted.
+		* @param OutWrappedLineData An optional array to fill with the indices from the source string marking the begin and end points of the wrapped lines
+		*/
+		static void Execute(FString& InString, const FSlateFontInfo& InFontInfo, const int32 InWrapWidth, const float InFontScale, TArray< FWrappedStringSlice >& OutWrappedText, FWrappedLineData* const OutWrappedLineData);
+	protected:
+		void AddLine(const int32 StartIndex, const int32 EndIndex);
+	private:
+		FSlateWordWrapperMulti(FString& InString, const FSlateFontInfo& InFontInfo, const int32 InWrapWidth, const float InFontScale, TArray< FWrappedStringSlice >& OutWrappedText, FWrappedLineData* const OutWrappedLineData);
+	private:
+		FString& StringReference;
+		TArray< struct FWrappedStringSlice >& Result;
+	};
+
+	void FSlateWordWrapperMulti::Execute(FString& InString, const FSlateFontInfo& InFontInfo, const int32 InWrapWidth, const float InFontScale, TArray< FWrappedStringSlice >& OutWrappedText, FWrappedLineData* const OutWrappedLineData)
+	{
+		FSlateWordWrapperMulti WordWrapper(InString, InFontInfo, InWrapWidth, InFontScale, OutWrappedText, OutWrappedLineData);
+		for(int i = 0; i < InString.Len(); ++i) // Sanity check: Doesn't seem valid to have more lines than code units.
+		{
+			if( !WordWrapper.ProcessLine() )
+			{
+				break;
+			}
+		}
+	}
+
+	void FSlateWordWrapperMulti::AddLine(const int32 StartIndex, const int32 EndIndex)
+	{
+		const TSharedRef< FSlateFontMeasure > FontMeasureService = FSlateApplication::Get().GetRenderer()->GetFontMeasureService();
+		FVector2D StringSize( FontMeasureService->Measure( StringReference.Mid(StartIndex, EndIndex - StartIndex), FontInfo, FontScale ) );
+		Result.Add( FWrappedStringSlice(StringReference, StartIndex, EndIndex, StringSize) );
+	}
+
+	FSlateWordWrapperMulti::FSlateWordWrapperMulti(FString& InString, const FSlateFontInfo& InFontInfo, const int32 InWrapWidth, const float InFontScale, TArray< FWrappedStringSlice >& OutWrappedText, FWrappedLineData* const OutWrappedLineData)
+		: FSlateWordWrapperBase(InString, InFontInfo, InWrapWidth, InFontScale, OutWrappedLineData), StringReference(InString), Result(OutWrappedText)
+	{
+	}
+}
+
+void SlateWordWrapper::WrapText( const TArray< TSharedRef<FString> >& InTextLines, const FSlateFontInfo& InFontInfo, const float InWrapWidth, const float InFontScale, TArray< FWrappedStringSlice >& OutWrappedText, FWordWrapper::FWrappedLineData* const OutWrappedLineData )
+{
+	OutWrappedText.Empty();
+
+	// Iterate over every source line, generating wrapping information into OutWrappedText.
+	for (int32 LineIndex=0; LineIndex < InTextLines.Num(); ++LineIndex )
+	{
+		FString& CurrentSourceLine = InTextLines[LineIndex].Get();
+		FSlateWordWrapperMulti::Execute(CurrentSourceLine, InFontInfo, InFontScale, InWrapWidth, OutWrappedText, OutWrappedLineData);
+	}
+}
\ No newline at end of file
diff --git a/Engine/Source/Runtime/Slate/Private/Framework/Text/TextDecorators.cpp b/Engine/Source/Runtime/Slate/Private/Framework/Text/TextDecorators.cpp
new file mode 100644
--- /dev/null
+++ b/Engine/Source/Runtime/Slate/Private/Framework/Text/TextDecorators.cpp
@@ -0,0 +1,140 @@
+// Copyright 1998-2014 Epic Games, Inc. All Rights Reserved.
+
+#include "SlatePrivatePCH.h"
+#include "TextDecorators.h"
+
+TSharedRef< FWidgetDecorator > FWidgetDecorator::Create( FString InRunName, const FCreateWidget& InCreateWidgetDelegate )
+{
+	return MakeShareable( new FWidgetDecorator( InRunName, InCreateWidgetDelegate ) );
+}
+
+FWidgetDecorator::FWidgetDecorator( FString InRunName, const FCreateWidget& InCreateWidgetDelegate )
+	: RunName( InRunName )
+	, CreateWidgetDelegate( InCreateWidgetDelegate )
+{
+
+}
+
+bool FWidgetDecorator::Supports( const FTextRunParseResults& RunInfo, const FString& Text ) const
+{
+	return ( RunInfo.Name == RunName );
+}
+
+TSharedRef< ISlateRun > FWidgetDecorator::Create( const FTextRunParseResults& Run, const FString& OriginalText, const TSharedRef< FString >& InOutModelText, const ISlateStyle* Style )
+{
+	FTextRange ModelRange;
+	ModelRange.BeginIndex = InOutModelText->Len();
+	*InOutModelText += TEXT('\x200B'); // Zero-Width Breaking Space
+	ModelRange.EndIndex = InOutModelText->Len();
+
+	FTextRunInfo RunInfo( Run.Name, FText::FromString( OriginalText.Mid( Run.ContentRange.BeginIndex, Run.ContentRange.EndIndex - Run.ContentRange.BeginIndex ) ) );
+	for(TPair<FString, FTextRange> Pair : Run.MetaData)
+	{
+		RunInfo.MetaData[Pair.Key] = OriginalText.Mid( Pair.Value.BeginIndex, Pair.Value.EndIndex - Pair.Value.BeginIndex);
+	}
+
+	return FSlateWidgetRun::Create( InOutModelText, CreateWidgetDelegate.Execute( RunInfo, Style ), ModelRange );
+}
+
+TSharedRef< FImageDecorator > FImageDecorator::Create( FString InRunName, const ISlateStyle* const InStyle )
+{
+	return MakeShareable( new FImageDecorator( InRunName, InStyle ) );
+}
+
+FImageDecorator::FImageDecorator( FString InRunName, const ISlateStyle* const InOverrideStyle )
+	: RunName( InRunName )
+	, OverrideStyle( InOverrideStyle )
+{
+
+}
+
+bool FImageDecorator::Supports( const FTextRunParseResults& RunInfo, const FString& Text ) const
+{
+	return ( RunInfo.Name == RunName );
+}
+
+TSharedRef< ISlateRun > FImageDecorator::Create( const FTextRunParseResults& RunInfo, const FString& OriginalText, const TSharedRef< FString >& InOutModelText, const ISlateStyle* Style )
+{
+	const FTextRange* const BrushNameRange = RunInfo.MetaData.Find( TEXT("src") );
+
+	FTextRange ModelRange;
+	ModelRange.BeginIndex = InOutModelText->Len();
+	*InOutModelText += TEXT('\x200B'); // Zero-Width Breaking Space
+	ModelRange.EndIndex = InOutModelText->Len();
+
+	if ( BrushNameRange != NULL )
+	{
+		const FString BrushNameString = OriginalText.Mid(BrushNameRange->BeginIndex, BrushNameRange->EndIndex - BrushNameRange->BeginIndex);
+		if ( OverrideStyle != NULL )
+		{
+			Style = OverrideStyle;
+		}
+
+		FName BrushName( *BrushNameString );
+		if ( Style->HasWidgetStyle< FInlineTextImageStyle >( BrushName ) )
+		{
+			const FInlineTextImageStyle& ImageStyle = Style->GetWidgetStyle< FInlineTextImageStyle >( BrushName );
+			return FSlateImageRun::Create( InOutModelText, &ImageStyle.Image, ImageStyle.Baseline, ModelRange );
+		}
+
+		return FSlateImageRun::Create( InOutModelText, Style->GetBrush( BrushName ), 0, ModelRange );
+	}
+
+	const FInlineTextImageStyle& ImageStyle = FInlineTextImageStyle::GetDefault();
+	return FSlateImageRun::Create( InOutModelText, &ImageStyle.Image, ImageStyle.Baseline, ModelRange );
+}
+
+TSharedRef< FHyperlinkDecorator > FHyperlinkDecorator::Create( FString Id, const FSlateHyperlinkRun::FOnClick& NavigateDelegate )
+{
+	return MakeShareable( new FHyperlinkDecorator( Id, NavigateDelegate ) );
+}
+
+FHyperlinkDecorator::FHyperlinkDecorator( FString InId, const FSlateHyperlinkRun::FOnClick& InNavigateDelegate )
+	: Id( InId )
+	, NavigateDelegate( InNavigateDelegate )
+{
+
+}
+
+bool FHyperlinkDecorator::Supports( const FTextRunParseResults& RunInfo, const FString& Text ) const
+{
+	const FTextRange* const MetaDataIdRange = RunInfo.MetaData.Find( TEXT("id") );
+	FString MetaDataId;
+	if(MetaDataIdRange)
+	{
+		MetaDataId = Text.Mid(MetaDataIdRange->BeginIndex, MetaDataIdRange->EndIndex - MetaDataIdRange->BeginIndex);
+	}
+
+	return ( RunInfo.Name == TEXT("a") && MetaDataIdRange != NULL && MetaDataId == Id );
+}
+
+TSharedRef< ISlateRun > FHyperlinkDecorator::Create( const FTextRunParseResults& RunInfo, const FString& OriginalText, const TSharedRef< FString >& InOutModelText, const ISlateStyle* Style )
+{
+	FString StyleName = TEXT("Hyperlink");
+
+	const FTextRange* const MetaDataStyleNameRange = RunInfo.MetaData.Find( TEXT("style") );
+	if ( MetaDataStyleNameRange != NULL )
+	{
+		const FString MetaDataStyleName = OriginalText.Mid(MetaDataStyleNameRange->BeginIndex, MetaDataStyleNameRange->EndIndex - MetaDataStyleNameRange->BeginIndex);
+		StyleName = *MetaDataStyleName;
+	}
+
+	FTextRange ModelRange;
+	ModelRange.BeginIndex = InOutModelText->Len();
+	*InOutModelText += OriginalText.Mid(RunInfo.ContentRange.BeginIndex, RunInfo.ContentRange.EndIndex - RunInfo.ContentRange.BeginIndex);
+	ModelRange.EndIndex = InOutModelText->Len();
+
+	if ( !Style->HasWidgetStyle<FHyperlinkStyle>( FName( *StyleName ) ) )
+	{
+		Style = &FCoreStyle::Get();
+	}
+
+	TMap<FString, FString> MetaData;
+	for(const TPair<FString, FTextRange>& Pair : RunInfo.MetaData)
+	{
+		MetaData.Add(Pair.Key, OriginalText.Mid( Pair.Value.BeginIndex, Pair.Value.EndIndex - Pair.Value.BeginIndex));
+		//MetaData[Pair.Key] = OriginalText.Mid( Pair.Value.BeginIndex, Pair.Value.EndIndex - Pair.Value.BeginIndex);
+	}
+
+	return FSlateHyperlinkRun::Create( InOutModelText, Style->GetWidgetStyle<FHyperlinkStyle>( FName( *StyleName ) ), MetaData, NavigateDelegate, ModelRange );
+}
\ No newline at end of file
diff --git a/Engine/Source/Runtime/Slate/Private/Framework/Text/TextEditHelper.cpp b/Engine/Source/Runtime/Slate/Private/Framework/Text/TextEditHelper.cpp
new file mode 100644
--- /dev/null
+++ b/Engine/Source/Runtime/Slate/Private/Framework/Text/TextEditHelper.cpp
@@ -0,0 +1,259 @@
+// Copyright 1998-2014 Epic Games, Inc. All Rights Reserved.
+
+#include "SlatePrivatePCH.h"
+#include "TextEditHelper.h"
+
+/**
+ * Ensure that text transactions are always completed.
+ * i.e. never forget to call FinishChangingText();
+ */
+struct FScopedTextTransaction
+{
+	explicit FScopedTextTransaction( ITextEditorWidget& InTextEditor )
+	: TextEditor( InTextEditor )
+	{
+		TextEditor.StartChangingText();
+	}
+
+	~FScopedTextTransaction()
+	{
+		TextEditor.FinishChangingText();
+	}
+
+	ITextEditorWidget& TextEditor;
+};
+
+FReply FTextEditHelper::OnKeyChar( const FGeometry& MyGeometry, const FCharacterEvent& InCharacterEvent, class ITextEditorWidget& TextEditor )
+{
+	FReply Reply = FReply::Unhandled();
+
+	// Check for special characters
+	const TCHAR Character = InCharacterEvent.GetCharacter();
+
+	switch( Character )
+	{
+		// Backspace
+	case TCHAR( 8 ):
+		{
+			if( !TextEditor.GetIsReadOnly() )
+			{
+				FScopedTextTransaction TextTransaction(TextEditor);
+
+				TextEditor.BackspaceChar();
+				
+				Reply = FReply::Handled();
+			}
+		}
+		break;
+
+
+		// Tab
+	case TCHAR( '\t' ):
+		{
+			Reply = FReply::Handled();
+		}
+		break;
+
+
+		// Swallow OnKeyChar keys that we don't want to be entered into the buffer
+	case 1:		// Swallow Ctrl+A, we handle that through OnKeyDown
+	case 3:		// Swallow Ctrl+C, we handle that through OnKeyDown
+	case 13:	// Swallow Enter, we handle that through OnKeyDown
+	case 22:	// Swallow Ctrl+V, we handle that through OnKeyDown
+	case 24:	// Swallow Ctrl+X, we handle that through OnKeyDown
+	case 25:	// Swallow Ctrl+Y, we handle that through OnKeyDown
+	case 26:	// Swallow Ctrl+Z, we handle that through OnKeyDown
+	case 27:	// Swallow ESC, we handle that through OnKeyDown
+	case 127:	// Swallow CTRL+Backspace, we handle that through OnKeyDown
+		Reply = FReply::Handled();
+		break;
+
+
+		// Any other character!
+	default:
+		{
+			// Type the character, but only if it is allowed.
+			if (!TextEditor.GetIsReadOnly() && TextEditor.CanTypeCharacter(Character))
+			{
+				FScopedTextTransaction TextTransaction(TextEditor);
+				TextEditor.TypeChar( Character );
+				Reply = FReply::Handled();
+			}
+			else
+			{
+				Reply = FReply::Unhandled();
+			}
+		}
+		break;
+
+	}
+
+
+	return Reply;
+}
+
+FReply FTextEditHelper::OnKeyDown( const FGeometry& MyGeometry, const FKeyboardEvent& InKeyboardEvent, ITextEditorWidget& TextEditor )
+{
+	const FKey Key = InKeyboardEvent.GetKey();
+
+	if( Key == EKeys::Left )
+	{
+		return TextEditor.MoveCursor(
+			// Ctrl moves a whole word instead of one character.	
+			InKeyboardEvent.IsControlDown() ? ECursorMoveMethod::Word : ECursorMoveMethod::CharacterHorizontal,
+			// Move left
+			-1,
+			// Shift selects text.	
+			InKeyboardEvent.IsShiftDown() ? ECursorAction::SelectText : ECursorAction::MoveCursor
+		);
+	}
+	else if( Key == EKeys::Right )
+	{
+		return TextEditor.MoveCursor(
+			// Ctrl moves a whole word instead of one character.	
+			InKeyboardEvent.IsControlDown() ? ECursorMoveMethod::Word : ECursorMoveMethod::CharacterHorizontal,
+			// Move right
+			+1,
+			// Shift selects text.	
+			InKeyboardEvent.IsShiftDown() ? ECursorAction::SelectText : ECursorAction::MoveCursor
+		);
+	}
+	else if ( Key == EKeys::Up )
+	{
+		return TextEditor.MoveCursor(
+			ECursorMoveMethod::CharacterVertical,
+			// Move up
+			-1,
+			// Shift selects text.	
+			InKeyboardEvent.IsShiftDown() ? ECursorAction::SelectText : ECursorAction::MoveCursor
+		);
+	}
+	else if ( Key == EKeys::Down )
+	{
+		return TextEditor.MoveCursor(
+			ECursorMoveMethod::CharacterVertical,
+			// Move down
+			+1,
+			// Shift selects text.	
+			InKeyboardEvent.IsShiftDown() ? ECursorAction::SelectText : ECursorAction::MoveCursor
+		);
+	}
+	else if( Key == EKeys::Home )
+	{
+		// Go to the beginning of the document; select text if Shift is down.
+		TextEditor.JumpTo(
+			(InKeyboardEvent.IsControlDown() ) ? ETextLocation::BeginningOfDocument : ETextLocation::BeginningOfLine,
+			(InKeyboardEvent.IsShiftDown()) ? ECursorAction::SelectText : ECursorAction::MoveCursor );
+
+		return FReply::Handled();
+	}
+	else if ( Key == EKeys::End )
+	{
+		// Go to the end of the document; select text if Shift is down.
+		TextEditor.JumpTo(
+			(InKeyboardEvent.IsControlDown() ) ? ETextLocation::EndOfDocument : ETextLocation::EndOfLine,
+			(InKeyboardEvent.IsShiftDown()) ? ECursorAction::SelectText : ECursorAction::MoveCursor );
+
+		return FReply::Handled();
+	}
+	else if( Key == EKeys::Enter )
+	{
+		TextEditor.OnEnter();
+
+		return FReply::Handled();
+	}
+	else if( Key == EKeys::Delete && !TextEditor.GetIsReadOnly() )
+	{
+		// @Todo: Slate keybindings support more than one set of keys. 
+		// Delete to next word boundary (Ctrl+Delete)
+		if (InKeyboardEvent.IsControlDown() && !InKeyboardEvent.IsAltDown() && !InKeyboardEvent.IsShiftDown())
+		{
+			TextEditor.MoveCursor(ECursorMoveMethod::Word, +1, ECursorAction::SelectText);
+		}
+
+		TextEditor.DeleteChar();
+		
+		return FReply::Handled();
+	}	
+	else if( Key == EKeys::Escape )
+	{
+		return TextEditor.OnEscape();;
+	}
+
+	// @Todo: Slate keybindings support more than one set of keys. 
+	//Alternate key for cut (Shift+Delete)
+	else if( Key == EKeys::Delete && InKeyboardEvent.IsShiftDown() && TextEditor.CanExecuteCut() )
+	{
+		// Cut text to clipboard
+		TextEditor.CutSelectedTextToClipboard();
+		
+		return FReply::Handled();
+	}
+
+	// @Todo: Slate keybindings support more than one set of keys. 
+	// Alternate key for copy (Ctrl+Insert) 
+	else if( Key == EKeys::Insert && InKeyboardEvent.IsControlDown() && TextEditor.CanExecuteCopy() ) 
+	{
+		// Copy text to clipboard
+		TextEditor.CopySelectedTextToClipboard();
+		
+		return FReply::Handled();
+	}
+
+
+	// @Todo: Slate keybindings support more than one set of keys. 
+	// Alternate key for paste (Shift+Insert) 
+	else if( Key == EKeys::Insert && InKeyboardEvent.IsShiftDown() && TextEditor.CanExecutePaste() )
+	{
+		// Paste text from clipboard
+		TextEditor.PasteTextFromClipboard();
+		
+		return FReply::Handled();
+	}
+
+	// @Todo: Slate keybindings support more than one set of keys. 
+	//Alternate key for undo (Alt+Backspace)
+	else if( Key == EKeys::BackSpace && InKeyboardEvent.IsAltDown() && !InKeyboardEvent.IsShiftDown() && TextEditor.CanExecuteUndo() )
+	{
+		// Undo
+		TextEditor.Undo();
+		
+		return FReply::Handled();
+	}
+
+	// @Todo: Slate keybindings support more than one set of keys. 
+	// Delete to previous word boundary (Ctrl+Backspace)
+	else if( Key == EKeys::BackSpace && InKeyboardEvent.IsControlDown() && !InKeyboardEvent.IsAltDown() && !InKeyboardEvent.IsShiftDown() && !TextEditor.GetIsReadOnly() )
+	{
+		FScopedTextTransaction TextTransaction(TextEditor);
+
+		TextEditor.MoveCursor(ECursorMoveMethod::Word, -1, ECursorAction::SelectText);
+		TextEditor.BackspaceChar();
+
+		return FReply::Handled();
+	}
+
+
+	// Ctrl+Y (or Ctrl+Shift+Z, or Alt+Shift+Backspace) to redo
+	else if( !TextEditor.GetIsReadOnly() && ( ( Key == EKeys::Y && InKeyboardEvent.IsControlDown() ) ||
+		( Key == EKeys::Z && InKeyboardEvent.IsControlDown() && InKeyboardEvent.IsShiftDown() ) ||
+		( Key == EKeys::BackSpace && InKeyboardEvent.IsAltDown() && InKeyboardEvent.IsShiftDown() ) ) )
+	{
+		// Redo
+		TextEditor.Redo();
+		
+		return FReply::Handled();
+	}
+	else if( !InKeyboardEvent.IsAltDown() && !InKeyboardEvent.IsControlDown() && InKeyboardEvent.GetKey() != EKeys::Tab && InKeyboardEvent.GetCharacter() != 0 )
+	{
+		// Shift and a character was pressed or a single character was pressed.  We will type something in an upcoming OnKeyChar event.  
+		// Absorb this event so it is not bubbled and handled by other widgets that could have something bound to the key press.
+
+		//Note: Tab can generate a character code but not result in a typed character in this box so we ignore it.
+		return FReply::Handled();
+	}
+	else
+	{
+		return FReply::Unhandled();
+	}
+}
\ No newline at end of file
diff --git a/Engine/Source/Runtime/Slate/Private/Framework/Text/TextEditHelper.h b/Engine/Source/Runtime/Slate/Private/Framework/Text/TextEditHelper.h
new file mode 100644
--- /dev/null
+++ b/Engine/Source/Runtime/Slate/Private/Framework/Text/TextEditHelper.h
@@ -0,0 +1,70 @@
+// Copyright 1998-2014 Epic Games, Inc. All Rights Reserved.
+
+#pragma once
+
+/**
+ * Design constraints
+ */
+namespace EditableTextDefs
+{
+	/** Maximum number of undo levels to store */
+	static const int32 MaxUndoLevels = 99;
+
+	/** Vertical text offset from top of widget, as a scalar percentage of maximum character height */
+	static const float TextVertOffsetPercent = 0.0f;
+
+	/** How quickly the cursor's spring will pull the cursor around.  Larger is faster. */
+	static const float CaretSpringConstant = 65.0f;
+
+	/** Vertical offset of caret from text location, in pixels */
+	static const float CaretVertOffset = 1.0f;
+
+	/** Width of the caret, as a scalar percentage of the font's maximum character height */
+	static const float CaretWidthPercent = 0.08f;
+
+	/** Height of the caret, as a scalar percentage of the font's maximum character height */
+	static const float CaretHeightPercent = 0.95f;
+
+	/** How long after the user last interacted with the keyboard should we keep the caret at full opacity? */
+	static const float CaretBlinkPauseTime = 0.1f;
+
+	/** How many times should the caret blink per second (full on/off cycles) */
+	static const float BlinksPerSecond = 1.0f;
+
+	/** How many pixels to extend the selection rectangle's left side horizontally */
+	static const float SelectionRectLeftOffset = 0.0f;
+
+	/** How many pixels to extend the selection rectangle's right side horizontally */
+	static const float SelectionRectRightOffset = 0.0f;
+
+	/** How quickly the selection 'targeting' rectangle will slide around.  Larger is faster. */
+	static const float SelectionTargetSpringConstant = 25.0f;
+
+	/** Duration of animation selection target effects */
+	static const float SelectionTargetEffectDuration = 0.25f;
+
+	/** Opacity of the selection target effect overlay */
+	static const float SelectionTargetOpacity = 0.8f;
+
+	/** How large the selection target effect will be when selecting text, as a scalar percentage of font height */
+	static const float SelectingAnimOffsetPercent = 0.15f;
+
+	/** How large the selection target effect will be when deselecting text, as a scalar percentage of font height */
+	static const float DeselectingAnimOffsetPercent = 0.25f;
+}
+
+/**
+ * Contains helper routines that translate input into actions for an ITextEditorWidget.
+ * Used to share code between SEditableText and SMultiLineEditableText
+ */
+class FTextEditHelper
+{
+public:
+
+	/** Someone typed a character, translate the input event into a call to a ITextEditorWidget action */
+	static FReply OnKeyChar( const FGeometry& MyGeometry, const FCharacterEvent& InCharacterEvent, class ITextEditorWidget& TextEditor );
+
+	/** Someone pressed a key, translate the input event into a call to a ITextEditorWidget action */
+	static FReply OnKeyDown( const FGeometry& MyGeometry, const FKeyboardEvent& InKeyboardEvent, ITextEditorWidget& TextEditor );
+
+};
\ No newline at end of file
