Slate: Reorganized private Slate files; proper pre-compiled header for all files.

diff --git a/Engine/Source/Runtime/Slate/Private/Framework/Text/TextLayout.cpp b/Engine/Source/Runtime/Slate/Private/Framework/Text/TextLayout.cpp
new file mode 100644
--- /dev/null
+++ b/Engine/Source/Runtime/Slate/Private/Framework/Text/TextLayout.cpp
@@ -0,0 +1,890 @@
+// Copyright 1998-2014 Epic Games, Inc. All Rights Reserved.
+
+#include "SlatePrivatePCH.h"
+
+FTextLayout::FBreakCandidate FTextLayout::CreateBreakCandidate( int32& OutRunIndex, FLineModel& Line, int32 PreviousBreak, int32 CurrentBreak )
+{
+	bool SuccessfullyMeasuredSlice = false;
+	int16 MaxAboveBaseline = 0;
+	int16 MaxBelowBaseline = 0;
+	FVector2D BreakSize( ForceInitToZero );
+	FVector2D BreakSizeWithoutTrailingWhitespace( ForceInitToZero );
+	int32 WhitespaceStopIndex = CurrentBreak;
+	uint8 Kerning = 0;
+
+	if ( Line.Runs.IsValidIndex( OutRunIndex ) )
+	{
+		FRunModel& Run = Line.Runs[ OutRunIndex ];
+		const FTextRange Range = Run.GetTextRange();
+		int32 BeginIndex = FMath::Max( PreviousBreak, Range.BeginIndex );
+
+		if ( BeginIndex > 0 )
+		{
+			Kerning = Run.GetKerning( BeginIndex, Scale );
+		}
+	}
+
+	// We need to consider the Runs when detecting and measuring the text lengths of Lines because
+	// the font style used makes a difference.
+	for (; OutRunIndex < Line.Runs.Num(); OutRunIndex++)
+	{
+		FRunModel& Run = Line.Runs[ OutRunIndex ];
+		const FTextRange Range = Run.GetTextRange();
+
+		FVector2D SliceSize;
+		FVector2D SliceSizeWithoutTrailingWhitespace;
+		int32 StopIndex = PreviousBreak;
+
+		WhitespaceStopIndex = StopIndex = FMath::Min( Range.EndIndex, CurrentBreak );
+		int32 BeginIndex = FMath::Max( PreviousBreak, Range.BeginIndex );
+
+		//@todo Use ICU's whitespace functionality [12/5/2013 justin.sargent]
+		while( WhitespaceStopIndex > BeginIndex && TChar<TCHAR>::IsWhitespace( (*Line.Text)[ WhitespaceStopIndex - 1 ] ) )
+		{
+			--WhitespaceStopIndex;
+		}
+
+		if ( WhitespaceStopIndex != StopIndex && BeginIndex != WhitespaceStopIndex )
+		{
+			SliceSize = SliceSizeWithoutTrailingWhitespace = Run.Measure( BeginIndex, WhitespaceStopIndex, Scale );
+			SliceSize.X += Run.Measure( WhitespaceStopIndex, StopIndex, Scale ).X;
+		}
+		else
+		{
+			SliceSize = Run.Measure( BeginIndex, StopIndex, Scale );
+			SliceSizeWithoutTrailingWhitespace = SliceSize;
+		}
+
+		BreakSize.X += SliceSize.X; // We accumulate the slice widths
+		BreakSizeWithoutTrailingWhitespace.X += SliceSizeWithoutTrailingWhitespace.X; // We accumulate the slice widths
+
+		// Get the baseline and flip it's sign; Baselines are generally negative
+		const int16 Baseline = -(Run.GetBaseLine( Scale ));
+
+		// For the height of the slice we need to take into account the largest value below and above the baseline and add those together
+		MaxAboveBaseline = FMath::Max( MaxAboveBaseline, (int16)( Run.GetMaxHeight( Scale ) - Baseline ) );
+		MaxBelowBaseline = FMath::Max( MaxBelowBaseline, Baseline );
+
+		if ( StopIndex == CurrentBreak )
+		{
+			SuccessfullyMeasuredSlice = true;
+
+			if ( OutRunIndex < Line.Runs.Num() && StopIndex == Line.Runs[ OutRunIndex ].GetTextRange().EndIndex )
+			{
+				++OutRunIndex;
+			}
+			break;
+		}
+	}
+	check( SuccessfullyMeasuredSlice == true );
+
+	BreakSize.Y = BreakSizeWithoutTrailingWhitespace.Y = MaxAboveBaseline + MaxBelowBaseline;
+
+	FBreakCandidate BreakCandidate;
+	BreakCandidate.Size = BreakSize;
+	BreakCandidate.SizeWithoutTrailingWhitespace = BreakSizeWithoutTrailingWhitespace;
+	BreakCandidate.Range = FTextRange( PreviousBreak, CurrentBreak );
+	BreakCandidate.RangeWithoutTrailingWhitespace = FTextRange( PreviousBreak, WhitespaceStopIndex );
+	BreakCandidate.MaxAboveBaseline = MaxAboveBaseline;
+	BreakCandidate.MaxBelowBaseline = MaxBelowBaseline;
+	BreakCandidate.Kerning = Kerning;
+
+#if TEXT_LAYOUT_DEBUG
+	BreakCandidate.DebugSlice = FString( BreakCandidate.Range.EndIndex - BreakCandidate.Range.BeginIndex, (**Line.Text) + BreakCandidate.Range.BeginIndex );
+#endif
+
+	return BreakCandidate;
+}
+
+void FTextLayout::CreateLineViewBlocks( const FLineModel& Line, const int32 StopIndex, const int32 VisualStopIndex, int32& OutRunIndex, int32& OutHighlightIndex, int32& OutPreviousBlockEnd, TArray< TSharedRef< ILayoutBlock > >& OutSoftLine )
+{
+	const bool VisualStopIndexIsValid = VisualStopIndex != INDEX_NONE;
+	int16 MaxAboveBaseline = 0;
+	int16 MaxBelowBaseline = 0;
+
+	for (; OutRunIndex < Line.Runs.Num(); )
+	{
+		const FRunModel& Run = Line.Runs[ OutRunIndex ];
+		const FTextRange RunRange = Run.GetTextRange();
+
+		TSharedPtr< IRunHighlighter > BlockHighlighter = NULL;
+
+		int32 BlockStopIndex = RunRange.EndIndex;
+		if ( OutHighlightIndex != INDEX_NONE )
+		{
+			const FTextHighlight& Highlight = Line.Highlights[ OutHighlightIndex ];
+			if ( OutPreviousBlockEnd >= Highlight.Range.BeginIndex )
+			{
+				if ( Highlight.Range.EndIndex <= RunRange.EndIndex )
+				{
+					BlockStopIndex = Highlight.Range.EndIndex;
+					BlockHighlighter = Highlight.Highlighter;
+				}
+			}
+			else
+			{
+				if ( Highlight.Range.BeginIndex <= RunRange.EndIndex )
+				{
+					BlockStopIndex = Highlight.Range.BeginIndex;
+					BlockHighlighter = NULL;
+				}
+			}
+		}
+
+		if ( VisualStopIndex != INDEX_NONE )
+		{
+			BlockStopIndex = FMath::Min( VisualStopIndex, BlockStopIndex );
+		}
+
+		int32 BlockBeginIndex = FMath::Max( OutPreviousBlockEnd, RunRange.BeginIndex );
+		const bool IsLastBlock = BlockStopIndex == VisualStopIndex;
+
+		if ( RunRange.BeginIndex < BlockStopIndex && RunRange.EndIndex > BlockBeginIndex )
+		{
+			check( BlockBeginIndex <= BlockStopIndex );
+
+			FBlockDefinition BlockDefine;
+			BlockDefine.Range = BlockDefine.VisualRange = FTextRange( BlockBeginIndex, BlockStopIndex );
+			BlockDefine.Highlighter = BlockHighlighter;
+
+			OutSoftLine.Add( Run.CreateBlock( BlockDefine, Scale ) );
+			OutPreviousBlockEnd = BlockStopIndex;
+		}
+		else
+		{
+			FBlockDefinition BlockDefine;
+			BlockDefine.Range = RunRange;
+			BlockDefine.VisualRange = FTextRange( RunRange.BeginIndex, ( IsLastBlock && VisualStopIndexIsValid ) ? VisualStopIndex : RunRange.EndIndex );
+			BlockDefine.Highlighter = BlockHighlighter;
+
+			OutSoftLine.Add( Run.CreateBlock( BlockDefine, Scale ) );
+			OutPreviousBlockEnd = RunRange.EndIndex;
+		}
+
+		// Get the baseline and flip it's sign; Baselines are generally negative
+		const int16 Baseline = -(Run.GetBaseLine( Scale ));
+
+		// For the height of the slice we need to take into account the largest value below and above the baseline and add those together
+		MaxAboveBaseline = FMath::Max( MaxAboveBaseline, (int16)( Run.GetMaxHeight( Scale ) - Baseline ) );
+		MaxBelowBaseline = FMath::Max( MaxBelowBaseline, Baseline );
+
+		if ( BlockStopIndex == RunRange.EndIndex )
+		{
+			++OutRunIndex;
+		}
+
+		if ( OutHighlightIndex != INDEX_NONE && BlockStopIndex == Line.Highlights[ OutHighlightIndex ].Range.EndIndex )
+		{
+			++OutHighlightIndex;
+
+			if ( OutHighlightIndex >= Line.Highlights.Num() )
+			{
+				OutHighlightIndex = INDEX_NONE;
+			}
+		}
+
+		if ( IsLastBlock )
+		{
+			break;
+		}
+	}
+
+	//Skip any highlighted ranges which we have skipped due to ignoring trailing whitespace
+	while ( OutHighlightIndex != INDEX_NONE && Line.Highlights[ OutHighlightIndex ].Range.EndIndex <= StopIndex )
+	{
+		OutHighlightIndex++;
+
+		if ( OutHighlightIndex >= Line.Highlights.Num() )
+		{
+			OutHighlightIndex = INDEX_NONE;
+		}
+	}
+
+	FVector2D LineSize( ForceInitToZero );
+	FVector2D CurrentOffset( Margin.Left, Margin.Top + DrawSize.Y );
+
+	if ( OutSoftLine.Num() > 0 )
+	{
+		float CurrentHorizontalPos = 0.0f;
+		for (int Index = 0; Index < OutSoftLine.Num(); Index++)
+		{
+			const TSharedRef< ILayoutBlock > Block = OutSoftLine[ Index ];
+			const int16 BlockBaseline = Block->GetRun()->GetBaseLine( Scale );
+			const int16 VerticalOffset = MaxAboveBaseline - Block->GetSize().Y - BlockBaseline;
+
+			Block->SetLocationOffset( FVector2D( CurrentOffset.X + CurrentHorizontalPos, CurrentOffset.Y + VerticalOffset ) );
+
+			CurrentHorizontalPos += Block->GetSize().X;
+		}
+
+		LineSize.X = CurrentHorizontalPos;
+		LineSize.Y = MaxAboveBaseline + MaxBelowBaseline;
+	}
+
+	FTextLayout::FLineView LineView;
+	LineView.Size = LineSize;
+	LineView.Offset = CurrentOffset;
+	LineView.Blocks.Append( OutSoftLine );
+
+	LineViews.Add( LineView );
+
+	DrawSize.X = FMath::Max( DrawSize.X, LineSize.X ); //DrawSize.X is the size of the longest line + the Margin
+	const float LineHeight = LineSize.Y * LineHeightPercentage; 
+	DrawSize.Y += LineHeight; //DrawSize.Y is the total height of all lines
+}
+
+void FTextLayout::JustifyLayout()
+{
+	if ( Justification == ETextJustify::Left )
+	{
+		return;
+	}
+
+	const float LayoutWidthNoMargin = DrawSize.X - Margin.GetTotalSpaceAlong<Orient_Horizontal>() * Scale;
+
+	for (int LineViewIndex = 0; LineViewIndex < LineViews.Num(); LineViewIndex++)
+	{
+		FLineView& LineView = LineViews[ LineViewIndex ];
+		const float ExtraSpace = LayoutWidthNoMargin - LineView.Size.X;
+
+		FVector2D BlockOffset( ForceInitToZero );
+		if ( Justification == ETextJustify::Center )
+		{
+			BlockOffset = FVector2D( ExtraSpace / 2, 0 );
+		}
+		else if ( Justification == ETextJustify::Right )
+		{
+			BlockOffset = FVector2D( ExtraSpace, 0 );
+		}
+
+		for (int BlockIndex = 0; BlockIndex < LineView.Blocks.Num(); BlockIndex++)
+		{
+			const TSharedRef< ILayoutBlock >& Block = LineView.Blocks[ BlockIndex ];
+			Block->SetLocationOffset( Block->GetLocationOffset() + BlockOffset );
+		}
+	}
+}
+
+void FTextLayout::FlowLayout()
+{
+	check( WrappingWidth >= 0 );
+	const bool IsWrapping = WrappingWidth > 0;
+	const float WrappingDrawWidth = FMath::Max( 0.01f, ( WrappingWidth - ( Margin.Left + Margin.Right ) ) * Scale );
+
+	if ( IsWrapping )
+	{
+		ReadyForWrapping();
+	}
+
+	TArray< TSharedRef< ILayoutBlock > > SoftLine;
+
+	for (int LineModelIndex = 0; LineModelIndex < LineModels.Num(); LineModelIndex++)
+	{
+		const FLineModel& LineModel = LineModels[ LineModelIndex ];
+
+		float CurrentWidth = 0.0f;
+		int32 CurrentRunIndex = 0;
+		int32 PreviousBlockEnd = 0;
+
+		int32 CurrentHighlightIndex = 0;
+		if ( CurrentHighlightIndex >= LineModel.Highlights.Num() )
+		{
+			CurrentHighlightIndex = INDEX_NONE;
+		}
+
+		// if the Line doesn't have any BreakCandidates
+		if ( LineModel.BreakCandidates.Num() == 0 )
+		{
+			//Then iterate over all of it's runs
+			CreateLineViewBlocks( LineModel, INDEX_NONE, INDEX_NONE, /*OUT*/CurrentRunIndex, /*OUT*/CurrentHighlightIndex, /*OUT*/PreviousBlockEnd, SoftLine );
+			check( CurrentRunIndex == LineModel.Runs.Num() );
+			CurrentWidth = 0;
+			SoftLine.Empty();
+		}
+		else
+		{
+			for (int BreakIndex = 0; BreakIndex < LineModel.BreakCandidates.Num(); BreakIndex++)
+			{
+				const FBreakCandidate& Break = LineModel.BreakCandidates[ BreakIndex ];
+
+				const bool IsLastBreak = BreakIndex + 1 == LineModel.BreakCandidates.Num();
+				const bool IsFirstBreakOnSoftLine = CurrentWidth == 0;
+				const uint8 Kerning = ( IsFirstBreakOnSoftLine ) ? Break.Kerning : 0;
+				const bool BreakDoesFit = !IsWrapping || CurrentWidth + Break.Size.X + Kerning <= WrappingDrawWidth;
+
+				if ( !BreakDoesFit || IsLastBreak )
+				{
+					const bool BreakWithoutTrailingWhitespaceDoesFit = !IsWrapping || CurrentWidth + Break.SizeWithoutTrailingWhitespace.X + Kerning <= WrappingDrawWidth;
+					const bool IsFirstBreak = BreakIndex == 0;
+
+					const FBreakCandidate& FinalBreakOnSoftLine = ( !IsFirstBreak && !IsFirstBreakOnSoftLine && !BreakWithoutTrailingWhitespaceDoesFit ) ? LineModel.BreakCandidates[ --BreakIndex ] : Break;
+
+					CreateLineViewBlocks( LineModel, FinalBreakOnSoftLine.Range.EndIndex, FinalBreakOnSoftLine.RangeWithoutTrailingWhitespace.EndIndex, /*OUT*/CurrentRunIndex, /*OUT*/CurrentHighlightIndex, /*OUT*/PreviousBlockEnd, SoftLine );
+
+					if ( CurrentRunIndex < LineModel.Runs.Num() && FinalBreakOnSoftLine.Range.EndIndex == LineModel.Runs[ CurrentRunIndex ].GetTextRange().EndIndex )
+					{
+						++CurrentRunIndex;
+					}
+
+					PreviousBlockEnd = FinalBreakOnSoftLine.Range.EndIndex;
+
+					CurrentWidth = 0;
+					SoftLine.Empty();
+				} 
+				else
+				{
+					CurrentWidth += Break.Size.X;
+				}
+			}
+		}
+	}
+
+	//Add on the margins to the DrawSize
+	DrawSize.X += Margin.GetTotalSpaceAlong<Orient_Horizontal>() * Scale;
+	DrawSize.Y += Margin.GetTotalSpaceAlong<Orient_Vertical>() * Scale;
+}
+
+void FTextLayout::CreateView()
+{
+	ClearView();
+	BeginLayout();
+
+	FlowLayout();
+	JustifyLayout();
+
+	EndLayout();
+	HasChanged = false;
+}
+
+void FTextLayout::EndLayout()
+{
+	for (int LineModelIndex = 0; LineModelIndex < LineModels.Num(); LineModelIndex++)
+	{
+		FLineModel& LineModel = LineModels[ LineModelIndex ];
+		for (int RunIndex = 0; RunIndex < LineModel.Runs.Num(); RunIndex++)
+		{
+			LineModel.Runs[ RunIndex ].EndLayout();
+		}
+	}
+}
+
+void FTextLayout::BeginLayout()
+{
+	for (int LineModelIndex = 0; LineModelIndex < LineModels.Num(); LineModelIndex++)
+	{
+		FLineModel& LineModel = LineModels[ LineModelIndex ];
+		for (int RunIndex = 0; RunIndex < LineModel.Runs.Num(); RunIndex++)
+		{
+			LineModel.Runs[ RunIndex ].BeginLayout();
+		}
+	}
+}
+
+void FTextLayout::ClearView()
+{
+	DrawSize = FVector2D::ZeroVector;
+	LineViews.Empty();
+}
+
+void FTextLayout::ReadyForWrapping()
+{
+	if ( HasWrapInformation )
+	{
+		return;
+	}
+
+	const TSharedRef< FSlateFontMeasure > FontMeasure = FSlateApplication::Get().GetRenderer()->GetFontMeasureService();
+
+	for (int LineModelIndex = 0; LineModelIndex < LineModels.Num(); LineModelIndex++)
+	{
+		FLineModel& LineModel = LineModels[ LineModelIndex ];
+
+		for (int RunIndex = 0; RunIndex < LineModel.Runs.Num(); RunIndex++)
+		{
+			LineModel.Runs[ RunIndex ].ClearCache();
+		}
+
+		FLineBreakIterator LineBreakIterator( **LineModel.Text );
+
+		int32 PreviousBreak = 0;
+		int32 CurrentBreak = 0;
+		int32 CurrentRunIndex = 0;
+
+		while( ( CurrentBreak = LineBreakIterator.MoveToNext() ) != INDEX_NONE )
+		{
+			LineModel.BreakCandidates.Add( CreateBreakCandidate(/*OUT*/CurrentRunIndex, LineModel, PreviousBreak, CurrentBreak) );
+			PreviousBreak = CurrentBreak;
+		}
+	}
+
+	HasWrapInformation = true;
+}
+
+FTextLayout::FTextLayout() : LineModels()
+	, LineViews()
+	, HasWrapInformation( false )
+	, HasChanged( false )
+	, Scale( 1.0f )
+	, WrappingWidth( 0 )
+	, Margin()
+	, LineHeightPercentage( 1.0f )
+	, DrawSize( ForceInitToZero )
+{
+
+}
+
+void FTextLayout::UpdateIfNeeded()
+{
+	if ( HasChanged )
+	{
+		// if something has changed then create a new View
+		CreateView();
+	}
+}
+
+void FTextLayout::SetHighlights( const TArray< FTextHighlight >& Highlights )
+{
+	ClearHighlights();
+
+	for (int Index = 0; Index < Highlights.Num(); Index++)
+	{
+		AddHighlight( Highlights[ Index ] );
+	}
+}
+
+void FTextLayout::AddHighlight( const FTextHighlight& Highlight )
+{
+	checkf( LineModels.IsValidIndex( Highlight.LineIndex ), TEXT("Highlights must be for a valid Line Index") );
+
+	//Highlights needs to be in order and not overlap
+	bool InsertSuccessful = false;
+	FLineModel& LineModel = LineModels[ Highlight.LineIndex ];
+	for (int Index = 0; Index < LineModel.Highlights.Num(); Index++)
+	{
+		if ( LineModel.Highlights[ Index ].Range.BeginIndex > Highlight.Range.BeginIndex )
+		{
+			checkf( Index == 0 || LineModel.Highlights[ Index - 1 ].Range.EndIndex <= Highlight.Range.BeginIndex, TEXT("Highlights cannot overlap") );
+			LineModel.Highlights.Insert( Highlight, Index - 1 );
+			InsertSuccessful = true;
+		}
+		else if ( LineModel.Highlights[ Index ].Range.EndIndex > Highlight.Range.EndIndex )
+		{
+			checkf( LineModel.Highlights[ Index ].Range.BeginIndex >= Highlight.Range.EndIndex, TEXT("Highlights cannot overlap") );
+			LineModel.Highlights.Insert( Highlight, Index - 1 );
+			InsertSuccessful = true;
+		}
+	}
+
+	if ( !InsertSuccessful )
+	{
+		LineModel.Highlights.Add( Highlight );
+	}
+
+	HasChanged = true;
+}
+
+void FTextLayout::ClearHighlights()
+{
+	for (int32 Index = 0; Index < LineModels.Num(); Index++)
+	{
+		LineModels[ Index ].Highlights.Empty();
+	}
+
+	HasChanged = true;
+}
+
+void FTextLayout::AddLine( const TSharedRef< FString >& Text, const TArray< TSharedRef< IRun > >& Runs )
+{
+	FLineModel LineModel( Text );
+
+	for (int Index = 0; Index < Runs.Num(); Index++)
+	{
+		LineModel.Runs.Add( FRunModel( Runs[ Index ] ) );
+	}
+
+	LineModels.Add( LineModel );
+	HasChanged = true;
+}
+
+void FTextLayout::ClearLines()
+{
+	LineModels.Empty();
+	HasChanged = true;
+}
+
+void FTextLayout::SetMargin( const FMargin& InMargin )
+{
+	if ( Margin == InMargin )
+	{
+		return;
+	}
+
+	const FMargin PreviousMargin = Margin;
+	Margin = InMargin;
+
+	if ( WrappingWidth > 0 )
+	{
+		// Since we are wrapping our text we'll need to rebuild the view as the actual wrapping width includes the margin.
+		HasChanged = true;
+	}
+	else
+	{
+		FVector2D OffsetAdjustment( Margin.Left - PreviousMargin.Left, Margin.Top - PreviousMargin.Top );
+		OffsetAdjustment *= Scale;
+		for (int LineViewIndex = 0; LineViewIndex < LineViews.Num(); LineViewIndex++)
+		{
+			FLineView& LineView = LineViews[ LineViewIndex ];
+
+			for (int BlockIndex = 0; BlockIndex < LineView.Blocks.Num(); BlockIndex++)
+			{
+				const TSharedRef< ILayoutBlock >& Block = LineView.Blocks[ BlockIndex ];
+				Block->SetLocationOffset( Block->GetLocationOffset() + OffsetAdjustment );
+			}
+		}
+
+		DrawSize.X += ( Margin.GetTotalSpaceAlong<Orient_Horizontal>() - PreviousMargin.GetTotalSpaceAlong<Orient_Horizontal>() ) * Scale;
+		DrawSize.Y += ( Margin.GetTotalSpaceAlong<Orient_Vertical>() - PreviousMargin.GetTotalSpaceAlong<Orient_Vertical>() ) * Scale;
+	}
+}
+
+FMargin FTextLayout::GetMargin() const
+{
+	return Margin;
+}
+
+void FTextLayout::SetScale( float Value )
+{
+	if ( Scale == Value )
+	{
+		return;
+	}
+
+	Scale = Value;
+	HasChanged = true;
+
+	if ( HasWrapInformation )
+	{
+		for (int LineModelIndex = 0; LineModelIndex < LineModels.Num(); LineModelIndex++)
+		{
+			FLineModel& LineModel = LineModels[ LineModelIndex ];
+			int32 CurrentRunIndex = 0;
+			int32 PreviousBreak = 0;
+			TArray< FBreakCandidate > UpdatedBreakCandidates;
+
+			for (int BreakIndex = 0; BreakIndex < LineModel.BreakCandidates.Num(); BreakIndex++)
+			{
+				const FBreakCandidate& Break = LineModel.BreakCandidates[ BreakIndex ];
+				UpdatedBreakCandidates.Add( CreateBreakCandidate( /*OUT*/CurrentRunIndex, LineModel, PreviousBreak, Break.Range.EndIndex ) );
+				PreviousBreak = Break.Range.EndIndex;
+			}
+
+			check( LineModel.BreakCandidates.Num() == UpdatedBreakCandidates.Num() );
+			LineModel.BreakCandidates = UpdatedBreakCandidates;
+
+			for (int RunIndex = 0; RunIndex < LineModel.Runs.Num(); RunIndex++)
+			{
+				FRunModel& Run = LineModel.Runs[ RunIndex ];
+				Run.ClearCache();
+			}
+		}
+	}
+}
+
+float FTextLayout::GetScale() const
+{
+	return Scale;
+}
+
+void FTextLayout::SetJustification( ETextJustify::Type Value )
+{
+	if ( Justification == Value )
+	{
+		return;
+	}
+
+	Justification = Value;
+	HasChanged = true;
+}
+
+ETextJustify::Type FTextLayout::GetJustification() const
+{
+	return Justification;
+}
+
+void FTextLayout::SetLineHeightPercentage( float Value )
+{
+	if ( LineHeightPercentage != Value )
+	{
+		LineHeightPercentage = Value; 
+		HasChanged = true;
+	}
+}
+
+float FTextLayout::GetLineHeightPercentage() const
+{
+	return LineHeightPercentage;
+}
+
+void FTextLayout::SetWrappingWidth( float Value )
+{
+	if ( WrappingWidth != Value )
+	{
+		WrappingWidth = Value; 
+		HasChanged = true;
+	}
+}
+
+float FTextLayout::GetWrappingWidth() const
+{
+	return WrappingWidth;
+}
+
+FVector2D FTextLayout::GetDrawSize() const
+{
+	return DrawSize;
+}
+
+FVector2D FTextLayout::GetSize() const
+{
+	return DrawSize * ( 1 / Scale );
+}
+
+const TArray< FTextLayout::FLineModel >& FTextLayout::GetLineModels() const
+{
+	return LineModels;
+}
+
+const TArray< FTextLayout::FLineView >& FTextLayout::GetLineViews() const
+{
+	return LineViews;
+}
+
+FTextLayout::~FTextLayout()
+{
+
+}
+
+FTextLayout::FLineModel::FLineModel( const TSharedRef< FString >& InText ) : Text( InText )
+	, Runs()
+	, BreakCandidates()
+	, Highlights()
+{
+
+}
+
+void FTextLayout::FRunModel::ClearCache()
+{
+	MeasuredRanges.Empty();
+	MeasuredRangeSizes.Empty();
+}
+
+TSharedRef< ILayoutBlock > FTextLayout::FRunModel::CreateBlock( const FBlockDefinition& BlockDefine, float Scale ) const
+{
+	check( BlockDefine.Range.EndIndex >= BlockDefine.VisualRange.EndIndex );
+	check( BlockDefine.Range.BeginIndex <= BlockDefine.VisualRange.BeginIndex );
+
+	const FTextRange& SizeRange = BlockDefine.VisualRange;
+
+	if ( MeasuredRanges.Num() == 0 )
+	{
+		FTextRange RunRange = Run->GetTextRange();
+		return Run->CreateBlock( BlockDefine.Range.BeginIndex, BlockDefine.Range.EndIndex, Run->Measure( SizeRange.BeginIndex, SizeRange.EndIndex, Scale ), BlockDefine.Highlighter );
+	}
+
+	int32 StartRangeIndex = 0;
+	int32 EndRangeIndex = 0;
+
+	if ( MeasuredRanges.Num() > 16 )
+	{
+		if ( SizeRange.BeginIndex != 0 )
+		{
+			StartRangeIndex = BinarySearchForBeginIndex( MeasuredRanges, SizeRange.BeginIndex );
+			check( StartRangeIndex != INDEX_NONE);
+		}
+
+		EndRangeIndex = StartRangeIndex;
+		if ( StartRangeIndex != MeasuredRanges.Num() - 1 )
+		{
+			EndRangeIndex = BinarySearchForEndIndex( MeasuredRanges, StartRangeIndex, SizeRange.EndIndex );
+			check( EndRangeIndex != INDEX_NONE);
+		}
+	}
+	else
+	{
+		const int32 MaxValidIndex = MeasuredRanges.Num() - 1;
+
+		if ( SizeRange.BeginIndex != 0 )
+		{
+			bool FoundBeginIndexMatch = false;
+			for (; StartRangeIndex < MaxValidIndex; StartRangeIndex++)
+			{
+				FoundBeginIndexMatch = MeasuredRanges[ StartRangeIndex ].BeginIndex >= SizeRange.BeginIndex;
+				if ( FoundBeginIndexMatch ) 
+				{
+					break;
+				}
+			}
+			check( FoundBeginIndexMatch == true || StartRangeIndex == MaxValidIndex );
+		}
+
+		EndRangeIndex = StartRangeIndex;
+		if ( StartRangeIndex != MaxValidIndex )
+		{
+			bool FoundEndIndexMatch = false;
+			for (; EndRangeIndex < MeasuredRanges.Num(); EndRangeIndex++)
+			{
+				FoundEndIndexMatch = MeasuredRanges[ EndRangeIndex ].EndIndex >= SizeRange.EndIndex;
+				if ( FoundEndIndexMatch ) 
+				{
+					break;
+				}
+			}
+			check( FoundEndIndexMatch == true || EndRangeIndex == MaxValidIndex );
+		}
+	}
+
+	FVector2D BlockSize( ForceInitToZero );
+	if ( StartRangeIndex == EndRangeIndex )
+	{
+		if ( MeasuredRanges[ StartRangeIndex ].BeginIndex == SizeRange.BeginIndex && 
+			MeasuredRanges[ StartRangeIndex ].EndIndex == SizeRange.EndIndex )
+		{
+			BlockSize += MeasuredRangeSizes[ StartRangeIndex ];
+		}
+		else
+		{
+			BlockSize += Run->Measure( SizeRange.BeginIndex, SizeRange.EndIndex, Scale );
+		}
+	}
+	else
+	{
+		if ( MeasuredRanges[ StartRangeIndex ].BeginIndex == SizeRange.BeginIndex )
+		{
+			BlockSize += MeasuredRangeSizes[ StartRangeIndex ];
+		}
+		else
+		{
+			BlockSize += Run->Measure( SizeRange.BeginIndex, MeasuredRanges[ StartRangeIndex ].EndIndex, Scale );
+		}
+
+		for (int Index = StartRangeIndex + 1; Index < EndRangeIndex; Index++)
+		{
+			BlockSize.X += MeasuredRangeSizes[ Index ].X;
+			BlockSize.Y = FMath::Max( MeasuredRangeSizes[ Index ].Y, BlockSize.Y );
+		}
+
+		if ( MeasuredRanges[ EndRangeIndex ].EndIndex == SizeRange.EndIndex )
+		{
+			BlockSize.X += MeasuredRangeSizes[ EndRangeIndex ].X;
+			BlockSize.Y = FMath::Max( MeasuredRangeSizes[ EndRangeIndex ].Y, BlockSize.Y );
+		}
+		else
+		{
+			FVector2D Size = Run->Measure( MeasuredRanges[ EndRangeIndex ].BeginIndex, SizeRange.EndIndex, Scale );
+			BlockSize.X += Size.X;
+			BlockSize.Y = FMath::Max( Size.Y, BlockSize.Y );
+		}
+	}
+
+	return Run->CreateBlock( BlockDefine.Range.BeginIndex, BlockDefine.Range.EndIndex, BlockSize, BlockDefine.Highlighter );
+}
+
+int FTextLayout::FRunModel::BinarySearchForEndIndex( const TArray< FTextRange >& Ranges, int32 RangeBeginIndex, int32 EndIndex )
+{
+	int32 Min = RangeBeginIndex;
+	int32 Mid = 0;
+	int32 Max = Ranges.Num() - 1;
+	while (Max >= Min)
+	{
+		Mid = Min + ((Max - Min) / 2);
+		if ( Ranges[Mid].EndIndex == EndIndex )
+		{
+			return Mid; 
+		}
+		else if ( Ranges[Mid].EndIndex < EndIndex )
+		{
+			Min = Mid + 1;
+		}
+		else
+		{
+			Max = Mid - 1;
+		}
+	}
+
+	return Mid;
+}
+
+int FTextLayout::FRunModel::BinarySearchForBeginIndex( const TArray< FTextRange >& Ranges, int32 BeginIndex )
+{
+	int32 Min = 0;
+	int32 Mid = 0;
+	int32 Max = Ranges.Num() - 1;
+	while (Max >= Min)
+	{
+		Mid = Min + ((Max - Min) / 2);
+		if ( Ranges[Mid].BeginIndex == BeginIndex )
+		{
+			return Mid; 
+		}
+		else if ( Ranges[Mid].BeginIndex < BeginIndex )
+		{
+			Min = Mid + 1;
+		}
+		else
+		{
+			Max = Mid - 1;
+		}
+	}
+
+	return Mid;
+}
+
+uint8 FTextLayout::FRunModel::GetKerning(int32 CurrentIndex, float InScale)
+{
+	return Run->GetKerning(CurrentIndex, InScale);
+}
+
+FVector2D FTextLayout::FRunModel::Measure(int32 BeginIndex, int32 EndIndex, float InScale)
+{
+	FVector2D Size = Run->Measure(BeginIndex, EndIndex, InScale);
+
+	MeasuredRanges.Add( FTextRange( BeginIndex, EndIndex ) );
+	MeasuredRangeSizes.Add( Size );
+
+	return Size;
+}
+
+int16 FTextLayout::FRunModel::GetMaxHeight(float InScale) const
+{
+	return Run->GetMaxHeight(InScale);
+}
+
+int16 FTextLayout::FRunModel::GetBaseLine( float InScale ) const
+{
+	return Run->GetBaseLine(InScale);
+}
+
+FTextRange FTextLayout::FRunModel::GetTextRange() const
+{
+	return Run->GetTextRange();
+}
+
+void FTextLayout::FRunModel::EndLayout()
+{
+	Run->EndLayout();
+}
+
+void FTextLayout::FRunModel::BeginLayout()
+{
+	Run->BeginLayout();
+}
+
+TSharedRef< IRun > FTextLayout::FRunModel::GetRun() const
+{
+	return Run;
+}
+
+FTextLayout::FRunModel::FRunModel( const TSharedRef< IRun >& InRun ) : Run( InRun )
+{
+
+}
diff --git a/Engine/Source/Runtime/Slate/Private/SlateModule.cpp b/Engine/Source/Runtime/Slate/Private/SlateModule.cpp
--- a/Engine/Source/Runtime/Slate/Private/SlateModule.cpp
+++ b/Engine/Source/Runtime/Slate/Private/SlateModule.cpp
@@ -4,11 +4,34 @@
 	SlateModule.cpp: Implements the FSlateModule class.
 =============================================================================*/
 
-#include "Slate.h"
 #include "SlatePrivatePCH.h"
 #include "ModuleManager.h"
 
 
+DEFINE_STAT(STAT_SlateCacheDesiredSize_STextBlock);
+DEFINE_STAT(STAT_SlateCacheDesiredSize);
+
+DEFINE_STAT(STAT_SlateOnPaint_SGraphPanel);
+DEFINE_STAT(STAT_SlateOnPaint_STextBlock);
+DEFINE_STAT(STAT_SlateOnPaint_SRichTextBlock);
+DEFINE_STAT(STAT_SlateOnPaint_SEditableText);
+DEFINE_STAT(STAT_SlateOnPaint_SOverlay);
+DEFINE_STAT(STAT_SlateOnPaint_SPanel);
+DEFINE_STAT(STAT_SlateOnPaint_SBorder);
+DEFINE_STAT(STAT_SlateOnPaint_SCompoundWidget);
+DEFINE_STAT(STAT_SlateOnPaint_SBox);
+DEFINE_STAT(STAT_SlateOnPaint_SImage);
+DEFINE_STAT(STAT_SlateOnPaint_SViewport);
+DEFINE_STAT(STAT_SlateOnPaint);
+DEFINE_STAT(STAT_SlateOnPaint_SSlider);
+
+DEFINE_STAT(STAT_SlateDrawWindowTime);
+DEFINE_STAT(STAT_SlateTickWidgets);
+DEFINE_STAT(STAT_SlateTickTime);
+
+DEFINE_STAT(STAT_SlateMiscTime); // Useful for locally pref testing miscellaneous Slate areas. Usage of this should not be checked in!
+
+
 /**
  * Implements the Slate module.
  */
diff --git a/Engine/Source/Runtime/Slate/Private/SlatePrivatePCH.h b/Engine/Source/Runtime/Slate/Private/SlatePrivatePCH.h
--- a/Engine/Source/Runtime/Slate/Private/SlatePrivatePCH.h
+++ b/Engine/Source/Runtime/Slate/Private/SlatePrivatePCH.h
@@ -6,25 +6,24 @@
 
 #pragma once
 
-//#include "Slate.h"
+#include "Slate.h"
 
 
 /* Dependencies
  *****************************************************************************/
 
+#include "InputCore.h"
+#include "Json.h"
+#include "ModuleManager.h"
+#include "Regex.h"
+#include "RemoteConfigIni.h"
 
 
 /* Private includes
  *****************************************************************************/
 
-
-
-/* Global functions
- *****************************************************************************/
-
-/**
- * Creates a widget reflector widget.
- *
- * @todo Slate: move this into a plug-in.
- */
-//TSharedRef<SWidgetReflector> MakeWidgetReflector( );
\ No newline at end of file
+#include "ButtonWidgetStyle.h"
+#include "CheckBoxWidgetStyle.h"
+#include "ComboBoxWidgetStyle.h"
+#include "ComboButtonWidgetStyle.h"
+#include "TextBlockWidgetStyle.h"
diff --git a/Engine/Source/Runtime/Slate/Private/Widgets/Colors/SColorBlock.cpp b/Engine/Source/Runtime/Slate/Private/Widgets/Colors/SColorBlock.cpp
new file mode 100644
--- /dev/null
+++ b/Engine/Source/Runtime/Slate/Private/Widgets/Colors/SColorBlock.cpp
@@ -0,0 +1,99 @@
+// Copyright 1998-2014 Epic Games, Inc. All Rights Reserved.
+
+#include "SlatePrivatePCH.h"
+
+
+/**
+ * Construct this widget
+ *
+ * @param	InArgs	The declaration data for this widget
+ */
+void SColorBlock::Construct( const FArguments& InArgs )
+{
+	Color = InArgs._Color;
+	ColorIsHSV = InArgs._ColorIsHSV;
+	IgnoreAlpha = InArgs._IgnoreAlpha;
+	ShowBackgroundForAlpha = InArgs._ShowBackgroundForAlpha;
+	MouseButtonDownHandler = InArgs._OnMouseButtonDown;
+	bUseSRGB = InArgs._UseSRGB;
+	ColorBlockSize = InArgs._Size;
+}
+
+int32 SColorBlock::OnPaint( const FGeometry& AllottedGeometry, const FSlateRect& MyClippingRect, FSlateWindowElementList& OutDrawElements, int32 LayerId, const FWidgetStyle& InWidgetStyle, bool bParentEnabled ) const
+{
+	const FSlateBrush* GenericBrush = FCoreStyle::Get().GetBrush( "GenericWhiteBox" );
+
+	const ESlateDrawEffect::Type DrawEffects = ESlateDrawEffect::None;
+	
+	FLinearColor InColor = Color.Get();
+	if (ColorIsHSV.Get())
+	{
+		InColor = InColor.HSVToLinearRGB();
+	}
+	if (IgnoreAlpha.Get())
+	{
+		InColor.A = 1.f;
+	}
+	
+	const FColor DrawColor = InColor.ToFColor(bUseSRGB.Get());
+	if( ShowBackgroundForAlpha.Get() && DrawColor.A < 255 )
+	{
+		// If we are showing a background pattern and the colors is transparent, draw a checker pattern
+		const FSlateBrush* CheckerBrush = FCoreStyle::Get().GetBrush("ColorPicker.AlphaBackground");
+		FSlateDrawElement::MakeBox( OutDrawElements, LayerId, AllottedGeometry.ToPaintGeometry(), CheckerBrush, MyClippingRect, DrawEffects );
+	}
+	
+	// determine if it is HDR
+	const float MaxRGB = FMath::Max3(InColor.R, InColor.G, InColor.B);
+	if (MaxRGB > 1.f)
+	{
+		FLinearColor NormalizedLinearColor = InColor / MaxRGB;
+		NormalizedLinearColor.A = InColor.A;
+		const FColor DrawNormalizedColor = InWidgetStyle.GetColorAndOpacityTint() * NormalizedLinearColor.ToFColor(bUseSRGB.Get());
+
+		FLinearColor ClampedLinearColor = InColor;
+		ClampedLinearColor.A = InColor.A * MaxRGB;
+		const FColor DrawClampedColor = InWidgetStyle.GetColorAndOpacityTint() * ClampedLinearColor.ToFColor(bUseSRGB.Get());
+
+		TArray<FSlateGradientStop> GradientStops;
+		
+		GradientStops.Add( FSlateGradientStop( FVector2D::ZeroVector, DrawNormalizedColor ) );
+		GradientStops.Add( FSlateGradientStop( AllottedGeometry.Size * 0.5f, DrawClampedColor ) );
+		GradientStops.Add( FSlateGradientStop( AllottedGeometry.Size, DrawNormalizedColor ) );
+
+		FSlateDrawElement::MakeGradient(
+			OutDrawElements,
+			LayerId + 1,
+			AllottedGeometry.ToPaintGeometry(),
+			GradientStops,
+			(AllottedGeometry.Size.X > AllottedGeometry.Size.Y) ? Orient_Vertical : Orient_Horizontal,
+			MyClippingRect,
+			DrawEffects
+		);
+	}
+	else
+	{
+		FSlateDrawElement::MakeBox( OutDrawElements, LayerId + 1, AllottedGeometry.ToPaintGeometry(), GenericBrush, MyClippingRect, DrawEffects, InWidgetStyle.GetColorAndOpacityTint() * DrawColor );
+	}
+
+	return LayerId + 1;
+}
+
+FReply SColorBlock::OnMouseButtonDown( const FGeometry& MyGeometry, const FPointerEvent& MouseEvent )
+{
+	if ( MouseButtonDownHandler.IsBound() )
+	{
+		// If a handler is assigned, call it.
+		return MouseButtonDownHandler.Execute(MyGeometry, MouseEvent);
+	}
+	else
+	{
+		// otherwise the event is unhandled.
+		return FReply::Unhandled();
+	}
+}
+
+FVector2D SColorBlock::ComputeDesiredSize() const
+{
+	return ColorBlockSize.Get();
+}
diff --git a/Engine/Source/Runtime/Slate/Private/Widgets/Colors/SColorPicker.cpp b/Engine/Source/Runtime/Slate/Private/Widgets/Colors/SColorPicker.cpp
new file mode 100644
--- /dev/null
+++ b/Engine/Source/Runtime/Slate/Private/Widgets/Colors/SColorPicker.cpp
@@ -0,0 +1,2255 @@
+// Copyright 1998-2014 Epic Games, Inc. All Rights Reserved.
+
+#include "SlatePrivatePCH.h"
+
+
+class SSimpleGradient : public SCompoundWidget
+{
+public:
+	SLATE_BEGIN_ARGS( SSimpleGradient )
+		: _StartColor( FLinearColor(0,0,0) )
+		, _EndColor( FLinearColor(1,1,1) )
+		, _HasAlphaBackground( false )
+		, _Orientation( Orient_Vertical )
+		, _UseSRGB( true )
+		{}
+
+		/** The leftmost gradient color */
+		SLATE_ATTRIBUTE( FLinearColor, StartColor )
+		
+		/** The rightmost gradient color */
+		SLATE_ATTRIBUTE( FLinearColor, EndColor )
+
+		/** Whether a checker background is displayed for alpha viewing */
+		SLATE_ATTRIBUTE( bool, HasAlphaBackground )
+
+		/** Horizontal or vertical gradient */
+		SLATE_ATTRIBUTE( EOrientation, Orientation)
+
+		/** Whether to display sRGB color */
+		SLATE_ATTRIBUTE( bool, UseSRGB )
+
+	SLATE_END_ARGS()
+
+	/**
+	 * Construct this widget
+	 *
+	 * @param	InArgs	The declaration data for this widget
+	 */
+	void Construct( const FArguments& InArgs )
+	{
+		StartColor = InArgs._StartColor;
+		EndColor = InArgs._EndColor;
+		bHasAlphaBackground = InArgs._HasAlphaBackground.Get();
+		Orientation = InArgs._Orientation.Get();
+	}
+
+private:
+	/**
+	 * The widget should respond by populating the OutDrawElements array with FDrawElements 
+	 * that represent it and any of its children.
+	 *
+	 * @param AllottedGeometry  The FGeometry that describes an area in which the widget should appear.
+	 * @param MyClippingRect    The clipping rectangle allocated for this widget and its children.
+	 * @param OutDrawElements   A list of FDrawElements to populate with the output.
+	 * @param LayerId           The Layer onto which this widget should be rendered.
+	 * @param InColorAndOpacity Color and Opacity to be applied to all the descendants of the widget being painted
+ 	 * @param bParentEnabled	True if the parent of this widget is enabled.
+	 *
+	 * @return The maximum layer ID attained by this widget or any of its children.
+	 */
+	virtual int32 OnPaint( const FGeometry& AllottedGeometry, const FSlateRect& MyClippingRect, FSlateWindowElementList& OutDrawElements, int32 LayerId, const FWidgetStyle& InWidgetStyle, bool bParentEnabled ) const
+	{
+		const ESlateDrawEffect::Type DrawEffects = this->ShouldBeEnabled(bParentEnabled) ? ESlateDrawEffect::None : ESlateDrawEffect::DisabledEffect;
+
+		if (bHasAlphaBackground)
+		{
+			const FSlateBrush* StyleInfo = FCoreStyle::Get().GetBrush("ColorPicker.AlphaBackground");
+
+			FSlateDrawElement::MakeBox(
+				OutDrawElements,
+				LayerId,
+				AllottedGeometry.ToPaintGeometry(),
+				StyleInfo,
+				MyClippingRect,
+				DrawEffects
+			);
+		}
+
+		TArray<FSlateGradientStop> GradientStops;
+
+		GradientStops.Add(FSlateGradientStop(FVector2D::ZeroVector, StartColor.Get()));
+		GradientStops.Add(FSlateGradientStop(AllottedGeometry.Size, EndColor.Get()));
+
+		FSlateDrawElement::MakeGradient(
+			OutDrawElements,
+			LayerId + 1,
+			AllottedGeometry.ToPaintGeometry(),
+			GradientStops,
+			Orientation,
+			MyClippingRect,
+			DrawEffects,
+			false
+		);
+
+		return LayerId + 1;
+	}
+	
+	/** The leftmost gradient color */
+	TAttribute<FLinearColor> StartColor;
+
+	/** The rightmost gradient color */
+	TAttribute<FLinearColor> EndColor;
+
+	/** Whether a checker background is displayed for alpha viewing */
+	bool bHasAlphaBackground;
+	
+	/** Horizontal or vertical gradient */
+	EOrientation Orientation;
+
+	/** Whether to display sRGB color */
+	TAttribute<bool> UseSRGB;
+};
+
+
+class SComplexGradient : public SCompoundWidget
+{
+public:
+	SLATE_BEGIN_ARGS( SComplexGradient )
+		: _GradientColors()
+		, _HasAlphaBackground( false )
+		, _Orientation( Orient_Vertical )
+		{}
+
+		/** The colors used in the gradient */
+		SLATE_ATTRIBUTE( TArray<FLinearColor>, GradientColors )
+
+		/** Whether a checker background is displayed for alpha viewing */
+		SLATE_ATTRIBUTE( bool, HasAlphaBackground )
+		
+		/** Horizontal or vertical gradient */
+		SLATE_ATTRIBUTE( EOrientation, Orientation)
+
+	SLATE_END_ARGS()
+
+	/**
+	 * Construct this widget
+	 *
+	 * @param	InArgs	The declaration data for this widget
+	 */
+	void Construct( const FArguments& InArgs )
+	{
+		GradientColors = InArgs._GradientColors;
+		bHasAlphaBackground = InArgs._HasAlphaBackground.Get();
+		Orientation = InArgs._Orientation.Get();
+	}
+
+private:
+	/**
+	 * The widget should respond by populating the OutDrawElements array with FDrawElements 
+	 * that represent it and any of its children.
+	 *
+	 * @param AllottedGeometry  The FGeometry that describes an area in which the widget should appear.
+	 * @param MyClippingRect    The clipping rectangle allocated for this widget and its children.
+	 * @param OutDrawElements   A list of FDrawElements to populate with the output.
+	 * @param LayerId           The Layer onto which this widget should be rendered.
+	 * @param InColorAndOpacity Color and Opacity to be applied to all the descendants of the widget being painted
+ 	 * @param bParentEnabled	True if the parent of this widget is enabled.
+	 *
+	 * @return The maximum layer ID attained by this widget or any of its children.
+	 */
+	virtual int32 OnPaint( const FGeometry& AllottedGeometry, const FSlateRect& MyClippingRect, FSlateWindowElementList& OutDrawElements, int32 LayerId, const FWidgetStyle& InWidgetStyle, bool bParentEnabled ) const
+	{
+		ESlateDrawEffect::Type DrawEffects = (bParentEnabled && IsEnabled()) ? ESlateDrawEffect::None : ESlateDrawEffect::DisabledEffect;
+
+		if (bHasAlphaBackground)
+		{
+			const FSlateBrush* StyleInfo = FCoreStyle::Get().GetBrush("ColorPicker.AlphaBackground");
+
+			FSlateDrawElement::MakeBox(
+				OutDrawElements,
+				LayerId,
+				AllottedGeometry.ToPaintGeometry(),
+				StyleInfo,
+				MyClippingRect,
+				DrawEffects
+				);
+		}
+
+		if (GradientColors.Get().Num())
+		{
+			TArray<FSlateGradientStop> GradientStops;
+			for (int32 i = 0; i < GradientColors.Get().Num(); ++i)
+			{
+				GradientStops.Add( FSlateGradientStop( AllottedGeometry.Size * (float(i) / (GradientColors.Get().Num() - 1)),
+					GradientColors.Get()[i] ) );
+			}
+
+			FSlateDrawElement::MakeGradient(
+				OutDrawElements,
+				LayerId + 1,
+				AllottedGeometry.ToPaintGeometry(),
+				GradientStops,
+				Orientation,
+				MyClippingRect,
+				DrawEffects
+				);
+		}
+
+		return LayerId + 1;
+	}
+	
+	/** The colors used in the gradient */
+	TAttribute< TArray<FLinearColor> > GradientColors;
+
+	/** Whether a checker background is displayed for alpha viewing */
+	bool bHasAlphaBackground;
+	
+	/** Horizontal or vertical gradient */
+	EOrientation Orientation;
+};
+
+
+/** The value slider is a simple control like the color wheel for selecting value */
+class SValueSlider : public SLeafWidget
+{
+public:
+	SLATE_BEGIN_ARGS( SValueSlider )
+		: _SelectedColor()
+		, _OnValueChanged()
+		, _OnMouseCaptureBegin()
+		, _OnMouseCaptureEnd()
+		{}
+		
+		/** The current color selected by the user */
+		SLATE_ATTRIBUTE( FLinearColor, SelectedColor )
+		
+		/** Invoked when a new value is selected on the color wheel */
+		SLATE_EVENT( FOnLinearColorValueChanged, OnValueChanged )
+		
+		/** Invoked when the mouse is pressed and sliding begins */
+		SLATE_EVENT( FSimpleDelegate, OnMouseCaptureBegin )
+
+		/** Invoked when the mouse is released and sliding ends */
+		SLATE_EVENT( FSimpleDelegate, OnMouseCaptureEnd )
+
+	SLATE_END_ARGS()
+	
+	void Construct( const FArguments& InArgs )
+	{
+		SelectorImage = FCoreStyle::Get().GetBrush("ColorPicker.Selector");
+
+		OnValueChanged = InArgs._OnValueChanged;
+		OnMouseCaptureBegin = InArgs._OnMouseCaptureBegin;
+		OnMouseCaptureEnd = InArgs._OnMouseCaptureEnd;
+		SelectedColor = InArgs._SelectedColor;
+	}
+
+	virtual int32 OnPaint( const FGeometry& AllottedGeometry, const FSlateRect& MyClippingRect, FSlateWindowElementList& OutDrawElements, int32 LayerId, const FWidgetStyle& InWidgetStyle, bool bParentEnabled ) const
+	{
+		const bool bIsEnabled = ShouldBeEnabled(bParentEnabled);
+		const uint32 DrawEffects = bIsEnabled ? ESlateDrawEffect::None : ESlateDrawEffect::DisabledEffect;
+	
+		FLinearColor FullValueColor = SelectedColor.Get();
+		FullValueColor.B = FullValueColor.A = 1.f;
+		FLinearColor StopColor = FullValueColor.HSVToLinearRGB();
+
+		TArray<FSlateGradientStop> GradientStops;
+		GradientStops.Add( FSlateGradientStop( FVector2D::ZeroVector, FLinearColor(0,0,0,1.0f) ) );
+		GradientStops.Add( FSlateGradientStop( AllottedGeometry.Size, StopColor ) );
+
+		FSlateDrawElement::MakeGradient(
+			OutDrawElements,
+			LayerId,
+			AllottedGeometry.ToPaintGeometry(),
+			GradientStops,
+			Orient_Vertical,
+			MyClippingRect,
+			DrawEffects
+		);
+	
+		float Value = SelectedColor.Get().B;
+		FVector2D RelativeSelectedPosition = FVector2D(Value, 0.5f);
+
+		FSlateDrawElement::MakeBox(
+			OutDrawElements,
+			LayerId + 1,
+			AllottedGeometry.ToPaintGeometry( RelativeSelectedPosition*AllottedGeometry.Size - SelectorImage->ImageSize*0.5, SelectorImage->ImageSize ),
+			SelectorImage,
+			MyClippingRect,
+			DrawEffects,
+			InWidgetStyle.GetColorAndOpacityTint() * SelectorImage->GetTint( InWidgetStyle ) );
+
+		return LayerId + 1;
+	}
+	
+	virtual FReply OnMouseButtonDown( const FGeometry& MyGeometry, const FPointerEvent& MouseEvent )
+	{
+		if ( MouseEvent.GetEffectingButton() == EKeys::LeftMouseButton )
+		{
+			OnMouseCaptureBegin.ExecuteIfBound();
+
+			return FReply::Handled().CaptureMouse( SharedThis(this) );
+		}
+		else
+		{
+			return FReply::Unhandled();
+		}
+	}
+	
+	virtual FReply OnMouseButtonUp( const FGeometry& MyGeometry, const FPointerEvent& MouseEvent )
+	{
+		if ( MouseEvent.GetEffectingButton() == EKeys::LeftMouseButton && HasMouseCapture() )
+		{
+			OnMouseCaptureEnd.ExecuteIfBound();
+
+			return FReply::Handled().ReleaseMouseCapture();
+		}
+		else
+		{
+			return FReply::Unhandled();
+		}
+	}
+
+	virtual FReply OnMouseMove( const FGeometry& MyGeometry, const FPointerEvent& MouseEvent )
+	{
+		if (!HasMouseCapture())
+		{
+			return FReply::Unhandled();
+		}
+
+		FVector2D LocalMouseCoordinate = MyGeometry.AbsoluteToLocal( MouseEvent.GetScreenSpacePosition() );
+		FVector2D Location = LocalMouseCoordinate / (MyGeometry.Size);
+		float Value = FMath::Clamp(Location.X, 0.f, 1.f);
+
+		FLinearColor NewColor = SelectedColor.Get();
+		NewColor.B = Value;
+
+		OnValueChanged.ExecuteIfBound(NewColor);
+
+		return FReply::Handled();
+	}
+	
+	virtual FReply OnMouseButtonDoubleClick( const FGeometry& InMyGeometry, const FPointerEvent& InMouseEvent )
+	{
+		return FReply::Handled();
+	}
+
+	virtual FVector2D ComputeDesiredSize() const
+	{
+		return SelectorImage->ImageSize * 2.f;
+	}
+	
+protected:
+
+	/** The color selector image to show */
+	const FSlateBrush* SelectorImage;
+	
+	/** The current color selected by the user */
+	TAttribute< FLinearColor > SelectedColor;
+
+	/** Invoked when a new value is selected on the color wheel */
+	FOnLinearColorValueChanged OnValueChanged;
+
+	/** Invoked when the mouse is pressed */
+	FSimpleDelegate OnMouseCaptureBegin;
+
+	/** Invoked when the mouse is let up */
+	FSimpleDelegate OnMouseCaptureEnd;
+};
+
+
+
+#define LOCTEXT_NAMESPACE "EyeDroppperButton"
+
+/**
+ * Class for placing a color picker eye-dropper button.
+ * A self contained until that only needs client code to set the display gamma and listen
+ * for the OnValueChanged events. It toggles the dropper when clicked.
+ * When active it captures the mouse, shows a dropper cursor and samples the pixel color constantly.
+ * It is stopped normally by hitting the Esc key.
+ */
+class SEyeDropperButton : public SButton
+{
+public:
+	DECLARE_DELEGATE(FOnDropperComplete)
+
+	SLATE_BEGIN_ARGS( SEyeDropperButton )
+		: _OnValueChanged()
+		, _OnBegin()
+		, _OnComplete()
+		, _DisplayGamma()
+		{}
+
+		/** Invoked when a new value is selected by the dropper */
+		SLATE_EVENT(FOnLinearColorValueChanged, OnValueChanged)
+
+		/** Invoked when the dropper goes from inactive to active */
+		SLATE_EVENT(FSimpleDelegate, OnBegin)
+
+		/** Invoked when the dropper goes from active to inactive */
+		SLATE_EVENT(FOnDropperComplete, OnComplete)
+
+		/** Sets the display Gamma setting - used to correct colors sampled from the screen */
+		SLATE_ATTRIBUTE(float, DisplayGamma)
+
+	SLATE_END_ARGS()
+
+	void Construct(const FArguments& InArgs )
+	{
+		OnValueChanged = InArgs._OnValueChanged;
+		OnBegin = InArgs._OnBegin;
+		OnComplete = InArgs._OnComplete;
+		DisplayGamma = InArgs._DisplayGamma;
+
+		// This is a button containing an dropper image and text that tells the user to hit Esc.
+		// Their visibility and colors are changed according to whether dropper mode is active or not.
+		SButton::Construct(
+			SButton::FArguments()
+				.ContentPadding(1.0f)
+				.OnClicked(this, &SEyeDropperButton::OnClicked)
+				[
+					SNew(SOverlay)
+
+					+ SOverlay::Slot()
+						.Padding(FMargin(1.0f, 0.0f))
+						[
+							SNew(SImage)
+								.Image(FCoreStyle::Get().GetBrush("ColorPicker.EyeDropper"))
+								.ToolTipText(LOCTEXT("EyeDropperButton_ToolTip", "Activates the eye-dropper for selecting a colored pixel from any window."))
+								.ColorAndOpacity(this, &SEyeDropperButton::GetDropperImageColor)				
+						]
+
+					+ SOverlay::Slot()
+						[
+							SNew(STextBlock)
+								.Text(LOCTEXT("EscapeCue", "Esc"))
+								.ToolTipText(LOCTEXT("EyeDropperEscapeCue_ToolTip", "Hit Escape key to stop the eye dropper"))
+								.Visibility(this, &SEyeDropperButton::GetEscapeTextVisibility)
+						]
+				]
+		);
+	}
+
+	virtual FReply OnMouseButtonDown( const FGeometry& MyGeometry, const FPointerEvent& MouseEvent ) OVERRIDE
+	{
+		// Clicking ANY mouse button when the dropper isn't active resets the active dropper states ready to activate
+		if (!HasMouseCapture())
+		{
+			bWasClickActivated = false;
+			bWasLeft = false;
+			bWasReEntered = false;
+		}
+
+		return SButton::OnMouseButtonDown(MyGeometry, MouseEvent);
+	}
+
+	virtual FReply OnMouseButtonUp( const FGeometry& MyGeometry, const FPointerEvent& MouseEvent ) OVERRIDE
+	{
+		// If a mouse click is completing and the dropper is active ALWAYS deactivate
+		bool bDeactivating = false;
+		if (bWasClickActivated)
+		{
+			bDeactivating = true;
+		}
+
+		// bWasClicked is reset here because if it is set during SButton::OnMouseButtonUp()
+		// then the button was 'clicked' according to the usual rules. We might want to capture the
+		// mouse when the button is clicked but can't do it in the Clicked callback.
+		bWasClicked = false;
+		FReply Reply = SButton::OnMouseButtonUp(MyGeometry, MouseEvent);
+
+		// Switching dropper mode off?
+		if (bDeactivating)
+		{
+			// These state flags clear dropper mode
+			bWasClickActivated = false;
+			bWasLeft = false;
+			bWasReEntered = false;
+
+			Reply.ReleaseMouseCapture();
+
+			OnComplete.ExecuteIfBound();
+		}
+		else if (MouseEvent.GetEffectingButton() == EKeys::LeftMouseButton && bWasClicked)
+		{
+			// A normal LMB mouse click on the button occurred
+			// Set the initial dropper mode state and capture the mouse
+			bWasClickActivated = true;
+			bWasLeft = false;
+			bWasReEntered = false;
+
+			OnBegin.ExecuteIfBound();
+
+			Reply.CaptureMouse(this->AsShared());
+		}
+		bWasClicked = false;
+
+		return Reply;
+	}
+
+	virtual FReply OnMouseMove( const FGeometry& MyGeometry, const FPointerEvent& MouseEvent ) OVERRIDE
+	{
+		// If the mouse is captured and bWasClickActivated is set then we are in dropper mode
+		if (HasMouseCapture() && bWasClickActivated)
+		{
+			if (IsMouseOver(MyGeometry, MouseEvent))
+			{
+				if (bWasLeft)
+				{
+					// Mouse is over the button having left it once
+					bWasReEntered = true;
+				}
+			}
+			else
+			{
+				// Mouse is outside the button
+				bWasLeft = true;
+				bWasReEntered = false;
+
+				// In dropper mode and outside the button - sample the pixel color and push it to the client
+				FLinearColor ScreenColor = FPlatformMisc::GetScreenPixelColor(FSlateApplication::Get().GetCursorPos(), false);
+				OnValueChanged.ExecuteIfBound(ScreenColor);
+			}
+		}
+
+		return SButton::OnMouseMove( MyGeometry, MouseEvent );
+	}
+
+	FReply OnKeyDown( const FGeometry& MyGeometry, const FKeyboardEvent& InKeyboardEvent )
+	{
+		// Escape key when in dropper mode
+		if (InKeyboardEvent.GetKey() == EKeys::Escape &&
+			HasMouseCapture() &&
+			bWasClickActivated)
+		{
+			// Clear the dropper mode states
+			bWasClickActivated = false;
+			bWasLeft = false;
+			bWasReEntered = false;
+
+			// This is needed to switch the dropper cursor off immediately so the user can see the Esc key worked
+			FSlateApplication::Get().QueryCursor();
+
+			FReply ReleaseReply = FReply::Handled().ReleaseMouseCapture();
+
+			OnComplete.ExecuteIfBound();
+
+			return ReleaseReply;
+		}
+
+		return FReply::Unhandled();
+	}
+
+	virtual FCursorReply OnCursorQuery( const FGeometry& MyGeometry, const FPointerEvent& CursorEvent ) const OVERRIDE
+	{
+		// Cursor is changed to the dropper when dropper mode is active and the states are correct
+		if (HasMouseCapture() &&
+			bWasClickActivated && bWasLeft && !bWasReEntered )
+		{
+			return FCursorReply::Cursor( EMouseCursor::EyeDropper );
+		}
+
+		return SButton::OnCursorQuery( MyGeometry, CursorEvent );
+	}
+
+	FReply OnClicked()
+	{
+		// Log clicks so that OnMouseButtonUp() can post process clicks
+		bWasClicked = true;
+		return FReply::Handled();
+	}
+
+private:
+
+	EVisibility GetEscapeTextVisibility() const
+	{
+		// Show the Esc key message in the button when dropper mode is active
+		if (HasMouseCapture() && bWasClickActivated)
+		{
+			return EVisibility::Visible;
+		}
+		return EVisibility::Hidden;
+	}
+
+	FSlateColor GetDropperImageColor() const
+	{
+		// Make the dropper image in the button pale when dropper mode is active
+		if (HasMouseCapture() && bWasClickActivated)
+		{
+			return FLinearColor(0.3f, 0.3f, 0.3f, 1.0f);
+		}
+		return FSlateColor::UseForeground();
+	}
+
+	bool IsMouseOver(const FGeometry& MyGeometry, const FPointerEvent& MouseEvent) const
+	{ 
+		return MyGeometry.IsUnderLocation(MouseEvent.GetScreenSpacePosition());
+	}
+
+	/** Invoked when a new value is selected by the dropper */
+	FOnLinearColorValueChanged OnValueChanged;
+
+	/** Invoked when the dropper goes from inactive to active */
+	FSimpleDelegate OnBegin;
+
+	/** Invoked when the dropper goes from active to inactive - can be used to commit colors by the owning picker */
+	FOnDropperComplete OnComplete;
+
+	/** Sets the display Gamma setting - used to correct colors sampled from the screen */
+	TAttribute<float> DisplayGamma;
+
+	// Dropper states
+	bool bWasClicked;
+	bool bWasClickActivated;
+	bool bWasLeft;
+	bool bWasReEntered;
+};
+
+#undef LOCTEXT_NAMESPACE
+
+
+FColorTheme::FColorTheme( const FString& InName, const TArray< TSharedPtr<FLinearColor> >& InColors )
+	: Name(InName)
+	, Colors(InColors)
+	, RefreshEvent()
+{ }
+
+
+void FColorTheme::InsertNewColor( TSharedPtr<FLinearColor> InColor, int32 InsertPosition )
+{
+	Colors.Insert(InColor, InsertPosition);
+	RefreshEvent.Broadcast();
+}
+
+
+int32 FColorTheme::FindApproxColor( const FLinearColor& InColor, float Tolerance ) const
+{
+	for (int32 ColorIndex = 0; ColorIndex < Colors.Num(); ++ColorIndex)
+	{
+		if (Colors[ColorIndex]->Equals(InColor, Tolerance))
+		{
+			return ColorIndex;
+		}
+	}
+
+	return INDEX_NONE;
+}
+
+
+void FColorTheme::RemoveAll()
+{
+	Colors.Empty();
+	RefreshEvent.Broadcast();
+}
+
+
+void FColorTheme::RemoveColor( int32 ColorIndex )
+{
+	Colors.RemoveAt(ColorIndex);
+	RefreshEvent.Broadcast();
+}
+
+
+int32 FColorTheme::RemoveColor( const TSharedPtr<FLinearColor> InColor )
+{
+	const int32 Position = Colors.Find(InColor);
+	if (Position != INDEX_NONE)
+	{
+		RemoveColor(Position);
+	}
+
+	return Position;
+}
+
+
+#define LOCTEXT_NAMESPACE "ColorPicker"
+
+/** A default window size for the color picker which looks nice */
+const FVector2D SColorPicker::DEFAULT_WINDOW_SIZE = FVector2D(308, 458);
+
+/** The max time allowed for updating before we shut off auto-updating */
+const double SColorPicker::MAX_ALLOWED_UPDATE_TIME = 0.1;
+
+TSharedPtr<SColorThemesViewer> SColorPicker::ColorThemesViewer;
+
+
+/* SColorPicker structors
+ *****************************************************************************/
+
+SColorPicker::~SColorPicker()
+{
+	if (ColorThemesViewer.IsValid())
+	{
+		ColorThemesViewer->OnCurrentThemeChanged().RemoveAll( this );
+	}
+}
+
+
+/* SColorPicker methods
+ *****************************************************************************/
+
+void SColorPicker::Construct( const FArguments& InArgs )
+{
+	TargetColorAttribute = InArgs._TargetColorAttribute;
+	CurrentColorHSV = OldColor = TargetColorAttribute.Get().LinearRGBToHSV();
+	CurrentColorRGB = TargetColorAttribute.Get();
+	CurrentMode = EColorPickerModes::Wheel;
+	TargetFColors = InArgs._TargetFColors.Get();
+	TargetLinearColors = InArgs._TargetLinearColors.Get();
+	TargetColorChannels = InArgs._TargetColorChannels.Get();
+	bUseAlpha = InArgs._UseAlpha;
+	bOnlyRefreshOnMouseUp = InArgs._OnlyRefreshOnMouseUp.Get();
+	bOnlyRefreshOnOk = InArgs._OnlyRefreshOnOk.Get();
+	OnColorCommitted = InArgs._OnColorCommitted;
+	PreColorCommitted = InArgs._PreColorCommitted;
+	OnColorPickerCancelled = InArgs._OnColorPickerCancelled;
+	OnInteractivePickBegin = InArgs._OnInteractivePickBegin;
+	OnInteractivePickEnd = InArgs._OnInteractivePickEnd;
+	OnColorPickerWindowClosed = InArgs._OnColorPickerWindowClosed;
+	ParentWindowPtr = InArgs._ParentWindow.Get();
+	DisplayGamma = InArgs._DisplayGamma;
+	bClosedViaOkOrCancel = false;
+
+	// We need a parent window to set the close callback
+	if (ParentWindowPtr.IsValid())
+	{
+		ParentWindowPtr.Pin()->SetOnWindowClosed(FOnWindowClosed::CreateSP(this, &SColorPicker::HandleParentWindowClosed));
+	}
+
+	bColorPickerIsInlineVersion = InArgs._DisplayInlineVersion;
+	bIsInteractive = false;
+	bPerfIsTooSlowToUpdate = false;
+
+	BackupColors();
+
+	BeginAnimation(FLinearColor(ForceInit), CurrentColorHSV);
+
+	bool bAdvancedSectionExpanded = false;
+
+	if (FPaths::FileExists(GEditorUserSettingsIni))
+	{
+		bool WheelMode = true;
+
+		GConfig->GetBool(TEXT("ColorPickerUI"), TEXT("bWheelMode"), WheelMode, GEditorUserSettingsIni);
+		GConfig->GetBool(TEXT("ColorPickerUI"), TEXT("bAdvancedSectionExpanded"), bAdvancedSectionExpanded, GEditorUserSettingsIni);
+		GConfig->GetBool(TEXT("ColorPickerUI"), TEXT("bSRGBEnabled"), SColorThemesViewer::bSRGBEnabled, GEditorUserSettingsIni);
+
+		CurrentMode = WheelMode ? EColorPickerModes::Wheel : EColorPickerModes::Spectrum;
+	}
+
+	if (bColorPickerIsInlineVersion)
+	{
+		GenerateInlineColorPickerContent();
+	}
+	else
+	{
+		GenerateDefaultColorPickerContent(bAdvancedSectionExpanded);
+	}
+}
+
+
+/* SColorPicker implementation
+ *****************************************************************************/
+
+void SColorPicker::BackupColors()
+{
+	OldTargetFColors.Empty();
+	for (int32 i = 0; i < TargetFColors.Num(); ++i)
+	{
+		OldTargetFColors.Add( *TargetFColors[i] );
+	}
+
+	OldTargetLinearColors.Empty();
+	for (int32 i = 0; i < TargetLinearColors.Num(); ++i)
+	{
+		OldTargetLinearColors.Add( *TargetLinearColors[i] );
+	}
+
+	OldTargetColorChannels.Empty();
+	for (int32 i = 0; i < TargetColorChannels.Num(); ++i)
+	{
+		// Remap the color channel as a linear color for ease
+		const FColorChannels& Channel = TargetColorChannels[i];
+		const FLinearColor Color( Channel.Red ? *Channel.Red : 0.f, Channel.Green ? *Channel.Green : 0.f, Channel.Blue ? *Channel.Blue : 0.f, Channel.Alpha ? *Channel.Alpha : 0.f );
+		OldTargetColorChannels.Add( Color );
+	}
+}
+
+
+void SColorPicker::RestoreColors()
+{
+	check(TargetFColors.Num() == OldTargetFColors.Num());
+
+	for (int32 i = 0; i < TargetFColors.Num(); ++i)
+	{
+		*TargetFColors[i] = OldTargetFColors[i];
+	}
+
+	check(TargetLinearColors.Num() == OldTargetLinearColors.Num());
+
+	for (int32 i = 0; i < TargetLinearColors.Num(); ++i)
+	{
+		*TargetLinearColors[i] = OldTargetLinearColors[i];
+	}
+
+	check(TargetColorChannels.Num() == OldTargetColorChannels.Num());
+
+	for (int32 i = 0; i < TargetColorChannels.Num(); ++i)
+	{
+		// Copy back out of the linear to the color channel
+		FColorChannels& Channel = TargetColorChannels[i];
+		const FLinearColor& OldChannel = OldTargetColorChannels[i];
+		if (Channel.Red)
+		{
+			*Channel.Red = OldChannel.R;
+		}
+		if (Channel.Green)
+		{
+			*Channel.Green = OldChannel.G;
+		}
+		if (Channel.Blue)
+		{
+			*Channel.Blue = OldChannel.B;
+		}
+		if (Channel.Alpha)
+		{
+			*Channel.Alpha = OldChannel.A;
+		}
+	}
+}
+
+
+void SColorPicker::SetColors(const FLinearColor& InColor)
+{
+	for (int32 i = 0; i < TargetFColors.Num(); ++i)
+	{
+		*TargetFColors[i] = InColor.ToFColor(false);
+	}
+
+	for (int32 i = 0; i < TargetLinearColors.Num(); ++i)
+	{
+		*TargetLinearColors[i] = InColor;
+	}
+
+	for (int32 i = 0; i < TargetColorChannels.Num(); ++i)
+	{
+		// Only set those channels who have a valid ptr
+		FColorChannels& Channel = TargetColorChannels[i];
+		if (Channel.Red)
+		{
+			*Channel.Red = InColor.R;
+		}
+		if (Channel.Green)
+		{
+			*Channel.Green = InColor.G;
+		}
+		if (Channel.Blue)
+		{
+			*Channel.Blue = InColor.B;
+		}
+		if (Channel.Alpha)
+		{
+			*Channel.Alpha = InColor.A;
+		}
+	}
+}
+
+
+BEGIN_SLATE_FUNCTION_BUILD_OPTIMIZATION
+void SColorPicker::GenerateDefaultColorPickerContent( bool bAdvancedSectionExpanded )
+{	
+	// The height of the gradient bars beneath the sliders
+	const FSlateFontInfo SmallLayoutFont = FCoreStyle::Get().GetFontStyle("ColorPicker.Font");
+
+	if (!ColorThemesViewer.IsValid())
+	{
+		ColorThemesViewer = SNew(SColorThemesViewer);
+	}
+
+	ColorThemesViewer->OnCurrentThemeChanged().AddSP(this, &SColorPicker::HandleThemesViewerThemeChanged);
+	ColorThemesViewer->SetUseAlpha(bUseAlpha);
+	ColorThemesViewer->MenuToStandardNoReturn();
+
+	// The standard button to open the color themes can temporarily become a trash for colors
+	ColorThemeComboButton = SNew(SComboButton)
+		.ContentPadding(3.0f)
+		.MenuPlacement(MenuPlacement_ComboBox)
+		.ToolTipText(LOCTEXT("OpenThemeManagerToolTip", "Open Color Theme Manager"));
+
+	ColorThemeComboButton->SetMenuContent(ColorThemesViewer.ToSharedRef());
+
+	SmallTrash = SNew(SHorizontalBox)
+	+ SHorizontalBox::Slot()
+		.AutoWidth()
+		[
+			SNew(SColorTrash)
+				.UsesSmallIcon(true)
+		];
+	
+	this->ChildSlot
+	[
+		SNew(SVerticalBox)
+
+		+ SVerticalBox::Slot()
+			.AutoHeight()
+			[
+				SNew(SGridPanel)
+					.FillColumn(0, 1.0f)
+
+				+ SGridPanel::Slot(0, 0)
+					.Padding(0.0f, 1.0f, 20.0f, 1.0f)
+					[
+						SNew(SHorizontalBox)
+
+						+ SHorizontalBox::Slot()
+							.FillWidth(1.0f)
+							.Padding(0.0f, 1.0f)
+							[
+								SNew(SOverlay)
+
+								+ SOverlay::Slot()
+									[
+										// color theme bar
+										SAssignNew(CurrentThemeBar, SThemeColorBlocksBar)
+											.ColorTheme(this, &SColorPicker::HandleThemeBarColorTheme)
+											.EmptyText(LOCTEXT("EmptyBarHint", "Drag & drop colors here to save"))
+											.HideTrashCallback(this, &SColorPicker::HideSmallTrash)
+											.ShowTrashCallback(this, &SColorPicker::ShowSmallTrash)
+											.ToolTipText(LOCTEXT("CurrentThemeBarToolTip", "Current Color Theme"))
+											.UseAlpha(SharedThis(this), &SColorPicker::HandleThemeBarUseAlpha)
+											.UseSRGB(SharedThis(this), &SColorPicker::HandleColorPickerUseSRGB)
+											.OnSelectColor(this, &SColorPicker::HandleThemeBarColorSelected)
+									]
+
+								// hack: need to fix SThemeColorBlocksBar::EmptyText to render properly
+								+ SOverlay::Slot()
+									.HAlign(HAlign_Center)
+									.VAlign(VAlign_Center)
+									[
+										SNew(STextBlock)
+											.Text(LOCTEXT("EmptyBarHint", "Drag & drop colors here to save"))
+											.Visibility(this, &SColorPicker::HandleThemeBarHintVisibility)
+									]
+							]
+
+						+ SHorizontalBox::Slot()
+							.AutoWidth()
+							[
+								// color theme selector
+								SAssignNew(ColorThemeButtonOrSmallTrash, SBorder)
+									.BorderImage(FStyleDefaults::GetNoBrush())
+									.Padding(0.0f)							
+							]
+					]
+
+				+ SGridPanel::Slot(1, 0)
+					.HAlign(HAlign_Right)
+					.VAlign(VAlign_Center)
+					[
+						// sRGB check box
+						SNew(SCheckBox)
+							.ToolTipText(LOCTEXT("SRGBCheckboxToolTip", "Toggle gamma corrected sRGB previewing"))
+							.IsChecked(this, &SColorPicker::HandleSRGBCheckBoxIsChecked)
+							.OnCheckStateChanged(this, &SColorPicker::HandleSRGBCheckBoxCheckStateChanged)
+							[
+								SNew(STextBlock)
+									.Text(LOCTEXT("SRGBCheckboxLabel", "sRGB"))
+							]
+					]
+
+				+ SGridPanel::Slot(0, 1)
+					.Padding(0.0f, 8.0f, 20.0f, 0.0f)
+					[
+						SNew(SBorder)
+							.BorderImage(FCoreStyle::Get().GetBrush("NoBorder"))
+							.Padding(0.0f)
+							.OnMouseButtonDown(this, &SColorPicker::HandleColorAreaMouseDown)
+							[
+								SNew(SOverlay)
+
+								// color wheel
+								+ SOverlay::Slot()
+									[
+										SNew(SHorizontalBox)
+
+										+ SHorizontalBox::Slot()
+											.FillWidth(1.0f)
+											.HAlign(HAlign_Center)
+											[
+												SNew(SColorWheel)
+													.SelectedColor(this, &SColorPicker::GetCurrentColor)
+													.Visibility(this, &SColorPicker::HandleColorPickerModeVisibility, EColorPickerModes::Wheel)
+													.OnValueChanged(this, &SColorPicker::HandleColorSpectrumValueChanged)
+													.OnMouseCaptureBegin(this, &SColorPicker::HandleInteractiveChangeBegin)
+													.OnMouseCaptureEnd(this, &SColorPicker::HandleInteractiveChangeEnd)
+											]
+
+										+ SHorizontalBox::Slot()
+											.AutoWidth()
+											.Padding(4.0f, 0.0f)
+											[
+												// saturation slider
+												MakeColorSlider(EColorPickerChannels::Saturation)
+											]
+
+										+ SHorizontalBox::Slot()
+											.AutoWidth()
+											[
+												// value slider
+												MakeColorSlider(EColorPickerChannels::Value)
+											]
+									]
+
+								// color spectrum
+								+ SOverlay::Slot()
+									[
+										SNew(SBox)
+											.HeightOverride(200.0f)
+											.WidthOverride(292.0f)
+											[
+												SNew(SColorSpectrum)
+													.SelectedColor(this, &SColorPicker::GetCurrentColor)
+													.Visibility(this, &SColorPicker::HandleColorPickerModeVisibility, EColorPickerModes::Spectrum)
+													.OnValueChanged(this, &SColorPicker::HandleColorSpectrumValueChanged)
+													.OnMouseCaptureBegin(this, &SColorPicker::HandleInteractiveChangeBegin)
+													.OnMouseCaptureEnd(this, &SColorPicker::HandleInteractiveChangeEnd)
+											]
+									]
+							]
+					]
+
+				+ SGridPanel::Slot(1, 1)
+					.Padding(0.0f, 8.0f, 0.0f, 0.0f)
+					[
+						SNew(SVerticalBox)
+
+						+ SVerticalBox::Slot()
+							.AutoHeight()
+							[
+								SNew(SBox)
+									.HeightOverride(100.0f)
+									.WidthOverride(70.0f)
+									[
+										// color preview
+										MakeColorPreviewBox()
+									]
+							]
+
+						+ SVerticalBox::Slot()
+							.AutoHeight()
+							.Padding(0.0f, 16.0f, 0.0f, 0.0f)
+							.VAlign(VAlign_Top)
+							[
+								SNew(SHorizontalBox)
+
+								+ SHorizontalBox::Slot()
+									.HAlign(HAlign_Left)
+									[
+										// mode selector
+										SNew(SButton)
+											.OnClicked(this, &SColorPicker::HandleColorPickerModeButtonClicked)
+											.Content()
+											[
+												SNew(SImage)
+													.Image(FCoreStyle::Get().GetBrush("ColorPicker.Mode"))
+													.ToolTipText(LOCTEXT("ColorPickerModeEToolTip", "Toggle between color wheel and color spectrum."))
+											]									
+									]
+
+								+ SHorizontalBox::Slot()
+									.HAlign(HAlign_Right)
+									[
+										// eye dropper
+										SNew(SEyeDropperButton)
+											.OnValueChanged(this, &SColorPicker::HandleRGBColorChanged)
+											.OnBegin(this, &SColorPicker::HandleInteractiveChangeBegin)
+											.OnComplete(this, &SColorPicker::HandleEyeDropperButtonComplete)
+											.DisplayGamma(DisplayGamma)								
+									]
+							]
+					]
+			]
+
+		// advanced settings
+		+ SVerticalBox::Slot()
+			.AutoHeight()
+			.Padding(0.0f, 4.0f, 0.0f, 0.0f)
+			[
+				SNew(SExpandableArea)
+					.AreaTitle(LOCTEXT("AdvancedAreaTitle", "Advanced"))
+					.BorderBackgroundColor(FLinearColor::Transparent)
+					.InitiallyCollapsed(!bAdvancedSectionExpanded)
+					.OnAreaExpansionChanged(this, &SColorPicker::HandleAdvancedAreaExpansionChanged)
+					.Padding(FMargin(0.0f, 1.0f, 0.0f, 8.0f))
+					.BodyContent()
+					[
+						SNew(SHorizontalBox)
+
+						// RGBA inputs
+						+ SHorizontalBox::Slot()
+							.Padding(0.0f, 0.0f, 4.0f, 0.0f)
+							[
+								SNew(SVerticalBox)
+
+								// Red
+								+ SVerticalBox::Slot()
+									[
+										MakeColorSpinBox(EColorPickerChannels::Red)
+									]
+
+								// Green
+								+ SVerticalBox::Slot()
+									.Padding(0.0f, 8.0f, 0.0f, 0.0f)
+									[
+										MakeColorSpinBox(EColorPickerChannels::Green)
+									]
+
+								// Blue
+								+ SVerticalBox::Slot()
+									.Padding(0.0f, 8.0f, 0.0f, 0.0f)
+									[
+										MakeColorSpinBox(EColorPickerChannels::Blue)
+									]
+
+								// Alpha
+								+ SVerticalBox::Slot()
+									.Padding(0.0f, 8.0f, 0.0f, 0.0f)
+									[
+										MakeColorSpinBox(EColorPickerChannels::Alpha)
+									]
+							]
+
+						// HSV & Hex inputs
+						+ SHorizontalBox::Slot()
+							.Padding(4.0f, 0.0f, 0.0f, 0.0f)
+							[
+								SNew(SVerticalBox)
+							
+								// Hue
+								+ SVerticalBox::Slot()
+									[
+										MakeColorSpinBox(EColorPickerChannels::Hue)
+									]
+
+								// Saturation
+								+ SVerticalBox::Slot()
+									.Padding(0.0f, 8.0f, 0.0f, 0.0f)
+									[
+										MakeColorSpinBox(EColorPickerChannels::Saturation)
+									]
+
+								// Value
+								+ SVerticalBox::Slot()
+									.Padding(0.0f, 8.0f, 0.0f, 0.0f)
+									[
+										MakeColorSpinBox(EColorPickerChannels::Value)
+									]
+
+								// Hex
+								+ SVerticalBox::Slot()
+									.HAlign(HAlign_Right)
+									.VAlign(VAlign_Top)
+									.Padding(0.0f, 8.0f, 0.0f, 0.0f)
+									[
+										SNew(SHorizontalBox)
+											.ToolTipText(LOCTEXT("HexSliderToolTip", "Hexadecimal Value"))
+
+										+ SHorizontalBox::Slot()
+											.AutoWidth()
+											.Padding(0.0f, 0.0f, 4.0f, 0.0f)
+											.VAlign(VAlign_Center)
+											[
+												SNew(STextBlock)
+													.Text(LOCTEXT("HexInputLabel", "Hex"))
+											]
+
+										+ SHorizontalBox::Slot()
+											.AutoWidth()
+											[
+												SNew(SEditableTextBox)
+													.MinDesiredWidth(72.0f)
+													.Text(this, &SColorPicker::HandleHexBoxText)
+													.OnTextCommitted(this, &SColorPicker::HandleHexInputTextCommitted)
+											]
+									]
+							]
+					]
+			]
+
+		// dialog buttons
+		+ SVerticalBox::Slot()
+			.AutoHeight()
+			.HAlign(HAlign_Right)
+			.VAlign(VAlign_Center)
+			[
+				SNew(SUniformGridPanel)
+					.MinDesiredSlotHeight(FCoreStyle::Get().GetFloat("StandardDialog.MinDesiredSlotHeight"))
+					.MinDesiredSlotWidth(FCoreStyle::Get().GetFloat("StandardDialog.MinDesiredSlotWidth"))
+					.SlotPadding(FCoreStyle::Get().GetMargin("StandardDialog.SlotPadding"))
+					.Visibility(ParentWindowPtr.IsValid() ? EVisibility::Visible : EVisibility::Collapsed)
+
+				+ SUniformGridPanel::Slot(0, 0)
+					[
+						// ok button
+						SNew(SButton)
+							.ContentPadding( FCoreStyle::Get().GetMargin("StandardDialog.ContentPadding") )
+							.HAlign(HAlign_Center)
+							.Text(LOCTEXT("OKButton", "OK"))
+							.OnClicked(this, &SColorPicker::HandleOkButtonClicked)
+					]
+
+				+ SUniformGridPanel::Slot(1, 0)
+					[
+						// cancel button
+						SNew(SButton)
+							.ContentPadding( FCoreStyle::Get().GetMargin("StandardDialog.ContentPadding") )
+							.HAlign(HAlign_Center)
+							.Text(LOCTEXT("CancelButton", "Cancel"))
+							.OnClicked(this, &SColorPicker::HandleCancelButtonClicked)
+					]
+			]
+	];
+	
+	HideSmallTrash();
+}
+END_SLATE_FUNCTION_BUILD_OPTIMIZATION
+
+
+void SColorPicker::Tick( const FGeometry& AllottedGeometry, const double InCurrentTime, const float InDeltaTime )
+{
+	// animate the wheel and all the sliders
+	static const float AnimationTime = 0.2f;
+	if (CurrentTime < AnimationTime)
+	{
+		CurrentColorHSV = FMath::Lerp(ColorBegin, ColorEnd, CurrentTime / AnimationTime);
+		if (CurrentColorHSV.R < 0.f)
+		{
+			CurrentColorHSV.R += 360.f;
+		}
+		else if (CurrentColorHSV.R > 360.f)
+		{
+			CurrentColorHSV.R -= 360.f;
+		}
+
+		CurrentTime += InDeltaTime;
+		if (CurrentTime >= AnimationTime)
+		{
+			CurrentColorHSV = ColorEnd;
+		}
+
+		CurrentColorRGB = CurrentColorHSV.HSVToLinearRGB();
+	}
+}
+
+
+void SColorPicker::GenerateInlineColorPickerContent()
+{
+	ChildSlot
+	[
+		SNew(SHorizontalBox)
+
+		+ SHorizontalBox::Slot()
+			.HAlign(HAlign_Center)
+			.AutoWidth()
+			[
+				SNew(SVerticalBox)
+
+				+ SVerticalBox::Slot()
+					.AutoHeight()
+					[
+						SNew(SColorWheel)
+							.SelectedColor(this, &SColorPicker::GetCurrentColor)
+							.OnValueChanged(this, &SColorPicker::HandleHSVColorChanged)
+							.OnMouseCaptureBegin(this, &SColorPicker::HandleInteractiveChangeBegin)
+							.OnMouseCaptureEnd(this, &SColorPicker::HandleInteractiveChangeEnd)
+					]
+
+				+ SVerticalBox::Slot()
+					.Padding(FMargin(0.0f, 2.0f, 0.0f, 0.0f))
+					.AutoHeight()
+					[
+						SNew(SValueSlider)
+							.SelectedColor(this, &SColorPicker::GetCurrentColor)
+							.OnValueChanged(this, &SColorPicker::HandleHSVColorChanged)
+							.OnMouseCaptureBegin(this, &SColorPicker::HandleInteractiveChangeBegin)
+							.OnMouseCaptureEnd(this, &SColorPicker::HandleInteractiveChangeEnd)
+					]
+			]
+	];
+}
+
+
+
+
+
+void SColorPicker::DiscardColor()
+{
+	if (OnColorPickerCancelled.IsBound())
+	{
+		// let the user decide what to do about cancel
+		OnColorPickerCancelled.Execute( OldColor );
+	}
+	else
+	{	
+		SetNewTargetColorHSV(OldColor, true);
+		RestoreColors();
+	}
+}
+
+
+bool SColorPicker::SetNewTargetColorHSV( const FLinearColor& NewValue, bool bForceUpdate /*= false*/ )
+{
+	CurrentColorHSV = NewValue;
+	CurrentColorRGB = NewValue.HSVToLinearRGB();
+
+	return ApplyNewTargetColor(bForceUpdate);
+}
+
+
+bool SColorPicker::SetNewTargetColorRGB( const FLinearColor& NewValue, bool bForceUpdate /*= false*/ )
+{
+	CurrentColorRGB = NewValue;
+	CurrentColorHSV = NewValue.LinearRGBToHSV();
+
+	return ApplyNewTargetColor(bForceUpdate);
+}
+
+
+bool SColorPicker::ApplyNewTargetColor( bool bForceUpdate /*= false*/ )
+{
+	bool bUpdated = false;
+
+	if ((bForceUpdate || (!bOnlyRefreshOnMouseUp && !bPerfIsTooSlowToUpdate)) && (!bOnlyRefreshOnOk || bColorPickerIsInlineVersion))
+	{
+		double StartUpdateTime = FPlatformTime::Seconds();
+		UpdateColorPickMouseUp();
+		double EndUpdateTime = FPlatformTime::Seconds();
+
+		if (EndUpdateTime - StartUpdateTime > MAX_ALLOWED_UPDATE_TIME)
+		{
+			bPerfIsTooSlowToUpdate = true;
+		}
+
+		bUpdated = true;
+	}
+
+	return bUpdated;
+}
+
+
+void SColorPicker::UpdateColorPickMouseUp()
+{
+	if (!bOnlyRefreshOnOk || bColorPickerIsInlineVersion)
+	{
+		UpdateColorPick();
+	}
+}
+
+
+void SColorPicker::UpdateColorPick()
+{
+	bPerfIsTooSlowToUpdate = false;
+	FLinearColor OutColor = CurrentColorRGB;
+
+	PreColorCommitted.ExecuteIfBound(OutColor);
+
+	SetColors(OutColor);
+	OnColorCommitted.ExecuteIfBound(OutColor);
+	
+	// This callback is only necessary for wx backwards compatibility
+	FCoreDelegates::ColorPickerChanged.Broadcast();
+}
+
+
+void SColorPicker::BeginAnimation( FLinearColor Start, FLinearColor End )
+{
+	ColorEnd = End;
+	ColorBegin = Start;
+	CurrentTime = 0.f;
+	
+	// wraparound with hue
+	float HueDif = FMath::Abs(ColorBegin.R - ColorEnd.R);
+	if (FMath::Abs(ColorBegin.R + 360.f - ColorEnd.R) < HueDif)
+	{
+		ColorBegin.R += 360.f;
+	}
+	else if (FMath::Abs(ColorBegin.R - 360.f - ColorEnd.R) < HueDif)
+	{
+		ColorBegin.R -= 360.f;
+	}
+}
+
+
+void SColorPicker::HideSmallTrash()
+{
+	if (ColorThemeComboButton.IsValid())
+	{
+		ColorThemeButtonOrSmallTrash->SetContent(ColorThemeComboButton.ToSharedRef());
+	}
+	else
+	{
+		ColorThemeButtonOrSmallTrash->ClearContent();
+	}
+}
+
+
+void SColorPicker::ShowSmallTrash()
+{
+	if ( SmallTrash.IsValid() )
+	{
+		ColorThemeButtonOrSmallTrash->SetContent(SmallTrash.ToSharedRef());
+	}
+	else
+	{
+		ColorThemeButtonOrSmallTrash->ClearContent();
+	}
+}
+
+
+/* SColorPicker implementation
+ *****************************************************************************/
+
+void SColorPicker::CycleMode( )
+{
+	if (CurrentMode == EColorPickerModes::Spectrum)
+	{
+		CurrentMode = EColorPickerModes::Wheel;
+	}
+	else
+	{
+		CurrentMode = EColorPickerModes::Spectrum;
+	}
+}
+
+
+BEGIN_SLATE_FUNCTION_BUILD_OPTIMIZATION
+TSharedRef<SWidget> SColorPicker::MakeColorSlider( EColorPickerChannels::Type Channel ) const
+{
+	FText SliderTooltip;
+
+	switch (Channel)
+	{
+	case EColorPickerChannels::Red:			SliderTooltip = LOCTEXT("RedSliderToolTip", "Red"); break;
+	case EColorPickerChannels::Green:		SliderTooltip = LOCTEXT("GreenSliderToolTip", "Green"); break;
+	case EColorPickerChannels::Blue:		SliderTooltip = LOCTEXT("BlueSliderToolTip", "Blue"); break;
+	case EColorPickerChannels::Alpha:		SliderTooltip = LOCTEXT("AlphaSliderToolTip", "Alpha"); break;
+	case EColorPickerChannels::Hue:			SliderTooltip = LOCTEXT("HueSliderToolTip", "Hue"); break;
+	case EColorPickerChannels::Saturation:	SliderTooltip = LOCTEXT("SaturationSliderToolTip", "Saturation"); break;
+	case EColorPickerChannels::Value:		SliderTooltip = LOCTEXT("ValueSliderToolTip", "Value"); break;
+	default:
+		return SNullWidget::NullWidget;
+	}
+
+	return SNew(SOverlay)
+		.ToolTipText(SliderTooltip)
+
+	+ SOverlay::Slot()
+		.Padding(FMargin(4.0f, 0.0f))
+		[
+			SNew(SSimpleGradient)
+				.EndColor(this, &SColorPicker::HandleColorSliderEndColor, Channel)
+				.StartColor(this, &SColorPicker::HandleColorSliderStartColor, Channel)
+				.Orientation(Orient_Horizontal)
+				.UseSRGB(SharedThis(this), &SColorPicker::HandleColorPickerUseSRGB)
+		]
+
+	+ SOverlay::Slot()
+		[
+			SNew(SSlider)
+				.IndentHandle(false)
+				.Orientation(Orient_Vertical)
+				.SliderBarColor(FLinearColor::Transparent)
+				.Style(&FCoreStyle::Get().GetWidgetStyle<FSliderStyle>("ColorPicker.Slider"))
+				.Value(this, &SColorPicker::HandleColorSpinBoxValue, Channel)
+				.OnMouseCaptureBegin(this, &SColorPicker::HandleInteractiveChangeBegin)
+				.OnMouseCaptureEnd(this, &SColorPicker::HandleInteractiveChangeEnd)
+				.OnValueChanged(this, &SColorPicker::HandleColorSpinBoxValueChanged, Channel)
+		];
+}
+
+
+TSharedRef<SWidget> SColorPicker::MakeColorSpinBox( EColorPickerChannels::Type Channel ) const
+{
+	if ((Channel == EColorPickerChannels::Alpha) && !bUseAlpha.Get())
+	{
+		return SNullWidget::NullWidget;
+	}
+
+	const int32 GradientHeight = 6;
+	const float HDRMaxValue = (TargetFColors.Num()) ? 1.f : FLT_MAX;
+	const FSlateFontInfo SmallLayoutFont = FCoreStyle::Get().GetFontStyle("ColorPicker.Font");
+
+	// create gradient widget
+	TSharedPtr<SWidget> GradientWidget;
+	
+	if (Channel == EColorPickerChannels::Hue)
+	{
+		TArray<FLinearColor> HueGradientColors;
+
+		for (int32 i = 0; i < 7; ++i)
+		{
+			HueGradientColors.Add( FLinearColor((i % 6) * 60.f, 1.f, 1.f).HSVToLinearRGB() );
+		}
+
+		GradientWidget = SNew(SComplexGradient)
+			.GradientColors(HueGradientColors);
+	}
+	else
+	{
+		GradientWidget = SNew(SSimpleGradient)
+			.StartColor(this, &SColorPicker::GetGradientStartColor, Channel)
+			.EndColor(this, &SColorPicker::GetGradientEndColor, Channel)
+			.HasAlphaBackground(Channel == EColorPickerChannels::Alpha)
+			.UseSRGB(SharedThis(this), &SColorPicker::HandleColorPickerUseSRGB);
+	}
+	
+	// create spin box
+	float MaxValue;
+	FText SliderLabel;
+	FText SliderTooltip;
+
+	switch (Channel)
+	{
+	case EColorPickerChannels::Red:
+		MaxValue = HDRMaxValue;
+		SliderLabel = LOCTEXT("RedSliderLabel", "R");
+		SliderTooltip = LOCTEXT("RedSliderToolTip", "Red");
+		break;
+		
+	case EColorPickerChannels::Green:
+		MaxValue = HDRMaxValue;
+		SliderLabel = LOCTEXT("GreenSliderLabel", "G");
+		SliderTooltip = LOCTEXT("GreenSliderToolTip", "Green");
+		break;
+		
+	case EColorPickerChannels::Blue:
+		MaxValue = HDRMaxValue;
+		SliderLabel = LOCTEXT("BlueSliderLabel", "B");
+		SliderTooltip = LOCTEXT("BlueSliderToolTip", "Blue");
+		break;
+		
+	case EColorPickerChannels::Alpha:
+		MaxValue = HDRMaxValue;
+		SliderLabel = LOCTEXT("AlphaSliderLabel", "A");
+		SliderTooltip = LOCTEXT("AlphaSliderToolTip", "Alpha");
+		break;
+		
+	case EColorPickerChannels::Hue:
+		MaxValue = 359.0f;
+		SliderLabel = LOCTEXT("HueSliderLabel", "H");
+		SliderTooltip = LOCTEXT("HueSliderToolTip", "Hue");
+		break;
+		
+	case EColorPickerChannels::Saturation:
+		MaxValue = 1.0f;
+		SliderLabel = LOCTEXT("SaturationSliderLabel", "S");
+		SliderTooltip = LOCTEXT("SaturationSliderToolTip", "Saturation");
+		break;
+		
+	case EColorPickerChannels::Value:
+		MaxValue = HDRMaxValue;
+		SliderLabel = LOCTEXT("ValueSliderLabel", "V");
+		SliderTooltip = LOCTEXT("ValueSliderToolTip", "Value");
+		break;
+
+	default:
+		return SNullWidget::NullWidget;
+	}
+	
+	return SNew(SHorizontalBox)
+
+	+ SHorizontalBox::Slot()
+		.AutoWidth()
+		.Padding(0.0f, 0.0f, 4.0f, 0.0f)
+		.VAlign(VAlign_Center)
+		[
+			SNew(STextBlock)
+				.Text(SliderLabel)
+		]
+
+	+ SHorizontalBox::Slot()
+		.FillWidth(1.0f)
+		[
+			SNew(SVerticalBox)
+				.ToolTipText(SliderTooltip)
+
+			+ SVerticalBox::Slot()
+				.AutoHeight()
+				[
+					SNew(SSpinBox<float>)
+						.MinValue(0.0f)
+						.MaxValue(MaxValue)
+						.MinSliderValue(0.0f)
+						.MaxSliderValue(Channel == EColorPickerChannels::Hue ? 359.0f : 1.0f)
+						.Delta(Channel == EColorPickerChannels::Hue ? 1.0f : 0.001f)
+						.Font(SmallLayoutFont)
+						.Value(this, &SColorPicker::HandleColorSpinBoxValue, Channel)
+						.OnBeginSliderMovement(this, &SColorPicker::HandleInteractiveChangeBegin)
+						.OnEndSliderMovement(this, &SColorPicker::HandleInteractiveChangeEnd)
+						.OnValueChanged(this, &SColorPicker::HandleColorSpinBoxValueChanged, Channel)
+				]
+
+			+ SVerticalBox::Slot()
+				.AutoHeight()
+				[
+					SNew(SBox)
+						.HeightOverride(GradientHeight)
+						[
+							GradientWidget.ToSharedRef()
+						]
+				]
+		];
+}
+
+
+TSharedRef<SWidget> SColorPicker::MakeColorPreviewBox( ) const
+{
+	return SNew(SOverlay)
+
+	+ SOverlay::Slot()
+		[
+			// preview blocks
+			SNew(SVerticalBox)
+
+			+ SVerticalBox::Slot()
+				.AutoHeight()
+				[
+					SNew(STextBlock)
+						.Text(LOCTEXT("OldColorLabel", "Old"))
+				]
+
+			+ SVerticalBox::Slot()
+				[
+					SNew(SHorizontalBox)
+
+					+ SHorizontalBox::Slot()
+						[
+							// old color
+							SNew(SColorBlock) 
+								.ColorIsHSV(true) 
+								.IgnoreAlpha(true)
+								.ToolTipText(LOCTEXT("OldColorToolTip", "Old color without alpha (drag to theme bar to save)"))
+								.Color(OldColor) 
+								.OnMouseButtonDown(this, &SColorPicker::HandleOldColorBlockMouseButtonDown, false)
+								.UseSRGB(SharedThis(this), &SColorPicker::HandleColorPickerUseSRGB)
+								.Cursor(EMouseCursor::GrabHand)
+						]
+
+					+ SHorizontalBox::Slot()
+						[
+							// old color (alpha)
+							SNew(SColorBlock) 
+								.ColorIsHSV(true) 
+								.ShowBackgroundForAlpha(true)
+								.ToolTipText(LOCTEXT("OldColorAlphaToolTip", "Old color with alpha (drag to theme bar to save)"))
+								.Color(OldColor) 
+								.Visibility(SharedThis(this), &SColorPicker::HandleAlphaColorBlockVisibility)
+								.OnMouseButtonDown(this, &SColorPicker::HandleOldColorBlockMouseButtonDown, true)
+								.UseSRGB(SharedThis(this), &SColorPicker::HandleColorPickerUseSRGB)
+								.Cursor(EMouseCursor::GrabHand)
+						]
+				]
+
+			+ SVerticalBox::Slot()
+				[
+					SNew(SHorizontalBox)
+
+					+ SHorizontalBox::Slot()
+						[
+							// new color
+							SNew(SColorBlock) 
+								.ColorIsHSV(true) 
+								.IgnoreAlpha(true)
+								.ToolTipText(LOCTEXT("NewColorToolTip", "New color without alpha (drag to theme bar to save)"))
+								.Color(this, &SColorPicker::GetCurrentColor)
+								.OnMouseButtonDown(this, &SColorPicker::HandleNewColorBlockMouseButtonDown, false)
+								.UseSRGB(SharedThis(this), &SColorPicker::HandleColorPickerUseSRGB)
+								.Cursor(EMouseCursor::GrabHand)
+						]
+
+					+ SHorizontalBox::Slot()
+						[
+							// new color (alpha)
+							SNew(SColorBlock) 
+								.ColorIsHSV(true) 
+								.ShowBackgroundForAlpha(true)
+								.ToolTipText(LOCTEXT("NewColorAlphaToolTip", "New color with alpha (drag to theme bar to save)"))
+								.Color(this, &SColorPicker::GetCurrentColor)
+								.Visibility(SharedThis(this), &SColorPicker::HandleAlphaColorBlockVisibility)
+								.OnMouseButtonDown(this, &SColorPicker::HandleNewColorBlockMouseButtonDown, true)
+								.UseSRGB(SharedThis(this), &SColorPicker::HandleColorPickerUseSRGB)
+								.Cursor(EMouseCursor::GrabHand)
+						]
+				]
+
+			+ SVerticalBox::Slot()
+				.AutoHeight()
+				[
+					SNew(STextBlock)
+						.Text(LOCTEXT("NewColorLabel", "New"))
+				]
+		]
+
+	+ SOverlay::Slot()
+		.VAlign(VAlign_Center)
+		[
+			// block separators
+			SNew(SHorizontalBox)
+
+			+ SHorizontalBox::Slot()
+				.FillWidth(0.5f)
+				.HAlign(HAlign_Left)
+				[
+					SNew(SBox)
+						.HeightOverride(2.0f)
+						.WidthOverride(4.0f)
+						[
+							SNew(SBorder)
+								.BorderImage(FCoreStyle::Get().GetBrush("ColorPicker.Separator"))
+								.Padding(0.0f)
+						]
+				]
+
+			+ SHorizontalBox::Slot()
+				.FillWidth(0.5f)
+				.HAlign(HAlign_Right)
+				[
+					SNew(SBox)
+						.HeightOverride(2.0f)
+						.WidthOverride(4.0f)
+						[
+							SNew(SBorder)
+								.BorderImage(FCoreStyle::Get().GetBrush("ColorPicker.Separator"))
+								.Padding(0.0f)
+						]
+				]
+		];
+}
+END_SLATE_FUNCTION_BUILD_OPTIMIZATION
+
+
+/* SColorPicker callbacks
+ *****************************************************************************/
+
+FLinearColor SColorPicker::GetGradientEndColor( EColorPickerChannels::Type Channel ) const
+{
+	switch (Channel)
+	{
+	case EColorPickerChannels::Red:			return FLinearColor(1.0f, CurrentColorRGB.G, CurrentColorRGB.B, 1.0f);
+	case EColorPickerChannels::Green:		return FLinearColor(CurrentColorRGB.R, 1.0f, CurrentColorRGB.B, 1.0f);
+	case EColorPickerChannels::Blue:		return FLinearColor(CurrentColorRGB.R, CurrentColorRGB.G, 1.0f, 1.0f);
+	case EColorPickerChannels::Alpha:		return FLinearColor(CurrentColorRGB.R, CurrentColorRGB.G, CurrentColorRGB.B, 1.0f);
+	case EColorPickerChannels::Saturation:	return FLinearColor(CurrentColorHSV.R, 1.0f, CurrentColorHSV.B, 1.0f).HSVToLinearRGB();
+	case EColorPickerChannels::Value:		return FLinearColor(CurrentColorHSV.R, CurrentColorHSV.G, 1.0f, 1.0f).HSVToLinearRGB();
+	default:								return FLinearColor();
+	}
+}
+
+
+FLinearColor SColorPicker::GetGradientStartColor( EColorPickerChannels::Type Channel ) const
+{
+	switch (Channel)
+	{
+	case EColorPickerChannels::Red:			return FLinearColor(0.0f, CurrentColorRGB.G, CurrentColorRGB.B, 1.0f);
+	case EColorPickerChannels::Green:		return FLinearColor(CurrentColorRGB.R, 0.0f, CurrentColorRGB.B, 1.0f);
+	case EColorPickerChannels::Blue:		return FLinearColor(CurrentColorRGB.R, CurrentColorRGB.G, 0.0f, 1.0f);
+	case EColorPickerChannels::Alpha:		return FLinearColor(CurrentColorRGB.R, CurrentColorRGB.G, CurrentColorRGB.B, 0.0f);
+	case EColorPickerChannels::Saturation:	return FLinearColor(CurrentColorHSV.R, 0.0f, CurrentColorHSV.B, 1.0f).HSVToLinearRGB();
+	case EColorPickerChannels::Value:		return FLinearColor(CurrentColorHSV.R, CurrentColorHSV.G, 0.0f, 1.0f).HSVToLinearRGB();
+	default:								return FLinearColor();
+	}
+}
+
+
+void SColorPicker::HandleAdvancedAreaExpansionChanged( bool Expanded )
+{
+	if (FPaths::FileExists(GEditorUserSettingsIni))
+	{
+		GConfig->SetBool(TEXT("ColorPickerUI"), TEXT("bAdvancedSectionExpanded"), Expanded, GEditorUserSettingsIni);
+	}
+}
+
+
+EVisibility SColorPicker::HandleAlphaColorBlockVisibility( ) const
+{
+	return bUseAlpha.Get() ? EVisibility::Visible : EVisibility::Collapsed;
+}
+
+
+FReply SColorPicker::HandleCancelButtonClicked( )
+{
+	bClosedViaOkOrCancel = true;
+
+	DiscardColor();
+	ParentWindowPtr.Pin()->RequestDestroyWindow();
+
+	return FReply::Handled();
+}
+
+
+EVisibility SColorPicker::HandleColorPickerModeVisibility( EColorPickerModes::Type Mode ) const
+{
+	return (CurrentMode == Mode) ? EVisibility::Visible : EVisibility::Hidden;
+}
+
+
+FLinearColor SColorPicker::HandleColorSliderEndColor( EColorPickerChannels::Type Channel ) const
+{
+	switch (Channel)
+	{
+	case EColorPickerChannels::Red:			return FLinearColor(0.0f, 0.0f, 0.0f, 1.0f);
+	case EColorPickerChannels::Green:		return FLinearColor(0.0f, 0.0f, 0.0f, 1.0f);
+	case EColorPickerChannels::Blue:		return FLinearColor(0.0f, 0.0f, 0.0f, 1.0f);
+	case EColorPickerChannels::Alpha:		return FLinearColor(CurrentColorRGB.R, CurrentColorRGB.G, CurrentColorRGB.B, 0.0f);
+	case EColorPickerChannels::Saturation:	return FLinearColor(CurrentColorHSV.R, 0.0f, 1.0f, 1.0f).HSVToLinearRGB();
+	case EColorPickerChannels::Value:		return FLinearColor(CurrentColorHSV.R, CurrentColorHSV.G, 0.0f, 1.0f).HSVToLinearRGB();
+	default:								return FLinearColor();
+	}
+}
+
+
+FLinearColor SColorPicker::HandleColorSliderStartColor( EColorPickerChannels::Type Channel ) const
+{
+	switch (Channel)
+	{
+	case EColorPickerChannels::Red:			return FLinearColor(1.0f, 0.0f, 0.0f, 1.0f);
+	case EColorPickerChannels::Green:		return FLinearColor(0.0f, 1.0f, 0.0f, 1.0f);
+	case EColorPickerChannels::Blue:		return FLinearColor(0.0f, 0.0f, 1.0f, 1.0f);
+	case EColorPickerChannels::Alpha:		return FLinearColor(CurrentColorRGB.R, CurrentColorRGB.G, CurrentColorRGB.B, 1.0f);
+	case EColorPickerChannels::Saturation:	return FLinearColor(CurrentColorHSV.R, 1.0f, 1.0f, 1.0f).HSVToLinearRGB();
+	case EColorPickerChannels::Value:		return FLinearColor(CurrentColorHSV.R, CurrentColorHSV.G, 1.0f, 1.0f).HSVToLinearRGB();
+	default:								return FLinearColor();
+	}
+}
+
+
+void SColorPicker::HandleColorSpectrumValueChanged( FLinearColor NewValue )
+{
+	SetNewTargetColorHSV(NewValue);
+}
+
+
+float SColorPicker::HandleColorSpinBoxValue( EColorPickerChannels::Type Channel ) const
+{
+	switch (Channel)
+	{
+	case EColorPickerChannels::Red:			return CurrentColorRGB.R;
+	case EColorPickerChannels::Green:		return CurrentColorRGB.G;
+	case EColorPickerChannels::Blue:		return CurrentColorRGB.B;
+	case EColorPickerChannels::Alpha:		return CurrentColorRGB.A;
+	case EColorPickerChannels::Hue:			return CurrentColorHSV.R;
+	case EColorPickerChannels::Saturation:	return CurrentColorHSV.G;
+	case EColorPickerChannels::Value:		return CurrentColorHSV.B;
+	default:								return 0.0f;
+	}
+}
+
+
+void SColorPicker::HandleColorSpinBoxValueChanged( float NewValue, EColorPickerChannels::Type Channel )
+{
+	int32 ComponentIndex;
+	bool IsHSV = false;
+
+	switch (Channel)
+	{
+	case EColorPickerChannels::Red:			ComponentIndex = 0; break;
+	case EColorPickerChannels::Green:		ComponentIndex = 1; break;
+	case EColorPickerChannels::Blue:		ComponentIndex = 2; break;
+	case EColorPickerChannels::Alpha:		ComponentIndex = 3; break;
+	case EColorPickerChannels::Hue:			ComponentIndex = 0; IsHSV = true; break;
+	case EColorPickerChannels::Saturation:	ComponentIndex = 1; IsHSV = true; break;
+	case EColorPickerChannels::Value:		ComponentIndex = 2; IsHSV = true; break;
+	default:								
+		return;
+	}
+
+	FLinearColor& NewColor = IsHSV ? CurrentColorHSV : CurrentColorRGB;
+
+	if (FMath::IsNearlyEqual(NewValue, NewColor.Component(ComponentIndex), KINDA_SMALL_NUMBER))
+	{
+		return;
+	}
+
+	NewColor.Component(ComponentIndex) = NewValue;
+
+	if (IsHSV)
+	{
+		SetNewTargetColorHSV(NewColor, !bIsInteractive);
+	}
+	else
+	{
+		SetNewTargetColorRGB(NewColor, !bIsInteractive);
+	}
+}
+
+
+void SColorPicker::HandleEyeDropperButtonComplete( )
+{
+	bIsInteractive = false;
+
+	if (bOnlyRefreshOnMouseUp || bPerfIsTooSlowToUpdate)
+	{
+		UpdateColorPick();
+	}
+
+	OnInteractivePickEnd.ExecuteIfBound();
+}
+
+
+FText SColorPicker::HandleHexBoxText( ) const
+{
+	return FText::FromString(CurrentColorRGB.ToFColor(false).ToHex());
+}
+
+
+void SColorPicker::HandleHexInputTextCommitted( const FText& Text, ETextCommit::Type CommitType )
+{
+	if ((CommitType == ETextCommit::OnEnter) || (CommitType == ETextCommit::OnUserMovedFocus))
+	{
+		SetNewTargetColorRGB(FColor::FromHex(Text.ToString()), false);
+	}	
+}
+
+
+void SColorPicker::HandleHSVColorChanged( FLinearColor NewValue )
+{
+	SetNewTargetColorHSV(NewValue);
+}
+
+
+void SColorPicker::HandleInteractiveChangeBegin( )
+{
+	if( bIsInteractive && OnInteractivePickEnd.IsBound() )
+	{
+		OnInteractivePickEnd.Execute();
+	}
+
+	OnInteractivePickBegin.ExecuteIfBound();
+	bIsInteractive = true;
+}
+
+
+void SColorPicker::HandleInteractiveChangeEnd( )
+{
+	HandleInteractiveChangeEnd(0.0f);
+}
+
+
+void SColorPicker::HandleInteractiveChangeEnd( float NewValue )
+{
+	bIsInteractive = false;
+
+	UpdateColorPickMouseUp();
+	OnInteractivePickEnd.ExecuteIfBound();
+}
+
+
+FReply SColorPicker::HandleColorAreaMouseDown( const FGeometry& MyGeometry, const FPointerEvent& MouseEvent )
+{
+	if (MouseEvent.GetEffectingButton() == EKeys::RightMouseButton)
+	{
+		CycleMode();
+
+		return FReply::Handled();
+	}
+
+	return FReply::Unhandled();
+}
+
+
+FReply SColorPicker::HandleColorPickerModeButtonClicked( )
+{
+	CycleMode();
+
+	if (FPaths::FileExists(GEditorUserSettingsIni))
+	{
+		GConfig->SetBool(TEXT("ColorPickerUI"), TEXT("bWheelMode"), CurrentMode == EColorPickerModes::Wheel, GEditorUserSettingsIni);
+	}
+
+	return FReply::Handled();
+}
+
+
+FReply SColorPicker::HandleNewColorBlockMouseButtonDown( const FGeometry& MyGeometry, const FPointerEvent& MouseEvent, bool bCheckAlpha )
+{
+	if (MouseEvent.GetEffectingButton() == EKeys::LeftMouseButton)
+	{
+		TSharedRef<FColorDragDrop> Operation = FColorDragDrop::New(
+			CurrentColorHSV,
+			SColorThemesViewer::bSRGBEnabled,
+			bCheckAlpha ? bUseAlpha.Get() : false,
+			FSimpleDelegate::CreateSP(this, &SColorPicker::ShowSmallTrash),
+			FSimpleDelegate::CreateSP(this, &SColorPicker::HideSmallTrash)
+		);
+
+		return FReply::Handled().BeginDragDrop(Operation);
+	}
+
+	return FReply::Unhandled();
+}
+
+
+FReply SColorPicker::HandleOkButtonClicked( )
+{
+	bClosedViaOkOrCancel = true;
+
+	if (bOnlyRefreshOnOk)
+	{
+		UpdateColorPick();
+	}
+
+	ParentWindowPtr.Pin()->RequestDestroyWindow();
+
+	return FReply::Handled();
+}
+
+
+FReply SColorPicker::HandleOldColorBlockMouseButtonDown( const FGeometry& MyGeometry, const FPointerEvent& MouseEvent, bool bCheckAlpha )
+{
+	if (MouseEvent.GetEffectingButton() == EKeys::LeftMouseButton)
+	{
+		TSharedRef<FColorDragDrop> Operation = FColorDragDrop::New(
+			OldColor,
+			SColorThemesViewer::bSRGBEnabled,
+			bCheckAlpha ? bUseAlpha.Get() : false,
+			FSimpleDelegate::CreateSP(this, &SColorPicker::ShowSmallTrash),
+			FSimpleDelegate::CreateSP(this, &SColorPicker::HideSmallTrash)
+		);
+
+		return FReply::Handled().BeginDragDrop(Operation);
+	}
+
+	return FReply::Unhandled();
+}
+
+
+bool SColorPicker::HandleColorPickerUseSRGB( ) const
+{
+	return SColorThemesViewer::bSRGBEnabled;
+}
+
+
+void SColorPicker::HandleParentWindowClosed( const TSharedRef<SWindow>& Window )
+{
+	check(Window == ParentWindowPtr.Pin());
+
+	// End picking interaction if still active
+	if( bIsInteractive && OnInteractivePickEnd.IsBound() )
+	{
+		OnInteractivePickEnd.Execute();
+		bIsInteractive = false;
+	}
+
+	// We always have to call the close callback
+	if (OnColorPickerWindowClosed.IsBound())
+	{
+		OnColorPickerWindowClosed.Execute(Window);
+	}
+
+	// If we weren't closed via the OK or Cancel button, we need to perform the default close action
+	if (!bClosedViaOkOrCancel && bOnlyRefreshOnOk)
+	{
+		DiscardColor();
+	}
+}
+
+
+void SColorPicker::HandleRGBColorChanged( FLinearColor NewValue )
+{
+	SetNewTargetColorRGB(NewValue);
+}
+
+
+void SColorPicker::HandleSRGBCheckBoxCheckStateChanged( ESlateCheckBoxState::Type InIsChecked )
+{
+	SColorThemesViewer::bSRGBEnabled = (InIsChecked == ESlateCheckBoxState::Checked);
+
+	if (FPaths::FileExists(GEditorUserSettingsIni))
+	{
+		GConfig->SetBool(TEXT("ColorPickerUI"), TEXT("bSRGBEnabled"), SColorThemesViewer::bSRGBEnabled, GEditorUserSettingsIni);
+	}
+}
+
+
+ESlateCheckBoxState::Type SColorPicker::HandleSRGBCheckBoxIsChecked( ) const
+{
+	return SColorThemesViewer::bSRGBEnabled ? ESlateCheckBoxState::Checked : ESlateCheckBoxState::Unchecked;
+}
+
+
+void SColorPicker::HandleThemeBarColorSelected( FLinearColor NewValue )
+{
+	// Force the alpha component to 1 when we don't care about the alpha
+	if (!bUseAlpha.Get())
+	{
+		NewValue.A = 1.0f;
+	}
+
+	BeginAnimation(CurrentColorHSV, NewValue);
+	SetNewTargetColorHSV(NewValue, true);
+}
+
+
+TSharedPtr<FColorTheme> SColorPicker::HandleThemeBarColorTheme() const
+{
+	return ColorThemesViewer->GetCurrentColorTheme();
+}
+
+
+EVisibility SColorPicker::HandleThemeBarHintVisibility( ) const
+{
+	TSharedPtr<FColorTheme> SelectedTheme = ColorThemesViewer->GetCurrentColorTheme();
+
+	if (SelectedTheme.IsValid() && SelectedTheme->GetColors().Num() == 0)
+	{
+		return EVisibility::HitTestInvisible;
+	}
+
+	return EVisibility::Hidden;
+}
+
+
+bool SColorPicker::HandleThemeBarUseAlpha( ) const
+{
+	return bUseAlpha.Get();
+}
+
+
+void SColorPicker::HandleThemesViewerThemeChanged()
+{
+	if (CurrentThemeBar.IsValid())
+	{
+		CurrentThemeBar->RemoveRefreshCallback();
+		CurrentThemeBar->AddRefreshCallback();
+		CurrentThemeBar->Refresh();
+	}
+}
+
+
+/* Global functions
+ *****************************************************************************/
+
+/** A static color picker that everything should use. */
+static TWeakPtr<SWindow> ColorPickerWindow;
+
+
+bool OpenColorPicker(const FColorPickerArgs& Args)
+{
+	DestroyColorPicker();
+
+	bool Result = false;
+	
+	//guard to verify that we're not in game using edit actor
+	if (GIsEditor)
+	{
+		FLinearColor OldColor = Args.InitialColorOverride;
+
+		if (Args.ColorArray && Args.ColorArray->Num() > 0)
+		{
+			OldColor = (*Args.ColorArray)[0]->ReinterpretAsLinear();
+		}
+		else if (Args.LinearColorArray && Args.LinearColorArray->Num() > 0)
+		{
+			OldColor = *(*Args.LinearColorArray)[0];
+		}
+		else if (Args.ColorChannelsArray && Args.ColorChannelsArray->Num() > 0)
+		{
+			OldColor.R = (*Args.ColorChannelsArray)[0].Red ? *(*Args.ColorChannelsArray)[0].Red : 0.0f;
+			OldColor.G = (*Args.ColorChannelsArray)[0].Green ? *(*Args.ColorChannelsArray)[0].Green : 0.0f;
+			OldColor.B = (*Args.ColorChannelsArray)[0].Blue ? *(*Args.ColorChannelsArray)[0].Blue : 0.0f;
+			OldColor.A = (*Args.ColorChannelsArray)[0].Alpha ? *(*Args.ColorChannelsArray)[0].Alpha : 0.0f;
+		}
+		else
+		{
+			check(Args.OnColorCommitted.IsBound());
+		}
+		
+		// Determine the position of the window so that it will spawn near the mouse, but not go off the screen.
+		const FVector2D CursorPos = FSlateApplication::Get().GetCursorPos();
+		FSlateRect Anchor(CursorPos.X, CursorPos.Y, CursorPos.X, CursorPos.Y);
+
+		FVector2D AdjustedSummonLocation = FSlateApplication::Get().CalculatePopupWindowPosition( Anchor, SColorPicker::DEFAULT_WINDOW_SIZE, Orient_Horizontal );
+
+		TSharedPtr<SWindow> Window = SNew(SWindow)
+			.AutoCenter(EAutoCenter::None)
+			.ScreenPosition(AdjustedSummonLocation)
+			.SupportsMaximize(false)
+			.SupportsMinimize(false)
+			.SizingRule(ESizingRule::Autosized)
+			.Title(LOCTEXT("WindowHeader", "Color Picker"));
+
+		TSharedRef<SColorPicker> ColorPicker = SNew(SColorPicker)
+			.TargetColorAttribute(OldColor)
+			.TargetFColors(Args.ColorArray ? *Args.ColorArray : TArray<FColor*>())
+			.TargetLinearColors(Args.LinearColorArray ? *Args.LinearColorArray : TArray<FLinearColor*>())
+			.TargetColorChannels(Args.ColorChannelsArray ? *Args.ColorChannelsArray : TArray<FColorChannels>())
+			.UseAlpha(Args.bUseAlpha)
+			.OnlyRefreshOnMouseUp(Args.bOnlyRefreshOnMouseUp && !Args.bIsModal)
+			.OnlyRefreshOnOk(Args.bOnlyRefreshOnOk || Args.bIsModal)
+			.OnColorCommitted(Args.OnColorCommitted)
+			.PreColorCommitted(Args.PreColorCommitted)
+			.OnColorPickerCancelled(Args.OnColorPickerCancelled)
+			.OnInteractivePickBegin(Args.OnInteractivePickBegin)
+			.OnInteractivePickEnd(Args.OnInteractivePickEnd)
+			.OnColorPickerWindowClosed(Args.OnColorPickerWindowClosed)
+			.ParentWindow(Window)
+			.DisplayGamma(Args.DisplayGamma);
+		
+		Window->SetContent(
+			SNew(SBox)
+				[
+					SNew(SBorder)
+						.BorderImage(FCoreStyle::Get().GetBrush("ToolPanel.GroupBorder"))
+						.Padding(FMargin(8.0f, 8.0f))
+						[
+							ColorPicker
+						]
+				]
+		);
+
+		if (Args.bIsModal)
+		{
+			FSlateApplication::Get().AddModalWindow(Window.ToSharedRef(), Args.ParentWidget);
+		}
+		else
+		{
+			if ( Args.ParentWidget.IsValid() )
+			{
+				// Find the window of the parent widget
+				FWidgetPath WidgetPath;
+				FSlateApplication::Get().GeneratePathToWidgetChecked( Args.ParentWidget.ToSharedRef(), WidgetPath );
+				Window = FSlateApplication::Get().AddWindowAsNativeChild( Window.ToSharedRef(), WidgetPath.GetWindow() );
+			}
+			else
+			{
+				Window = FSlateApplication::Get().AddWindow(Window.ToSharedRef());
+			}
+			
+		}
+
+		Result = true;
+
+		//hold on to the window created for external use...
+		ColorPickerWindow = Window;
+	}
+
+	return Result;
+}
+
+
+/**
+ * Destroys the current color picker. Necessary if the values the color picker
+ * currently targets become invalid.
+ */
+void DestroyColorPicker()
+{
+	if (ColorPickerWindow.IsValid())
+	{
+		ColorPickerWindow.Pin()->RequestDestroyWindow();
+		ColorPickerWindow.Reset();
+	}
+}
+
+
+#undef LOCTEXT_NAMESPACE
