Slate: Reorganized private Slate files; proper pre-compiled header for all files.

diff --git a/Engine/Source/Runtime/Slate/Private/Widgets/Colors/SColorSpectrum.cpp b/Engine/Source/Runtime/Slate/Private/Widgets/Colors/SColorSpectrum.cpp
new file mode 100644
--- /dev/null
+++ b/Engine/Source/Runtime/Slate/Private/Widgets/Colors/SColorSpectrum.cpp
@@ -0,0 +1,140 @@
+// Copyright 1998-2014 Epic Games, Inc. All Rights Reserved.
+
+/*=============================================================================
+	SColorSpectrum.cpp: Implements the SColorSpectrum class.
+=============================================================================*/
+
+#include "SlatePrivatePCH.h"
+
+
+/* SColorSpectrum methods
+ *****************************************************************************/
+
+void SColorSpectrum::Construct( const FArguments& InArgs )
+{
+	Image = FCoreStyle::Get().GetBrush("ColorSpectrum.Spectrum");
+	SelectorImage = FCoreStyle::Get().GetBrush("ColorSpectrum.Selector");
+	SelectedColor = InArgs._SelectedColor;
+
+	OnMouseCaptureBegin = InArgs._OnMouseCaptureBegin;
+	OnMouseCaptureEnd = InArgs._OnMouseCaptureEnd;
+	OnValueChanged = InArgs._OnValueChanged;
+}
+
+
+/* SWidget overrides
+ *****************************************************************************/
+
+FVector2D SColorSpectrum::ComputeDesiredSize( ) const
+{
+	return Image->ImageSize;
+}
+
+
+FReply SColorSpectrum::OnMouseButtonDown( const FGeometry& MyGeometry, const FPointerEvent& MouseEvent )
+{
+	if (MouseEvent.GetEffectingButton() == EKeys::LeftMouseButton)
+	{
+		OnMouseCaptureBegin.ExecuteIfBound();
+
+		return FReply::Handled().CaptureMouse(SharedThis(this));
+	}
+
+	return FReply::Unhandled();
+}
+
+
+FReply SColorSpectrum::OnMouseButtonUp( const FGeometry& MyGeometry, const FPointerEvent& MouseEvent )
+{
+	if ((MouseEvent.GetEffectingButton() == EKeys::LeftMouseButton) && HasMouseCapture())
+	{
+		OnMouseCaptureEnd.ExecuteIfBound();
+
+		return FReply::Handled().ReleaseMouseCapture();
+	}
+
+	return FReply::Unhandled();
+}
+
+
+FReply SColorSpectrum::OnMouseMove( const FGeometry& MyGeometry, const FPointerEvent& MouseEvent )
+{
+	if (!HasMouseCapture())
+	{
+		return FReply::Unhandled();
+	}
+
+	FVector2D NormalizedMousePosition = MyGeometry.AbsoluteToLocal(MouseEvent.GetScreenSpacePosition()) / MyGeometry.Size;
+	NormalizedMousePosition = NormalizedMousePosition.ClampAxes(0.0f, 1.0f);
+
+	FLinearColor NewColor = SelectedColor.Get();
+	NewColor.R = 360.0f * NormalizedMousePosition.X;
+
+	if (NormalizedMousePosition.Y > 0.5f)
+	{
+		NewColor.G = 1.0f;
+		NewColor.B = 2.0f * (1.0f - NormalizedMousePosition.Y);
+	}
+	else
+	{
+		NewColor.G = 2.0f * NormalizedMousePosition.Y;
+		NewColor.B = 1.0f;
+	}
+
+	OnValueChanged.ExecuteIfBound(NewColor);
+
+	return FReply::Handled();
+}
+
+
+int32 SColorSpectrum::OnPaint( const FGeometry& AllottedGeometry, const FSlateRect& MyClippingRect, FSlateWindowElementList& OutDrawElements, int32 LayerId, const FWidgetStyle& InWidgetStyle, bool bParentEnabled ) const
+{
+	const bool bIsEnabled = ShouldBeEnabled(bParentEnabled);
+	const uint32 DrawEffects = bIsEnabled ? ESlateDrawEffect::None : ESlateDrawEffect::DisabledEffect;
+	
+	// draw gradient
+	FSlateDrawElement::MakeBox(
+		OutDrawElements,
+		LayerId,
+		AllottedGeometry.ToPaintGeometry(),
+		Image,
+		MyClippingRect,
+		DrawEffects,
+		InWidgetStyle.GetColorAndOpacityTint() * Image->GetTint(InWidgetStyle));
+
+	// ignore colors that can't be represented in spectrum
+	const FLinearColor& Color = SelectedColor.Get();
+
+	if ((Color.G < 1.0f) && (Color.B < 1.0f))
+	{
+		return LayerId;
+	}
+
+	// draw cursor
+	FSlateDrawElement::MakeBox(
+		OutDrawElements,
+		LayerId + 1,
+		AllottedGeometry.ToPaintGeometry(CalcRelativeSelectedPosition() * AllottedGeometry.Size - SelectorImage->ImageSize * 0.5f, SelectorImage->ImageSize),
+		SelectorImage,
+		MyClippingRect,
+		DrawEffects,
+		InWidgetStyle.GetColorAndOpacityTint() * SelectorImage->GetTint(InWidgetStyle));
+
+	return LayerId + 1;
+}
+
+
+/* SColorSpectrum implementation
+ *****************************************************************************/
+
+FVector2D SColorSpectrum::CalcRelativeSelectedPosition( ) const
+{
+	const FLinearColor& Color = SelectedColor.Get();
+
+	if (Color.G == 1.0f)
+	{
+		return FVector2D(Color.R / 360.0f, 1.0f - 0.5f * Color.B);
+	}
+
+	return FVector2D(Color.R / 360.0f, 0.5f * Color.G);
+}
diff --git a/Engine/Source/Runtime/Slate/Private/Widgets/Colors/SColorWheel.cpp b/Engine/Source/Runtime/Slate/Private/Widgets/Colors/SColorWheel.cpp
new file mode 100644
--- /dev/null
+++ b/Engine/Source/Runtime/Slate/Private/Widgets/Colors/SColorWheel.cpp
@@ -0,0 +1,138 @@
+// Copyright 1998-2014 Epic Games, Inc. All Rights Reserved.
+
+/*=============================================================================
+	SColorWheel.cpp: Implements the SColorWheel class.
+=============================================================================*/
+
+#include "SlatePrivatePCH.h"
+
+
+/* SColorWheel methods
+ *****************************************************************************/
+
+void SColorWheel::Construct( const FArguments& InArgs )
+{
+	Image = FCoreStyle::Get().GetBrush("ColorWheel.HueValueCircle");
+	SelectorImage = FCoreStyle::Get().GetBrush("ColorWheel.Selector");
+	SelectedColor = InArgs._SelectedColor;
+
+	OnMouseCaptureBegin = InArgs._OnMouseCaptureBegin;
+	OnMouseCaptureEnd = InArgs._OnMouseCaptureEnd;
+	OnValueChanged = InArgs._OnValueChanged;
+}
+
+
+/* SWidget overrides
+ *****************************************************************************/
+
+FVector2D SColorWheel::ComputeDesiredSize( ) const
+{
+	return Image->ImageSize + SelectorImage->ImageSize;
+}
+
+
+FReply SColorWheel::OnMouseButtonDown( const FGeometry& MyGeometry, const FPointerEvent& MouseEvent )
+{
+	if (MouseEvent.GetEffectingButton() == EKeys::LeftMouseButton)
+	{
+		FVector2D LocalMouseCoordinate = MyGeometry.AbsoluteToLocal( MouseEvent.GetScreenSpacePosition());
+		FVector2D Location = LocalMouseCoordinate / (MyGeometry.Size * 0.5f) - FVector2D(1.0f, 1.0f);
+		float Radius = Location.Size();
+	
+		if (Radius > 1.f)
+		{
+			return FReply::Unhandled();
+		}
+
+		OnMouseCaptureBegin.ExecuteIfBound();
+
+		return FReply::Handled().CaptureMouse(SharedThis(this));
+	}
+	else
+	{
+		return FReply::Unhandled();
+	}
+}
+
+
+FReply SColorWheel::OnMouseButtonUp( const FGeometry& MyGeometry, const FPointerEvent& MouseEvent )
+{
+	if (MouseEvent.GetEffectingButton() == EKeys::LeftMouseButton && HasMouseCapture())
+	{
+		OnMouseCaptureEnd.ExecuteIfBound();
+
+		return FReply::Handled().ReleaseMouseCapture();
+	}
+	else
+	{
+		return FReply::Unhandled();
+	}
+}
+
+
+FReply SColorWheel::OnMouseMove( const FGeometry& MyGeometry, const FPointerEvent& MouseEvent )
+{
+	if (!HasMouseCapture())
+	{
+		return FReply::Unhandled();
+	}
+
+	FVector2D LocalMouseCoordinate = MyGeometry.AbsoluteToLocal(MouseEvent.GetScreenSpacePosition());
+	FVector2D Location = LocalMouseCoordinate / (MyGeometry.Size * 0.5f) - FVector2D(1.0f, 1.0f);
+	float Angle = FMath::Atan2(Location.Y, Location.X);
+	float Radius = FMath::Min(Location.Size(), 1.0f);
+	
+	if (Angle < 0.0f)
+	{
+		Angle += 2.0f * PI;
+	}
+
+	FLinearColor NewColor = SelectedColor.Get();
+	NewColor.R = Angle * 180.0f * INV_PI;
+	NewColor.G = Radius;
+
+	OnValueChanged.ExecuteIfBound(NewColor);
+
+	return FReply::Handled();
+}
+
+
+int32 SColorWheel::OnPaint( const FGeometry& AllottedGeometry, const FSlateRect& MyClippingRect, FSlateWindowElementList& OutDrawElements, int32 LayerId, const FWidgetStyle& InWidgetStyle, bool bParentEnabled ) const
+{
+	const bool bIsEnabled = ShouldBeEnabled(bParentEnabled);
+	const uint32 DrawEffects = bIsEnabled ? ESlateDrawEffect::None : ESlateDrawEffect::DisabledEffect;
+	
+	FSlateDrawElement::MakeBox(
+		OutDrawElements,
+		LayerId,
+		AllottedGeometry.ToPaintGeometry(),
+		Image,
+		MyClippingRect,
+		DrawEffects,
+		InWidgetStyle.GetColorAndOpacityTint() * Image->GetTint(InWidgetStyle));
+
+	FSlateDrawElement::MakeBox(
+		OutDrawElements,
+		LayerId + 1,
+		AllottedGeometry.ToPaintGeometry(CalcRelativeSelectedPosition() * AllottedGeometry.Size * 0.5f - SelectorImage->ImageSize * 0.5f, SelectorImage->ImageSize),
+		SelectorImage,
+		MyClippingRect,
+		DrawEffects,
+		InWidgetStyle.GetColorAndOpacityTint() * SelectorImage->GetTint(InWidgetStyle));
+
+	return LayerId + 1;
+}
+
+
+/* SColorWheel implementation
+ *****************************************************************************/
+
+FVector2D SColorWheel::CalcRelativeSelectedPosition( ) const
+{
+	float Hue = SelectedColor.Get().R;
+	float Saturation = SelectedColor.Get().G;
+	float Angle = Hue / 180.0f * PI;
+	float Radius = Saturation;
+
+	return FVector2D(FMath::Cos(Angle), FMath::Sin(Angle)) * Radius + FVector2D(1.0f, 1.0f);
+}
diff --git a/Engine/Source/Runtime/Slate/Private/Widgets/Docking/SDockableTab.cpp b/Engine/Source/Runtime/Slate/Private/Widgets/Docking/SDockableTab.cpp
new file mode 100644
--- /dev/null
+++ b/Engine/Source/Runtime/Slate/Private/Widgets/Docking/SDockableTab.cpp
@@ -0,0 +1,63 @@
+// Copyright 1998-2014 Epic Games, Inc. All Rights Reserved.
+
+#include "SlatePrivatePCH.h"
+
+
+void SDockableTab::SetContent( TSharedRef<SWidget> InContent )
+{
+
+}
+
+bool SDockableTab::IsActive() const
+{
+	return false;
+}
+
+bool SDockableTab::IsForeground() const
+{
+	return false;
+}
+
+TSharedRef<SWidget> SDockableTab::GetContent()
+{
+	return SNullWidget::NullWidget;
+}
+
+TSharedPtr<SDockTabStack> SDockableTab::GetParentDockTabStack() const
+{
+	return TSharedPtr<SDockTabStack>();
+}
+
+void SDockableTab::BringToFrontInParent()
+{
+
+}
+
+void SDockableTab::RemoveTabFromParent()
+{
+
+}
+
+void SDockableTab::ActivateInParent(ETabActivationCause::Type InActivationCause)
+{
+}
+
+void SDockableTab::Construct( const FArguments& InArgs )
+{
+
+}
+
+bool SDockableTab::ShouldAutosize() const
+{
+	return false;
+}
+
+void SDockableTab::SetOnTabStackMenuOpening( const FOnTabStackMenuOpening& InHandler )
+{
+}
+
+void SDockableTab::RequestCloseTab()
+{	
+
+}
+
diff --git a/Engine/Source/Runtime/Slate/Private/Widgets/Images/SImage.cpp b/Engine/Source/Runtime/Slate/Private/Widgets/Images/SImage.cpp
new file mode 100644
--- /dev/null
+++ b/Engine/Source/Runtime/Slate/Private/Widgets/Images/SImage.cpp
@@ -0,0 +1,94 @@
+// Copyright 1998-2014 Epic Games, Inc. All Rights Reserved.
+
+#include "SlatePrivatePCH.h"
+
+
+/**
+ * Construct this widget
+ *
+ * @param	InArgs	The declaration data for this widget
+ */
+void SImage::Construct( const FArguments& InArgs )
+{
+	Image = InArgs._Image;
+	ColorAndOpacity = InArgs._ColorAndOpacity;
+	OnMouseButtonDownHandler = InArgs._OnMouseButtonDown;
+}
+
+
+/**
+ * The widget should respond by populating the OutDrawElements array with FDrawElements 
+ * that represent it and any of its children.
+ *
+ * @param AllottedGeometry  The FGeometry that describes an area in which the widget should appear.
+ * @param MyClippingRect    The clipping rectangle allocated for this widget and its children.
+ * @param OutDrawElements   A list of FDrawElements to populate with the output.
+ * @param LayerId           The Layer onto which this widget should be rendered.
+ * @param InColorAndOpacity Color and Opacity to be applied to all the descendants of the widget being painted
+ * @param bParentEnabled	True if the parent of this widget is enabled.
+ *
+ * @return The maximum layer ID attained by this widget or any of its children.
+ */
+int32 SImage::OnPaint( const FGeometry& AllottedGeometry, const FSlateRect& MyClippingRect, FSlateWindowElementList& OutDrawElements, int32 LayerId, const FWidgetStyle& InWidgetStyle, bool bParentEnabled ) const
+{
+#if SLATE_HD_STATS
+	SCOPE_CYCLE_COUNTER( STAT_SlateOnPaint_SImage );
+#endif
+	const FSlateBrush* ImageBrush = Image.Get();
+
+	if ((ImageBrush != NULL) && (ImageBrush->DrawAs != ESlateBrushDrawType::NoDrawType))
+	{
+		const bool bIsEnabled = ShouldBeEnabled(bParentEnabled);
+		const uint32 DrawEffects = bIsEnabled ? ESlateDrawEffect::None : ESlateDrawEffect::DisabledEffect;
+
+		const FColor FinalColorAndOpacity( InWidgetStyle.GetColorAndOpacityTint() * ColorAndOpacity.Get().GetColor(InWidgetStyle) * ImageBrush->GetTint( InWidgetStyle ) );
+
+		FSlateDrawElement::MakeBox(OutDrawElements, LayerId, AllottedGeometry.ToPaintGeometry(), ImageBrush, MyClippingRect, DrawEffects, FinalColorAndOpacity );
+	}
+	return LayerId;
+}
+	
+/**
+ * The system calls this method to notify the widget that a mouse button was pressed within it. This event is bubbled.
+ *
+ * @param MyGeometry The Geometry of the widget receiving the event
+ * @param MouseEvent Information about the input event
+ *
+ * @return Whether the event was handled along with possible requests for the system to take action.
+ */
+FReply SImage::OnMouseButtonDown( const FGeometry& MyGeometry, const FPointerEvent& MouseEvent )
+{
+	if ( OnMouseButtonDownHandler.IsBound() )
+	{
+		// If a handler is assigned, call it.
+		return OnMouseButtonDownHandler.Execute(MyGeometry, MouseEvent);
+	}
+	else
+	{
+		// otherwise the event is unhandled.
+		return FReply::Unhandled();
+	}	
+}
+
+/**
+ * An Image's desired size is whatever the image looks best as. This is decided on a case by case basis via the Width and Height properties.
+ */
+FVector2D SImage::ComputeDesiredSize() const
+{
+	const FSlateBrush* ImageBrush = Image.Get();
+	if (ImageBrush != NULL)
+	{
+		return ImageBrush->ImageSize;
+	}
+	return FVector2D::ZeroVector;
+}
+
+void SImage::SetColorAndOpacity( const TAttribute<FSlateColor>& InColorAndOpacity )
+{
+	ColorAndOpacity = InColorAndOpacity;
+}
+
+void SImage::SetColorAndOpacity( FLinearColor InColorAndOpacity )
+{
+	ColorAndOpacity = InColorAndOpacity;
+}
diff --git a/Engine/Source/Runtime/Slate/Private/Widgets/Images/SSpinningImage.cpp b/Engine/Source/Runtime/Slate/Private/Widgets/Images/SSpinningImage.cpp
new file mode 100644
--- /dev/null
+++ b/Engine/Source/Runtime/Slate/Private/Widgets/Images/SSpinningImage.cpp
@@ -0,0 +1,44 @@
+// Copyright 1998-2014 Epic Games, Inc. All Rights Reserved.
+
+#include "SlatePrivatePCH.h"
+
+
+void SSpinningImage::Construct(const FArguments& InArgs)
+{
+	Image = InArgs._Image;
+	ColorAndOpacity = InArgs._ColorAndOpacity;
+	OnMouseButtonDownHandler = InArgs._OnMouseButtonDown;
+	
+	Sequence = FCurveSequence();
+	Curve = Sequence.AddCurve(0.f, InArgs._Period);
+	Sequence.Play();
+}
+
+int32 SSpinningImage::OnPaint( const FGeometry& AllottedGeometry, const FSlateRect& MyClippingRect, FSlateWindowElementList& OutDrawElements, int32 LayerId, const FWidgetStyle& InWidgetStyle, bool bParentEnabled ) const
+{
+	const FSlateBrush* ImageBrush = Image.Get();
+
+	if ((ImageBrush != NULL) && (ImageBrush->DrawAs != ESlateBrushDrawType::NoDrawType))
+	{
+		const bool bIsEnabled = ShouldBeEnabled(bParentEnabled);
+		const uint32 DrawEffects = bIsEnabled ? ESlateDrawEffect::None : ESlateDrawEffect::DisabledEffect;
+
+		const FColor FinalColorAndOpacity( InWidgetStyle.GetColorAndOpacityTint() * ColorAndOpacity.Get().GetColor(InWidgetStyle) * ImageBrush->GetTint( InWidgetStyle ) );
+		
+		const float Angle = Curve.GetLerpLooping() * 2.0f * PI;
+
+		FSlateDrawElement::MakeRotatedBox( 
+			OutDrawElements,
+			LayerId,
+			AllottedGeometry.ToPaintGeometry(),
+			ImageBrush,
+			MyClippingRect,
+			DrawEffects,
+			Angle,
+			TOptional<FVector2D>(), // Will auto rotate about center
+			FSlateDrawElement::RelativeToElement,
+			FinalColorAndOpacity
+			);
+	}
+	return LayerId;
+}
\ No newline at end of file
diff --git a/Engine/Source/Runtime/Slate/Private/Widgets/Images/SThrobber.cpp b/Engine/Source/Runtime/Slate/Private/Widgets/Images/SThrobber.cpp
new file mode 100644
--- /dev/null
+++ b/Engine/Source/Runtime/Slate/Private/Widgets/Images/SThrobber.cpp
@@ -0,0 +1,115 @@
+// Copyright 1998-2014 Epic Games, Inc. All Rights Reserved.
+
+#include "SlatePrivatePCH.h"
+
+
+void SThrobber::Construct(const FArguments& InArgs)
+{
+	Animate = InArgs._Animate;
+	AnimCurves = FCurveSequence();
+	{
+		for (int32 PieceIndex=0; PieceIndex < InArgs._NumPieces; ++PieceIndex)
+		{
+			ThrobberCurve.Add( AnimCurves.AddCurve(PieceIndex*0.05f,1.5f) );
+		}
+		AnimCurves.Play();
+	}
+
+	TSharedRef<SHorizontalBox> HBox = SNew(SHorizontalBox);
+	
+	for (int32 PieceIndex=0; PieceIndex < InArgs._NumPieces; ++PieceIndex)
+	{
+		HBox->AddSlot()
+		.AutoWidth()
+		[
+			SNew(SBorder)
+			.BorderImage( FStyleDefaults::GetNoBrush() )
+			.ContentScale( this, &SThrobber::GetPieceScale, PieceIndex )
+			.ColorAndOpacity( this, &SThrobber::GetPieceColor, PieceIndex )
+			.HAlign(HAlign_Center)
+			.VAlign(VAlign_Center)
+			[
+				SNew(SImage)
+				.Image( InArgs._PieceImage )
+			]
+		];
+	}
+
+	this->ChildSlot
+	[
+		HBox
+	];
+}
+
+FVector2D SThrobber::GetPieceScale(int32 PieceIndex) const
+{
+	const float Value = FMath::Sin( 2 * PI * ThrobberCurve[PieceIndex].GetLerpLooping());
+	
+	const bool bAnimateHorizontally = ((0 != (Animate & Horizontal)));
+	const bool bAnimateVertically = (0 != (Animate & Vertical));
+	
+	return FVector2D(
+		bAnimateHorizontally ? Value : 1.0f,
+		bAnimateVertically ? Value : 1.0f
+	);
+}
+
+FLinearColor SThrobber::GetPieceColor(int32 PieceIndex) const
+{
+	const bool bAnimateOpacity = (0 != (Animate & Opacity));
+	if (bAnimateOpacity)
+	{
+		const float Value = FMath::Sin( 2 * PI * ThrobberCurve[PieceIndex].GetLerpLooping());
+		return FLinearColor(1.0f,1.0f,1.0f, Value);
+	}
+	else
+	{
+		return FLinearColor::White;
+	}
+}
+
+
+// SCircularThrobber
+void SCircularThrobber::Construct(const FArguments& InArgs)
+{
+	Sequence = FCurveSequence();
+	Curve = Sequence.AddCurve(0.f, InArgs._Period);
+	Sequence.Play();
+	
+	PieceImage = InArgs._PieceImage;
+	NumPieces = InArgs._NumPieces;
+	Radius = InArgs._Radius;
+}
+
+int32 SCircularThrobber::OnPaint( const FGeometry& AllottedGeometry, const FSlateRect& MyClippingRect, FSlateWindowElementList& OutDrawElements, int32 LayerId, const FWidgetStyle& InWidgetStyle, bool bParentEnabled ) const
+{
+	const FColor FinalColorAndOpacity( InWidgetStyle.GetColorAndOpacityTint() * PieceImage->GetTint( InWidgetStyle ) );
+	
+	const float Scale = AllottedGeometry.Scale; 
+	const float OffsetX = (AllottedGeometry.Size.X * Scale * 0.5f) - (PieceImage->ImageSize.X * Scale * 0.5f);
+	const float OffsetY = (AllottedGeometry.Size.Y * Scale * 0.5f) - (PieceImage->ImageSize.Y * Scale * 0.5f);
+
+	FPaintGeometry PaintGeom = AllottedGeometry.ToPaintGeometry();
+	FVector2D Origin = PaintGeom.DrawPosition;
+	Origin.X += OffsetX;
+	Origin.Y += OffsetY;
+
+	const float DeltaAngle = NumPieces > 0 ? 2 * PI / NumPieces : 0;
+	const float Phase = Curve.GetLerpLooping() * 2 * PI;
+
+	for (int32 PieceIdx = 0; PieceIdx < NumPieces; ++PieceIdx)
+	{
+		PaintGeom.DrawPosition.X = Origin.X + FMath::Sin(DeltaAngle * PieceIdx + Phase) * OffsetX;
+		PaintGeom.DrawPosition.Y = Origin.Y + FMath::Cos(DeltaAngle * PieceIdx + Phase) * OffsetY;
+		PaintGeom.DrawSize = PieceImage->ImageSize * Scale * (PieceIdx + 1) / NumPieces;
+
+		FSlateDrawElement::MakeBox(OutDrawElements, LayerId, PaintGeom, PieceImage, MyClippingRect, ESlateDrawEffect::None, InWidgetStyle.GetColorAndOpacityTint() );
+	}
+	
+	return LayerId;
+}
+
+FVector2D SCircularThrobber::ComputeDesiredSize() const
+{
+	return FVector2D(Radius, Radius) * 2;
+}
diff --git a/Engine/Source/Runtime/Slate/Private/Widgets/Input/SButton.cpp b/Engine/Source/Runtime/Slate/Private/Widgets/Input/SButton.cpp
new file mode 100644
--- /dev/null
+++ b/Engine/Source/Runtime/Slate/Private/Widgets/Input/SButton.cpp
@@ -0,0 +1,260 @@
+// Copyright 1998-2014 Epic Games, Inc. All Rights Reserved.
+
+#include "SlatePrivatePCH.h"
+
+
+/**
+ * Construct this widget
+ *
+ * @param	InArgs	The declaration data for this widget
+ */
+void SButton::Construct( const FArguments& InArgs )
+{
+	bIsPressed = false;
+
+	// Text overrides button content. If nothing is specified, put an null widget in the button.
+	// Null content makes the button enter a special mode where it will ask to be as big as the image used for its border.
+	struct
+	{
+		TSharedRef<SWidget> operator()( const FArguments& InOpArgs ) const
+		{
+			if ((InOpArgs._Content.Widget == SNullWidget::NullWidget) && (InOpArgs._Text.IsBound() || !InOpArgs._Text.Get().IsEmpty()) )
+			{
+				return SNew(STextBlock)
+					.Text( InOpArgs._Text )
+					.TextStyle( InOpArgs._TextStyle );
+			}
+			else
+			{
+				return InOpArgs._Content.Widget;
+			}
+		}
+	} DetermineContent; 
+
+	SBorder::Construct( SBorder::FArguments()
+		.ContentScale(InArgs._ContentScale)
+		.DesiredSizeScale(InArgs._DesiredSizeScale)
+		.BorderBackgroundColor(InArgs._ButtonColorAndOpacity)
+		.ForegroundColor(InArgs._ForegroundColor)
+		.BorderImage( this, &SButton::GetBorder )
+		.HAlign( InArgs._HAlign )
+		.VAlign( InArgs._VAlign )
+		.Padding( TAttribute<FMargin>(this, &SButton::GetCombinedPadding) )
+		[
+			DetermineContent(InArgs)
+		]
+	);
+
+	ContentPadding = InArgs._ContentPadding;
+
+	/* Get pointer to the button style */
+	const FButtonStyle* const ButtonStyle = InArgs._ButtonStyle;
+
+	NormalImage = &ButtonStyle->Normal;
+	HoverImage = &ButtonStyle->Hovered;
+	PressedImage = &ButtonStyle->Pressed;
+
+	BorderPadding = ButtonStyle->NormalPadding;
+	PressedBorderPadding = ButtonStyle->PressedPadding;
+
+	bIsFocusable = InArgs._IsFocusable;
+
+	OnClicked = InArgs._OnClicked;
+
+	ClickMethod = InArgs._ClickMethod;
+
+	HoveredSound = InArgs._HoveredSoundOverride.Get(ButtonStyle->HoveredSlateSound);
+	PressedSound = InArgs._PressedSoundOverride.Get(ButtonStyle->PressedSlateSound);
+}
+
+FMargin SButton::GetCombinedPadding() const
+{
+	return ( IsPressed() )
+		? ContentPadding.Get() + PressedBorderPadding
+		: ContentPadding.Get() + BorderPadding;	
+}
+
+
+/** @return An image that represents this button's border*/
+const FSlateBrush* SButton::GetBorder() const
+{
+	if (IsPressed())
+	{
+		return PressedImage;
+	}
+	else if (IsHovered())
+	{
+		return HoverImage;
+	}
+	else
+	{
+		return NormalImage;
+	}
+}
+
+
+bool SButton::SupportsKeyboardFocus() const
+{
+	// Buttons are focusable by default
+	return bIsFocusable;
+}
+
+FReply SButton::OnKeyDown( const FGeometry& MyGeometry, const FKeyboardEvent& InKeyboardEvent )
+{
+	//see if we pressed the Enter or Spacebar keys
+	if(InKeyboardEvent.GetKey() == EKeys::Enter || InKeyboardEvent.GetKey() == EKeys::SpaceBar)
+	{
+		if(IsEnabled())
+		{
+			PlayPressedSound();
+		}		
+
+		//execute our "OnClicked" delegate, if we have one
+		if(OnClicked.IsBound() == true)
+		{
+			return OnClicked.Execute();
+		}
+	}
+
+	//if it was some other button, ignore it
+	return FReply::Unhandled();
+}
+
+
+FReply SButton::OnMouseButtonDown( const FGeometry& MyGeometry, const FPointerEvent& MouseEvent )
+{
+	FReply Reply = FReply::Unhandled();
+	if ( MouseEvent.GetEffectingButton() == EKeys::LeftMouseButton )
+	{
+		bIsPressed = true;
+
+		if(IsEnabled())
+		{
+			PlayPressedSound();
+		}
+		
+		if( ClickMethod == EButtonClickMethod::MouseDown )
+		{
+			//get the reply from the execute function
+			Reply = OnClicked.IsBound() ? OnClicked.Execute() : FReply::Handled();
+
+			//You should ALWAYS handle the OnClicked event.
+			ensure(Reply.IsEventHandled() == true);
+		}
+		else
+		{
+			//we need to capture the mouse for MouseUp events
+			Reply =  FReply::Handled().CaptureMouse( AsShared() );
+		}
+	}
+
+	//return the constructed reply
+	return Reply;
+}
+
+
+FReply SButton::OnMouseButtonDoubleClick( const FGeometry& InMyGeometry, const FPointerEvent& InMouseEvent )
+{
+	return OnMouseButtonDown( InMyGeometry, InMouseEvent );
+}
+
+
+FReply SButton::OnMouseButtonUp( const FGeometry& MyGeometry, const FPointerEvent& MouseEvent )
+{
+	FReply Reply = FReply::Unhandled();
+	if ( MouseEvent.GetEffectingButton() == EKeys::LeftMouseButton )
+	{
+		bIsPressed = false;
+
+		if( ClickMethod == EButtonClickMethod::MouseDown )
+		{
+			// NOTE: If we're configured to click on mouse-down, then we never capture the mouse thus
+			//       may never receive an OnMouseButtonUp() call.  We make sure that our bIsPressed
+			//       state is reset by overriding OnMouseLeave().
+		}
+		else
+		{
+			const bool bIsUnderMouse = MyGeometry.IsUnderLocation(MouseEvent.GetScreenSpacePosition());
+			if( bIsUnderMouse )
+			{
+				// If we were asked to allow the button to be clicked on mouse up, regardless of whether the user
+				// pressed the button down first, then we'll allow the click to proceed without an active capture
+				if( ( ClickMethod == EButtonClickMethod::MouseUp || HasMouseCapture() ) && OnClicked.IsBound() == true )
+				{
+					Reply = OnClicked.Execute();
+					ensure(Reply.IsEventHandled() == true);
+				}
+			}
+		}
+		
+		//If the user of the button didn't handle this click, then the button's
+		//default behavior handles it.
+		if(Reply.IsEventHandled() == false)
+		{
+			Reply = FReply::Handled();
+		}
+
+		//If the user hasn't requested a new mouse captor, then the default
+		//behavior of the button is to release mouse capture.
+		if(Reply.GetMouseCaptor().IsValid() == false)
+		{
+			Reply.ReleaseMouseCapture();
+		}
+	}
+
+	return Reply;
+}
+
+
+FReply SButton::OnMouseMove( const FGeometry& MyGeometry, const FPointerEvent& MouseEvent )
+{
+	return FReply::Unhandled();
+}
+
+void SButton::OnMouseEnter( const FGeometry& MyGeometry, const FPointerEvent& MouseEvent )
+{
+	if(IsEnabled())
+	{
+		PlayHoverSound();
+	}
+	
+	SBorder::OnMouseEnter( MyGeometry, MouseEvent );
+}
+
+void SButton::OnMouseLeave( const FPointerEvent& MouseEvent )
+{
+	// Call parent implementation
+	SWidget::OnMouseLeave( MouseEvent );
+
+	// If we're setup to click on mouse-down, then we never capture the mouse and may not receive a
+	// mouse up event, so we need to make sure our pressed state is reset properly here
+	if( ClickMethod == EButtonClickMethod::MouseDown )
+	{
+		bIsPressed = false;
+	}
+}
+
+
+void SButton::PlayPressedSound() const
+{
+	FSlateApplication::Get().PlaySound( PressedSound );
+}
+
+void SButton::PlayHoverSound() const
+{
+	FSlateApplication::Get().PlaySound( HoveredSound );
+}
+
+FVector2D SButton::ComputeDesiredSize() const
+{
+	// When there is no widget in the button, it sizes itself based on
+	// the border image specified by the style.
+	if (ChildSlot.Widget == SNullWidget::NullWidget)
+	{
+		return GetBorder()->ImageSize;
+	}
+	else
+	{
+		return SBorder::ComputeDesiredSize();
+	}
+}
\ No newline at end of file
diff --git a/Engine/Source/Runtime/Slate/Private/Widgets/Input/SCheckBox.cpp b/Engine/Source/Runtime/Slate/Private/Widgets/Input/SCheckBox.cpp
new file mode 100644
--- /dev/null
+++ b/Engine/Source/Runtime/Slate/Private/Widgets/Input/SCheckBox.cpp
@@ -0,0 +1,340 @@
+// Copyright 1998-2014 Epic Games, Inc. All Rights Reserved.
+
+#include "SlatePrivatePCH.h"
+
+
+/**
+ * Construct this widget
+ *
+ * @param	InArgs	The declaration data for this widget
+ */
+void SCheckBox::Construct( const SCheckBox::FArguments& InArgs )
+{
+	check(InArgs._Style != NULL);
+	TAttribute<FMargin> Padding = InArgs._Padding.IsSet() ? InArgs._Padding : InArgs._Style->Padding;
+	TAttribute<FSlateColor> ForegroundColor = InArgs._ForegroundColor.IsSet() ? InArgs._ForegroundColor : InArgs._Style->ForegroundColor;
+	TAttribute<FSlateColor> BorderBackgroundColor = InArgs._BorderBackgroundColor.IsSet() ? InArgs._BorderBackgroundColor : InArgs._Style->BorderBackgroundColor;
+	ESlateCheckBoxType::Type CheckBoxType = InArgs._Type.IsSet() ? InArgs._Type.GetValue() : InArgs._Style->CheckBoxType.GetValue();
+
+	bIsPressed = false;
+	bReadOnly = InArgs._ReadOnly.Get();
+	bIsFocusable = InArgs._IsFocusable;
+
+	if( CheckBoxType == ESlateCheckBoxType::CheckBox )
+	{
+		// Check boxes use a separate check button to the side of the user's content (often, a text label or icon.)
+		this->ChildSlot
+		[
+			SNew(SHorizontalBox)
+			// Make sure we aren't trying to compute the desired size when the check box is collapsed
+			.Visibility(InArgs._Visibility)
+			+SHorizontalBox::Slot()
+			.AutoWidth()
+			.VAlign( VAlign_Center )
+			.HAlign( HAlign_Center )
+			[
+				SNew(SImage)
+				.Image( this, &SCheckBox::OnGetCheckImage )
+				.ColorAndOpacity( ForegroundColor )
+			]
+			+SHorizontalBox::Slot()
+			.AutoWidth()
+			.Padding( Padding )
+			.VAlign( VAlign_Center )
+			[
+				InArgs._Content.Widget
+			]
+		];
+	}
+	else if( ensure( CheckBoxType == ESlateCheckBoxType::ToggleButton ) )
+	{
+
+		// Toggle buttons have a visual appearance that is similar to a Slate button
+		this->ChildSlot
+		[
+			SNew( SBorder )
+			// Make sure we aren't trying to compute the desired size when the check box is collapsed
+			.Visibility(InArgs._Visibility)
+			// Bind the border background to our method that gets a slate brush for the current state of the control
+			.BorderImage( this, &SCheckBox::OnGetCheckImage )
+			// Let the user decide the padding amount
+			.Padding( Padding )
+			// Set the user's incoming content as the content of our border
+			.ForegroundColor( ForegroundColor )
+			// Set the color of the border image
+			.BorderBackgroundColor( BorderBackgroundColor )
+			// Content alignment
+			.HAlign( InArgs._HAlign )
+			[
+				InArgs._Content.Widget
+			]
+		];
+	}
+
+	/* Set images to use in various SCheckBox states */
+	UncheckedImage = &InArgs._Style->UncheckedImage;
+	UncheckedHoveredImage = &InArgs._Style->UncheckedHoveredImage;
+	UncheckedPressedImage = &InArgs._Style->UncheckedPressedImage;
+	CheckedImage = &InArgs._Style->CheckedImage;
+	CheckedHoveredImage = &InArgs._Style->CheckedHoveredImage;
+	CheckedPressedImage = &InArgs._Style->CheckedPressedImage;
+	UndeterminedImage = &InArgs._Style->UndeterminedImage;
+	UndeterminedHoveredImage = &InArgs._Style->UndeterminedHoveredImage;
+	UndeterminedPressedImage = &InArgs._Style->UndeterminedPressedImage;
+
+	IsCheckboxChecked = InArgs._IsChecked;
+	OnCheckStateChanged = InArgs._OnCheckStateChanged;
+
+	ClickMethod = InArgs._ClickMethod.Get();
+
+	OnGetMenuContent = InArgs._OnGetMenuContent;
+
+	HoveredSound = InArgs._HoveredSoundOverride.Get(InArgs._Style->HoveredSlateSound);
+	CheckedSound = InArgs._CheckedSoundOverride.Get(InArgs._Style->CheckedSlateSound);
+	UncheckedSound = InArgs._UncheckedSoundOverride.Get(InArgs._Style->UncheckedSlateSound);
+}
+
+/**
+ * See SWidget::SupportsKeyboardFocus().
+ *
+ * @return  True if this widget can take keyboard focus
+ */
+bool SCheckBox::SupportsKeyboardFocus() const
+{
+	// Buttons are focusable by default
+	return !bReadOnly && bIsFocusable;
+}
+
+FReply SCheckBox::OnKeyUp( const FGeometry& MyGeometry, const FKeyboardEvent& InKeyboardEvent )
+{
+	if ( InKeyboardEvent.GetKey() == EKeys::SpaceBar )
+	{
+		ToggleCheckedState();
+		const TAttribute<ESlateCheckBoxState::Type>& State = IsCheckboxChecked.Get();
+		if(State == ESlateCheckBoxState::Checked)
+		{
+			PlayCheckedSound();
+		}
+		else if(State == ESlateCheckBoxState::Unchecked)
+		{
+			PlayUncheckedSound();
+		}
+		
+		return FReply::Handled();
+	}
+
+	return FReply::Unhandled();
+}
+
+/**
+ * See SWidget::OnMouseButtonDown.
+ *
+ * @param MyGeometry The Geometry of the widget receiving the event
+ * @param MouseEvent Information about the input event
+ *
+ * @return Whether the event was handled along with possible requests for the system to take action.
+ */
+FReply SCheckBox::OnMouseButtonDown( const FGeometry& MyGeometry, const FPointerEvent& MouseEvent )
+{
+	if ( MouseEvent.GetEffectingButton() == EKeys::LeftMouseButton && !bReadOnly )
+	{
+		bIsPressed = true;
+
+		if( ClickMethod == EButtonClickMethod::MouseDown )
+		{
+			ToggleCheckedState();
+			const TAttribute<ESlateCheckBoxState::Type>& State = IsCheckboxChecked.Get();
+			if(State == ESlateCheckBoxState::Checked)
+			{
+				PlayCheckedSound();
+			}
+			else if(State == ESlateCheckBoxState::Unchecked)
+			{
+				PlayUncheckedSound();
+			}
+
+			// Set focus to this button, but don't capture the mouse
+			return FReply::Handled().SetKeyboardFocus( AsShared(), EKeyboardFocusCause::Mouse );
+		}
+		else
+		{
+			// Capture the mouse, and also set focus to this button
+			return FReply::Handled().CaptureMouse( AsShared() ).SetKeyboardFocus( AsShared(), EKeyboardFocusCause::Mouse );
+		}
+	}
+	else if ( MouseEvent.GetEffectingButton() == EKeys::RightMouseButton && OnGetMenuContent.IsBound() )
+	{
+		FSlateApplication::Get().PushMenu(
+			AsShared(),
+			OnGetMenuContent.Execute(),
+			MouseEvent.GetScreenSpacePosition(),
+			FPopupTransitionEffect( FPopupTransitionEffect::ContextMenu )
+			);
+
+		return FReply::Handled();
+	}
+	else
+	{
+		return FReply::Unhandled();
+	}
+}
+
+/**
+ * See SWidget::OnMouseButtonDoubleClick.
+ *
+ * @param MyGeometry The Geometry of the widget receiving the event
+ * @param MouseEvent Information about the input event
+ *
+ * @return Whether the event was handled along with possible requests for the system to take action.
+ */
+FReply SCheckBox::OnMouseButtonDoubleClick( const FGeometry& InMyGeometry, const FPointerEvent& InMouseEvent )
+{
+	return OnMouseButtonDown( InMyGeometry, InMouseEvent );
+}
+
+/**
+ * See SWidget::OnMouseButtonUp.
+ *
+ * @param MyGeometry The Geometry of the widget receiving the event
+ * @param MouseEvent Information about the input event
+ *
+ * @return Whether the event was handled along with possible requests for the system to take action.
+ */
+FReply SCheckBox::OnMouseButtonUp( const FGeometry& MyGeometry, const FPointerEvent& MouseEvent )
+{
+	if ( MouseEvent.GetEffectingButton() == EKeys::LeftMouseButton )
+	{
+		bIsPressed = false;
+
+		if( ClickMethod == EButtonClickMethod::MouseDown )
+		{
+			// NOTE: If we're configured to click on mouse-down, then we never capture the mouse thus
+			//       may never receive an OnMouseButtonUp() call.  We make sure that our bIsPressed
+			//       state is reset by overriding OnMouseLeave().
+		}
+		else if ( !bReadOnly )
+		{
+			const bool IsUnderMouse = MyGeometry.IsUnderLocation( MouseEvent.GetScreenSpacePosition() );
+			if ( IsUnderMouse )
+			{
+				// If we were asked to allow the button to be clicked on mouse up, regardless of whether the user
+				// pressed the button down first, then we'll allow the click to proceed without an active capture
+				if( ClickMethod == EButtonClickMethod::MouseUp || HasMouseCapture() )
+				{
+					ToggleCheckedState();
+					const TAttribute<ESlateCheckBoxState::Type>& State = IsCheckboxChecked.Get();
+					if(State == ESlateCheckBoxState::Checked)
+					{
+						PlayCheckedSound();
+					}
+					else if(State == ESlateCheckBoxState::Unchecked)
+					{
+						PlayUncheckedSound();
+					}
+				}
+			}
+		}
+
+		return FReply::Handled().ReleaseMouseCapture();
+	}
+
+	return FReply::Unhandled();
+}
+
+void SCheckBox::OnMouseEnter( const FGeometry& MyGeometry, const FPointerEvent& MouseEvent )
+{
+	PlayHoverSound();
+	SCompoundWidget::OnMouseEnter( MyGeometry, MouseEvent );
+}
+
+
+void SCheckBox::OnMouseLeave( const FPointerEvent& MouseEvent )
+{
+	// Call parent implementation
+	SWidget::OnMouseLeave( MouseEvent );
+
+	// If we're setup to click on mouse-down, then we never capture the mouse and may not receive a
+	// mouse up event, so we need to make sure our pressed state is reset properly here
+	if( ClickMethod == EButtonClickMethod::MouseDown )
+	{
+		bIsPressed = false;
+	}
+}
+
+
+/**
+ * Gets the check image to display for the current state of the check box
+ * @return	The name of the image to display
+ */
+const FSlateBrush* SCheckBox::OnGetCheckImage() const
+{
+	ESlateCheckBoxState::Type State = IsCheckboxChecked.Get();
+
+	const FSlateBrush* ImageToUse;
+	switch( State )
+	{
+		case ESlateCheckBoxState::Unchecked:
+			ImageToUse = IsPressed() ? UncheckedPressedImage : ( ( IsHovered() && !bReadOnly ) ? UncheckedHoveredImage : UncheckedImage );
+			break;
+	
+		case ESlateCheckBoxState::Checked:
+			ImageToUse = IsPressed() ? CheckedPressedImage : ( ( IsHovered() && !bReadOnly ) ? CheckedHoveredImage : CheckedImage );
+			break;
+	
+		default:
+		case ESlateCheckBoxState::Undetermined:
+			ImageToUse = IsPressed() ? UndeterminedPressedImage : ( ( IsHovered() && !bReadOnly ) ? UndeterminedHoveredImage : UndeterminedImage );
+			break;
+	}
+
+	return ImageToUse;
+}
+
+
+
+/**
+ * Toggles the checked state for this check box, fire events as needed
+ */
+void SCheckBox::ToggleCheckedState()
+{
+	const TAttribute<ESlateCheckBoxState::Type>& State = IsCheckboxChecked.Get();
+
+	// If the current check box state is checked OR undetermined we set the check box to checked.
+	if( State == ESlateCheckBoxState::Checked || State == ESlateCheckBoxState::Undetermined )
+	{
+		if ( !IsCheckboxChecked.IsBound() )
+		{
+			// When we are not bound, just toggle the current state.
+			IsCheckboxChecked.Set( ESlateCheckBoxState::Unchecked );
+		}
+
+		// The state of the check box changed.  Execute the delegate to notify users
+		OnCheckStateChanged.ExecuteIfBound( ESlateCheckBoxState::Unchecked );
+	}
+	else if( State == ESlateCheckBoxState::Unchecked )
+	{
+		if ( !IsCheckboxChecked.IsBound() )
+		{
+			// When we are not bound, just toggle the current state.
+			IsCheckboxChecked.Set( ESlateCheckBoxState::Checked );
+		}
+
+		// The state of the check box changed.  Execute the delegate to notify users
+		OnCheckStateChanged.ExecuteIfBound( ESlateCheckBoxState::Checked );
+	}
+}
+
+void SCheckBox::PlayCheckedSound() const
+{
+	FSlateApplication::Get().PlaySound( CheckedSound );
+}
+
+void SCheckBox::PlayUncheckedSound() const
+{
+	FSlateApplication::Get().PlaySound( UncheckedSound );
+}
+
+void SCheckBox::PlayHoverSound() const
+{
+	FSlateApplication::Get().PlaySound( HoveredSound );
+}
diff --git a/Engine/Source/Runtime/Slate/Private/Widgets/Input/SComboButton.cpp b/Engine/Source/Runtime/Slate/Private/Widgets/Input/SComboButton.cpp
new file mode 100644
--- /dev/null
+++ b/Engine/Source/Runtime/Slate/Private/Widgets/Input/SComboButton.cpp
@@ -0,0 +1,138 @@
+// Copyright 1998-2014 Epic Games, Inc. All Rights Reserved.
+
+#include "SlatePrivatePCH.h"
+
+
+void SComboButton::Construct( const FArguments& InArgs )
+{
+	check(InArgs._ComboButtonStyle);
+
+	// Work out which values we should use based on whether we were given an override, or should use the style's version
+	const FButtonStyle* const OurButtonStyle = InArgs._ButtonStyle ? InArgs._ButtonStyle : &InArgs._ComboButtonStyle->ButtonStyle;
+
+	MenuBorderBrush = &InArgs._ComboButtonStyle->MenuBorderBrush;
+	MenuBorderPadding = InArgs._ComboButtonStyle->MenuBorderPadding;
+	
+	OnGetMenuContent = InArgs._OnGetMenuContent;
+	OnComboBoxOpened = InArgs._OnComboBoxOpened;
+	MenuHeight = InArgs._MenuHeight;
+	MenuWidth = InArgs._MenuWidth;
+	Method = InArgs._Method;
+	ContentWidgetPtr = InArgs._MenuContent.Widget;
+	ContentScale = InArgs._ContentScale;
+	Placement = InArgs._MenuPlacement;
+
+	TSharedPtr<SHorizontalBox> HBox;
+
+	this->ChildSlot
+	.HAlign( InArgs._HAlign )
+	.VAlign( InArgs._VAlign )
+	[
+		SNew(SButton)
+		.ContentPadding(FMargin(1,0))
+		.ButtonStyle(OurButtonStyle)	
+		.ClickMethod( EButtonClickMethod::MouseDown )
+		.OnClicked( this, &SComboButton::OnButtonClicked )
+		.ContentPadding( InArgs._ContentPadding )
+		.ForegroundColor( InArgs._ForegroundColor )
+		.ButtonColorAndOpacity( InArgs._ButtonColorAndOpacity )
+		.IsFocusable( InArgs._IsFocusable )
+		[
+			// Button and down arrow on the right
+			// +-------------------+---+
+			// | Button Content    | v |
+			// +-------------------+---+
+			SAssignNew(HBox, SHorizontalBox)
+			+ SHorizontalBox::Slot()
+			.  Expose( ButtonContentSlot )
+			.  FillWidth(1)
+			.  VAlign(VAlign_Center)
+			[
+				InArgs._ButtonContent.Widget
+			]
+			+SHorizontalBox::Slot()
+			.  AutoWidth()
+			.  HAlign(HAlign_Center)
+			.  VAlign(VAlign_Center)
+			.  Padding( InArgs._HasDownArrow ? 2 : 0)
+			[
+				SNew(SImage)
+				. Visibility( InArgs._HasDownArrow ? EVisibility::Visible : EVisibility::Collapsed )
+				. Image(&InArgs._ComboButtonStyle->DownArrowImage)
+				// Inherit tinting from parent
+				. ColorAndOpacity( FSlateColor::UseForeground() )
+			]
+		]
+	];
+	
+	// The menu that pops up when we press the button.
+	// We keep this content around, and then put it into a new window when we need to pop
+	// it up.
+	SetMenuContent( InArgs._MenuContent.Widget );
+}
+
+FReply SComboButton::OnButtonClicked()
+{
+	TSharedPtr<SWidget> Content = NULL;
+	if( OnGetMenuContent.IsBound() )
+	{
+		Content = OnGetMenuContent.Execute();
+		SetMenuContent( Content.ToSharedRef() );
+	}
+
+	// Button was clicked; show the popup.
+	// Do nothing if clicking on the button also dismissed the menu, because we will end up doing the same thing twice.
+	this->SetIsOpen( ShouldOpenDueToClick() );	
+
+	// If the menu is open, execute the related delegate.
+	if( IsOpen() && OnComboBoxOpened.IsBound() )
+	{
+		OnComboBoxOpened.Execute();
+	}
+
+	// Focusing any newly-created widgets must occur after they have been added to the UI root.
+	FReply ButtonClickedReply = FReply::Handled();
+	
+	TSharedPtr<SWidget> WidgetToFocus = WidgetToFocusPtr.Pin();
+	if (!WidgetToFocus.IsValid())
+	{
+		// no explicitly focused widget, try to focus the content
+		WidgetToFocus = Content;
+	}
+
+	if (!WidgetToFocus.IsValid())
+	{
+		// no content, so try to focus the original widget set on construction
+		WidgetToFocus = ContentWidgetPtr.Pin();
+	}
+
+	if (WidgetToFocus.IsValid())
+	{
+		ButtonClickedReply.SetKeyboardFocus( WidgetToFocus.ToSharedRef(), EKeyboardFocusCause::SetDirectly );
+	}
+
+	return ButtonClickedReply;
+}
+
+
+/**
+ * Sets the content for this border
+ *
+ * @param	InContent	The widget to use as content for the border
+ */
+void SComboButton::SetMenuContent( const TSharedRef< SWidget >& InContent )
+{
+	MenuContent = 
+		// Wrap in configurable box to restrain height/width of menu
+		SNew(SBox)
+		.WidthOverride(MenuWidth)
+		.HeightOverride(MenuHeight)
+		[
+			SNew(SBorder)
+			.BorderImage(MenuBorderBrush)
+			.Padding(MenuBorderPadding)
+			[
+				InContent
+			]
+		];
+}
