Slate: Reorganized private Slate files; proper pre-compiled header for all files.

diff --git a/Engine/Source/Runtime/Slate/Private/Widgets/Input/SEditableText.cpp b/Engine/Source/Runtime/Slate/Private/Widgets/Input/SEditableText.cpp
new file mode 100644
--- /dev/null
+++ b/Engine/Source/Runtime/Slate/Private/Widgets/Input/SEditableText.cpp
@@ -0,0 +1,2052 @@
+// Copyright 1998-2014 Epic Games, Inc. All Rights Reserved.
+
+#include "SlatePrivatePCH.h"
+#include "TextEditHelper.h"
+
+/** A pointer to the editable text currently under the mouse cursor. NULL when there isn't one. */
+SEditableText* SEditableText::EditableTextUnderCursor = NULL;
+
+
+/** Constructor */
+SEditableText::SEditableText()
+	: ScrollHelper(),
+	  CaretPosition( 0 ),
+	  CaretVisualPositionSpring(),
+	  LastCaretInteractionTime( -1000.0 ),
+	  LastSelectionInteractionTime( -1000.0 ),
+	  bIsDragSelecting( false ),
+	  bWasFocusedByLastMouseDown( false ),
+	  bHasDragSelectedSinceFocused( false ),
+	  bHasDragSelectedSinceMouseDown( false ),
+	  CurrentUndoLevel( INDEX_NONE ),
+	  bIsChangingText( false ),
+	  UICommandList( new FUICommandList() )
+{
+	// Setup springs
+	FFloatSpring1D::FSpringConfig SpringConfig;
+	SpringConfig.SpringConstant = EditableTextDefs::CaretSpringConstant;
+	CaretVisualPositionSpring.SetConfig( SpringConfig );
+
+	SpringConfig.SpringConstant = EditableTextDefs::SelectionTargetSpringConstant;
+	SelectionTargetLeftSpring.SetConfig( SpringConfig );
+	SelectionTargetRightSpring.SetConfig( SpringConfig );
+}
+
+/** Destructor */
+SEditableText::~SEditableText()
+{
+	ITextInputMethodSystem* const TextInputMethodSystem = FSlateApplication::IsInitialized() ? FSlateApplication::Get().GetTextInputMethodSystem() : nullptr;
+	if(TextInputMethodSystem)
+	{
+		TextInputMethodSystem->UnregisterContext(TextInputMethodContext.ToSharedRef());
+	}
+}
+
+/**
+ * Construct this widget
+ *
+ * @param	InArgs	The declaration data for this widget
+ */
+void SEditableText::Construct( const FArguments& InArgs )
+{
+	Text = InArgs._Text;
+	HintText = InArgs._HintText;
+
+	// Set colors and font using override attributes Font and ColorAndOpacity if set, otherwise use the Style argument.
+	check (InArgs._Style);
+	Font = InArgs._Font.IsSet() ? InArgs._Font : InArgs._Style->Font;
+	ColorAndOpacity = InArgs._ColorAndOpacity.IsSet() ? InArgs._ColorAndOpacity : InArgs._Style->ColorAndOpacity;
+	BackgroundImageSelected = InArgs._BackgroundImageSelected.IsSet() ? InArgs._BackgroundImageSelected : &InArgs._Style->BackgroundImageSelected;
+	BackgroundImageSelectionTarget = InArgs._BackgroundImageSelectionTarget.IsSet() ? InArgs._BackgroundImageSelectionTarget : &InArgs._Style->BackgroundImageSelectionTarget;
+	BackgroundImageComposing = InArgs._BackgroundImageComposing.IsSet() ? InArgs._BackgroundImageComposing : &InArgs._Style->BackgroundImageComposing;	
+	CaretImage = InArgs._CaretImage.IsSet() ? InArgs._CaretImage : &InArgs._Style->CaretImage;
+
+	IsReadOnly = InArgs._IsReadOnly;
+	IsPassword = InArgs._IsPassword;
+	IsCaretMovedWhenGainFocus = InArgs._IsCaretMovedWhenGainFocus;
+	SelectAllTextWhenFocused = InArgs._SelectAllTextWhenFocused;
+	RevertTextOnEscape = InArgs._RevertTextOnEscape;
+	ClearKeyboardFocusOnCommit = InArgs._ClearKeyboardFocusOnCommit;
+	OnIsTypedCharValid = InArgs._OnIsTypedCharValid;
+	OnTextChanged = InArgs._OnTextChanged;
+	OnTextCommitted = InArgs._OnTextCommitted;
+	MinDesiredWidth = InArgs._MinDesiredWidth;
+	SelectAllTextOnCommit = InArgs._SelectAllTextOnCommit;
+
+	// Map UI commands to delegates which are called when the command should be executed
+	UICommandList->MapAction( FGenericCommands::Get().Undo,
+		FExecuteAction::CreateSP( this, &SEditableText::Undo ),
+		FCanExecuteAction::CreateSP( this, &SEditableText::CanExecuteUndo ) );
+
+	UICommandList->MapAction( FGenericCommands::Get().Cut,
+		FExecuteAction::CreateSP( this, &SEditableText::CutSelectedTextToClipboard ),
+		FCanExecuteAction::CreateSP( this, &SEditableText::CanExecuteCut ) );
+
+	UICommandList->MapAction( FGenericCommands::Get().Paste,
+		FExecuteAction::CreateSP( this, &SEditableText::PasteTextFromClipboard ),
+		FCanExecuteAction::CreateSP( this, &SEditableText::CanExecutePaste ) );
+
+	UICommandList->MapAction( FGenericCommands::Get().Copy,
+		FExecuteAction::CreateSP( this, &SEditableText::CopySelectedTextToClipboard ),
+		FCanExecuteAction::CreateSP( this, &SEditableText::CanExecuteCopy ) );
+
+	UICommandList->MapAction( FGenericCommands::Get().Delete,
+		FExecuteAction::CreateSP( this, &SEditableText::ExecuteDeleteAction ),
+		FCanExecuteAction::CreateSP( this, &SEditableText::CanExecuteDelete ) );
+
+	UICommandList->MapAction( FGenericCommands::Get().SelectAll,
+		FExecuteAction::CreateSP( this, &SEditableText::SelectAllText ),
+		FCanExecuteAction::CreateSP( this, &SEditableText::CanExecuteSelectAll ) );
+
+	// build context menu extender
+	MenuExtender = MakeShareable(new FExtender);
+	MenuExtender->AddMenuExtension("EditText", EExtensionHook::Before, TSharedPtr<FUICommandList>(), InArgs._ContextMenuExtender);
+
+	TextInputMethodContext = MakeShareable( new FTextInputMethodContext( StaticCastSharedRef<SEditableText>(AsShared()) ) );
+	ITextInputMethodSystem* const TextInputMethodSystem = FSlateApplication::Get().GetTextInputMethodSystem();
+	if(TextInputMethodSystem)
+	{
+		TextInputMethodChangeNotifier = TextInputMethodSystem->RegisterContext(TextInputMethodContext.ToSharedRef());
+	}
+	if(TextInputMethodChangeNotifier.IsValid())
+	{
+		TextInputMethodChangeNotifier->NotifyLayoutChanged(ITextInputMethodChangeNotifier::ELayoutChangeType::Created);
+	}
+}
+
+
+/**
+ * Sets the text string currently being edited 
+ *
+ * @param  InNewText  The new text string
+ */
+void SEditableText::SetText( const TAttribute< FText >& InNewText )
+{
+	const bool HasTextChanged = HasKeyboardFocus() ? !InNewText.Get().EqualTo( EditedText ) : !InNewText.Get().EqualTo( Text.Get() );
+
+	// NOTE: This will unbind any getter that is currently assigned to the Text attribute!  You should only be calling
+	//       SetText() if you know what you're doing.
+	Text = InNewText;
+
+	// Update the edited text, if the user isn't already editing it
+	if( HasKeyboardFocus() && !InNewText.Get().EqualTo(EditedText))
+	{
+		ClearSelection();
+		EditedText = FText();
+
+		// Load the new text
+		LoadText();
+
+		// Move the cursor to the end of the string
+		SetCaretPosition( EditedText.ToString().Len() );
+
+	}
+
+	if( HasTextChanged )
+	{
+		// Let outsiders know that the text content has been changed
+		OnTextChanged.ExecuteIfBound( HasKeyboardFocus() ? EditedText : Text.Get() );
+	}
+}
+
+
+void SEditableText::RestoreOriginalText()
+{
+	if( HasTextChangedFromOriginal() )
+	{
+		EditedText = OriginalText;
+
+		SaveText();
+
+		// Let outsiders know that the text content has been changed
+		OnTextChanged.ExecuteIfBound( EditedText );
+	}
+}
+
+
+bool SEditableText::HasTextChangedFromOriginal() const
+{
+	return ( !IsReadOnly.Get() && !EditedText.EqualTo(OriginalText)  );
+}
+
+
+void SEditableText::SetFont( const TAttribute< FSlateFontInfo >& InNewFont )
+{
+	Font = InNewFont;
+}
+
+
+/**
+ * Ticks this widget.  Override in derived classes, but always call the parent implementation.
+ *
+ * @param  AllottedGeometry Space allotted to this widget
+ * @param  InCurrentTime  Current absolute real time
+ * @param  InDeltaTime  Real time passed since last tick
+ */
+void SEditableText::Tick( const FGeometry& AllottedGeometry, const double InCurrentTime, const float InDeltaTime )
+{
+	// Call parent implementation.
+	SLeafWidget::Tick( AllottedGeometry, InCurrentTime, InDeltaTime );
+
+	if(TextInputMethodChangeNotifier.IsValid() && TextInputMethodContext.IsValid() && TextInputMethodContext->CachedGeometry != AllottedGeometry)
+	{
+		TextInputMethodContext->CachedGeometry = AllottedGeometry;
+		TextInputMethodChangeNotifier->NotifyLayoutChanged(ITextInputMethodChangeNotifier::ELayoutChangeType::Changed);
+	}
+
+	// We'll draw with the 'focused' look if we're either focused or we have a context menu summoned
+	const bool bShouldAppearFocused = HasKeyboardFocus() || ContextMenuWindow.IsValid();
+
+
+	const FString VisibleText = GetStringToRender();
+	const FSlateFontInfo& FontInfo = Font.Get();
+	const TSharedRef< FSlateFontMeasure > FontMeasureService = FSlateApplication::Get().GetRenderer()->GetFontMeasureService();
+	const float FontMaxCharHeight = FontMeasureService->Measure( FString(), FontInfo ).Y;
+
+
+	// Update caret position
+	{
+		// If we don't have keyboard focus, then push the caret back to the beginning of the line so that
+		// the start of the text is always scrolled into view
+		float FinalCaretTargetPosition = 0.0f;
+		if( bShouldAppearFocused )
+		{
+			FString TextBeforeCaret( VisibleText.Left( CaretPosition ) );
+			const FVector2D TextBeforeCaretSize( FontMeasureService->Measure( TextBeforeCaret, FontInfo ) );
+
+			FinalCaretTargetPosition = TextBeforeCaretSize.X;
+		}
+
+		// Spring toward the caret's position!
+		if( FSlateApplication::Get().IsRunningAtTargetFrameRate() )
+		{
+			CaretVisualPositionSpring.SetTarget( FinalCaretTargetPosition );
+			CaretVisualPositionSpring.Tick( InDeltaTime );
+		}
+		else
+		{
+			CaretVisualPositionSpring.SetPosition( FinalCaretTargetPosition );
+		}
+	}
+
+
+	// Make sure the caret is scrolled into view.  Note that even if we don't have keyboard focus (and the caret
+	// is hidden), we still do this to make sure the beginning of the string is in view.
+	{
+		const float CaretWidth = CalculateCaretWidth(FontMaxCharHeight);
+
+		// Grab the caret position in the scrolled area space
+		const float CaretLeft = CaretVisualPositionSpring.GetPosition();
+		const float CaretRight = CaretVisualPositionSpring.GetPosition() + CaretWidth;
+
+		// Figure out where the caret is in widget space
+		const float WidgetSpaceCaretLeft = ScrollHelper.FromScrollerSpace( FVector2D( CaretLeft, ScrollHelper.GetPosition().Y ) ).X;
+		const float WidgetSpaceCaretRight = ScrollHelper.FromScrollerSpace( FVector2D( CaretRight, ScrollHelper.GetPosition().Y ) ).X;
+
+		// If caret is to the left of the scrolled area, start scrolling left
+		if( WidgetSpaceCaretLeft < 0.0f )
+		{
+			ScrollHelper.SetPosition( FVector2D( CaretLeft, ScrollHelper.GetPosition().Y ) );
+		}
+
+		// If caret is to the right of the scrolled area, start scrolling right
+		else if( WidgetSpaceCaretRight > AllottedGeometry.Size.X )
+		{
+			ScrollHelper.SetPosition( FVector2D( CaretRight - AllottedGeometry.Size.X, ScrollHelper.GetPosition().Y ) );
+		}
+
+
+		// Make sure text is never scrolled out of view when it doesn't need to be!
+		{
+			// Measure the total widget of text string, plus the caret's width!
+			const float TotalTextWidth = FontMeasureService->Measure( VisibleText, FontInfo ).X + CaretWidth;
+			if( ScrollHelper.GetPosition().X > TotalTextWidth - AllottedGeometry.Size.X )
+			{
+				ScrollHelper.SetPosition( FVector2D( TotalTextWidth - AllottedGeometry.Size.X, ScrollHelper.GetPosition().Y ) );
+			}
+			if( ScrollHelper.GetPosition().X < 0.0f )
+			{
+				ScrollHelper.SetPosition( FVector2D( 0.0f, ScrollHelper.GetPosition().Y ) );
+			}
+		}
+	}
+
+
+	// Update selection 'target' effect
+	{
+		if( AnyTextSelected() )
+		{
+			const float SelectionLeftX = EditableTextDefs::SelectionRectLeftOffset + FontMeasureService->Measure( VisibleText.Left( Selection.GetMinIndex() ), FontInfo ).X;
+			const float SelectionRightX = EditableTextDefs::SelectionRectRightOffset + FontMeasureService->Measure( VisibleText.Left( Selection.GetMaxIndex() ), FontInfo ).X;
+
+			SelectionTargetLeftSpring.SetTarget( SelectionLeftX );
+			SelectionTargetRightSpring.SetTarget( SelectionRightX );
+		}
+
+		const float TimeSinceSelectionInteraction = (float)( InCurrentTime - LastSelectionInteractionTime );
+		if( TimeSinceSelectionInteraction <= EditableTextDefs::SelectionTargetEffectDuration || ( bShouldAppearFocused && AnyTextSelected() ) )
+		{
+			SelectionTargetLeftSpring.Tick( InDeltaTime );
+			SelectionTargetRightSpring.Tick( InDeltaTime );
+		}
+	}
+}
+
+bool SEditableText::FTextInputMethodContext::IsReadOnly()
+{
+	bool Return = true;
+	if(OwningWidget.IsValid())
+	{
+		const TSharedRef<SEditableText> This = OwningWidget.Pin().ToSharedRef();
+		Return = This->GetIsReadOnly();
+	}
+	return Return;
+}
+
+uint32 SEditableText::FTextInputMethodContext::GetTextLength()
+{
+	uint32 TextLength = 0;
+	if(OwningWidget.IsValid())
+	{
+		const TSharedRef<SEditableText> This = OwningWidget.Pin().ToSharedRef();
+		const FString& EditedText = This->EditedText.ToString();
+		TextLength = EditedText.Len();
+	}
+	return TextLength;
+}
+
+void SEditableText::FTextInputMethodContext::GetSelectionRange(uint32& BeginIndex, uint32& Length, ECaretPosition& OutCaretPosition)
+{
+	if(OwningWidget.IsValid())
+	{
+		const TSharedRef<SEditableText> This = OwningWidget.Pin().ToSharedRef();
+		if(This->Selection.StartIndex != This->Selection.FinishIndex)
+		{
+			BeginIndex = This->Selection.GetMinIndex();
+			Length = This->Selection.GetMaxIndex() - BeginIndex;
+		}
+		else
+		{
+			BeginIndex = This->CaretPosition;
+			Length = 0;
+		}
+
+		if(This->CaretPosition == BeginIndex + Length)
+		{
+			OutCaretPosition = ITextInputMethodContext::ECaretPosition::Ending;
+		}
+		else if(This->CaretPosition == BeginIndex)
+		{
+			OutCaretPosition = ITextInputMethodContext::ECaretPosition::Beginning;
+		}
+	}
+}
+
+void SEditableText::FTextInputMethodContext::SetSelectionRange(const uint32 BeginIndex, const uint32 Length, const ECaretPosition InCaretPosition)
+{
+	if(OwningWidget.IsValid())
+	{
+		const TSharedRef<SEditableText> This = OwningWidget.Pin().ToSharedRef();
+		const uint32 MinIndex = BeginIndex;
+		const uint32 MaxIndex = MinIndex + Length;
+
+		uint32 SelectionBeginIndex = 0;
+		uint32 SelectionEndIndex = 0;
+
+		switch (InCaretPosition)
+		{
+		case ITextInputMethodContext::ECaretPosition::Beginning:
+			{
+				SelectionBeginIndex = MaxIndex;
+				SelectionEndIndex = MinIndex;
+			}
+			break;
+		case ITextInputMethodContext::ECaretPosition::Ending:
+			{
+				SelectionBeginIndex = MinIndex;
+				SelectionEndIndex = MaxIndex;
+			}
+			break;
+		}
+
+		This->SetCaretPosition(SelectionEndIndex);
+		This->ClearSelection();
+		This->SelectText(SelectionBeginIndex);
+	}
+}
+
+void SEditableText::FTextInputMethodContext::GetTextInRange(const uint32 BeginIndex, const uint32 Length, FString& OutString)
+{
+	if(OwningWidget.IsValid())
+	{
+		const TSharedRef<SEditableText> This = OwningWidget.Pin().ToSharedRef();
+		const FString& EditedText = This->EditedText.ToString();
+		OutString = EditedText.Mid(BeginIndex, Length);
+	}
+}
+
+void SEditableText::FTextInputMethodContext::SetTextInRange(const uint32 BeginIndex, const uint32 Length, const FString& InString)
+{
+	if(OwningWidget.IsValid())
+	{
+		const TSharedRef<SEditableText> This = OwningWidget.Pin().ToSharedRef();
+		const FString& OldText = This->EditedText.ToString();
+		
+		// We don't use Start/FinishEditing text here because the whole IME operation handles that.
+		// Also, we don't want to support undo for individual characters added during an IME context
+		This->EditedText = FText::FromString(OldText.Left( BeginIndex ) + InString + OldText.Mid( BeginIndex + Length ));
+		This->SaveText();
+
+		This->OnTextChanged.ExecuteIfBound(This->EditedText);
+	}
+}
+
+int32 SEditableText::FTextInputMethodContext::GetCharacterIndexFromPoint(const FVector2D& Point)
+{
+	int32 CharacterIndex = INDEX_NONE;
+	if(OwningWidget.IsValid())
+	{
+		const TSharedRef<SEditableText> This = OwningWidget.Pin().ToSharedRef();
+		CharacterIndex = This->FindClickedCharacterIndex(CachedGeometry, Point);
+	}
+	return CharacterIndex;
+}
+
+bool SEditableText::FTextInputMethodContext::GetTextBounds(const uint32 BeginIndex, const uint32 Length, FVector2D& Position, FVector2D& Size)
+{
+	bool IsClipped = false;
+	if(OwningWidget.IsValid())
+	{
+		const TSharedRef<SEditableText> This = OwningWidget.Pin().ToSharedRef();
+
+		const FSlateFontInfo& FontInfo = This->Font.Get();
+		const FString VisibleText = This->GetStringToRender();
+		const TSharedRef< FSlateFontMeasure > FontMeasureService = FSlateApplication::Get().GetRenderer()->GetFontMeasureService();
+
+		const float TextLeft = FontMeasureService->Measure( VisibleText.Left( BeginIndex ), FontInfo ).X;
+		const float TextRight = FontMeasureService->Measure( VisibleText.Left( BeginIndex + Length ), FontInfo ).X;
+
+		const float ScrollAreaLeft = This->ScrollHelper.ToScrollerSpace( FVector2D::ZeroVector ).X;
+		const float ScrollAreaRight = This->ScrollHelper.ToScrollerSpace( CachedGeometry.Size ).X;
+
+		const int32 FirstVisibleCharacter = FontMeasureService->FindCharacterIndexAtOffset( VisibleText, FontInfo, ScrollAreaLeft );
+		const int32 LastVisibleCharacter = FontMeasureService->FindCharacterIndexAtOffset( VisibleText, FontInfo, ScrollAreaRight );
+
+		if(FirstVisibleCharacter == INDEX_NONE || LastVisibleCharacter == INDEX_NONE || static_cast<uint32>(FirstVisibleCharacter) > BeginIndex || static_cast<uint32>(LastVisibleCharacter) <= BeginIndex + Length)
+		{
+			IsClipped = true;
+		}
+
+		const float FontHeight = This->GetFontHeight();
+		const float DrawPositionY = ( CachedGeometry.Size.Y / 2 ) - ( FontHeight / 2 );
+
+		const float TextVertOffset = EditableTextDefs::TextVertOffsetPercent * FontHeight;
+
+		Position = CachedGeometry.LocalToAbsolute( This->ScrollHelper.FromScrollerSpace( FVector2D( TextLeft, DrawPositionY + TextVertOffset ) ) );
+		Size = CachedGeometry.Scale * This->ScrollHelper.SizeFromScrollerSpace( FVector2D( TextRight - TextLeft, CachedGeometry.Size.Y ) );
+	}
+	return IsClipped;
+}
+
+void SEditableText::FTextInputMethodContext::GetScreenBounds(FVector2D& Position, FVector2D& Size)
+{
+	if(OwningWidget.IsValid())
+	{
+		const TSharedRef<SEditableText> This = OwningWidget.Pin().ToSharedRef();
+		Position = CachedGeometry.AbsolutePosition;
+		Size = CachedGeometry.GetDrawSize();
+	}
+}
+
+TSharedPtr<FGenericWindow> SEditableText::FTextInputMethodContext::GetWindow()
+{
+	TSharedPtr<FGenericWindow> Window;
+	if(OwningWidget.IsValid())
+	{
+		const TSharedRef<SEditableText> This = OwningWidget.Pin().ToSharedRef();
+		const TSharedPtr<SWindow> SlateWindow = FSlateApplication::Get().FindWidgetWindow( This );
+		Window = SlateWindow.IsValid() ? SlateWindow->GetNativeWindow() : nullptr;
+	}
+	return Window;
+}
+
+void SEditableText::FTextInputMethodContext::BeginComposition()
+{
+	if(!IsComposing)
+	{
+		IsComposing = true;
+
+		if(OwningWidget.IsValid())
+		{
+			const TSharedRef<SEditableText> This = OwningWidget.Pin().ToSharedRef();
+			This->StartChangingText();
+		}
+	}
+}
+
+void SEditableText::FTextInputMethodContext::UpdateCompositionRange(const int32 InBeginIndex, const uint32 InLength)
+{
+	check(IsComposing);
+
+	CompositionBeginIndex = InBeginIndex;
+	CompositionLength = InLength;
+}
+
+void SEditableText::FTextInputMethodContext::EndComposition()
+{
+	if(IsComposing)
+	{
+		if(OwningWidget.IsValid())
+		{
+			const TSharedRef<SEditableText> This = OwningWidget.Pin().ToSharedRef();
+			This->FinishChangingText();
+		}
+
+		IsComposing = false;
+	}
+}
+
+/**
+ * Gets the height of the largest character in the font
+ *
+ * @return  The fonts height
+ */
+float SEditableText::GetFontHeight() const
+{
+	// Take a local copy of the FontInfo and remove any flags that could affect the outcome in a negative way
+	FSlateFontInfo FontCopy = Font.Get();
+
+	// Aesthetic choice: editable text can use a little more room, and should never be smaller than a W of the font.
+	const TSharedRef< FSlateFontMeasure > FontMeasureService = FSlateApplication::Get().GetRenderer()->GetFontMeasureService();
+	const FVector2D FontSize = FontMeasureService->Measure( FString(TEXT("W")), FontCopy );
+	return FontSize.Y;
+}
+
+void SEditableText::StartChangingText()
+{
+	// Never change text on read only controls! 
+	check( !IsReadOnly.Get() );
+
+	// We're starting to (potentially) change text
+	check( !bIsChangingText );
+	bIsChangingText = true;
+
+	// Save off an undo state in case we actually change the text
+	MakeUndoState( StateBeforeChangingText );
+}
+
+
+void SEditableText::FinishChangingText()
+{
+	// We're no longer changing text
+	check( bIsChangingText );
+	bIsChangingText = false;
+
+	// Has the text changed?
+	if( !EditedText.EqualTo(StateBeforeChangingText.Text) )
+	{
+		// Save text state
+		SaveText();
+
+		// Text was actually changed, so push the undo state we previously saved off
+		PushUndoState( StateBeforeChangingText );
+		
+		// We're done with this state data now.  Clear out any old data.
+		StateBeforeChangingText = FUndoState();
+
+		// Let outsiders know that the text content has been changed
+		OnTextChanged.ExecuteIfBound( EditedText );
+	}
+}
+
+
+bool SEditableText::GetIsReadOnly() const
+{
+	return IsReadOnly.Get();
+}
+
+
+void SEditableText::BackspaceChar()
+{
+	if( AnyTextSelected() )
+	{
+		// Delete selected text
+		DeleteSelectedText();
+	}
+	else
+	{
+		// Delete character to the left of the caret
+		if( CaretPosition > 0 )
+		{
+			const FString& OldText = EditedText.ToString();
+			EditedText = FText::FromString(OldText.Left( CaretPosition - 1 ) + OldText.Mid( CaretPosition ));
+
+			// Move the caret back too
+			SetCaretPosition( CaretPosition - 1 );
+		}
+	}
+}
+
+
+void SEditableText::DeleteChar()
+{
+	if( !IsReadOnly.Get() )
+	{
+		StartChangingText();
+
+		if( AnyTextSelected() )
+		{
+			// Delete selected text
+			DeleteSelectedText();
+		}
+		else
+		{
+			// Delete character to the right of the caret
+			const FString& OldText = EditedText.ToString();
+			if( CaretPosition < OldText.Len() )
+			{
+				EditedText = FText::FromString(OldText.Left( CaretPosition ) + OldText.Mid( CaretPosition + 1 ));
+			}
+		}
+
+		FinishChangingText();
+	}
+}
+
+
+bool SEditableText::CanTypeCharacter(const TCHAR CharInQuestion) const
+{
+	return !OnIsTypedCharValid.IsBound() || OnIsTypedCharValid.Execute(CharInQuestion);
+}
+
+
+void SEditableText::TypeChar( const int32 InChar )
+{
+	if( AnyTextSelected() )
+	{
+		// Delete selected text
+		DeleteSelectedText();
+	}
+
+	// Certain characters are not allowed
+	bool bIsCharAllowed = true;
+	{
+		if( InChar <= 0x1F )
+		{
+			bIsCharAllowed = false;
+		}
+	}
+
+	if( bIsCharAllowed )
+	{
+		// Insert character at caret position
+		const FString& OldText = EditedText.ToString();
+		EditedText = FText::FromString(OldText.Left( CaretPosition ) + InChar + OldText.Mid( CaretPosition ));
+
+		// Advance caret position
+		SetCaretPosition( CaretPosition + 1 );
+	}
+}
+
+/** When scanning forward we look at the location of the cursor. When scanning backwards we look at the location before the cursor. */
+static int32 OffsetToTestBasedOnDirection( int8 Direction )
+{
+	return (Direction >= 0) ? 0 : -1;
+}
+
+FReply SEditableText::MoveCursor( ECursorMoveMethod::Type Method, int8 Direction, ECursorAction::Type Action )
+{
+	const int32 OldCaretPosition = CaretPosition;
+	int32 NewCaretPosition = OldCaretPosition;
+
+	const int32 StringLength = EditedText.ToString().Len();
+	const FString& EditedTextAsString = EditedText.ToString();
+
+
+	// If control is held down, jump to beginning of the current word (or, the previous word, if
+	// the caret is currently on a whitespace character
+	if( Method == ECursorMoveMethod::Word )
+	{
+		if (Direction > 0)
+		{
+			// Scan right for text
+			int32 CurCharIndex = CaretPosition;
+			while( CurCharIndex < StringLength && !FChar::IsWhitespace( EditedText.ToString().GetCharArray()[ CurCharIndex ] ) )
+			{
+				++CurCharIndex;
+			}
+
+			// Scan right for whitespace
+			while( CurCharIndex < StringLength && FChar::IsWhitespace( EditedText.ToString().GetCharArray()[ CurCharIndex ] ) )
+			{
+				++CurCharIndex;
+			}
+
+			NewCaretPosition = CurCharIndex;
+		}
+		else
+		{
+			// Scan left for whitespace
+			int32 CurCharIndex = CaretPosition;
+			while( CurCharIndex > 0 && FChar::IsWhitespace( EditedText.ToString().GetCharArray()[ CurCharIndex - 1 ] ) )
+			{
+				--CurCharIndex;
+			}
+
+			// Scan left for text
+			while( CurCharIndex > 0 && !FChar::IsWhitespace( EditedText.ToString().GetCharArray()[ CurCharIndex - 1 ] ) )
+			{
+				--CurCharIndex;
+			}
+
+			NewCaretPosition = CurCharIndex;
+		}
+	}
+	else if ( Method == ECursorMoveMethod::CharacterHorizontal )
+	{
+		NewCaretPosition = FMath::Clamp( CaretPosition + Direction, 0, StringLength );
+	}
+	else
+	{
+		ensure( Method == ECursorMoveMethod::CharacterVertical );
+		return FReply::Unhandled();
+	}
+
+	SetCaretPosition( NewCaretPosition );
+
+	if( Action == ECursorAction::SelectText )
+	{
+		SelectText( OldCaretPosition );
+	}
+	else
+	{
+		ClearSelection();
+	}
+
+	return FReply::Handled();
+}
+
+
+void SEditableText::JumpTo(ETextLocation::Type JumpLocation, ECursorAction::Type Action)
+{
+	switch( JumpLocation )
+	{
+		case ETextLocation::BeginningOfLine:
+		case ETextLocation::BeginningOfDocument:
+		{
+			const int32 OldCaretPosition = CaretPosition;
+			SetCaretPosition( 0 );
+
+			if( Action == ECursorAction::SelectText )
+			{
+				SelectText( OldCaretPosition );
+			}
+			else
+			{
+				ClearSelection();
+			}
+		}
+		break;
+
+		case ETextLocation::EndOfLine:
+		case ETextLocation::EndOfDocument:
+		{
+			const int32 OldCaretPosition = CaretPosition;
+			SetCaretPosition( EditedText.ToString().Len() );
+
+			if( Action == ECursorAction::SelectText )
+			{
+				SelectText( OldCaretPosition );
+			}
+			else
+			{
+				ClearSelection();
+			}
+		}
+		break;
+	};
+}
+
+
+void SEditableText::ClearSelection()
+{
+	// Only animate clearing of the selection if there was actually a character selected
+	if( AnyTextSelected() )
+	{
+		RestartSelectionTargetAnimation();
+	}
+
+	// Explicitly invalidate the selection range indices.  The user no longer has a selection
+	// start and finish point
+	Selection.Clear();
+}
+
+
+void SEditableText::SelectAllText()
+{
+	// NOTE: Caret position is left unchanged
+
+	// NOTE: The selection anchor will always be after the last character in the string
+	
+	Selection.StartIndex = EditedText.ToString().Len();
+	Selection.FinishIndex = 0;
+	RestartSelectionTargetAnimation();
+
+	// Make sure the animated selection effect starts roughly where the cursor is
+	const TSharedRef< FSlateFontMeasure > FontMeasureService = FSlateApplication::Get().GetRenderer()->GetFontMeasureService();
+	const float SelectionTargetX = EditableTextDefs::SelectionRectLeftOffset + FontMeasureService->Measure( EditedText.ToString().Left( CaretPosition ), Font.Get() ).X;
+	SelectionTargetLeftSpring.SetPosition( SelectionTargetX );
+	SelectionTargetRightSpring.SetPosition( SelectionTargetX );
+}
+
+
+FReply SEditableText::OnEscape()
+{
+	FReply MyReply = FReply::Unhandled();
+
+	if ( AnyTextSelected() )
+	{
+		// Clear selection
+		ClearSelection();
+		MyReply = FReply::Handled();
+	}
+
+	if ( !GetIsReadOnly() )
+	{
+		// Restore the text if the revert flag is set
+		if ( RevertTextOnEscape.Get() && HasTextChangedFromOriginal() )
+		{
+			RestoreOriginalText();
+			MyReply = FReply::Handled();
+		}
+	}
+
+	return MyReply;
+}
+
+
+void SEditableText::OnEnter()
+{
+	//Always clear the local undo chain on commit.
+	ClearUndoStates();
+
+	// When enter is pressed text is committed.  Let anyone interested know about it.
+	OnTextCommitted.ExecuteIfBound( EditedText, ETextCommit::OnEnter );
+
+	// Clear keyboard focus if we were configured to do that
+	if( HasKeyboardFocus() )
+	{
+		// Reload underlying value now it is committed  (commit may alter the value) 
+		// so it can be re-displayed in the edit box if it retains focus
+		LoadText();
+
+		if ( ClearKeyboardFocusOnCommit.Get() )
+		{
+			FSlateApplication::Get().ClearKeyboardFocus( EKeyboardFocusCause::Cleared );
+		}
+		else
+		{
+			// If we aren't clearing keyboard focus, make it easy to type text again
+			SelectAllText();
+		}
+	}
+		
+	if( SelectAllTextOnCommit.Get() )
+	{
+		SelectAllText();
+	}	
+}
+
+
+bool SEditableText::CanExecuteCut() const
+{
+	bool bCanExecute = true;
+
+	// Can't execute if this is a read-only control
+	if( IsReadOnly.Get() )
+	{
+		bCanExecute = false;
+	}
+
+	// Can't execute if it's a password or there is no text selected
+	if( IsPassword.Get() || !AnyTextSelected() )
+	{
+		bCanExecute = false;
+	}
+
+	return bCanExecute;
+}
+
+
+void SEditableText::CutSelectedTextToClipboard()
+{
+	if( !IsPassword.Get() && AnyTextSelected() )
+	{
+		// Copy the text to the clipboard...
+		CopySelectedTextToClipboard();
+
+		if( !IsReadOnly.Get() )
+		{
+			StartChangingText();
+
+			// Delete the text!
+			DeleteSelectedText();
+
+			FinishChangingText();
+		}
+		else
+		{
+			// When read only, cut will simply copy the text
+		}
+	}
+}
+
+
+bool SEditableText::CanExecuteCopy() const
+{
+	bool bCanExecute = true;
+
+	// Can't execute if it's a password or there is no text selected
+	if( IsPassword.Get() || !AnyTextSelected() )
+	{
+		bCanExecute = false;
+	}
+
+	return bCanExecute;
+}
+
+
+void SEditableText::CopySelectedTextToClipboard()
+{
+	if( !IsPassword.Get() && AnyTextSelected() )
+	{
+		// Grab the selected substring
+		const FString SelectedText = EditedText.ToString().Mid( Selection.GetMinIndex(), Selection.GetMaxIndex() - Selection.GetMinIndex() );
+
+		// Copy text to clipboard
+		FPlatformMisc::ClipboardCopy( *SelectedText );
+	}
+}
+
+
+bool SEditableText::CanExecutePaste() const
+{
+	bool bCanExecute = true;
+
+	// Can't execute if this is a read-only control
+	if( IsReadOnly.Get() )
+	{
+		bCanExecute = false;
+	}
+
+	// Can't paste unless the clipboard has a string in it
+	if( !DoesClipboardHaveAnyText() )
+	{
+		bCanExecute = false;
+	}
+
+	return bCanExecute;
+}
+
+
+void SEditableText::PasteTextFromClipboard()
+{
+	if( !IsReadOnly.Get() )
+	{
+		StartChangingText();
+
+		// Delete currently selected text, if any
+		DeleteSelectedText();
+
+		// Paste from the clipboard
+		FString PastedText;
+		FPlatformMisc::ClipboardPaste(PastedText);
+
+		for( int32 CurCharIndex = 0; CurCharIndex < PastedText.Len(); ++CurCharIndex )
+		{
+			TypeChar( PastedText.GetCharArray()[ CurCharIndex ] );
+		}
+
+		FinishChangingText();
+	}
+}
+
+
+bool SEditableText::CanExecuteUndo() const
+{
+	bool bCanExecute = true;
+
+	// Can't execute if this is a read-only control
+	if( IsReadOnly.Get() )
+	{
+		bCanExecute = false;
+	}
+
+	// Can't execute unless we have some undo history
+	if( UndoStates.Num() == 0 )
+	{
+		bCanExecute = false;
+	}
+
+	if(TextInputMethodContext->IsComposing)
+	{
+		bCanExecute = false;
+	}
+
+	return bCanExecute;
+}
+
+
+void SEditableText::Undo()
+{
+	if( !IsReadOnly.Get() && UndoStates.Num() > 0 && !TextInputMethodContext->IsComposing )
+	{
+		// Restore from undo state
+		int32 UndoStateIndex;
+		if( CurrentUndoLevel == INDEX_NONE )
+		{
+			// We haven't undone anything since the last time a new undo state was added
+			UndoStateIndex = UndoStates.Num() - 1;
+
+			// Store an undo state for the current state (before undo was pressed)
+			FUndoState NewUndoState;
+			MakeUndoState( NewUndoState );
+			PushUndoState( NewUndoState );
+		}
+		else
+		{
+			// Move down to the next undo level
+			UndoStateIndex = CurrentUndoLevel - 1;
+		}
+
+		// Is there anything else to undo?
+		if( UndoStateIndex >= 0 )
+		{
+			{
+				// NOTE: It's important the no code called here creates or destroys undo states!
+				const FUndoState& UndoState = UndoStates[ UndoStateIndex ];
+		
+				EditedText = UndoState.Text;
+				CaretPosition = UndoState.CaretPosition;
+				Selection = UndoState.Selection;
+
+				SaveText();
+			}
+
+			CurrentUndoLevel = UndoStateIndex;
+
+			// Let outsiders know that the text content has been changed
+			OnTextChanged.ExecuteIfBound( EditedText );
+		}
+	}
+}
+
+
+void SEditableText::Redo()
+{
+	// Is there anything to redo?  If we've haven't tried to undo since the last time new
+	// undo state was added, then CurrentUndoLevel will be INDEX_NONE
+	if( !IsReadOnly.Get() && CurrentUndoLevel != INDEX_NONE && !(TextInputMethodContext->IsComposing))
+	{
+		const int32 NextUndoLevel = CurrentUndoLevel + 1;
+		if( UndoStates.Num() > NextUndoLevel )
+		{
+			// Restore from undo state
+			{
+				// NOTE: It's important the no code called here creates or destroys undo states!
+				const FUndoState& UndoState = UndoStates[ NextUndoLevel ];
+		
+				EditedText = UndoState.Text;
+				CaretPosition = UndoState.CaretPosition;
+				Selection = UndoState.Selection;
+
+				SaveText();
+			}
+
+			CurrentUndoLevel = NextUndoLevel;
+
+			if( UndoStates.Num() <= CurrentUndoLevel + 1 )
+			{
+				// We've redone all available undo states
+				CurrentUndoLevel = INDEX_NONE;
+
+				// Pop last undo state that we created on initial undo
+				UndoStates.RemoveAt( UndoStates.Num() - 1 );
+			}
+
+			// Let outsiders know that the text content has been changed
+			OnTextChanged.ExecuteIfBound( EditedText );
+		}
+	}
+}
+
+FVector2D SEditableText::ComputeDesiredSize() const
+{
+	const TSharedRef< FSlateFontMeasure > FontMeasureService = FSlateApplication::Get().GetRenderer()->GetFontMeasureService();
+	const FSlateFontInfo& FontInfo = Font.Get();
+	
+	const float FontMaxCharHeight = GetFontHeight();
+	const float CaretWidth = CalculateCaretWidth(FontMaxCharHeight);
+
+	FVector2D TextSize;
+
+	const FString StringToRender = GetStringToRender();
+	if( !StringToRender.IsEmpty() )
+	{
+		TextSize = FontMeasureService->Measure( StringToRender, FontInfo );
+	}
+	else
+	{
+		TextSize = FontMeasureService->Measure( HintText.Get().ToString(), FontInfo );
+	}
+	
+	const float DesiredWidth = FMath::Max(TextSize.X, MinDesiredWidth.Get()) + CaretWidth;
+	const float DesiredHeight = FMath::Max(TextSize.Y, FontMaxCharHeight);
+	return FVector2D( DesiredWidth, DesiredHeight );
+}
+
+
+int32 SEditableText::OnPaint( const FGeometry& AllottedGeometry, const FSlateRect& MyClippingRect, FSlateWindowElementList& OutDrawElements, int32 LayerId, const FWidgetStyle& InWidgetStyle, bool bParentEnabled ) const
+{
+#if SLATE_HD_STATS
+	SCOPE_CYCLE_COUNTER( STAT_SlateOnPaint_SEditableText );
+#endif
+	// The text and some effects draws in front of the widget's background and selection.
+	const int32 SelectionLayer = 0;
+	const int32 TextLayer = 1;
+	const TSharedRef< FSlateFontMeasure > FontMeasureService = FSlateApplication::Get().GetRenderer()->GetFontMeasureService();
+
+	// See if a disabled effect should be used
+	bool bEnabled = ShouldBeEnabled( bParentEnabled );
+	bool bIsReadonly = IsReadOnly.Get();
+	ESlateDrawEffect::Type DrawEffects = (bEnabled) ? ESlateDrawEffect::None : ESlateDrawEffect::DisabledEffect;
+
+	const FSlateFontInfo& FontInfo = Font.Get();
+	const FString VisibleText = GetStringToRender();
+	const FLinearColor ThisColorAndOpacity = ColorAndOpacity.Get().GetColor(InWidgetStyle);
+	const FColor ColorAndOpacitySRGB = ThisColorAndOpacity * InWidgetStyle.GetColorAndOpacityTint();
+	const float FontMaxCharHeight = GetFontHeight();
+	const double CurrentTime = FSlateApplication::Get().GetCurrentTime();
+
+	// We'll draw with the 'focused' look if we're either focused or we have a context menu summoned
+	const bool bShouldAppearFocused = HasKeyboardFocus() || ContextMenuWindow.IsValid();
+
+	// Draw selection background
+	if( AnyTextSelected() && ( bShouldAppearFocused || bIsReadonly ) )
+	{
+		// Figure out a universally visible selection color.
+		const FColor SelectionBackgroundColorAndOpacity = ( (FLinearColor::White - ThisColorAndOpacity)*0.5f + FLinearColor(-0.2f, -0.05f, 0.15f)) * InWidgetStyle.GetColorAndOpacityTint();
+
+		const float SelectionLeftX = EditableTextDefs::SelectionRectLeftOffset + FontMeasureService->Measure( VisibleText.Left( Selection.GetMinIndex() ), FontInfo ).X;
+		const float SelectionRightX = EditableTextDefs::SelectionRectRightOffset + FontMeasureService->Measure( VisibleText.Left( Selection.GetMaxIndex() ), FontInfo ).X;
+		const float SelectionTopY = 0.0f;
+		const float SelectionBottomY = AllottedGeometry.GetDrawSize().Y;
+
+		const FVector2D DrawPosition = ScrollHelper.FromScrollerSpace( FVector2D( SelectionLeftX, SelectionTopY ) );
+		const FVector2D DrawSize = ScrollHelper.SizeFromScrollerSpace( FVector2D( SelectionRightX - SelectionLeftX, SelectionBottomY - SelectionTopY ) );
+
+		FSlateDrawElement::MakeBox(
+			OutDrawElements,
+			LayerId + SelectionLayer,
+			AllottedGeometry.ToPaintGeometry(DrawPosition, DrawSize),	// Position, Size, Scale
+			BackgroundImageSelected.Get(),								// Image
+			MyClippingRect,												// Clipping rect
+			DrawEffects,												// Effects to use
+			SelectionBackgroundColorAndOpacity );						// Color
+	}
+
+	// Draw composition background
+	if( TextInputMethodContext->IsComposing )
+	{
+		const float CompositionLeft = FontMeasureService->Measure( VisibleText.Left( TextInputMethodContext->CompositionBeginIndex ), FontInfo ).X;
+		const float CompositionRight = FontMeasureService->Measure( VisibleText.Left( TextInputMethodContext->CompositionBeginIndex + TextInputMethodContext->CompositionLength ), FontInfo ).X;
+		const float CompositionTop = 0.0f;
+		const float CompositionBottom = FontMeasureService->Measure( VisibleText.Mid( TextInputMethodContext->CompositionBeginIndex, TextInputMethodContext->CompositionLength ), FontInfo ).Y;
+
+		const FVector2D DrawPosition = ScrollHelper.FromScrollerSpace( FVector2D( CompositionLeft, CompositionTop ) );
+		const FVector2D DrawSize = ScrollHelper.SizeFromScrollerSpace( FVector2D( CompositionRight - CompositionLeft, CompositionBottom - CompositionTop ) );
+
+		FSlateDrawElement::MakeBox(
+			OutDrawElements,
+			LayerId + SelectionLayer,
+			AllottedGeometry.ToPaintGeometry(DrawPosition, DrawSize),	// Position, Size, Scale
+			BackgroundImageComposing.Get(),								// Image
+			MyClippingRect,												// Clipping rect
+			DrawEffects,												// Effects to use
+			ColorAndOpacitySRGB );						// Color
+	}
+
+	const float DrawPositionY = ( AllottedGeometry.Size.Y / 2 ) - ( FontMaxCharHeight / 2 );
+	if (VisibleText.Len() == 0)
+	{
+		// Draw the hint text.
+		const FLinearColor HintTextColor = FLinearColor(ColorAndOpacitySRGB.R, ColorAndOpacitySRGB.G, ColorAndOpacitySRGB.B, 0.35f);
+		const FString ThisHintText = this->HintText.Get().ToString();
+		FSlateDrawElement::MakeText(
+			OutDrawElements,
+			LayerId + TextLayer,
+			AllottedGeometry.ToPaintGeometry( FVector2D( 0, DrawPositionY ), AllottedGeometry.Size ),
+			ThisHintText,          // Text
+			FontInfo,              // Font information (font name, size)
+			MyClippingRect,        // Clipping rect
+			DrawEffects,           // Effects to use
+			HintTextColor          // Color
+		);
+	}
+	else
+	{
+		// Draw the text
+
+		// Only draw the text that's potentially visible
+		const float ScrollAreaLeft = ScrollHelper.ToScrollerSpace( FVector2D::ZeroVector ).X;
+		const float ScrollAreaRight = ScrollHelper.ToScrollerSpace( AllottedGeometry.Size ).X;
+		const int32 FirstVisibleCharacter = FontMeasureService->FindCharacterIndexAtOffset( VisibleText, FontInfo, ScrollAreaLeft );
+		const int32 LastVisibleCharacter = FontMeasureService->FindCharacterIndexAtOffset( VisibleText, FontInfo, ScrollAreaRight );
+		const FString PotentiallyVisibleText( VisibleText.Mid( FirstVisibleCharacter, ( LastVisibleCharacter - FirstVisibleCharacter ) + 1 ) );
+		const float FirstVisibleCharacterPosition = FontMeasureService->Measure( VisibleText.Left( FirstVisibleCharacter ), FontInfo ).X;
+
+		const float TextVertOffset = EditableTextDefs::TextVertOffsetPercent * FontMaxCharHeight;
+		const FVector2D DrawPosition = ScrollHelper.FromScrollerSpace( FVector2D( FirstVisibleCharacterPosition, DrawPositionY + TextVertOffset ) );
+		const FVector2D DrawSize = ScrollHelper.SizeFromScrollerSpace( AllottedGeometry.Size );
+
+		FSlateDrawElement::MakeText(
+			OutDrawElements,
+			LayerId + TextLayer,
+			AllottedGeometry.ToPaintGeometry( DrawPosition, DrawSize ),
+			PotentiallyVisibleText,          // Text
+			FontInfo,                        // Font information (font name, size)
+			MyClippingRect,                  // Clipping rect
+			DrawEffects,                     // Effects to use
+			ColorAndOpacitySRGB              // Color
+		);
+	}
+
+
+	// Draw selection targeting effect
+	const float TimeSinceSelectionInteraction = (float)( CurrentTime - LastSelectionInteractionTime );
+	if( TimeSinceSelectionInteraction <= EditableTextDefs::SelectionTargetEffectDuration )
+	{
+		// Don't draw selection effect when drag-selecting text using the mouse
+		if( !bIsDragSelecting || !bHasDragSelectedSinceMouseDown )
+		{
+			const bool bIsClearingSelection = !AnyTextSelected() || ( !bIsReadonly && !bShouldAppearFocused );
+
+			// Compute animation progress
+			float EffectAlpha = FMath::Clamp( TimeSinceSelectionInteraction / EditableTextDefs::SelectionTargetEffectDuration, 0.0f, 1.0f );
+			EffectAlpha = 1.0f - EffectAlpha * EffectAlpha;  // Inverse square falloff (looks nicer!)
+
+			// Apply extra opacity falloff when deselecting text
+			float EffectOpacity = EffectAlpha;
+			if( bIsClearingSelection )
+			{
+				EffectOpacity *= EffectOpacity;
+			}
+
+			const FSlateBrush* SelectionTargetBrush = BackgroundImageSelectionTarget.Get();
+
+			// Set final opacity of the effect
+			FLinearColor SelectionTargetColorAndOpacity = SelectionTargetBrush->GetTint( InWidgetStyle );
+			SelectionTargetColorAndOpacity.A *= EditableTextDefs::SelectionTargetOpacity * EffectOpacity;
+			const FColor SelectionTargetColorAndOpacitySRGB = SelectionTargetColorAndOpacity;
+
+			// Compute the bounds offset of the selection target from where the selection target spring
+			// extents currently lie.  This is used to "grow" or "shrink" the selection as needed.
+			const float SelectingAnimOffset = EditableTextDefs::SelectingAnimOffsetPercent * FontMaxCharHeight;
+			const float DeselectingAnimOffset = EditableTextDefs::DeselectingAnimOffsetPercent * FontMaxCharHeight;
+
+			// Choose an offset amount depending on whether we're selecting text, or clearing selection
+			const float EffectOffset = bIsClearingSelection ? ( 1.0f - EffectAlpha ) * DeselectingAnimOffset : EffectAlpha * SelectingAnimOffset;
+
+			const float SelectionLeftX = SelectionTargetLeftSpring.GetPosition() - EffectOffset;
+			const float SelectionRightX = SelectionTargetRightSpring.GetPosition() + EffectOffset;
+			const float SelectionTopY = 0.0f - EffectOffset;
+			const float SelectionBottomY = AllottedGeometry.Size.Y + EffectOffset;
+
+			const FVector2D DrawPosition = ScrollHelper.FromScrollerSpace( FVector2D( SelectionLeftX, SelectionTopY ) );
+			const FVector2D DrawSize = ScrollHelper.SizeFromScrollerSpace( FVector2D( SelectionRightX - SelectionLeftX, SelectionBottomY - SelectionTopY ) );
+
+ 			// NOTE: We rely on scissor clipping for the selection rectangle
+			FSlateDrawElement::MakeBox(
+				OutDrawElements,
+				LayerId + TextLayer,
+				AllottedGeometry.ToPaintGeometry( DrawPosition, DrawSize ),	// Position, Size, Scale
+				SelectionTargetBrush,										// Image
+				MyClippingRect,												// Clipping rect
+				DrawEffects,												// Effects to use
+				SelectionTargetColorAndOpacitySRGB );						// Color
+		}
+	}
+
+
+	// Draw the caret!
+	if( bShouldAppearFocused )
+	{
+		// Handle caret blinking
+		FLinearColor CursorColorAndOpacity = ThisColorAndOpacity;
+		{
+			// Only blink when the user isn't actively typing
+			const double BlinkPauseEndTime = LastCaretInteractionTime + EditableTextDefs::CaretBlinkPauseTime;
+			if( CurrentTime > BlinkPauseEndTime )
+			{
+				// Caret pulsate time is relative to the last time that we stopped interacting with
+				// the cursor.  This just makes sure that the caret starts out fully opaque.
+				const double PulsateTime = CurrentTime - BlinkPauseEndTime;
+
+				// Blink the caret!
+				float CaretOpacity = FMath::MakePulsatingValue( PulsateTime, EditableTextDefs::BlinksPerSecond );
+				CaretOpacity *= CaretOpacity;	// Squared falloff, because it looks more interesting
+				CursorColorAndOpacity.A = CaretOpacity;
+			}
+		}
+
+		const float CaretX = CaretVisualPositionSpring.GetPosition();
+		const float CaretY = DrawPositionY + EditableTextDefs::CaretVertOffset;
+		const float CaretWidth = CalculateCaretWidth(FontMaxCharHeight);
+		const float CaretHeight = EditableTextDefs::CaretHeightPercent * FontMaxCharHeight;
+
+		const FVector2D CaretDrawPosition = ScrollHelper.FromScrollerSpace( FVector2D( CaretX, CaretY ) );
+		const FVector2D CaretDrawSize = ScrollHelper.SizeFromScrollerSpace( FVector2D( CaretWidth, CaretHeight ) );
+		FSlateDrawElement::MakeBox(
+			OutDrawElements,
+			LayerId + TextLayer,
+			AllottedGeometry.ToPaintGeometry( CaretDrawPosition, CaretDrawSize ),		// Position, Size, Scale
+			CaretImage.Get(),															// Image
+			MyClippingRect,																// Clipping rect
+			DrawEffects,																// Effects to use
+			CursorColorAndOpacity*InWidgetStyle.GetColorAndOpacityTint() );				// Color
+	}
+	
+	return LayerId + TextLayer;
+}
+
+
+FReply SEditableText::OnDragOver( const FGeometry& MyGeometry, const FDragDropEvent& DragDropEvent )
+{
+	TSharedPtr<FExternalDragOperation> DragDropOp = DragDropEvent.GetOperationAs<FExternalDragOperation>();
+	if ( DragDropOp.IsValid() )
+	{
+		if ( DragDropOp->HasText() )
+		{
+			return FReply::Handled();
+		}
+	}
+	
+	return FReply::Unhandled();
+}
+
+
+FReply SEditableText::OnDrop( const FGeometry& MyGeometry, const FDragDropEvent& DragDropEvent )
+{
+	TSharedPtr<FExternalDragOperation> DragDropOp = DragDropEvent.GetOperationAs<FExternalDragOperation>();
+	if ( DragDropOp.IsValid() )
+	{
+		if ( DragDropOp->HasText() )
+		{
+			this->Text = FText::FromString(DragDropOp->GetText());
+			return FReply::Handled();
+		}
+	}
+	
+	return FReply::Unhandled();
+}
+
+
+bool SEditableText::SupportsKeyboardFocus() const
+{
+	return true;
+}
+
+
+FReply SEditableText::OnKeyboardFocusReceived( const FGeometry& MyGeometry, const FKeyboardFocusEvent& InKeyboardFocusEvent )
+{
+	// Skip the focus received code if it's due to the context menu closing
+	if ( !ContextMenuWindow.IsValid() )
+	{
+		// Don't unselect text with focus is set because another widget lost focus, as that may be in response to a
+		// dismissed context menu where the user clicked an option to select text.
+		if( InKeyboardFocusEvent.GetCause() != EKeyboardFocusCause::OtherWidgetLostFocus )
+		{
+			// Deselect text, but override the last selection interaction time, so that the user doesn't
+			// see a selection transition animation when keyboard focus is received.
+			ClearSelection();
+			LastSelectionInteractionTime = -1000.0;
+		}
+
+		// Reset the cursor position spring
+		CaretVisualPositionSpring.SetPosition( CaretVisualPositionSpring.GetPosition() );
+	 
+		bHasDragSelectedSinceFocused = false;
+
+		// The user wants to edit text. Make a copy of the observed text for the user to edit.
+		LoadText();
+
+		// For mouse-based focus, we'll handle the cursor positioning in the actual mouse down/up events
+		if( InKeyboardFocusEvent.GetCause() != EKeyboardFocusCause::Mouse )
+		{
+			// Don't move the cursor if we gained focus because another widget lost focus.  It may
+			// have been a dismissed context menu where the user is expecting the selection to be retained
+			if( InKeyboardFocusEvent.GetCause() != EKeyboardFocusCause::OtherWidgetLostFocus )
+			{
+				// Select all text on keyboard focus.  This mirrors mouse focus functionality.
+				if( SelectAllTextWhenFocused.Get() )
+				{
+					SelectAllText();
+				}
+			
+				// Move the cursor to the end of the string
+				if( IsCaretMovedWhenGainFocus.Get() )
+				{
+					SetCaretPosition( EditedText.ToString().Len() );
+				}
+			}
+		}
+
+		ITextInputMethodSystem* const TextInputMethodSystem = FSlateApplication::Get().GetTextInputMethodSystem();
+		if(TextInputMethodSystem)
+		{
+			TextInputMethodSystem->ActivateContext(TextInputMethodContext.ToSharedRef());
+		}
+	}
+
+	return FReply::Handled();
+}
+
+
+void SEditableText::OnKeyboardFocusLost( const FKeyboardFocusEvent& InKeyboardFocusEvent )
+{
+	// Skip the focus lost code if it's due to the context menu opening
+	if ( !ContextMenuWindow.IsValid() )
+	{
+		ITextInputMethodSystem* const TextInputMethodSystem = FSlateApplication::Get().GetTextInputMethodSystem();
+		if(TextInputMethodSystem)
+		{
+			TextInputMethodSystem->DeactivateContext(TextInputMethodContext.ToSharedRef());
+		}
+
+		// When focus is lost let anyone who is interested that text was committed
+		if( InKeyboardFocusEvent.GetCause() != EKeyboardFocusCause::WindowActivate  ) //Don't clear selection when activating a new window or cant copy and paste on right click )
+		{
+			ClearSelection();
+		}	
+
+		// See if user explicitly tabbed away or moved focus
+		ETextCommit::Type TextAction;
+		switch ( InKeyboardFocusEvent.GetCause() )
+		{
+		case EKeyboardFocusCause::Keyboard:
+		case EKeyboardFocusCause::Mouse:
+			TextAction = ETextCommit::OnUserMovedFocus;
+			break;
+
+		case EKeyboardFocusCause::Cleared:
+			TextAction = ETextCommit::OnCleared;
+			break;
+
+		default:
+			TextAction = ETextCommit::Default;
+			break;
+		}
+
+		//Always clear the local undo chain on commit.
+		ClearUndoStates();
+
+		OnTextCommitted.ExecuteIfBound( EditedText, TextAction);
+	}
+}
+
+
+FReply SEditableText::OnKeyChar( const FGeometry& MyGeometry, const FCharacterEvent& InCharacterEvent )
+{
+	return FTextEditHelper::OnKeyChar( MyGeometry, InCharacterEvent, *this );
+}
+
+
+FReply SEditableText::OnKeyDown( const FGeometry& MyGeometry, const FKeyboardEvent& InKeyboardEvent )
+{
+	FReply Reply = FTextEditHelper::OnKeyDown( MyGeometry, InKeyboardEvent, *this );
+
+	// Process keybindings if the event wasn't already handled
+	if( !Reply.IsEventHandled() && UICommandList->ProcessCommandBindings( InKeyboardEvent ) )
+	{
+		Reply = FReply::Handled();
+	}
+
+	return Reply;
+}
+
+
+FReply SEditableText::OnKeyUp( const FGeometry& MyGeometry, const FKeyboardEvent& InKeyboardEvent )
+{
+	FReply Reply = FReply::Unhandled();
+
+	// ...
+
+	return Reply;
+}
+
+
+FReply SEditableText::OnMouseButtonDown( const FGeometry& InMyGeometry, const FPointerEvent& InMouseEvent )
+{
+	FReply Reply = FReply::Unhandled();
+
+	// If the mouse is already captured, then don't allow a new action to be taken
+	if( !HasMouseCapture() )
+	{
+		if( InMouseEvent.GetEffectingButton() == EKeys::LeftMouseButton ||
+			InMouseEvent.GetEffectingButton() == EKeys::RightMouseButton )
+		{
+			const int32 OldCaretPosition = CaretPosition;
+
+			// Am I getting focus right now?
+			const bool bIsGettingFocus = !HasKeyboardFocus();
+			if( bIsGettingFocus )
+			{
+				// We might be receiving keyboard focus due to this event.  Because the keyboard focus received callback
+				// won't fire until after this function exits, we need to make sure our widget's state is in order early
+
+				// Assume we'll be given keyboard focus, so load text for editing
+				LoadText();
+
+				// Reset 'mouse has moved' state.  We'll use this in OnMouseMove to determine whether we
+				// should reset the selection range to the caret's position.
+				bWasFocusedByLastMouseDown = true;
+			}
+
+			// Figure out where in the string the user clicked
+			const int32 ClickedCharIndex = FindClickedCharacterIndex( InMyGeometry, InMouseEvent.GetScreenSpacePosition() );
+
+			// Move the caret to this position
+			SetCaretPosition( ClickedCharIndex );
+
+			// Reset 'is drag selecting' state
+			bHasDragSelectedSinceMouseDown = false;
+
+			if( InMouseEvent.GetEffectingButton() == EKeys::LeftMouseButton )
+			{
+				if( InMouseEvent.IsShiftDown() )
+				{
+					SelectText( OldCaretPosition );
+				}
+				else
+				{
+					// Deselect any text that was selected
+					ClearSelection();
+				}
+
+				// Start drag selection
+				bIsDragSelecting = true;
+			}
+			else if( InMouseEvent.GetEffectingButton() == EKeys::RightMouseButton )
+			{
+				// If the user clicked right clicked on a character that wasn't already selected, we'll clear
+				// the selection
+				if( AnyTextSelected() &&
+					( CaretPosition < Selection.GetMinIndex() || CaretPosition > Selection.GetMaxIndex() ) )
+				{
+					// Deselect any text that was selected
+					ClearSelection();
+				}
+			}
+
+
+			// Right clicking to summon context menu, but we'll do that on mouse-up.
+			Reply = FReply::Handled();
+			Reply.CaptureMouse( AsShared() );
+			Reply.SetKeyboardFocus( AsShared(), EKeyboardFocusCause::Mouse );
+		}
+	}
+
+	return Reply;
+}
+
+
+FReply SEditableText::OnMouseButtonUp( const FGeometry& InMyGeometry, const FPointerEvent& InMouseEvent )
+{
+	FReply Reply = FReply::Unhandled();
+
+	// The mouse must have been captured by either left or right button down before we'll process mouse ups
+	if( HasMouseCapture() )
+	{
+		if( InMouseEvent.GetEffectingButton() == EKeys::LeftMouseButton &&
+			bIsDragSelecting )
+		{
+			// No longer drag-selecting
+			bIsDragSelecting = false;
+
+
+			// If we received keyboard focus on this click, then we'll want to select all of the text
+			// when the user releases the mouse button, unless the user actually dragged the mouse
+			// while holding the button down, in which case they've already selected some text and
+			// we'll leave things alone!
+			if( bWasFocusedByLastMouseDown )
+			{
+				if( !bHasDragSelectedSinceFocused )
+				{
+					if( SelectAllTextWhenFocused.Get() )
+					{
+						// User wasn't dragging the mouse, so go ahead and select all of the text now
+						// that we've become focused
+						SelectAllText();
+
+						// Move the cursor to the end of the string
+						SetCaretPosition( EditedText.ToString().Len() );
+
+						// @todo Slate: In this state, the caret should actually stay hidden (until the user interacts again), and we should not move the caret
+					}
+				}
+
+				bWasFocusedByLastMouseDown = false;
+			}
+
+			// If we selected any text while dragging, then kick off the selection target animation
+			if( bHasDragSelectedSinceMouseDown && Selection.StartIndex != Selection.FinishIndex )
+			{
+				// Start the selection animation wherever the user started the click/drag
+				const TSharedRef< FSlateFontMeasure > FontMeasureService = FSlateApplication::Get().GetRenderer()->GetFontMeasureService();
+				const float SelectionTargetX = EditableTextDefs::SelectionRectLeftOffset + FontMeasureService->Measure( EditedText.ToString().Left( Selection.StartIndex ), Font.Get() ).X;
+				SelectionTargetLeftSpring.SetPosition( SelectionTargetX );
+				SelectionTargetRightSpring.SetPosition( SelectionTargetX );
+
+				RestartSelectionTargetAnimation();
+			}
+
+
+			// Release mouse capture
+			Reply = FReply::Handled();
+			Reply.ReleaseMouseCapture();
+		}
+		else if( InMouseEvent.GetEffectingButton() == EKeys::RightMouseButton )
+		{
+			if ( InMyGeometry.IsUnderLocation(InMouseEvent.GetScreenSpacePosition()) )
+			{
+				// Right clicked, so summon a context menu if the cursor is within the widget
+				SummonContextMenu( InMouseEvent.GetScreenSpacePosition() );
+			}
+
+			// Release mouse capture
+			Reply = FReply::Handled();
+			Reply.ReleaseMouseCapture();
+		}
+	}
+	
+	return Reply;
+}
+
+
+FReply SEditableText::OnMouseMove( const FGeometry& InMyGeometry, const FPointerEvent& InMouseEvent )
+{
+	FReply Reply = FReply::Unhandled();
+
+	if( bIsDragSelecting && HasMouseCapture() )
+	{
+		// Figure out where in the string the user clicked
+		const int32 ClickedCharIndex = FindClickedCharacterIndex( InMyGeometry, InMouseEvent.GetScreenSpacePosition() );
+
+		// Move the caret to this position
+		const int32 OldCaretPosition = CaretPosition;
+		SetCaretPosition( ClickedCharIndex );
+
+
+		// Check to see if the user dragged far enough to select any characters
+		if( OldCaretPosition != CaretPosition )
+		{
+			// OK, the user has dragged over a character since we received focus
+			bHasDragSelectedSinceFocused = true;
+			bHasDragSelectedSinceMouseDown = true;
+
+			// Select text that the user dragged over
+			SelectText( OldCaretPosition );
+		}
+
+		Reply = FReply::Handled();
+	}
+
+	return Reply;
+}
+
+
+FReply SEditableText::OnMouseButtonDoubleClick( const FGeometry& InMyGeometry, const FPointerEvent& InMouseEvent )
+{
+	FReply Reply = FReply::Unhandled();
+
+	if( InMouseEvent.GetEffectingButton() == EKeys::LeftMouseButton )
+	{
+		// Figure out where in the string the user clicked
+		const int32 ClickedCharIndex = FindClickedCharacterIndex( InMyGeometry, InMouseEvent.GetScreenSpacePosition() );
+
+		// Deselect any text that was selected
+		ClearSelection();
+
+		// Select the word that the user double-clicked on, and move the caret to the end of that word
+		{
+			// Find beginning of the word
+			int32 FirstWordCharIndex = ClickedCharIndex;
+			while( FirstWordCharIndex > 0 && !FChar::IsWhitespace( EditedText.ToString().GetCharArray()[ FirstWordCharIndex - 1 ] ) )
+			{
+				--FirstWordCharIndex;
+			}
+
+			// Find end of the word
+			int32 LastWordCharIndex = ClickedCharIndex;
+			while( LastWordCharIndex < EditedText.ToString().Len() && !FChar::IsWhitespace( EditedText.ToString().GetCharArray()[ LastWordCharIndex ] ) )
+			{
+				++LastWordCharIndex;
+			}
+
+			// Select the word!
+			Selection.StartIndex = FirstWordCharIndex;
+			Selection.FinishIndex = LastWordCharIndex;
+			RestartSelectionTargetAnimation();
+
+			// Move the caret to the end of the selected word
+			SetCaretPosition( ClickedCharIndex );
+		}
+
+		Reply = FReply::Handled();
+	}
+
+	return Reply;
+}
+
+
+FCursorReply SEditableText::OnCursorQuery( const FGeometry& MyGeometry, const FPointerEvent& CursorEvent ) const
+{
+	return FCursorReply::Cursor( EMouseCursor::TextEditBeam );
+}
+
+
+void SEditableText::SelectText( const int32 InOldCaretPosition )
+{
+	if( InOldCaretPosition != CaretPosition )
+	{
+		if( !AnyTextSelected() )
+		{
+			// Start selecting
+			Selection.StartIndex = Selection.FinishIndex = InOldCaretPosition;
+
+			RestartSelectionTargetAnimation();
+		}
+		else
+		{
+			// NOTE: We intentionally don't call RestartSelectionTargetAnimation() when text was already selected,
+			//		 as the effect is often distracting when shrinking or growing an existing selection.
+		}
+		Selection.FinishIndex = CaretPosition;
+	}
+}
+
+
+void SEditableText::DeleteSelectedText()
+{
+	if( AnyTextSelected() )
+	{
+		// Move the caret to the location of the first selected character
+		CaretPosition = Selection.GetMinIndex();
+
+		// Delete the selected text
+		const FString& OldText = EditedText.ToString();
+		const FString NewText = OldText.Left( Selection.GetMinIndex() ) + OldText.Mid( Selection.GetMaxIndex() );
+		EditedText = FText::FromString(NewText);
+
+		// Visually collapse the selection target right before we clear selection
+		const TSharedRef< FSlateFontMeasure > FontMeasureService = FSlateApplication::Get().GetRenderer()->GetFontMeasureService();
+		const float SelectionTargetX = EditableTextDefs::SelectionRectLeftOffset + FontMeasureService->Measure( NewText.Left( CaretPosition ), Font.Get() ).X;
+		SelectionTargetLeftSpring.SetTarget( SelectionTargetX );
+		SelectionTargetRightSpring.SetTarget( SelectionTargetX );
+
+		// Clear selection
+		ClearSelection();
+	}
+}
+
+
+int32 SEditableText::FindClickedCharacterIndex( const FGeometry& InMyGeometry, const FVector2D& InScreenCursorPosition ) const
+{
+	// Find the click position in this widget's local space
+	const FVector2D& ScreenCursorPosition = InScreenCursorPosition;
+	const FVector2D& LocalCursorPosition = InMyGeometry.AbsoluteToLocal( ScreenCursorPosition );
+
+	// Convert the click position from the local widget space to the scrolled area
+	const FVector2D& PositionInScrollableArea = ScrollHelper.ToScrollerSpace( LocalCursorPosition );
+	const float ClampedCursorX = FMath::Max( PositionInScrollableArea.X, 0.0f );
+
+	// Figure out where in the string the user clicked
+	const TSharedRef< FSlateFontMeasure > FontMeasureService = FSlateApplication::Get().GetRenderer()->GetFontMeasureService();
+	const int32 ClickedCharIndex = FontMeasureService->FindCharacterIndexAtOffset( EditedText.ToString(), Font.Get(), ClampedCursorX );
+
+	return ClickedCharIndex;
+}
+
+
+void SEditableText::PushUndoState( const SEditableText::FUndoState& InUndoState )
+{
+	// If we've already undone some state, then we'll remove any undo state beyond the level that
+	// we've already undone up to.
+	if( CurrentUndoLevel != INDEX_NONE )
+	{
+		UndoStates.RemoveAt( CurrentUndoLevel, UndoStates.Num() - CurrentUndoLevel );
+
+		// Reset undo level as we haven't undone anything since this latest undo
+		CurrentUndoLevel = INDEX_NONE;
+	}
+
+	// Cache new undo state
+	UndoStates.Add( InUndoState );
+
+	// If we've reached the maximum number of undo levels, then trim our array
+	if( UndoStates.Num() > EditableTextDefs::MaxUndoLevels )
+	{
+		UndoStates.RemoveAt( 0 );
+	}
+}
+
+
+void SEditableText::ClearUndoStates()
+{
+	CurrentUndoLevel = INDEX_NONE;
+	UndoStates.Empty();
+}
+
+
+void SEditableText::MakeUndoState( SEditableText::FUndoState& OutUndoState )
+{
+	OutUndoState.Text = EditedText;
+	OutUndoState.CaretPosition = CaretPosition;
+	OutUndoState.Selection = Selection;
+}
+
+
+void SEditableText::LoadText()
+{
+	// Check to see if the text is actually different.  This might often be the case the first time we
+	// check it after this control is created, but otherwise should only be different if the text attribute
+	// was set (or driven) by an external source
+	const FText NewText = Text.Get();
+	
+	// Store the original text.  This functionality is off by default so we only do it when the option is enabled
+	if( RevertTextOnEscape.Get() && !NewText.EqualTo(OriginalText))
+	{
+			OriginalText = NewText;
+	}
+
+	if( !NewText.EqualTo(EditedText) )
+	{
+		// @todo Slate: We should run any sort of format validation/sanitation here before accepting the new text
+		EditedText = NewText;
+
+		int32 StringLength = EditedText.ToString().Len();
+
+		// Validate cursor positioning
+		if( CaretPosition > StringLength )
+		{
+			CaretPosition = StringLength;
+		}
+
+
+		// Validate selection
+		if( Selection.StartIndex != INDEX_NONE && Selection.StartIndex > StringLength )
+		{
+			Selection.StartIndex = StringLength;
+		}
+		if( Selection.FinishIndex != INDEX_NONE && Selection.FinishIndex > StringLength )
+		{
+			Selection.FinishIndex = StringLength;
+		}
+	}
+}
+
+
+void SEditableText::SaveText()
+{
+	// Don't set text if the text attribute has a 'getter' binding on it, otherwise we'd blow away
+	// that binding.  If there is a getter binding, then we'll assume it will provide us with
+	// updated text after we've fired our 'text changed' callbacks
+	if( !Text.IsBound() )
+	{
+		Text.Set ( EditedText );
+	}
+}
+
+
+void SEditableText::SummonContextMenu( const FVector2D& InLocation )
+{
+	// Set the menu to automatically close when the user commits to a choice
+	const bool bShouldCloseWindowAfterMenuSelection = true;
+
+#define LOCTEXT_NAMESPACE "EditableTextContextMenu"
+
+	// This is a context menu which could be summoned from within another menu if this text block is in a menu
+	// it should not close the menu it is inside
+	bool bCloseSelfOnly = true;
+	FMenuBuilder MenuBuilder( bShouldCloseWindowAfterMenuSelection, UICommandList, MenuExtender, bCloseSelfOnly, &FCoreStyle::Get() );
+	{
+		MenuBuilder.BeginSection( "EditText", LOCTEXT( "Heading", "Modify Text" ) );
+		{
+			// Undo
+			MenuBuilder.AddMenuEntry( FGenericCommands::Get().Undo );
+		}
+		MenuBuilder.EndSection();
+
+		MenuBuilder.BeginSection("EditableTextModify2");
+		{
+			// Cut
+			MenuBuilder.AddMenuEntry( FGenericCommands::Get().Cut );
+
+			// Copy
+			MenuBuilder.AddMenuEntry( FGenericCommands::Get().Copy );
+
+			// Paste
+			MenuBuilder.AddMenuEntry( FGenericCommands::Get().Paste );
+
+			// Delete
+			MenuBuilder.AddMenuEntry( FGenericCommands::Get().Delete );
+		}
+		MenuBuilder.EndSection();
+
+		MenuBuilder.BeginSection("EditableTextModify3");
+		{
+			// Select All
+			MenuBuilder.AddMenuEntry( FGenericCommands::Get().SelectAll );
+		}
+		MenuBuilder.EndSection();
+	}
+
+#undef LOCTEXT_NAMESPACE
+
+	const bool bFocusImmediatley = false;
+	TSharedPtr< SWindow > ContextMenuWindowPinned = FSlateApplication::Get().PushMenu( SharedThis( this ), MenuBuilder.MakeWidget(), InLocation, FPopupTransitionEffect( FPopupTransitionEffect::ContextMenu ), bFocusImmediatley );
+	ContextMenuWindow = ContextMenuWindowPinned;
+
+	// Make sure the window is valid.  It's possible for the parent to already be in the destroy queue, for example if the editable text was configured to dismiss it's window during OnTextCommitted.
+	if( ContextMenuWindowPinned.IsValid() )
+	{
+		ContextMenuWindowPinned->SetOnWindowClosed( FOnWindowClosed::CreateSP( this, &SEditableText::OnWindowClosed ) );
+	}
+}
+
+
+void SEditableText::ExecuteDeleteAction()
+{
+	if( !IsReadOnly.Get() && AnyTextSelected() )
+	{
+		StartChangingText();
+
+		// Delete selected text
+		DeleteSelectedText();
+
+		FinishChangingText();
+	}
+}
+
+
+bool SEditableText::CanExecuteSelectAll() const
+{
+	bool bCanExecute = true;
+
+	// Can't select all if string is empty
+	if( EditedText.IsEmpty() )
+	{
+		bCanExecute = false;
+	}
+
+	// Can't select all if entire string is already selected
+	if( Selection.GetMinIndex() == 0 &&
+		Selection.GetMaxIndex() == EditedText.ToString().Len() - 1 )
+	{
+		bCanExecute = false;
+	}
+
+	return bCanExecute;
+}
+
+
+bool SEditableText::CanExecuteDelete() const
+{
+	bool bCanExecute = true;
+
+	// Can't execute if this is a read-only control
+	if( IsReadOnly.Get() )
+	{
+		bCanExecute = false;
+	}
+
+	// Can't execute unless there is some text selected
+	if( !AnyTextSelected() )
+	{
+		bCanExecute = false;
+	}
+
+	return bCanExecute;
+}
+
+
+void SEditableText::SetCaretPosition( int32 Position )
+{
+	CaretPosition = Position;
+	LastCaretInteractionTime = FSlateApplication::Get().GetCurrentTime();
+}
+
+
+bool SEditableText::DoesClipboardHaveAnyText() const 
+{
+	FString ClipboardContent;
+	FPlatformMisc::ClipboardPaste(ClipboardContent);
+	return ClipboardContent.Len() > 0;
+}
+
+
+FString SEditableText::GetStringToRender() const
+{
+	FString VisibleText;
+
+	const bool bShouldAppearFocused = HasKeyboardFocus() || ContextMenuWindow.IsValid();
+	if(bShouldAppearFocused)
+	{
+		VisibleText = EditedText.ToString();
+	}
+	else
+	{
+		VisibleText = Text.Get().ToString();
+	}
+
+	// If this string is a password we need to replace all the characters in the string with something else
+	if ( IsPassword.Get() )
+	{
+#if PLATFORM_TCHAR_IS_1_BYTE
+		const TCHAR Dot = '*';
+#else
+		const TCHAR Dot = 0x2022;
+#endif
+		int32 VisibleTextLen = VisibleText.Len();
+		VisibleText.Empty();
+		while ( VisibleTextLen )
+		{
+			VisibleText += Dot;
+			VisibleTextLen--;
+		}
+	}
+
+	return VisibleText;
+}
+
+float SEditableText::CalculateCaretWidth(const float FontMaxCharHeight) const
+{
+	// The caret with should be at least one pixel. For very small fonts the width could be < 1 which makes it not visible
+	return FMath::Max(1.0f, EditableTextDefs::CaretWidthPercent * FontMaxCharHeight);
+}
+
+void SEditableText::RestartSelectionTargetAnimation()
+{
+	// Don't bother with this unless we're running at a decent frame rate
+	if( FSlateApplication::Get().IsRunningAtTargetFrameRate() )
+	{
+		LastSelectionInteractionTime = FSlateApplication::Get().GetCurrentTime();
+	}
+}
+
+
+void SEditableText::OnWindowClosed(const TSharedRef<SWindow>&)
+{
+	// Give this focus when the context menu has been dismissed
+	FSlateApplication::Get().SetKeyboardFocus( AsShared(), EKeyboardFocusCause::OtherWidgetLostFocus );
+}
diff --git a/Engine/Source/Runtime/Slate/Private/Widgets/Input/SEditableTextBox.cpp b/Engine/Source/Runtime/Slate/Private/Widgets/Input/SEditableTextBox.cpp
new file mode 100644
--- /dev/null
+++ b/Engine/Source/Runtime/Slate/Private/Widgets/Input/SEditableTextBox.cpp
@@ -0,0 +1,167 @@
+// Copyright 1998-2014 Epic Games, Inc. All Rights Reserved.
+
+#include "SlatePrivatePCH.h"
+
+
+/**
+ * Construct this widget
+ *
+ * @param	InArgs	The declaration data for this widget
+ */
+void SEditableTextBox::Construct( const FArguments& InArgs )
+{
+	check (InArgs._Style);
+	BorderImageNormal = &InArgs._Style->BackgroundImageNormal;
+	BorderImageHovered = &InArgs._Style->BackgroundImageHovered;
+	BorderImageFocused = &InArgs._Style->BackgroundImageFocused;
+	BorderImageReadOnly = &InArgs._Style->BackgroundImageReadOnly;
+	TAttribute<FMargin> Padding = InArgs._Padding.IsSet() ? InArgs._Padding : InArgs._Style->Padding;
+	TAttribute<FSlateFontInfo> Font = InArgs._Font.IsSet() ? InArgs._Font : InArgs._Style->Font;
+	TAttribute<FSlateColor> ForegroundColor = InArgs._ForegroundColor.IsSet() ? InArgs._ForegroundColor : InArgs._Style->ForegroundColor;
+	TAttribute<FSlateColor> BackgroundColor = InArgs._BackgroundColor.IsSet() ? InArgs._BackgroundColor : InArgs._Style->BackgroundColor;
+	ReadOnlyForegroundColor = InArgs._ReadOnlyForegroundColor.IsSet() ? InArgs._ReadOnlyForegroundColor : InArgs._Style->ReadOnlyForegroundColor;
+
+	SBorder::Construct( SBorder::FArguments()
+		.BorderImage( this, &SEditableTextBox::GetBorderImage )
+		.BorderBackgroundColor( BackgroundColor )
+		.ForegroundColor( ForegroundColor )
+		.Padding( 0 )
+		[
+			SAssignNew( Box, SHorizontalBox)
+			+SHorizontalBox::Slot()
+			.VAlign(VAlign_Fill)
+			.HAlign(HAlign_Fill)
+			.FillWidth(1)
+			[
+				SNew( SBox )
+				.Padding( Padding )
+				[
+					SAssignNew( EditableText, SEditableText )
+					.Text( InArgs._Text )
+					.HintText( InArgs._HintText )
+					.Font( Font )
+					.IsReadOnly( InArgs._IsReadOnly )
+					.IsPassword( InArgs._IsPassword )
+					.IsCaretMovedWhenGainFocus( InArgs._IsCaretMovedWhenGainFocus )
+					.SelectAllTextWhenFocused( InArgs._SelectAllTextWhenFocused.Get() )
+					.RevertTextOnEscape( InArgs._RevertTextOnEscape.Get() )
+					.ClearKeyboardFocusOnCommit( InArgs._ClearKeyboardFocusOnCommit.Get() )
+					.OnTextChanged( InArgs._OnTextChanged )
+					.OnTextCommitted( InArgs._OnTextCommitted )
+					.MinDesiredWidth( InArgs._MinDesiredWidth )
+					.SelectAllTextOnCommit( InArgs._SelectAllTextOnCommit )
+				]
+			]
+		]
+	);
+
+	ErrorReporting = InArgs._ErrorReporting;
+	if ( ErrorReporting.IsValid() )
+	{
+		Box->AddSlot()
+		.AutoWidth()
+		.Padding(3,0)
+		[
+			ErrorReporting->AsWidget()
+		];
+	}
+
+}
+
+
+/**
+ * Sets the text string currently being edited 
+ *
+ * @param  InNewText  The new text string
+ */
+void SEditableTextBox::SetText( const TAttribute< FText >& InNewText )
+{
+	EditableText->SetText( InNewText );
+}
+
+void SEditableTextBox::SetError( const FText& InError )
+{
+	SetError( InError.ToString() );
+}
+
+void SEditableTextBox::SetError( const FString& InError )
+{
+	const bool bHaveError = !InError.IsEmpty();
+
+	if ( !ErrorReporting.IsValid() )
+	{
+		// No error reporting was specified; make a default one
+		TSharedPtr<SPopupErrorText> ErrorTextWidget;
+		Box->AddSlot()
+		.AutoWidth()
+		.Padding(3,0)
+		[
+			SAssignNew( ErrorTextWidget, SPopupErrorText )
+		];
+		ErrorReporting = ErrorTextWidget;
+	}
+
+	ErrorReporting->SetError( InError );
+}
+
+/** @return Border image for the text box based on the hovered and focused state */
+const FSlateBrush* SEditableTextBox::GetBorderImage() const
+{
+	if ( EditableText->GetIsReadOnly() )
+	{
+		return BorderImageReadOnly;
+	}
+	else if ( EditableText->HasKeyboardFocus() )
+	{
+		return BorderImageFocused;
+	}
+	else
+	{
+		if ( EditableText->IsHovered() )
+		{
+			return BorderImageHovered;
+		}
+		else
+		{
+			return BorderImageNormal;
+		}
+	}
+}
+
+bool SEditableTextBox::SupportsKeyboardFocus() const
+{
+	return StaticCastSharedPtr<SWidget>(EditableText)->SupportsKeyboardFocus();
+}
+
+bool SEditableTextBox::HasKeyboardFocus() const
+{
+	// Since keyboard focus is forwarded to our editable text, we will test it instead
+	return SBorder::HasKeyboardFocus() || EditableText->HasKeyboardFocus();
+}
+
+FReply SEditableTextBox::OnKeyboardFocusReceived( const FGeometry& MyGeometry, const FKeyboardFocusEvent& InKeyboardFocusEvent )
+{
+	FReply Reply = FReply::Handled();
+
+	if ( InKeyboardFocusEvent.GetCause() != EKeyboardFocusCause::Cleared )
+	{
+		// Forward keyboard focus to our editable text widget
+		Reply.SetKeyboardFocus( EditableText.ToSharedRef(), InKeyboardFocusEvent.GetCause() );
+	}
+
+	return Reply;
+}
+
+FReply SEditableTextBox::OnKeyDown( const FGeometry& MyGeometry, const FKeyboardEvent& InKeyboardEvent )
+{
+	FKey Key = InKeyboardEvent.GetKey();
+
+	if( Key == EKeys::Escape && EditableText->HasKeyboardFocus() )
+	{
+		// Clear selection
+		EditableText->ClearSelection();
+		return FReply::Handled().SetKeyboardFocus( SharedThis( this ), EKeyboardFocusCause::Cleared );
+	}
+
+	return FReply::Unhandled();
+}
\ No newline at end of file
diff --git a/Engine/Source/Runtime/Slate/Private/Widgets/Input/SMenuAnchor.cpp b/Engine/Source/Runtime/Slate/Private/Widgets/Input/SMenuAnchor.cpp
new file mode 100644
--- /dev/null
+++ b/Engine/Source/Runtime/Slate/Private/Widgets/Input/SMenuAnchor.cpp
@@ -0,0 +1,280 @@
+// Copyright 1998-2014 Epic Games, Inc. All Rights Reserved.
+ 
+#include "SlatePrivatePCH.h"
+
+
+/**
+ * Construct this widget
+ *
+ * @param	InArgs	The declaration data for this widget
+ */
+void SMenuAnchor::Construct( const FArguments& InArgs )
+{
+	PopupLayerSlot = NULL;
+
+	this->ChildSlot
+	.Padding(0)
+	[
+		InArgs._Content.Widget
+	];
+	MenuContent = InArgs._MenuContent;
+	OnGetMenuContent = InArgs._OnGetMenuContent;
+	Placement = InArgs._Placement;
+	Method = InArgs._Method;
+}
+
+void SMenuAnchor::ComputeMenuPlacement( const FGeometry& AllottedGeometry, const FVector2D& PopupDesiredSize, FVector2D& NewWindowSize, FVector2D& NewPosition )
+{
+	const EMenuPlacement PlacementMode = Placement.Get();
+
+	if ( PlacementMode == MenuPlacement_ComboBox )
+	{
+		NewWindowSize = FVector2D( FMath::Max(AllottedGeometry.Size.X, PopupDesiredSize.X), PopupDesiredSize.Y );
+		NewWindowSize *= AllottedGeometry.Scale;
+	}
+
+	const float AbsoluteX = AllottedGeometry.AbsolutePosition.X;
+	const float AbsoluteY = AllottedGeometry.AbsolutePosition.Y;
+	FSlateRect Anchor(AbsoluteX, AbsoluteY, AbsoluteX + AllottedGeometry.Size.X*AllottedGeometry.Scale, AbsoluteY + AllottedGeometry.Size.Y*AllottedGeometry.Scale);
+	EOrientation Orientation = (PlacementMode == MenuPlacement_MenuRight) ? Orient_Horizontal : Orient_Vertical;
+	NewPosition = FSlateApplication::Get().CalculatePopupWindowPosition( Anchor, NewWindowSize, Orientation );
+}
+
+void SMenuAnchor::RequestClosePopupWindow( const TSharedRef<SWindow>& PopupWindow )
+{
+	bDismissedThisTick = true;
+	if (ensure(Method == CreateNewWindow))
+	{
+		FSlateApplication::Get().RequestDestroyWindow(PopupWindow);
+	}
+}
+
+void SMenuAnchor::OnClickedOutsidePopup()
+{
+	bDismissedThisTick = true;
+	if (ensure(Method == UseCurrentWindow))
+	{
+		FSlateApplication::Get().GetPopupSupport().UnregisterClickNotification( FOnClickedOutside::CreateSP(this, &SMenuAnchor::OnClickedOutsidePopup) );
+		SetIsOpen(false);	
+	}
+}
+
+void SMenuAnchor::Tick( const FGeometry& AllottedGeometry, const double InCurrentTime, const float InDeltaTime )
+{
+	SCompoundWidget::Tick(AllottedGeometry, InCurrentTime, InDeltaTime);
+
+	TSharedPtr<SWindow> PopupWindow = PopupWindowPtr.Pin();
+	if ( PopupWindow.IsValid() )
+	{
+		// Figure out where our attached pop-up window should be placed.				
+		if (Method == CreateNewWindow)
+		{
+			FVector2D NewPosition = FVector2D::ZeroVector;
+			FVector2D NewSize = PopupWindow->GetSizeInScreen();
+			ComputeMenuPlacement( AllottedGeometry, PopupWindow->GetContent()->GetDesiredSize(), NewSize, NewPosition );
+
+			// We made a window for showing the popup.
+			// Update the window's position!
+			if( PopupWindow->IsMorphing() )
+			{
+				FSlateRect NewShape = FSlateRect(NewPosition.X, NewPosition.Y, NewPosition.X + NewSize.X, NewPosition.Y + NewSize.Y);
+				if( NewShape != PopupWindow->GetMorphTargetShape() )
+				{
+					// Update the target shape
+					PopupWindow->UpdateMorphTargetShape( NewShape );
+					// Set size immediately if not morphing size
+					if(!PopupWindow->IsMorphingSize())
+					{
+						PopupWindow->ReshapeWindow( PopupWindow->GetPositionInScreen(), NewSize );
+					}
+				}
+			}
+			else
+			{
+				if( NewPosition != PopupWindow->GetPositionInScreen() || NewSize != PopupWindow->GetSizeInScreen() )
+				{
+					PopupWindow->ReshapeWindow( NewPosition, NewSize );
+				}
+			}
+		}
+		else
+		{
+			// Using the PopupLayer to show this popup instead of a window.
+			// Update the Popup's position so it appears anchored to the menu anchor.
+			if (ensure(PopupLayerSlot != NULL))
+			{
+				FVector2D NewPosition = FVector2D::ZeroVector;
+				FVector2D NewSize = PopupLayerSlot->GetWidget()->GetDesiredSize();
+
+				ComputeMenuPlacement( AllottedGeometry, PopupLayerSlot->GetWidget()->GetDesiredSize(), NewSize, NewPosition );
+			
+				(*PopupLayerSlot)
+					.DesktopPosition( NewPosition )
+					.WidthOverride( NewSize.X )
+					.Scale(AllottedGeometry.Scale);
+			}
+		}
+	}
+
+	/** The tick is ending, so the window was not dismissed this tick. */
+	bDismissedThisTick = false;
+}
+
+/**
+ * Open or close the popup
+ *
+ * @param InIsOpen  If true, open the popup. Otherwise close it.
+ */
+void SMenuAnchor::SetIsOpen( bool InIsOpen, const bool bFocusMenu )
+{
+	// Prevent redundant opens/closes
+	if ( IsOpen() != InIsOpen )
+	{
+		TSharedPtr< SWidget > MenuContentPtr = OnGetMenuContent.IsBound() ? OnGetMenuContent.Execute() : MenuContent;
+
+		if (InIsOpen && MenuContentPtr.IsValid() )
+		{
+			// OPEN POPUP
+			
+			// This can be called at any time so we use the push menu override that explicitly allows us to specify our parent
+
+			// Figure out where the menu anchor is on the screen, so we can set the initial position of our pop-up window
+			SlatePrepass();
+
+			// NOTE: Careful, GeneratePathToWidget can be reentrant in that it can call visibility delegates and such
+			FWidgetPath MyWidgetPath;
+			FSlateApplication::Get().GeneratePathToWidgetChecked( AsShared(), MyWidgetPath );
+			const FGeometry& MyGeometry = MyWidgetPath.Widgets.Last().Geometry;
+
+			// Figure out how big the content widget is so we can set the window's initial size properly
+			TSharedRef< SWidget > MenuContentRef = MenuContentPtr.ToSharedRef();
+			MenuContentRef->SlatePrepass();
+
+
+			// Combo-boxes never size down smaller than the widget that spawned them, but all
+			// other pop-up menus are currently auto-sized
+			const FVector2D DesiredContentSize = MenuContentRef->GetDesiredSize();  // @todo: This is ignoring any window border size!
+			const EMenuPlacement PlacementMode = Placement.Get();
+			bool bShouldAutoSize = PlacementMode != MenuPlacement_ComboBox;
+
+			const FVector2D NewPosition = MyGeometry.AbsolutePosition;
+			FVector2D NewWindowSize = DesiredContentSize;
+			const FVector2D SummonLocationSize = MyGeometry.Size;
+
+			FPopupTransitionEffect TransitionEffect( FPopupTransitionEffect::None );
+			if( PlacementMode == MenuPlacement_ComboBox )
+			{
+				TransitionEffect = FPopupTransitionEffect( FPopupTransitionEffect::ComboButton );
+				NewWindowSize = FVector2D( FMath::Max(MyGeometry.Size.X, DesiredContentSize.X), DesiredContentSize.Y );
+			}
+			else if( PlacementMode == MenuPlacement_BelowAnchor )
+			{
+				TransitionEffect = FPopupTransitionEffect( FPopupTransitionEffect::TopMenu );
+			}
+			else if( PlacementMode == MenuPlacement_MenuRight )
+			{
+				TransitionEffect = FPopupTransitionEffect( FPopupTransitionEffect::SubMenu );
+			}
+
+			if (Method == CreateNewWindow)
+			{
+				// Open the pop-up
+				TSharedRef<SWindow> PopupWindow = FSlateApplication::Get().PushMenu( AsShared(), MenuContentRef, NewPosition, TransitionEffect, bFocusMenu, bShouldAutoSize, NewWindowSize, SummonLocationSize );
+				PopupWindow->SetRequestDestroyWindowOverride( FRequestDestroyWindowOverride::CreateSP(this, &SMenuAnchor::RequestClosePopupWindow) );
+				PopupWindowPtr = PopupWindow;
+			}
+			else
+			{
+				// We are re-using the current window instead of creating a new one.
+				// The popup will be presented via an overlay service.
+				ensure(Method==UseCurrentWindow);
+				TSharedRef<SWindow> PopupWindow = MyWidgetPath.GetWindow();
+				PopupWindowPtr = PopupWindow;
+				
+				TSharedPtr<SWeakWidget> PopupWrapperPtr;
+				PopupLayerSlot = &PopupWindow->AddPopupLayerSlot()
+				.DesktopPosition(NewPosition)
+				.WidthOverride(NewWindowSize.X)
+				[
+					SAssignNew( PopupWrapperPtr, SWeakWidget )
+					.PossiblyNullContent
+					(
+						MenuContentRef
+					)
+				];
+
+				// We want to support dismissing the popup widget when the user clicks outside it.
+				FSlateApplication::Get().GetPopupSupport().RegisterClickNotification( PopupWrapperPtr.ToSharedRef(), FOnClickedOutside::CreateSP(this, &SMenuAnchor::OnClickedOutsidePopup) );
+			}
+
+			if(bFocusMenu)
+			{
+				FSlateApplication::Get().SetKeyboardFocus(MenuContentRef, EKeyboardFocusCause::SetDirectly);
+			}
+		}
+		else
+		{
+			// CLOSE POPUP
+
+			if (Method == CreateNewWindow)
+			{
+				// Close the Popup window.
+				TSharedPtr<SWindow> PopupWindow = PopupWindowPtr.Pin();
+				if ( PopupWindow.IsValid() )
+				{
+					// Request that the popup be closed.
+					PopupWindow->RequestDestroyWindow();
+				}
+			}
+			else
+			{
+				// Close the popup overlay
+				ensure(Method==UseCurrentWindow);
+				check(PopupWindowPtr.IsValid());
+				check(PopupLayerSlot != NULL);
+				PopupWindowPtr.Pin()->RemovePopupLayerSlot(PopupLayerSlot->GetWidget());
+				// We no longer have a popup open, so reset all the tracking state associated.
+				PopupLayerSlot = NULL;
+				PopupWindowPtr.Reset();
+			}
+		}
+	}
+}
+
+bool SMenuAnchor::IsOpen() const
+{
+	return PopupWindowPtr.Pin().IsValid();
+}
+
+bool SMenuAnchor::ShouldOpenDueToClick() const
+{
+	return !IsOpen() && !bDismissedThisTick;
+}
+
+FVector2D SMenuAnchor::GetMenuPosition() const
+{
+	FVector2D Pos(0,0);
+	if( PopupWindowPtr.IsValid() )
+	{
+		Pos = PopupWindowPtr.Pin()->GetPositionInScreen();
+	}
+	
+	return Pos;
+}
+
+bool SMenuAnchor::HasOpenSubMenus() const
+{
+	bool Result = false;
+	if( PopupWindowPtr.IsValid() )
+	{
+		Result = FSlateApplication::Get().HasOpenSubMenus( PopupWindowPtr.Pin().ToSharedRef() );
+	}
+	return Result;
+}
+
+SMenuAnchor::SMenuAnchor()
+	: MenuContent( SNullWidget::NullWidget )
+	, bDismissedThisTick( false )
+	, Method(CreateNewWindow)
+{
+}
\ No newline at end of file
diff --git a/Engine/Source/Runtime/Slate/Private/Widgets/Input/SRotatorInputBox.cpp b/Engine/Source/Runtime/Slate/Private/Widgets/Input/SRotatorInputBox.cpp
new file mode 100644
--- /dev/null
+++ b/Engine/Source/Runtime/Slate/Private/Widgets/Input/SRotatorInputBox.cpp
@@ -0,0 +1,100 @@
+// Copyright 1998-2014 Epic Games, Inc. All Rights Reserved.
+
+#include "SlatePrivatePCH.h"
+
+
+#define LOCTEXT_NAMESPACE "SRotatorInputBox"
+
+
+/**
+ * Construct this widget
+ *
+ * @param	InArgs	The declaration data for this widget
+ */
+void SRotatorInputBox::Construct( const SRotatorInputBox::FArguments& InArgs )
+{
+	FLinearColor LabelColorX;
+	FLinearColor LabelColorY;
+	FLinearColor LabelColorZ;
+	LabelColorX = InArgs._bColorAxisLabels ? SNumericEntryBox<float>::RedLabelBackgroundColor : FLinearColor(0.0f,0.0f,0.0f,.5f);
+	LabelColorY = InArgs._bColorAxisLabels ? SNumericEntryBox<float>::GreenLabelBackgroundColor : FLinearColor(0.0f,0.0f,0.0f,.5f);
+	LabelColorZ = InArgs._bColorAxisLabels ? SNumericEntryBox<float>::BlueLabelBackgroundColor : FLinearColor(0.0f,0.0f,0.0f,.5f);
+
+	this->ChildSlot
+	[
+		SNew(SHorizontalBox)
+		+SHorizontalBox::Slot()
+		.VAlign( VAlign_Center )
+		.FillWidth( 1.0f )
+		.Padding( 0.0f, 1.0f, 2.0f, 1.0f )
+		[
+			SNew(SNumericEntryBox<float>)
+			.AllowSpin(InArgs._AllowSpin)
+			.MinSliderValue(0.0f)
+			.MaxSliderValue(359.999f)
+			.LabelPadding(0)
+			.Label()
+			[
+				SNumericEntryBox<float>::BuildLabel( LOCTEXT("Roll_Label", "X"), FLinearColor::White, SNumericEntryBox<float>::RedLabelBackgroundColor )
+			]
+			.Font( InArgs._Font )
+			.Value( InArgs._Roll )
+			.OnValueChanged( InArgs._OnRollChanged )
+			.OnValueCommitted( InArgs._OnRollCommitted )
+			.OnBeginSliderMovement( InArgs._OnBeginSliderMovement )
+			.OnEndSliderMovement( InArgs._OnEndSliderMovement )
+			.UndeterminedString( LOCTEXT("MultipleValues", "Multiple Values") )
+			.ToolTipText( LOCTEXT("Roll_ToolTip", "Roll Value") ) 
+		]
+		+SHorizontalBox::Slot()
+		.VAlign( VAlign_Center )
+		.FillWidth( 1.0f )
+		.Padding( 0.0f, 1.0f, 2.0f, 1.0f )
+		[
+			SNew(SNumericEntryBox<float>)
+			.AllowSpin(InArgs._AllowSpin)
+			.MinSliderValue(0.0f)
+			.MaxSliderValue(359.999f)
+			.LabelPadding(0)
+			.Label()
+			[
+				SNumericEntryBox<float>::BuildLabel( LOCTEXT("Pitch_Label", "Y"), FLinearColor::White, SNumericEntryBox<float>::GreenLabelBackgroundColor )
+			]
+			.Font( InArgs._Font )
+			.Value( InArgs._Pitch )
+			.OnValueChanged( InArgs._OnPitchChanged )
+			.OnValueCommitted( InArgs._OnPitchCommitted )
+			.OnBeginSliderMovement( InArgs._OnBeginSliderMovement )
+			.OnEndSliderMovement( InArgs._OnEndSliderMovement )
+			.UndeterminedString( LOCTEXT("MultipleValues", "Multiple Values") )
+			.ToolTipText( LOCTEXT("Pitch_ToolTip", "Pitch Value") ) 
+		]
+		+SHorizontalBox::Slot()
+		.VAlign( VAlign_Center )
+		.FillWidth( 1.0f )
+		.Padding( 0.0f, 1.0f, 0.0f, 1.0f )
+		[
+			SNew(SNumericEntryBox<float>)
+			.AllowSpin(InArgs._AllowSpin)
+			.MinSliderValue(0.0f)
+			.MaxSliderValue(359.999f)
+			.LabelPadding(0)
+			.Label()
+			[
+				SNumericEntryBox<float>::BuildLabel( LOCTEXT("Yaw_Label", "Z"), FLinearColor::White, SNumericEntryBox<float>::BlueLabelBackgroundColor )
+			]
+			.Font( InArgs._Font )
+			.Value( InArgs._Yaw )
+			.OnValueChanged( InArgs._OnYawChanged )
+			.OnValueCommitted( InArgs._OnYawCommitted )
+			.OnBeginSliderMovement( InArgs._OnBeginSliderMovement )
+			.OnEndSliderMovement( InArgs._OnEndSliderMovement )
+			.UndeterminedString( LOCTEXT("MultipleValues", "Multiple Values") )
+			.ToolTipText( LOCTEXT("Yaw_ToolTip", "Yaw Value") ) 			
+		]
+	];
+
+}
+
+#undef LOCTEXT_NAMESPACE
+
diff --git a/Engine/Source/Runtime/Slate/Private/Widgets/Input/SSlider.cpp b/Engine/Source/Runtime/Slate/Private/Widgets/Input/SSlider.cpp
new file mode 100644
--- /dev/null
+++ b/Engine/Source/Runtime/Slate/Private/Widgets/Input/SSlider.cpp
@@ -0,0 +1,191 @@
+// Copyright 1998-2014 Epic Games, Inc. All Rights Reserved.
+
+#include "SlatePrivatePCH.h"
+
+
+void SSlider::Construct( const SSlider::FArguments& InDeclaration )
+{
+	check(InDeclaration._Style);
+
+	DisabledHandleImage = &InDeclaration._Style->DisabledThumbImage;
+	NormalHandleImage = &InDeclaration._Style->NormalThumbImage;
+	IndentHandle = InDeclaration._IndentHandle;
+	LockedAttribute = InDeclaration._Locked;
+	Orientation = InDeclaration._Orientation;
+	ValueAttribute = InDeclaration._Value;
+	SliderBarColor = InDeclaration._SliderBarColor;
+	SliderHandleColor = InDeclaration._SliderHandleColor;
+	OnMouseCaptureBegin = InDeclaration._OnMouseCaptureBegin;
+	OnMouseCaptureEnd = InDeclaration._OnMouseCaptureEnd;
+
+	OnValueChanged = InDeclaration._OnValueChanged;
+}
+
+
+int32 SSlider::OnPaint( const FGeometry& AllottedGeometry, const FSlateRect& MyClippingRect, FSlateWindowElementList& OutDrawElements, int32 LayerId, const FWidgetStyle& InWidgetStyle, bool bParentEnabled ) const
+{
+#if SLATE_HD_STATS
+	SCOPE_CYCLE_COUNTER(STAT_SlateOnPaint_SSlider);
+#endif
+
+	float HandleRotation;
+	FVector2D HandleTopLeftPoint;
+	FVector2D SliderStartPoint;
+	FVector2D SliderEndPoint;
+
+	// calculate slider geometry
+	const FVector2D HalfHandleSize = 0.5f * NormalHandleImage->ImageSize;
+	const float Indentation = IndentHandle.Get() ? NormalHandleImage->ImageSize.X : 0.0f;
+
+	if (Orientation == Orient_Horizontal)
+	{
+		const float SliderLength = AllottedGeometry.Size.X - Indentation;
+		const float SliderHandleOffset = ValueAttribute.Get() * SliderLength;
+		const float SliderY = 0.5f * AllottedGeometry.Size.Y;
+
+		HandleRotation = 0.0f;
+		HandleTopLeftPoint = FVector2D(SliderHandleOffset - HalfHandleSize.X + 0.5f * Indentation, SliderY - HalfHandleSize.Y);
+
+		SliderStartPoint = FVector2D(HalfHandleSize.X, SliderY);
+		SliderEndPoint = FVector2D(AllottedGeometry.Size.X - HalfHandleSize.X, SliderY);
+	}
+	else
+	{
+		const float SliderLength = AllottedGeometry.Size.Y - Indentation;
+		const float SliderHandleOffset = ValueAttribute.Get() * SliderLength;
+		const float SliderX = 0.5f * AllottedGeometry.Size.X;
+
+		HandleRotation = -0.5f * PI;
+		HandleTopLeftPoint = FVector2D(SliderX - HalfHandleSize.X, AllottedGeometry.Size.Y - SliderHandleOffset - HalfHandleSize.Y - 0.5f * Indentation);
+
+		SliderStartPoint = FVector2D(SliderX, AllottedGeometry.Size.Y - HalfHandleSize.X);
+		SliderEndPoint = FVector2D(SliderX, HalfHandleSize.X);
+	}
+
+	// draw slider bar
+	TArray<FVector2D> LinePoints;
+	LinePoints.Add(SliderStartPoint);
+	LinePoints.Add(SliderEndPoint);
+
+	const bool bEnabled = ShouldBeEnabled(bParentEnabled);
+	const ESlateDrawEffect::Type DrawEffects = bEnabled ? ESlateDrawEffect::None : ESlateDrawEffect::DisabledEffect;
+
+	FSlateDrawElement::MakeLines( 
+		OutDrawElements,
+		LayerId,
+		AllottedGeometry.ToPaintGeometry(),
+		LinePoints,
+		MyClippingRect,
+		DrawEffects,
+		SliderBarColor.Get().GetColor(InWidgetStyle) * InWidgetStyle.GetColorAndOpacityTint()
+	);
+
+	++LayerId;
+
+	// draw slider handle
+	FSlateDrawElement::MakeRotatedBox( 
+		OutDrawElements,
+		LayerId,
+		AllottedGeometry.ToPaintGeometry(HandleTopLeftPoint, NormalHandleImage->ImageSize),
+		LockedAttribute.Get() ? DisabledHandleImage : NormalHandleImage,
+		MyClippingRect,
+		DrawEffects,
+		HandleRotation,
+		TOptional<FVector2D>(),
+		FSlateDrawElement::RelativeToElement,
+		SliderHandleColor.Get().GetColor(InWidgetStyle) * InWidgetStyle.GetColorAndOpacityTint()
+	);
+
+	return LayerId;
+}
+
+
+FVector2D SSlider::ComputeDesiredSize() const
+{
+	static const FVector2D SSliderDesiredSize(16.0f, 16.0f);
+
+	if (NormalHandleImage == nullptr)
+	{
+		return SSliderDesiredSize;
+	}
+
+	if (Orientation == Orient_Vertical)
+	{
+		return FVector2D(NormalHandleImage->ImageSize.Y, SSliderDesiredSize.Y);
+	}
+
+	return FVector2D(SSliderDesiredSize.X, NormalHandleImage->ImageSize.Y);
+}
+
+
+FReply SSlider::OnMouseButtonDown( const FGeometry& MyGeometry, const FPointerEvent& MouseEvent )
+{
+	if ((MouseEvent.GetEffectingButton() == EKeys::LeftMouseButton) && !LockedAttribute.Get())
+	{
+		OnMouseCaptureBegin.ExecuteIfBound();
+		CommitValue(PositionToValue(MyGeometry, MouseEvent.GetLastScreenSpacePosition()));
+
+		return FReply::Handled().CaptureMouse(SharedThis(this));
+	}
+
+	return FReply::Unhandled();
+}
+
+
+FReply SSlider::OnMouseButtonUp( const FGeometry& MyGeometry, const FPointerEvent& MouseEvent )
+{
+	if ((MouseEvent.GetEffectingButton() == EKeys::LeftMouseButton) && HasMouseCapture())
+	{
+		SetCursor(EMouseCursor::Default);
+		OnMouseCaptureEnd.ExecuteIfBound();
+
+		return FReply::Handled().ReleaseMouseCapture();	
+	}
+
+	return FReply::Unhandled();
+}
+
+
+FReply SSlider::OnMouseMove( const FGeometry& MyGeometry, const FPointerEvent& MouseEvent )
+{
+	if (this->HasMouseCapture() && !LockedAttribute.Get())
+	{
+		SetCursor((Orientation == Orient_Horizontal) ? EMouseCursor::ResizeLeftRight : EMouseCursor::ResizeUpDown);
+		CommitValue(PositionToValue(MyGeometry, MouseEvent.GetLastScreenSpacePosition()));
+
+		return FReply::Handled();
+	}
+
+	return FReply::Unhandled();
+}
+
+
+void SSlider::CommitValue(float NewValue)
+{
+	if (!ValueAttribute.IsBound())
+	{
+		ValueAttribute.Set(NewValue);
+	}
+
+	OnValueChanged.ExecuteIfBound(NewValue);
+}
+
+
+float SSlider::PositionToValue( const FGeometry& MyGeometry, const FVector2D& AbsolutePosition )
+{
+	const FVector2D LocalPosition = MyGeometry.AbsoluteToLocal(AbsolutePosition);
+	const float Indentation = IndentHandle.Get() ? NormalHandleImage->ImageSize.X : 0.0f;
+
+	float RelativeValue;
+
+	if (Orientation == Orient_Horizontal)
+	{
+		RelativeValue = (LocalPosition.X - 0.5f * Indentation) / (MyGeometry.Size.X - Indentation);
+	}
+	else
+	{
+		RelativeValue = (MyGeometry.Size.Y - LocalPosition.Y - 0.5f * Indentation) / (MyGeometry.Size.Y - Indentation);
+	}
+	
+	return FMath::Clamp(RelativeValue, 0.0f, 1.0f);
+}
diff --git a/Engine/Source/Runtime/Slate/Private/Widgets/Input/SVectorInputBox.cpp b/Engine/Source/Runtime/Slate/Private/Widgets/Input/SVectorInputBox.cpp
new file mode 100644
--- /dev/null
+++ b/Engine/Source/Runtime/Slate/Private/Widgets/Input/SVectorInputBox.cpp
@@ -0,0 +1,104 @@
+// Copyright 1998-2014 Epic Games, Inc. All Rights Reserved.
+
+#include "SlatePrivatePCH.h"
+
+
+#define LOCTEXT_NAMESPACE "SVectorInputBox"
+
+
+void SVectorInputBox::Construct( const FArguments& InArgs )
+{
+	TSharedRef<SHorizontalBox> HorizontalBox = SNew(SHorizontalBox);
+
+	ChildSlot
+	[
+		HorizontalBox
+	];
+
+	ConstructX( InArgs, HorizontalBox );
+	ConstructY( InArgs, HorizontalBox );
+	ConstructZ( InArgs, HorizontalBox );
+}
+
+void SVectorInputBox::ConstructX( const FArguments& InArgs, TSharedRef<SHorizontalBox> HorizontalBox )
+{	
+	FLinearColor LabelColor;
+	LabelColor = InArgs._bColorAxisLabels ?  SNumericEntryBox<float>::RedLabelBackgroundColor : FLinearColor(0.0f,0.0f,0.0f,.5f);
+
+	HorizontalBox->AddSlot()
+	.VAlign( VAlign_Center )
+	.FillWidth( 1.0f )
+	.Padding( 0.0f, 1.0f, 2.0f, 1.0f )
+	[
+		SNew( SNumericEntryBox<float> )
+		.Font( InArgs._Font )
+		.Value( InArgs._X )
+		.OnValueChanged( InArgs._OnXChanged )
+		.OnValueCommitted( InArgs._OnXCommitted )		
+		.ToolTipText( LOCTEXT("X_ToolTip", "X Value") )
+		.UndeterminedString( LOCTEXT("MultipleValues", "Multiple Values") )
+		.LabelPadding(0)
+		.ContextMenuExtender( InArgs._ContextMenuExtenderX )
+		.Label()
+		[
+			SNumericEntryBox<float>::BuildLabel( LOCTEXT("X_Label", "X"), FLinearColor::White, LabelColor )
+		]
+	];
+	
+}
+
+
+void SVectorInputBox::ConstructY( const FArguments& InArgs, TSharedRef<SHorizontalBox> HorizontalBox )
+{
+	FLinearColor LabelColor;
+	LabelColor = InArgs._bColorAxisLabels ?  SNumericEntryBox<float>::GreenLabelBackgroundColor : FLinearColor(0.0f,0.0f,0.0f,.5f);
+
+	HorizontalBox->AddSlot()
+	.VAlign( VAlign_Center )
+	.FillWidth( 1.0f )
+	.Padding( 0.0f, 1.0f, 2.0f, 1.0f )
+	[
+		SNew( SNumericEntryBox<float> )
+		.Font( InArgs._Font )
+		.Value( InArgs._Y )
+		.OnValueChanged( InArgs._OnYChanged )
+		.OnValueCommitted( InArgs._OnYCommitted )
+		.ToolTipText( LOCTEXT("Y_ToolTip", "Y Value") )
+		.UndeterminedString( LOCTEXT("MultipleValues", "Multiple Values") )
+		.LabelPadding(0)
+		.ContextMenuExtender( InArgs._ContextMenuExtenderY )
+		.Label()
+		[
+			SNumericEntryBox<float>::BuildLabel( LOCTEXT("Y_Label", "Y"), FLinearColor::White, LabelColor )
+		]
+	];
+
+}
+
+void SVectorInputBox::ConstructZ( const FArguments& InArgs, TSharedRef<SHorizontalBox> HorizontalBox )
+{
+	FLinearColor LabelColor;
+	LabelColor = InArgs._bColorAxisLabels ?  SNumericEntryBox<float>::BlueLabelBackgroundColor : FLinearColor(0.0f,0.0f,0.0f,.5f);
+
+	HorizontalBox->AddSlot()
+	.VAlign( VAlign_Center )
+	.FillWidth( 1.0f )
+	.Padding( 0.0f, 1.0f, 0.0f, 1.0f )
+	[
+		SNew( SNumericEntryBox<float> )
+		.Font( InArgs._Font )
+		.Value( InArgs._Z )
+		.OnValueChanged( InArgs._OnZChanged )
+		.OnValueCommitted( InArgs._OnZCommitted )
+		.ToolTipText( LOCTEXT("Z_ToolTip", "Z Value") )
+		.UndeterminedString( LOCTEXT("MultipleValues", "Multiple Values") )
+		.LabelPadding(0)
+		.ContextMenuExtender( InArgs._ContextMenuExtenderZ )
+		.Label()
+		[
+			SNumericEntryBox<float>::BuildLabel( LOCTEXT("Z_Label", "Z"), FLinearColor::White, LabelColor )
+		]
+	];
+}
+
+#undef LOCTEXT_NAMESPACE
