Slate: Reorganized private Slate files; proper pre-compiled header for all files.

diff --git a/Engine/Source/Runtime/Slate/Private/Widgets/Input/SVirtualJoystick.cpp b/Engine/Source/Runtime/Slate/Private/Widgets/Input/SVirtualJoystick.cpp
new file mode 100644
--- /dev/null
+++ b/Engine/Source/Runtime/Slate/Private/Widgets/Input/SVirtualJoystick.cpp
@@ -0,0 +1,463 @@
+// Copyright 1998-2014 Epic Games, Inc. All Rights Reserved.
+
+#include "SlatePrivatePCH.h"
+
+
+const float OPACITY_LERP_RATE = 3.f;
+
+
+static FORCEINLINE float GetScaleFactor(const FGeometry& Geometry)
+{
+	const float DesiredWidth = 1024.0f;
+	return Geometry.GetDrawSize().GetMax() / DesiredWidth;
+}
+
+FORCEINLINE float SVirtualJoystick::GetBaseOpacity()
+{
+	return (State == State_Active || State == State_CountingDownToInactive) ? ActiveOpacity : InactiveOpacity;
+}
+
+void SVirtualJoystick::FControlInfo::Reset()
+{
+	// snap the visual center back to normal (for controls that have a center on touch)
+	VisualCenter = Center;
+}
+
+void SVirtualJoystick::Construct( const FArguments& InArgs )
+{
+	State = State_Inactive;
+	bVisible = true;
+	bPreventReCenter = false;
+	
+	// just set some defaults
+	ActiveOpacity = 1.0f;
+	InactiveOpacity = 0.1f;
+	TimeUntilDeactive = 0.5f;
+	TimeUntilReset = 2.0f;
+	ActivationDelay = 0.f;
+	CurrentOpacity = InactiveOpacity;
+	StartupDelay = 0.f;
+}
+
+void SVirtualJoystick::SetGlobalParameters(float InActiveOpacity, float InInactiveOpacity, float InTimeUntilDeactive, float InTimeUntilReset, float InActivationDelay, bool InbPreventReCenter, float InStartupDelay)
+{
+	ActiveOpacity = InActiveOpacity;
+	InactiveOpacity = InInactiveOpacity;
+	TimeUntilDeactive = InTimeUntilDeactive;
+	TimeUntilReset = InTimeUntilReset;
+	ActivationDelay = InActivationDelay;
+	StartupDelay = InStartupDelay;
+
+	bPreventReCenter = InbPreventReCenter;
+
+	if (StartupDelay > 0.f)
+	{
+		State = State_WaitForStart;
+	}
+}
+
+bool SVirtualJoystick::ShouldDisplayTouchInterface()
+{
+	bool bAlwaysShowTouchInterface = false;
+	GConfig->GetBool(TEXT("/Script/Engine.InputSettings"), TEXT("bAlwaysShowTouchInterface"), bAlwaysShowTouchInterface, GInputIni);
+
+	// do we want to show virtual joysticks?
+	return FPlatformMisc::GetUseVirtualJoysticks() || bAlwaysShowTouchInterface || FSlateApplication::Get().IsFakingTouchEvents();
+}
+
+static int32 ResolveRelativePosition(float Position, float RelativeTo, float ScaleFactor)
+{
+	// absolute from edge
+	if (Position < -1.0f)
+	{
+		return RelativeTo + Position * ScaleFactor;
+	}
+	// relative from edge
+	else if (Position < 0.0f)
+	{
+		return RelativeTo + Position * RelativeTo;
+	}
+	// relative from 0
+	else if (Position <= 1.0f)
+	{
+		return Position * RelativeTo;
+	}
+	// absolute from 0
+	else
+	{
+		return Position * ScaleFactor;
+	}
+
+}
+
+static bool PositionIsInside(const SVirtualJoystick::FControlInfo& Control, const FVector2D& Position, const FVector2D& BoxSize)
+{
+	return
+		Position.X >= Control.Center.X - BoxSize.X * 0.5f &&
+		Position.X <= Control.Center.X + BoxSize.X * 0.5f &&
+		Position.Y >= Control.Center.Y - BoxSize.Y * 0.5f &&
+		Position.Y <= Control.Center.Y + BoxSize.Y * 0.5f;
+}
+
+int32 SVirtualJoystick::OnPaint( const FGeometry& AllottedGeometry, const FSlateRect& MyClippingRect, FSlateWindowElementList& OutDrawElements, int32 LayerId, const FWidgetStyle& InWidgetStyle, bool bParentEnabled ) const
+{
+	int32 RetLayerId = LayerId;
+
+	if (bVisible)
+	{
+		FColor ColorAndOpacitySRGB = InWidgetStyle.GetColorAndOpacityTint();
+		ColorAndOpacitySRGB.A = FMath::Round(255.f * CurrentOpacity);
+
+		for (int32 ControlIndex = 0; ControlIndex < Controls.Num(); ControlIndex++)
+		{
+			const FControlInfo& Control = Controls[ControlIndex];
+
+			FSlateDrawElement::MakeBox(
+				OutDrawElements,
+				RetLayerId++,
+				AllottedGeometry.ToPaintGeometry(
+				Control.VisualCenter - FVector2D(Control.VisualSize.X * 0.5f, Control.VisualSize.Y * 0.5f),
+				Control.VisualSize),
+				Control.Image2.Get(),
+				MyClippingRect,
+				ESlateDrawEffect::None,
+				ColorAndOpacitySRGB
+				);
+
+			FSlateDrawElement::MakeBox(
+				OutDrawElements,
+				RetLayerId++,
+				AllottedGeometry.ToPaintGeometry(
+				Control.VisualCenter + Control.ThumbPosition - FVector2D(Control.ThumbSize.X * 0.5f, Control.ThumbSize.Y * 0.5f),
+				Control.ThumbSize),
+				Control.Image1.Get(),
+				MyClippingRect,
+				ESlateDrawEffect::None,
+				ColorAndOpacitySRGB
+				);
+		}
+	}
+	
+	return RetLayerId;
+}
+
+FVector2D SVirtualJoystick::ComputeDesiredSize() const
+{
+	return FVector2D(100, 100);
+}
+
+
+FReply SVirtualJoystick::OnTouchStarted(const FGeometry& MyGeometry, const FPointerEvent& Event)
+{
+//	UE_LOG(LogTemp, Log, TEXT("Pointer index: %d"), Event.GetPointerIndex());
+
+	FVector2D LocalCoord = MyGeometry.AbsoluteToLocal( Event.GetScreenSpacePosition() );
+
+	for (int32 ControlIndex = 0; ControlIndex < Controls.Num(); ControlIndex++)
+	{
+		FControlInfo& Control = Controls[ControlIndex];
+
+		// skip controls already in use
+		if (Control.CapturedPointerIndex == -1)
+		{
+			if (PositionIsInside(Control, LocalCoord, Control.InteractionSize))
+			{
+				// Align Joystick inside of Screen
+				AlignBoxIntoScreen(LocalCoord, Control.VisualSize, MyGeometry.Size);
+
+				Control.CapturedPointerIndex = Event.GetPointerIndex();
+
+				if (ActivationDelay == 0.f)
+				{
+					CurrentOpacity = ActiveOpacity;
+
+					if (!bPreventReCenter)
+					{
+						Control.VisualCenter = LocalCoord;
+					}
+
+					if (HandleTouch(ControlIndex, LocalCoord, MyGeometry.Size)) // Never fail!
+					{
+						return FReply::Handled();
+					}
+				}
+				else
+				{
+					Control.bNeedUpdatedCenter = true;
+					Control.ElapsedTime = 0.f;
+					Control.NextCenter = LocalCoord;
+
+					return FReply::Unhandled();
+				}
+			}
+		}
+	}
+
+//	CapturedPointerIndex[CapturedJoystick] = -1;
+
+	return FReply::Unhandled();
+}
+
+FReply SVirtualJoystick::OnTouchMoved(const FGeometry& MyGeometry, const FPointerEvent& Event)
+{
+	FVector2D LocalCoord = MyGeometry.AbsoluteToLocal( Event.GetScreenSpacePosition() );
+
+	for (int32 ControlIndex = 0; ControlIndex < Controls.Num(); ControlIndex++)
+	{
+		FControlInfo& Control = Controls[ControlIndex];
+
+		// is this control the one captured to this pointer?
+		if (Control.CapturedPointerIndex == Event.GetPointerIndex())
+		{
+			if (Control.bNeedUpdatedCenter)
+			{
+				return FReply::Unhandled();
+			}
+			else if (HandleTouch(ControlIndex, LocalCoord, MyGeometry.Size))
+			{
+				return FReply::Handled();
+			}
+		}
+	}
+
+	return FReply::Unhandled();
+}
+
+FReply SVirtualJoystick::OnTouchEnded(const FGeometry& MyGeometry, const FPointerEvent& Event)
+{
+	for (int32 ControlIndex = 0; ControlIndex < Controls.Num(); ControlIndex++)
+	{
+		FControlInfo& Control = Controls[ControlIndex];
+
+		// is this control the one captured to this pointer?
+		if (Control.CapturedPointerIndex == Event.GetPointerIndex())
+		{
+			// release and center the joystick
+			Control.ThumbPosition = FVector2D(0, 0);
+			Control.CapturedPointerIndex = -1;
+
+			// send one more joystick update for the centering
+			Control.bSendOneMoreEvent = true;
+
+			// Pass event as unhandled if time is too short
+			if (Control.bNeedUpdatedCenter)
+			{
+				Control.bNeedUpdatedCenter = false;
+				return FReply::Unhandled();
+			}
+			return FReply::Handled();
+		}
+	}
+
+	return FReply::Unhandled();
+}
+
+bool SVirtualJoystick::HandleTouch(int32 ControlIndex, const FVector2D& LocalCoord, const FVector2D& ScreenSize)
+{
+	FControlInfo& Control = Controls[ControlIndex];
+
+	// figure out position around center
+	FVector2D Offset = LocalCoord - Controls[ControlIndex].VisualCenter;
+	// only do work if we aren't at the center
+	if (Offset == FVector2D(0, 0))
+	{
+		Control.ThumbPosition = Offset;
+	}
+	else
+	{
+		// clamp to the ellipse of the stick (snaps to the visual size, so, the art should go all the way to the edge of the texture)
+		float DistanceToTouch = Offset.Size();
+		float Angle = FMath::Atan2(Offset.Y, Offset.X);
+
+		// length along line to ellipse: L = 1.0 / sqrt(((sin(T)/Rx)^2 + (cos(T)/Ry)^2))
+		float CosAngle = FMath::Cos(Angle);
+		float SinAngle = FMath::Sin(Angle);
+		float XTerm = CosAngle / (Control.VisualSize.X * 0.5f);
+		float YTerm = SinAngle / (Control.VisualSize.Y * 0.5f);
+		float DistanceToEdge = FMath::InvSqrt(XTerm * XTerm + YTerm * YTerm);
+
+		// only clamp 
+		if (DistanceToTouch > DistanceToEdge)
+		{
+			Control.ThumbPosition = FVector2D(DistanceToEdge * CosAngle,  DistanceToEdge * SinAngle);
+		}
+		else
+		{
+			Control.ThumbPosition = Offset;
+		}
+	}
+
+	FVector2D AbsoluteThumbPos = Control.ThumbPosition + Controls[ControlIndex].VisualCenter;
+	AlignBoxIntoScreen(AbsoluteThumbPos, Control.ThumbSize, ScreenSize);
+	Control.ThumbPosition = AbsoluteThumbPos - Controls[ControlIndex].VisualCenter;
+
+	return true;
+}
+
+void SVirtualJoystick::Tick( const FGeometry& AllottedGeometry, const double InCurrentTime, const float InDeltaTime )
+{
+	SLeafWidget::Tick(AllottedGeometry, InCurrentTime, InDeltaTime);
+
+	if (State == State_WaitForStart || State == State_CountingDownToStart)
+	{
+		CurrentOpacity = 0.f;
+	}
+	else
+	{
+		// lerp to the desired opacity based on whether the user is interacting with the joystick
+		CurrentOpacity = FMath::Lerp(CurrentOpacity, GetBaseOpacity(), OPACITY_LERP_RATE * InDeltaTime);
+	}
+
+	// count how many controls are active
+	int32 NumActiveControls = 0;
+
+	for (int32 ControlIndex = 0; ControlIndex < Controls.Num(); ControlIndex++)
+	{
+		FControlInfo& Control = Controls[ControlIndex];
+
+		if (Control.bNeedUpdatedCenter)
+		{
+			Control.ElapsedTime += InDeltaTime;
+			if (Control.ElapsedTime > ActivationDelay)
+			{
+				Control.bNeedUpdatedCenter = false;
+				CurrentOpacity = ActiveOpacity;
+
+				if (!bPreventReCenter)
+				{
+					Control.VisualCenter = Control.NextCenter;
+				}
+
+				HandleTouch(ControlIndex, Control.NextCenter, AllottedGeometry.Size);
+			}
+		}
+
+		// calculate absolute positions based on geometry
+		// @todo: Need to manage geometry changing!
+		if (!Control.bHasBeenPositioned)
+		{
+			// figure out how much to scale the control sizes
+			float ScaleFactor = GetScaleFactor(AllottedGeometry);
+
+			// update all the sizes
+			Control.Center = FVector2D(ResolveRelativePosition(Control.Center.X, AllottedGeometry.Size.X, ScaleFactor), ResolveRelativePosition(Control.Center.Y, AllottedGeometry.Size.Y, ScaleFactor));
+			Control.VisualCenter = Control.Center;
+			Control.VisualSize = FVector2D(ResolveRelativePosition(Control.VisualSize.X, AllottedGeometry.Size.X, ScaleFactor), ResolveRelativePosition(Control.VisualSize.Y, AllottedGeometry.Size.Y, ScaleFactor));
+			Control.InteractionSize = FVector2D(ResolveRelativePosition(Control.InteractionSize.X, AllottedGeometry.Size.X, ScaleFactor), ResolveRelativePosition(Control.InteractionSize.Y, AllottedGeometry.Size.Y, ScaleFactor));
+			Control.ThumbSize = FVector2D(ResolveRelativePosition(Control.ThumbSize.X, AllottedGeometry.Size.X, ScaleFactor), ResolveRelativePosition(Control.ThumbSize.Y, AllottedGeometry.Size.Y, ScaleFactor));
+			Control.InputScale = Control.InputScale * ScaleFactor;
+			Control.bHasBeenPositioned = true;
+		}
+
+		if (Control.CapturedPointerIndex >= 0 || Control.bSendOneMoreEvent)
+		{
+			Control.bSendOneMoreEvent = false;
+
+			// now pass the fake joystick events to the game
+			// Assume that joystick size is all equal
+			float JoysticInputSize = Control.ThumbPosition.Size() * 2.f / Control.VisualSize.X;
+			FVector2D NormalizedOffset = Control.ThumbPosition.SafeNormal() * Control.InputScale * JoysticInputSize;
+			EControllerButtons::Type XAxis = ControlIndex == 0 ? EControllerButtons::LeftAnalogX : EControllerButtons::RightAnalogX;
+			EControllerButtons::Type YAxis = ControlIndex == 0 ? EControllerButtons::LeftAnalogY : EControllerButtons::RightAnalogY;
+
+	//		UE_LOG(LogTemp, Log, TEXT("Joysticking %f,%f"), NormalizedOffset.X, -NormalizedOffset.Y);
+			FSlateApplication::Get().SetJoystickCaptorToGameViewport();
+			FSlateApplication::Get().OnControllerAnalog(Control.MainInputKey, 0, NormalizedOffset.X);
+			FSlateApplication::Get().OnControllerAnalog(Control.AltInputKey, 0, -NormalizedOffset.Y);
+		}
+
+		// is this active?
+		if (Control.CapturedPointerIndex != -1)
+		{
+			NumActiveControls++;
+		}
+	}
+
+
+	// STATE MACHINE!
+	if (NumActiveControls > 0 || bPreventReCenter)
+	{
+		// any active control snaps the state to active immediately
+		State = State_Active;
+	}
+	else
+	{
+		switch (State)
+		{
+			case State_WaitForStart:
+				{
+					State = State_CountingDownToStart;
+					Countdown = StartupDelay;
+				}
+				break;
+			case State_CountingDownToStart:
+				// update the countdown
+				Countdown -= InDeltaTime;
+				if (Countdown <= 0.0f)
+				{
+					State = State_Inactive;
+				}
+				break;
+			case State_Active:
+				if (NumActiveControls == 0)
+				{
+					// start going to inactive
+					State = State_CountingDownToInactive;
+					Countdown = TimeUntilDeactive;
+				}
+				break;
+
+			case State_CountingDownToInactive:
+				// update the countdown
+				Countdown -= InDeltaTime;
+				if (Countdown <= 0.0f)
+				{
+					// should we start counting down to a control reset?
+					if (TimeUntilReset > 0.0f)
+					{
+						State = State_CountingDownToReset;
+						Countdown = TimeUntilReset;
+					}
+					else
+					{
+						// if not, then just go inactive
+						State = State_Inactive;
+					}
+				}
+				break;
+
+			case State_CountingDownToReset:
+				Countdown -= InDeltaTime;
+				if (Countdown <= 0.0f)
+				{
+					// reset all the controls
+					for (int32 ControlIndex = 0; ControlIndex < Controls.Num(); ControlIndex++)
+					{
+						Controls[ControlIndex].Reset();
+					}
+
+					// finally, go inactive
+					State = State_Inactive;
+				}
+				break;
+		}
+	}
+}
+
+void SVirtualJoystick::SetVisibility(const bool bInVisible, const bool bInFade)
+{
+	// if we aren't fading, then just set the current opacity to desired
+	if (!bInFade)
+	{
+		if (bInVisible)
+		{
+			CurrentOpacity = GetBaseOpacity();
+		}
+		else
+		{
+			CurrentOpacity = 0.f;
+		}
+	}
+
+	bVisible = bInVisible;
+}
diff --git a/Engine/Source/Runtime/Slate/Private/Widgets/Input/SVirtualKeyboardEntry.cpp b/Engine/Source/Runtime/Slate/Private/Widgets/Input/SVirtualKeyboardEntry.cpp
new file mode 100644
--- /dev/null
+++ b/Engine/Source/Runtime/Slate/Private/Widgets/Input/SVirtualKeyboardEntry.cpp
@@ -0,0 +1,328 @@
+// Copyright 1998-2014 Epic Games, Inc. All Rights Reserved.
+
+#include "SlatePrivatePCH.h"
+
+
+/** Constructor */
+SVirtualKeyboardEntry::SVirtualKeyboardEntry()
+	: ScrollHelper(),
+	  bWasFocusedByLastMouseDown( false ),
+	  bIsChangingText( false )
+{ }
+
+
+/**
+ * Construct this widget
+ *
+ * @param	InArgs	The declaration data for this widget
+ */
+void SVirtualKeyboardEntry::Construct( const FArguments& InArgs )
+{
+	Text = InArgs._Text;
+	HintText = InArgs._HintText;
+
+	Font = InArgs._Font;
+	ColorAndOpacity = InArgs._ColorAndOpacity;
+	IsReadOnly = InArgs._IsReadOnly;
+	ClearKeyboardFocusOnCommit = InArgs._ClearKeyboardFocusOnCommit;
+	OnTextChanged = InArgs._OnTextChanged;
+	OnTextCommitted = InArgs._OnTextCommitted;
+	MinDesiredWidth = InArgs._MinDesiredWidth;
+	KeyboardType = InArgs._KeyboardType;
+}
+
+/**
+ * Sets the text string currently being edited 
+ *
+ * @param  InNewText  The new text string
+ */
+void SVirtualKeyboardEntry::SetText( const TAttribute< FText >& InNewText )
+{
+	EditedText = InNewText.Get();
+
+	// Don't set text if the text attribute has a 'getter' binding on it, otherwise we'd blow away
+	// that binding.  If there is a getter binding, then we'll assume it will provide us with
+	// updated text after we've fired our 'text changed' callbacks
+	if( !Text.IsBound() )
+	{
+		Text.Set ( EditedText );
+	}
+
+	bNeedsUpdate = true;
+
+	// Let outsiders know that the text content has been changed
+	//OnTextChanged.ExecuteIfBound( HasKeyboardFocus() ? EditedText : Text.Get() );
+}
+
+/**
+ * Restores the text to the original state
+ */
+void SVirtualKeyboardEntry::RestoreOriginalText()
+{
+	if( HasTextChangedFromOriginal() )
+	{
+		SetText(OriginalText);
+
+		// Let outsiders know that the text content has been changed
+		OnTextChanged.ExecuteIfBound( EditedText );
+	}
+}
+
+bool SVirtualKeyboardEntry::HasTextChangedFromOriginal() const
+{
+	return ( !IsReadOnly.Get() && !EditedText.EqualTo(OriginalText)  );
+}
+
+/**
+ * Sets the font used to draw the text
+ *
+ * @param  InNewFont	The new font to use
+ */
+void SVirtualKeyboardEntry::SetFont( const TAttribute< FSlateFontInfo >& InNewFont )
+{
+	Font = InNewFont;
+}
+
+/**
+ * Checks to see if this widget supports keyboard focus.  Override this in derived classes.
+ *
+ * @return  True if this widget can take keyboard focus
+ */
+bool SVirtualKeyboardEntry::SupportsKeyboardFocus() const
+{
+	return true;
+}
+
+bool SVirtualKeyboardEntry::GetIsReadOnly() const
+{
+	return IsReadOnly.Get();
+}
+
+/**
+ * Ticks this widget.  Override in derived classes, but always call the parent implementation.
+ *
+ * @param  AllottedGeometry Space allotted to this widget
+ * @param  InCurrentTime  Current absolute real time
+ * @param  InDeltaTime  Real time passed since last tick
+ */
+void SVirtualKeyboardEntry::Tick( const FGeometry& AllottedGeometry, const double InCurrentTime, const float InDeltaTime )
+{
+	// Call parent implementation.
+	SLeafWidget::Tick( AllottedGeometry, InCurrentTime, InDeltaTime );
+
+	if(bNeedsUpdate)
+	{
+		// Let outsiders know that the text content has been changed
+		OnTextChanged.ExecuteIfBound( HasKeyboardFocus() ? EditedText : Text.Get() );
+		bNeedsUpdate = false;
+	}
+
+
+	// We'll draw with the 'focused' look if we're either focused or we have a context menu summoned
+	const bool bShouldAppearFocused = HasKeyboardFocus();
+}
+
+/**
+ * Gets the height of the largest character in the font
+ *
+ * @return  The fonts height
+ */
+float SVirtualKeyboardEntry::GetFontHeight() const
+{
+	// Take a local copy of the FontInfo and remove any flags that could affect the outcome in a negative way
+	FSlateFontInfo FontCopy = Font.Get();
+
+	// Aesthetic choice: editable text can use a little more room, and should never be smaller than a W of the font.
+	const TSharedRef< FSlateFontMeasure > FontMeasureService = FSlateApplication::Get().GetRenderer()->GetFontMeasureService();
+	const FVector2D FontSize = FontMeasureService->Measure( FString(TEXT("W")), FontCopy );
+	return FontSize.Y;
+}
+
+/**
+ * Computes the desired size of this widget (SWidget)
+ *
+ * @return  The widget's desired size
+ */
+FVector2D SVirtualKeyboardEntry::ComputeDesiredSize() const
+{
+	const FSlateFontInfo& FontInfo = Font.Get();
+	const TSharedRef< FSlateFontMeasure > FontMeasureService = FSlateApplication::Get().GetRenderer()->GetFontMeasureService();
+
+	const float FontMaxCharHeight = GetFontHeight();
+
+	FVector2D TextSize;
+
+	const FString StringToRender = GetStringToRender();
+	if( !StringToRender.IsEmpty() )
+	{
+		TextSize = FontMeasureService->Measure( StringToRender, FontInfo );
+	}
+	else
+	{
+		TextSize = FontMeasureService->Measure( HintText.Get().ToString(), FontInfo );
+	}
+	
+	return FVector2D( FMath::Max(TextSize.X, MinDesiredWidth.Get()), FMath::Max(TextSize.Y, FontMaxCharHeight) );
+}
+
+/**
+ * The widget should respond by populating the OutDrawElements array with FDrawElements 
+ * that represent it and any of its children.
+ *
+ * @param AllottedGeometry  The FGeometry that describes an area in which the widget should appear.
+ * @param MyClippingRect    The clipping rectangle allocated for this widget and its children.
+ * @param OutDrawElements   A list of FDrawElements to populate with the output.
+ * @param LayerId           The Layer onto which this widget should be rendered.
+ * @param InColorAndOpacity Color and Opacity to be applied to all the descendants of the widget being painted
+ * @param bParentEnabled	True if the parent of this widget is enabled.
+ *
+ * @return The maximum layer ID attained by this widget or any of its children.
+ */
+int32 SVirtualKeyboardEntry::OnPaint( const FGeometry& AllottedGeometry, const FSlateRect& MyClippingRect, FSlateWindowElementList& OutDrawElements, int32 LayerId, const FWidgetStyle& InWidgetStyle, bool bParentEnabled ) const
+{
+	const TSharedRef< FSlateFontMeasure > FontMeasureService = FSlateApplication::Get().GetRenderer()->GetFontMeasureService();
+
+	// The text and some effects draws in front of the widget's background and selection.
+	const int32 SelectionLayer = 0;
+	const int32 TextLayer = 1;
+
+	// See if a disabled effect should be used
+	bool bEnabled = ShouldBeEnabled( bParentEnabled );
+	bool bIsReadonly = IsReadOnly.Get();
+	ESlateDrawEffect::Type DrawEffects = (bEnabled) ? ESlateDrawEffect::None : ESlateDrawEffect::DisabledEffect;
+
+	const FSlateFontInfo& FontInfo = Font.Get();
+	const FString VisibleText = GetStringToRender();
+	const FLinearColor ThisColorAndOpacity = ColorAndOpacity.Get().GetColor(InWidgetStyle);
+	const FColor ColorAndOpacitySRGB = ThisColorAndOpacity * InWidgetStyle.GetColorAndOpacityTint();
+	const float FontMaxCharHeight = GetFontHeight();
+	const double CurrentTime = FSlateApplication::Get().GetCurrentTime();
+
+	// We'll draw with the 'focused' look if we're either focused or we have a context menu summoned
+	const bool bShouldAppearFocused = HasKeyboardFocus();
+
+	const float DrawPositionY = ( AllottedGeometry.Size.Y / 2 ) - ( FontMaxCharHeight / 2 );
+
+	if (VisibleText.Len() == 0)
+	{
+		// Draw the hint text.
+		const FLinearColor HintTextColor = FLinearColor(ColorAndOpacitySRGB.R, ColorAndOpacitySRGB.G, ColorAndOpacitySRGB.B, 0.35f);
+		const FString ThisHintText = this->HintText.Get().ToString();
+		FSlateDrawElement::MakeText(
+			OutDrawElements,
+			LayerId + TextLayer,
+			AllottedGeometry.ToPaintGeometry( FVector2D( 0, DrawPositionY ), AllottedGeometry.Size ),
+			ThisHintText,          // Text
+			FontInfo,              // Font information (font name, size)
+			MyClippingRect,        // Clipping rect
+			DrawEffects,           // Effects to use
+			HintTextColor          // Color
+		);
+	}
+	else
+	{
+		// Draw the text
+
+		// Only draw the text that's potentially visible
+		const float ScrollAreaLeft = ScrollHelper.ToScrollerSpace( FVector2D::ZeroVector ).X;
+		const float ScrollAreaRight = ScrollHelper.ToScrollerSpace( AllottedGeometry.Size ).X;
+		const int32 FirstVisibleCharacter = FontMeasureService->FindCharacterIndexAtOffset( VisibleText, FontInfo, ScrollAreaLeft );
+		const int32 LastVisibleCharacter = FontMeasureService->FindCharacterIndexAtOffset( VisibleText, FontInfo, ScrollAreaRight );
+		const FString PotentiallyVisibleText( VisibleText.Mid( FirstVisibleCharacter, ( LastVisibleCharacter - FirstVisibleCharacter ) + 1 ) );
+		const float FirstVisibleCharacterPosition = FontMeasureService->Measure( VisibleText.Left( FirstVisibleCharacter ), FontInfo ).X;
+
+		const float TextVertOffset = FontMaxCharHeight;
+		const FVector2D DrawPosition = ScrollHelper.FromScrollerSpace( FVector2D( FirstVisibleCharacterPosition, DrawPositionY + TextVertOffset ) );
+		const FVector2D DrawSize = ScrollHelper.SizeFromScrollerSpace( AllottedGeometry.Size );
+
+		FSlateDrawElement::MakeText(
+			OutDrawElements,
+			LayerId + TextLayer,
+			AllottedGeometry.ToPaintGeometry( DrawPosition, DrawSize ),
+			PotentiallyVisibleText,          // Text
+			FontInfo,                        // Font information (font name, size)
+			MyClippingRect,                  // Clipping rect
+			DrawEffects,                     // Effects to use
+			ColorAndOpacitySRGB              // Color
+		);
+	}
+	
+	return LayerId + TextLayer;
+}
+
+FReply SVirtualKeyboardEntry::OnKeyboardFocusReceived( const FGeometry& MyGeometry, const FKeyboardFocusEvent& InKeyboardFocusEvent )
+{
+	// The user wants to edit text. Make a copy of the observed text for the user to edit.
+	EditedText = Text.Get();
+
+	int32 CaretPosition = EditedText.ToString().Len();
+	FSlateApplication& CurrentApp = FSlateApplication::Get();
+	CurrentApp.ShowKeyboard(true, SharedThis(this));
+
+	return FReply::Handled();
+}
+
+/**
+ * Called when this widget loses the keyboard focus.  This event does not bubble.
+ *
+ * @param  InKeyboardFocusEvent  KeyboardFocusEvent
+ */
+void SVirtualKeyboardEntry::OnKeyboardFocusLost( const FKeyboardFocusEvent& InKeyboardFocusEvent )
+{
+	// See if user explicitly tabbed away or moved focus
+	ETextCommit::Type TextAction;
+	switch ( InKeyboardFocusEvent.GetCause() )
+	{
+		case EKeyboardFocusCause::Keyboard:
+		case EKeyboardFocusCause::Mouse:
+			TextAction = ETextCommit::OnUserMovedFocus;
+			break;
+
+		case EKeyboardFocusCause::Cleared:
+			TextAction = ETextCommit::OnCleared;
+			break;
+
+		default:
+			TextAction = ETextCommit::Default;
+			break;
+	}
+
+	FSlateApplication& CurrentApp = FSlateApplication::Get();
+	CurrentApp.ShowKeyboard(false);
+
+	OnTextCommitted.ExecuteIfBound( EditedText, TextAction );
+}
+
+/**
+ * Gets the text that needs to be rendered
+ *
+ * @return  Text string
+ */
+FString SVirtualKeyboardEntry::GetStringToRender() const
+{
+	FString VisibleText;
+
+	if(HasKeyboardFocus())
+	{
+		VisibleText = EditedText.ToString();
+	}
+	else
+	{
+		VisibleText = Text.Get().ToString();
+	}
+
+	// If this string is a password we need to replace all the characters in the string with something else
+	if ( KeyboardType.Get() == EKeyboardType::Keyboard_Password )
+	{
+		const TCHAR Dot = 0x25CF;
+		int32 VisibleTextLen = VisibleText.Len();
+		VisibleText.Empty();
+		while ( VisibleTextLen )
+		{
+			VisibleText += Dot;
+			VisibleTextLen--;
+		}
+	}
+
+	return VisibleText;
+}
diff --git a/Engine/Source/Runtime/Slate/Private/Widgets/Layout/SBorder.cpp b/Engine/Source/Runtime/Slate/Private/Widgets/Layout/SBorder.cpp
new file mode 100644
--- /dev/null
+++ b/Engine/Source/Runtime/Slate/Private/Widgets/Layout/SBorder.cpp
@@ -0,0 +1,219 @@
+// Copyright 1998-2014 Epic Games, Inc. All Rights Reserved.
+
+#include "SlatePrivatePCH.h"
+ 
+
+SBorder::SBorder()
+	: BorderImage( FCoreStyle::Get().GetBrush( "Border" ) )
+	, BorderBackgroundColor( FLinearColor::White )
+	, DesiredSizeScale(FVector2D(1,1))
+{ }
+
+
+/**
+ * Construct this widget
+ *
+ * @param	InArgs	The declaration data for this widget
+ */
+void SBorder::Construct( const SBorder::FArguments& InArgs )
+{
+	ContentScale = InArgs._ContentScale;
+	ColorAndOpacity = InArgs._ColorAndOpacity;
+	DesiredSizeScale = InArgs._DesiredSizeScale;
+
+	ShowDisabledEffect = InArgs._ShowEffectWhenDisabled;
+
+	ChildSlot
+		.HAlign(InArgs._HAlign)
+		.VAlign(InArgs._VAlign)
+		.Padding(InArgs._Padding)
+	[
+		InArgs._Content.Widget
+	];
+
+	BorderImage = InArgs._BorderImage;
+	BorderBackgroundColor = InArgs._BorderBackgroundColor;
+	ForegroundColor = InArgs._ForegroundColor;
+	MouseButtonDownHandler = InArgs._OnMouseButtonDown;
+	MouseButtonUpHandler = InArgs._OnMouseButtonUp;
+	MouseMoveHandler = InArgs._OnMouseMove;
+	MouseDoubleClickHandler = InArgs._OnMouseDoubleClick;
+}
+
+
+/**
+ * Sets the content for this border
+ *
+ * @param	InContent	The widget to use as content for the border
+ */
+void SBorder::SetContent( const TSharedRef< SWidget >& InContent )
+{
+	ChildSlot
+	[
+		InContent
+	];
+}
+
+const TSharedRef< SWidget >& SBorder::GetContent() const
+{
+	return ChildSlot.Widget;
+}
+
+/** Clears out the content for the border */
+void SBorder::ClearContent()
+{
+	ChildSlot = FSimpleSlot();
+}
+
+void SBorder::SetBorderBackgroundColor( const TAttribute<FSlateColor>& InColorAndOpacity )
+{
+	BorderBackgroundColor = InColorAndOpacity;
+}
+
+void SBorder::SetDesiredSizeScale( const TAttribute<FVector2D>& InDesiredSizeScale )
+{
+	DesiredSizeScale = InDesiredSizeScale;
+}
+
+
+/**
+ * The widget should respond by populating the OutDrawElements array with FDrawElements 
+ * that represent it and any of its children.
+ *
+ * @param AllottedGeometry  The FGeometry that describes an area in which the widget should appear.
+ * @param MyClippingRect    The clipping rectangle allocated for this widget and its children.
+ * @param OutDrawElements   A list of FDrawElements to populate with the output.
+ * @param LayerId           The Layer onto which this widget should be rendered.
+ * @param InColorAndOpacity Color and Opacity to be applied to all the descendants of the widget being painted
+ * @param bParentEnabled	True if the parent of this widget is enabled.
+ *
+ * @return The maximum layer ID attained by this widget or any of its children.
+ */
+int32 SBorder::OnPaint( const FGeometry& AllottedGeometry, const FSlateRect& MyClippingRect, FSlateWindowElementList& OutDrawElements, int32 LayerId, const FWidgetStyle& InWidgetStyle, bool bParentEnabled ) const
+{
+	bool bEnabled = false;
+	{
+#if SLATE_HD_STATS
+		SCOPE_CYCLE_COUNTER( STAT_SlateOnPaint_SBorder );
+#endif
+		const FSlateBrush* BrushResource = BorderImage.Get();
+		
+		bEnabled = ShouldBeEnabled( bParentEnabled );
+		bool bShowDisabledEffect = ShowDisabledEffect.Get();
+		ESlateDrawEffect::Type DrawEffects = bShowDisabledEffect && !bEnabled ? ESlateDrawEffect::DisabledEffect : ESlateDrawEffect::None;
+
+		if ( BrushResource && BrushResource->DrawAs != ESlateBrushDrawType::NoDrawType )
+		{
+			FSlateDrawElement::MakeBox(
+				OutDrawElements,
+				LayerId,
+				AllottedGeometry.ToPaintGeometry(),
+				BrushResource,
+				MyClippingRect,
+				DrawEffects,
+				BrushResource->GetTint( InWidgetStyle ) * InWidgetStyle.GetColorAndOpacityTint() * BorderBackgroundColor.Get().GetColor( InWidgetStyle )
+			);
+		}
+	}
+
+	FWidgetStyle CompoundedWidgetStyle = FWidgetStyle(InWidgetStyle)
+		.BlendColorAndOpacityTint(ColorAndOpacity.Get())
+		.SetForegroundColor( ForegroundColor.Get() );
+
+	return SCompoundWidget::OnPaint(AllottedGeometry, MyClippingRect.IntersectionWith( AllottedGeometry.GetClippingRect() ), OutDrawElements, LayerId, CompoundedWidgetStyle, bEnabled );
+}
+
+
+/**
+ * The system calls this method to notify the widget that a mouse button was pressed within it. This event is bubbled.
+ *
+ * @param MyGeometry The Geometry of the widget receiving the event
+ * @param MouseEvent Information about the input event
+ *
+ * @return Whether the event was handled along with possible requests for the system to take action.
+ */
+FReply SBorder::OnMouseButtonDown( const FGeometry& MyGeometry, const FPointerEvent& MouseEvent )
+{
+	if ( MouseButtonDownHandler.IsBound() )
+	{
+		// If a handler is assigned, call it.
+		return MouseButtonDownHandler.Execute(MyGeometry, MouseEvent);
+	}
+	else
+	{
+		// otherwise the event is unhandled.
+		return FReply::Unhandled();
+	}
+}
+
+/**
+ * The system calls this method to notify the widget that a mouse button was release within it. This event is bubbled.
+ *
+ * @param MyGeometry The Geometry of the widget receiving the event
+ * @param MouseEvent Information about the input event
+ *
+ * @return Whether the event was handled along with possible requests for the system to take action.
+ */
+FReply SBorder::OnMouseButtonUp( const FGeometry& MyGeometry, const FPointerEvent& MouseEvent )
+{
+	if ( MouseButtonUpHandler.IsBound() )
+	{
+		// If a handler is assigned, call it.
+		return MouseButtonUpHandler.Execute(MyGeometry, MouseEvent);
+	}
+	else
+	{
+		// otherwise the event is unhandled.
+		return FReply::Unhandled();
+	}
+}
+
+/**
+ * The system calls this method to notify the widget that a mouse moved within it. This event is bubbled.
+ *
+ * @param MyGeometry The Geometry of the widget receiving the event
+ * @param MouseEvent Information about the input event
+ *
+ * @return Whether the event was handled along with possible requests for the system to take action.
+ */
+FReply SBorder::OnMouseMove( const FGeometry& MyGeometry, const FPointerEvent& MouseEvent )
+{
+	if ( MouseMoveHandler.IsBound() )
+	{
+		// A valid handler is assigned for mouse move; let it handle the event.
+		return MouseMoveHandler.Execute( MyGeometry, MouseEvent );
+	}
+	else
+	{
+		// otherwise the event is unhandled
+		return FReply::Unhandled();
+	}
+}
+
+/**
+ * Called when a mouse button is double clicked
+ *
+ * @param  InMyGeometry  Widget geometry
+ * @param  InMouseEvent  Mouse button event
+ *
+ * @return  Returns whether the event was handled, along with other possible actions
+ */
+FReply SBorder::OnMouseButtonDoubleClick( const FGeometry& MyGeometry, const FPointerEvent& MouseEvent )
+{
+	if ( MouseDoubleClickHandler.IsBound() )
+	{
+		// A valid handler is assigned; let it handle the event.
+		return MouseDoubleClickHandler.Execute( MyGeometry, MouseEvent );
+	}
+	else
+	{
+		// otherwise the event is unhandled
+		return FReply::Unhandled();
+	}
+	return FReply::Unhandled();
+}
+
+FVector2D SBorder::ComputeDesiredSize() const
+{
+	return DesiredSizeScale.Get() * SCompoundWidget::ComputeDesiredSize();
+}
diff --git a/Engine/Source/Runtime/Slate/Private/Widgets/Layout/SBox.cpp b/Engine/Source/Runtime/Slate/Private/Widgets/Layout/SBox.cpp
new file mode 100644
--- /dev/null
+++ b/Engine/Source/Runtime/Slate/Private/Widgets/Layout/SBox.cpp
@@ -0,0 +1,89 @@
+// Copyright 1998-2014 Epic Games, Inc. All Rights Reserved.
+
+#include "SlatePrivatePCH.h"
+#include "LayoutUtils.h"
+
+void SBox::Construct( const FArguments& InArgs )
+{
+	WidthOverride = InArgs._WidthOverride;
+	HeightOverride = InArgs._HeightOverride;
+
+	ChildSlot
+		.HAlign( InArgs._HAlign )
+		.VAlign( InArgs._VAlign )
+		.Padding( InArgs._Padding )
+	[
+		InArgs._Content.Widget
+	];
+}
+
+
+
+FVector2D SBox::ComputeDesiredSize() const
+{
+	EVisibility ChildVisibility = ChildSlot.Widget->GetVisibility();
+
+	if ( ChildVisibility != EVisibility::Collapsed )
+	{
+		// If the user specified a fixed width or height, those values override the Box's content.
+		const FVector2D& UnmodifiedChildDesiredSize = ChildSlot.Widget->GetDesiredSize() + ChildSlot.SlotPadding.Get().GetDesiredSize();
+		const FOptionalSize CurrentWidthOverride = WidthOverride.Get();
+		const FOptionalSize CurrentHeightOverride = HeightOverride.Get();
+		return FVector2D(
+			(CurrentWidthOverride.IsSet()) ? CurrentWidthOverride.Get() : UnmodifiedChildDesiredSize.X,
+			(CurrentHeightOverride.IsSet()) ? CurrentHeightOverride.Get() : UnmodifiedChildDesiredSize.Y
+		);
+	}
+	
+	return FVector2D::ZeroVector;
+}
+
+void SBox::ArrangeChildren( const FGeometry& AllottedGeometry, FArrangedChildren& ArrangedChildren ) const
+{
+	const EVisibility& MyCurrentVisibility = this->GetVisibility();
+	if ( ArrangedChildren.Accepts( MyCurrentVisibility ) )
+	{
+		const FMargin SlotPadding(ChildSlot.SlotPadding.Get());
+		AlignmentArrangeResult XAlignmentResult = AlignChild<Orient_Horizontal>( AllottedGeometry.Size.X, ChildSlot, SlotPadding );
+		AlignmentArrangeResult YAlignmentResult = AlignChild<Orient_Vertical>( AllottedGeometry.Size.Y, ChildSlot, SlotPadding );
+
+		ArrangedChildren.AddWidget(
+			AllottedGeometry.MakeChild(
+				ChildSlot.Widget,
+				FVector2D(XAlignmentResult.Offset, YAlignmentResult.Offset),
+				FVector2D(XAlignmentResult.Size, YAlignmentResult.Size)
+			)
+		);
+	}
+}
+
+FChildren* SBox::GetChildren()
+{
+	return &ChildSlot;
+}
+
+int32 SBox::OnPaint( const FGeometry& AllottedGeometry, const FSlateRect& MyClippingRect, FSlateWindowElementList& OutDrawElements, int32 LayerId, const FWidgetStyle& InWidgetStyle, bool bParentEnabled ) const
+{
+	// An SBox just draws its only child
+	FArrangedChildren ArrangedChildren(EVisibility::Visible);
+	{
+#if SLATE_HD_STATS
+		SCOPE_CYCLE_COUNTER( STAT_SlateOnPaint_SBox );
+#endif
+		this->ArrangeChildren(AllottedGeometry, ArrangedChildren);
+	}
+
+	// Maybe none of our children are visible
+	if( ArrangedChildren.Num() > 0 )
+	{
+		check( ArrangedChildren.Num() == 1 );
+		FArrangedWidget& TheChild = ArrangedChildren(0);
+
+		const FSlateRect ChildClippingRect = AllottedGeometry.GetClippingRect().InsetBy( ChildSlot.SlotPadding.Get() * AllottedGeometry.Scale ).IntersectionWith(MyClippingRect);
+
+		return TheChild.Widget->OnPaint( TheChild.Geometry, ChildClippingRect, OutDrawElements, LayerId, InWidgetStyle, ShouldBeEnabled( bParentEnabled ) );
+	}
+	return LayerId;
+}
+
+
diff --git a/Engine/Source/Runtime/Slate/Private/Widgets/Layout/SDPIScaler.cpp b/Engine/Source/Runtime/Slate/Private/Widgets/Layout/SDPIScaler.cpp
new file mode 100644
--- /dev/null
+++ b/Engine/Source/Runtime/Slate/Private/Widgets/Layout/SDPIScaler.cpp
@@ -0,0 +1,38 @@
+// Copyright 1998-2014 Epic Games, Inc. All Rights Reserved.
+
+#include "SlatePrivatePCH.h"
+
+
+void SDPIScaler::Construct( const FArguments& InArgs )
+{
+	this->ChildSlot.Widget = InArgs._Content.Widget;
+	this->DPIScale = InArgs._DPIScale;
+}
+
+void SDPIScaler::ArrangeChildren( const FGeometry& AllottedGeometry, FArrangedChildren& ArrangedChildren ) const
+{
+	const EVisibility MyVisibility = this->GetVisibility();
+	if ( ArrangedChildren.Accepts( MyVisibility ) )
+	{
+		const float MyDPIScale = DPIScale.Get();
+
+		ArrangedChildren.AddWidget( AllottedGeometry.MakeChild(
+			this->ChildSlot.Widget,
+			FVector2D::ZeroVector,
+			AllottedGeometry.Size / MyDPIScale,
+			MyDPIScale
+		));
+
+	}
+}
+	
+FVector2D SDPIScaler::ComputeDesiredSize() const
+{
+	return DPIScale.Get() * ChildSlot.Widget->GetDesiredSize();
+}
+
+FChildren* SDPIScaler::GetChildren()
+{
+	return &ChildSlot;
+}
+
diff --git a/Engine/Source/Runtime/Slate/Private/Widgets/Layout/SExpandableArea.cpp b/Engine/Source/Runtime/Slate/Private/Widgets/Layout/SExpandableArea.cpp
new file mode 100644
--- /dev/null
+++ b/Engine/Source/Runtime/Slate/Private/Widgets/Layout/SExpandableArea.cpp
@@ -0,0 +1,194 @@
+// Copyright 1998-2014 Epic Games, Inc. All Rights Reserved.
+
+/*=============================================================================
+	SExpandableArea.cpp: Implements the SExpandableArea class.
+=============================================================================*/
+
+#include "SlatePrivatePCH.h"
+
+
+#define LOCTEXT_NAMESPACE "SExpandableArea"
+
+
+void SExpandableArea::Construct( const FArguments& InArgs )
+{
+	check(InArgs._Style);
+
+	bAreaCollapsed = InArgs._InitiallyCollapsed;
+	MaxHeight = InArgs._MaxHeight;
+	OnAreaExpansionChanged = InArgs._OnAreaExpansionChanged;
+	CollapsedImage = &InArgs._Style->CollapsedImage;
+	ExpandedImage = &InArgs._Style->ExpandedImage;
+
+	// If it should be initially visible, show it now
+	RolloutCurve = FCurveSequence(0.0f, 0.1f, ECurveEaseFunction::QuadOut);
+
+	if (!bAreaCollapsed)
+	{
+		RolloutCurve.JumpToEnd();
+	}
+
+	TSharedRef<SWidget> HeaderContent = InArgs._HeaderContent.Widget;
+	if( HeaderContent == SNullWidget::NullWidget )
+	{
+		HeaderContent = 
+			SNew(STextBlock)
+			.Text(InArgs._AreaTitle)
+			.Font(InArgs._AreaTitleFont)
+			.ShadowOffset(FVector2D(1.0f, 1.0f));
+	}
+
+	ChildSlot
+	[
+		SNew( SBorder )
+		.BorderImage( InArgs._BorderImage )
+		.BorderBackgroundColor( InArgs._BorderBackgroundColor )
+		.Padding(0.0f)
+		[
+			SNew(SVerticalBox)
+			+ SVerticalBox::Slot()
+			.AutoHeight()
+			[
+				SNew( SButton )
+				.ButtonStyle(FCoreStyle::Get(), "NoBorder")
+				.ForegroundColor(FSlateColor::UseForeground())
+				.OnClicked( this, &SExpandableArea::OnHeaderClicked )
+				[
+					ConstructHeaderWidget( InArgs, HeaderContent )
+				]
+			]
+			+ SVerticalBox::Slot()
+			.FillHeight(1.0f)
+			[
+				SNew(SBorder)
+				.Visibility(this, &SExpandableArea::OnGetContentVisibility)
+				.BorderImage(FStyleDefaults::GetNoBrush())
+				.Padding(InArgs._Padding)
+				.DesiredSizeScale(this, &SExpandableArea::GetSectionScale)
+				[
+					InArgs._BodyContent.Widget
+				]
+			]
+		]
+	];
+}
+
+void SExpandableArea::SetExpanded( bool bExpanded )
+{
+	// If setting to expanded and we are collapsed, change the state
+	if ( bAreaCollapsed == bExpanded )
+	{
+		bAreaCollapsed = !bExpanded;
+
+		if ( bExpanded )
+		{
+			RolloutCurve.JumpToEnd();
+		}
+		else
+		{
+			RolloutCurve.JumpToStart();
+		}
+
+		// Allow some section-specific code to be executed when the section becomes visible or collapsed
+		OnAreaExpansionChanged.ExecuteIfBound( bExpanded );
+	}
+}
+
+TSharedRef<SWidget> SExpandableArea::ConstructHeaderWidget( const FArguments& InArgs, TSharedRef<SWidget> HeaderContent )
+{
+	return
+		SNew(SHorizontalBox)
+		+ SHorizontalBox::Slot()
+		.AutoWidth()
+		.Padding(0, 0, 3, 0)
+		.VAlign(VAlign_Center)
+		[
+			SNew(SImage)
+			.Image(this, &SExpandableArea::OnGetCollapseImage)
+			.ColorAndOpacity(FSlateColor::UseForeground())
+		]
+		+ SHorizontalBox::Slot()
+		.FillWidth(1.0f)
+		.VAlign(VAlign_Center)
+		[
+			HeaderContent
+		];
+}
+
+EVisibility SExpandableArea::OnGetContentVisibility() const
+{
+	float Scale = GetSectionScale().Y;
+	// The section is visible if its scale in Y is greater than 0
+	return Scale > 0 ? EVisibility::Visible : EVisibility::Collapsed;
+}
+
+/*
+FReply SExpandableArea::OnMouseDown( const FGeometry& MyGeometry, const FPointerEvent& MouseEvent )
+{
+	if ( MouseEvent.GetEffectingButton() == EKeys::LeftMouseButton )
+	{
+		//we need to capture the mouse for MouseUp events
+		return FReply::Handled().CaptureMouse( HeaderBorder.ToSharedRef() ).SetKeyboardFocus( AsShared(), EKeyboardFocusCause::Mouse );
+	}
+
+	return FReply::Unhandled();
+}*/
+
+
+FReply SExpandableArea::OnHeaderClicked()
+{
+	// Only expand/collapse if the mouse was pressed down and up inside the area
+	OnToggleContentVisibility();
+
+	return FReply::Handled();
+}
+
+void SExpandableArea::OnToggleContentVisibility()
+{
+	bAreaCollapsed = !bAreaCollapsed;
+
+	if( !bAreaCollapsed )
+	{
+		RolloutCurve = FCurveSequence( 0.0f, 0.1f, ECurveEaseFunction::CubicOut );
+		RolloutCurve.Play();
+	}
+	else
+	{
+		RolloutCurve = FCurveSequence( 0.0f, 0.1f, ECurveEaseFunction::CubicIn );
+		RolloutCurve.PlayReverse();
+	}
+
+	// Allow some section-specific code to be executed when the section becomes visible or collapsed
+	OnAreaExpansionChanged.ExecuteIfBound( !bAreaCollapsed );
+}
+
+const FSlateBrush* SExpandableArea::OnGetCollapseImage() const
+{
+	return bAreaCollapsed ? CollapsedImage : ExpandedImage;
+}
+
+FVector2D SExpandableArea::GetSectionScale() const
+{
+	float Scale = RolloutCurve.GetLerp();
+	return FVector2D( 1.0f, Scale );
+}
+
+FVector2D SExpandableArea::ComputeDesiredSize() const
+{
+	EVisibility ChildVisibility = ChildSlot.Widget->GetVisibility();
+	if ( ChildVisibility != EVisibility::Collapsed )
+	{
+		FVector2D DesiredSize = ChildSlot.Widget->GetDesiredSize() + ChildSlot.SlotPadding.Get().GetDesiredSize();
+		
+		// Only clamp if the user specified a max height
+		if( MaxHeight > 0.0f )
+		{
+			DesiredSize.Y = FMath::Min( MaxHeight, DesiredSize.Y );
+		}
+		return DesiredSize;
+	}
+
+	return FVector2D::ZeroVector;
+}
+
+#undef LOCTEXT_NAMESPACE
\ No newline at end of file
diff --git a/Engine/Source/Runtime/Slate/Private/Widgets/Layout/SFxWidget.cpp b/Engine/Source/Runtime/Slate/Private/Widgets/Layout/SFxWidget.cpp
new file mode 100644
--- /dev/null
+++ b/Engine/Source/Runtime/Slate/Private/Widgets/Layout/SFxWidget.cpp
@@ -0,0 +1,125 @@
+// Copyright 1998-2014 Epic Games, Inc. All Rights Reserved.
+
+#include "SlatePrivatePCH.h"
+
+
+void SFxWidget::Construct( const FArguments& InArgs )
+{
+	RenderScale = InArgs._RenderScale;
+	RenderScaleOrigin = InArgs._RenderScaleOrigin;
+	LayoutScale = InArgs._LayoutScale;
+	VisualOffset = InArgs._VisualOffset;
+	bIgnoreClipping = InArgs._IgnoreClipping;
+	ColorAndOpacity = InArgs._ColorAndOpacity;
+
+	this->ChildSlot
+	.HAlign(HAlign_Center)
+	.VAlign(VAlign_Center)
+	[
+		InArgs._Content.Widget
+	];
+}
+
+void SFxWidget::SetVisualOffset( const TAttribute<FVector2D>& InOffset )
+{
+	VisualOffset = InOffset;
+}
+
+void SFxWidget::SetVisualOffset( FVector InOffset )
+{
+	VisualOffset = FVector2D(InOffset.X, InOffset.Y);
+}
+
+void SFxWidget::SetRenderScale( const TAttribute<float>& InScale )
+{
+	RenderScale = InScale;
+}
+
+void SFxWidget::SetRenderScale( float InScale )
+{
+	RenderScale = InScale;
+}
+
+void SFxWidget::SetColorAndOpacity( const TAttribute<FLinearColor>& InColorAndOpacity )
+{
+	ColorAndOpacity = InColorAndOpacity;
+}
+
+void SFxWidget::SetColorAndOpacity( FLinearColor InColorAndOpacity )
+{
+	ColorAndOpacity = InColorAndOpacity;
+}
+
+/**
+ * Transform the geometry for rendering only:
+ * 
+ * @param RenderScale     Zoom by this amount.
+ * @param ScaleOrigin     Where the origin of the transforms is within the object (range 0..1)
+ * @param Offset          Offset in either direction as multiple of unscaled object size.
+ */
+static FGeometry MakeScaledGeometry( const FGeometry& Geometry, const float RenderScale, const FVector2D& ScaleOrigin, const FVector2D& Offset )
+{
+	const FVector2D CenteringAdjustment = Geometry.Size*ScaleOrigin*Geometry.Scale - Geometry.Size*ScaleOrigin*Geometry.Scale*RenderScale;
+	const FVector2D OffsetAdjustment = Geometry.Size*Offset*Geometry.Scale;
+	
+	return FGeometry(
+		FVector2D::ZeroVector,
+		Geometry.AbsolutePosition + CenteringAdjustment + OffsetAdjustment,
+		Geometry.Size,
+		Geometry.Scale * RenderScale );
+}
+
+int32 SFxWidget::OnPaint( const FGeometry& AllottedGeometry, const FSlateRect& MyClippingRect, FSlateWindowElementList& OutDrawElements, int32 LayerId, const FWidgetStyle& InWidgetStyle, bool bParentEnabled ) const
+{
+	// Transformt the geometry for rendering only. Layout is unaffected.
+	const FGeometry ModifiedGeometry = MakeScaledGeometry( AllottedGeometry, RenderScale.Get(), RenderScaleOrigin.Get(), VisualOffset.Get() );
+	
+	FArrangedChildren ArrangedChildren(EVisibility::Visible);
+	this->ArrangeChildren(ModifiedGeometry, ArrangedChildren);
+
+	// There may be zero elements in this array if our child collapsed/hidden
+	if( ArrangedChildren.Num() > 0 )
+	{
+		// We can only have one direct descendant.
+		check( ArrangedChildren.Num() == 1 );
+		const FArrangedWidget& TheChild = ArrangedChildren(0);
+
+		// SFxWidgets are able to ignore parent clipping.
+		const FSlateRect ChildClippingRect = (bIgnoreClipping.Get())
+			? ModifiedGeometry.GetClippingRect()
+			: MyClippingRect.IntersectionWith(ModifiedGeometry.GetClippingRect());
+
+		FWidgetStyle CompoundedWidgetStyle = FWidgetStyle(InWidgetStyle)
+			.BlendColorAndOpacityTint(ColorAndOpacity.Get())
+			.SetForegroundColor( ForegroundColor );
+
+		return TheChild.Widget->OnPaint( TheChild.Geometry, ChildClippingRect, OutDrawElements, LayerId + 1, CompoundedWidgetStyle, ShouldBeEnabled( bParentEnabled ) );
+	}
+	return LayerId;
+
+}
+
+
+FVector2D SFxWidget::ComputeDesiredSize() const
+{
+	// Layout scale affects out desired size.
+	return LayoutScale.Get() * ChildSlot.Widget->GetDesiredSize();
+}
+
+
+void SFxWidget::ArrangeChildren( const FGeometry& AllottedGeometry, FArrangedChildren& ArrangedChildren ) const
+{
+	const EVisibility MyVisibility = this->GetVisibility();
+	if ( ArrangedChildren.Accepts( MyVisibility ) )
+	{
+		const float MyDPIScale = LayoutScale.Get();
+
+		// Only layout scale affects the arranged geometry. This part is identical to DPI scaling.
+		ArrangedChildren.AddWidget( AllottedGeometry.MakeChild(
+			this->ChildSlot.Widget,
+			FVector2D::ZeroVector,
+			AllottedGeometry.Size / MyDPIScale,
+			MyDPIScale
+		));
+	}
+}
diff --git a/Engine/Source/Runtime/Slate/Private/Widgets/Layout/SGridPanel.cpp b/Engine/Source/Runtime/Slate/Private/Widgets/Layout/SGridPanel.cpp
new file mode 100644
--- /dev/null
+++ b/Engine/Source/Runtime/Slate/Private/Widgets/Layout/SGridPanel.cpp
@@ -0,0 +1,394 @@
+// Copyright 1998-2014 Epic Games, Inc. All Rights Reserved.
+
+#include "SlatePrivatePCH.h"
+#include "LayoutUtils.h"
+
+
+SGridPanel::FSlot& SGridPanel::AddSlot( int32 Column, int32 Row, SGridPanel::Layer InLayer )
+{
+	return InsertSlot( new FSlot( Column, Row, InLayer.TheLayer ) );
+}
+
+void SGridPanel::ClearChildren()
+{
+	Columns.Empty();
+	Rows.Empty();
+	Slots.Empty();
+}
+
+void SGridPanel::Construct( const FArguments& InArgs )
+{
+	TotalDesiredSizes = FVector2D::ZeroVector;
+
+	// Populate the slots such that they are sorted by Layer (order preserved within layers)
+	// Also determine the grid size
+	for (int32 SlotIndex = 0; SlotIndex < InArgs.Slots.Num(); ++SlotIndex )
+	{
+		InsertSlot( InArgs.Slots[SlotIndex] );
+	}
+
+	ColFillCoefficients = InArgs.ColFillCoefficients;
+	RowFillCoefficients = InArgs.RowFillCoefficients;
+
+//	check( ColFillCoefficients.Num() <= Columns.Num() );
+//	check( RowFillCoefficients.Num() <= Rows.Num() );
+}
+
+
+int32 SGridPanel::OnPaint( const FGeometry& AllottedGeometry, const FSlateRect& MyClippingRect, FSlateWindowElementList& OutDrawElements, int32 LayerId, const FWidgetStyle& InWidgetStyle, bool bParentEnabled ) const
+{
+	FArrangedChildren ArrangedChildren(EVisibility::All);
+	this->ArrangeChildren(AllottedGeometry, ArrangedChildren);
+
+	// Because we paint multiple children, we must track the maximum layer id that they produced in case one of our parents
+	// wants to an overlay for all of its contents.
+	int32 MaxLayerId = LayerId;
+
+	// We need to iterate over slots, because slots know the GridLayers. This isn't available in the arranged children.
+	// Some slots do not show up (they are hidden/collapsed). We need a 2nd index to skip over them.
+	//
+	// GridLayers must ensure that everything in LayerN is below LayerN+1. In other words,
+	// every grid layer group must start at the current MaxLayerId (similar to how SOverlay works).
+	int32 LastGridLayer = 0;
+	for (int32 ChildIndex = 0; ChildIndex < Slots.Num(); ++ChildIndex)
+	{
+		FArrangedWidget& CurWidget = ArrangedChildren(ChildIndex);
+		if (CurWidget.Widget->GetVisibility().IsVisible())
+		{
+			const FSlot& CurSlot = Slots[ChildIndex];
+
+			FSlateRect ChildClipRect = MyClippingRect.IntersectionWith( CurWidget.Geometry.GetClippingRect() );
+
+			if ( LastGridLayer != CurSlot.LayerParam )
+			{
+				// We starting a new grid layer group?
+				LastGridLayer = CurSlot.LayerParam;
+				// Ensure that everything here is drawn on top of 
+				// previously drawn grid content.
+				LayerId = MaxLayerId+1;
+			}
+
+			const int32 CurWidgetsMaxLayerId = CurWidget.Widget->OnPaint(
+				CurWidget.Geometry,
+				ChildClipRect,
+				OutDrawElements,
+				LayerId,
+				InWidgetStyle,
+				ShouldBeEnabled( bParentEnabled )
+			);
+			
+			MaxLayerId = FMath::Max( MaxLayerId, CurWidgetsMaxLayerId );
+		}
+	}
+
+//#define LAYOUT_DEBUG
+
+#ifdef LAYOUT_DEBUG
+	LayerId = LayoutDebugPaint( AllottedGeometry, MyClippingRect, OutDrawElements, LayerId );
+#endif
+
+	return MaxLayerId;
+}
+
+
+void SGridPanel::ArrangeChildren( const FGeometry& AllottedGeometry, FArrangedChildren& ArrangedChildren ) const
+{
+	// PREPARE PHASE
+	// Prepare some data for arranging children.
+	// FinalColumns will be populated with column sizes that include the stretched column sizes.
+	// Then we will build partial sums so that we can easily handle column spans.
+	// Repeat the same for rows.
+
+	float ColumnCoeffTotal = 0.0f;
+	TArray<float> FinalColumns;
+	if ( Columns.Num() > 0 )
+	{
+		FinalColumns.AddUninitialized(Columns.Num());
+		FinalColumns[FinalColumns.Num()-1] = 0.0f;
+	}
+
+	float RowCoeffTotal = 0.0f;
+	TArray<float> FinalRows;
+	if ( Rows.Num() > 0 )
+	{
+		FinalRows.AddUninitialized(Rows.Num());
+		FinalRows[FinalRows.Num()-1] = 0.0f;
+	}
+	
+	FVector2D FlexSpace = AllottedGeometry.Size;
+	const int32 ColFillCoeffsLength = ColFillCoefficients.Num();
+	for(int32 ColIndex=0; ColIndex < Columns.Num(); ++ColIndex)
+	{
+		// Compute the total space available for stretchy columns.
+		if (ColIndex >= ColFillCoeffsLength || ColFillCoefficients[ColIndex] == 0)
+		{
+			FlexSpace.X -= Columns[ColIndex];
+		}
+		else //(ColIndex < ColFillCoeffsLength)
+		{
+			// Compute the denominator for dividing up the stretchy column space
+			ColumnCoeffTotal += ColFillCoefficients[ColIndex];
+		}
+
+	}
+
+	for(int32 ColIndex=0; ColIndex < Columns.Num(); ++ColIndex)
+	{
+		// Figure out how big each column needs to be
+		FinalColumns[ColIndex] = (ColIndex < ColFillCoeffsLength && ColFillCoefficients[ColIndex] != 0)
+			? (ColFillCoefficients[ColIndex] / ColumnCoeffTotal * FlexSpace.X)
+			: Columns[ColIndex];
+	}
+
+	const int32 RowFillCoeffsLength = RowFillCoefficients.Num();
+	for(int32 RowIndex=0; RowIndex < Rows.Num(); ++RowIndex)
+	{
+		// Compute the total space available for stretchy rows.
+		if (RowIndex >= RowFillCoeffsLength || RowFillCoefficients[RowIndex] == 0)
+		{
+			FlexSpace.Y -= Rows[RowIndex];
+		}
+		else //(RowIndex < RowFillCoeffsLength)
+		{
+			// Compute the denominator for dividing up the stretchy row space
+			RowCoeffTotal += RowFillCoefficients[RowIndex];
+		}
+
+	}
+
+	for(int32 RowIndex=0; RowIndex < Rows.Num(); ++RowIndex)
+	{
+		// Compute how big each row needs to be
+		FinalRows[RowIndex] = (RowIndex < RowFillCoeffsLength && RowFillCoefficients[RowIndex] != 0)
+			? (RowFillCoefficients[RowIndex] / RowCoeffTotal * FlexSpace.Y)
+			: Rows[RowIndex];
+	}
+	
+	// Build up partial sums for row and column sizes so that we can handle column and row spans conveniently.
+	ComputePartialSums(FinalColumns);
+	ComputePartialSums(FinalRows);
+	
+	// ARRANGE PHASE
+	for( int32 SlotIndex=0; SlotIndex < Slots.Num(); ++SlotIndex )
+	{
+		const FSlot& CurSlot = Slots[SlotIndex];
+		const EVisibility ChildVisibility = CurSlot.Widget->GetVisibility();
+		if ( ArrangedChildren.Accepts(ChildVisibility) )
+		{
+			// Figure out the position of this cell.
+			const FVector2D ThisCellOffset( FinalColumns[CurSlot.ColumnParam], FinalRows[CurSlot.RowParam] );
+			// Figure out the size of this slot; takes row span into account.
+			// We use the properties of partial sums arrays to achieve this.
+			const FVector2D CellSize(
+				FinalColumns[CurSlot.ColumnParam+CurSlot.ColumnSpanParam] - ThisCellOffset.X ,
+				FinalRows[CurSlot.RowParam+CurSlot.RowSpanParam] - ThisCellOffset.Y );
+
+			// Do the standard arrangement of elements within a slot
+			// Takes care of alignment and padding.
+			const FMargin SlotPadding(CurSlot.SlotPadding.Get());
+			AlignmentArrangeResult XAxisResult = AlignChild<Orient_Horizontal>( CellSize.X, CurSlot, SlotPadding );
+			AlignmentArrangeResult YAxisResult = AlignChild<Orient_Vertical>( CellSize.Y, CurSlot, SlotPadding );
+
+			// Output the result
+			ArrangedChildren.AddWidget( ChildVisibility, AllottedGeometry.MakeChild( 
+				CurSlot.Widget,
+				ThisCellOffset + FVector2D( XAxisResult.Offset, YAxisResult.Offset ) + CurSlot.NudgeParam,
+				FVector2D(XAxisResult.Size, YAxisResult.Size)
+			));
+		}
+	}
+}
+
+
+void SGridPanel::CacheDesiredSize()
+{
+	// The desired size of the grid is the sum of the desires sizes for every row and column.
+	ComputeDesiredCellSizes( Columns, Rows );
+
+	TotalDesiredSizes = FVector2D::ZeroVector;
+	for (int ColId=0; ColId < Columns.Num(); ++ColId)
+	{
+		TotalDesiredSizes.X += Columns[ColId];
+	}
+
+	for(int RowId=0; RowId < Rows.Num(); ++RowId)
+	{
+		TotalDesiredSizes.Y += Rows[RowId];
+	}
+	
+	SPanel::CacheDesiredSize();
+}
+
+
+FVector2D SGridPanel::ComputeDesiredSize() const
+{
+	return TotalDesiredSizes;
+}
+
+
+FChildren* SGridPanel::GetChildren()
+{
+	return &Slots;
+}
+
+
+FVector2D SGridPanel::GetDesiredSize( const FIntPoint& StartCell, int32 Width, int32 Height ) const
+{
+	if (Columns.Num() > 0 && Rows.Num() > 0)
+	{
+		const int32 FirstColumn = FMath::Clamp(StartCell.X, 0, Columns.Num()-1);
+		const int32 LastColumn = FMath::Clamp(StartCell.X + Width, 0, Columns.Num()-1);
+
+		const int32 FirstRow = FMath::Clamp(StartCell.Y, 0, Rows.Num()-1);
+		const int32 LastRow = FMath::Clamp(StartCell.Y + Height, 0,  Rows.Num()-1);
+
+		return FVector2D( Columns[LastColumn] - Columns[FirstColumn], Rows[LastRow] - Rows[FirstRow] );
+	}
+	else
+	{
+		return FVector2D::ZeroVector;
+	}
+}
+
+void SGridPanel::SetColumnFill( int32 ColumnId, float Coefficient )
+{
+	if (ColFillCoefficients.Num() <= ColumnId)
+	{
+		ColFillCoefficients.AddZeroed( ColumnId - ColFillCoefficients.Num() + 1 );
+	}
+	ColFillCoefficients[ColumnId] = Coefficient;
+}
+
+void SGridPanel::SetRowFill( int32 RowId, float Coefficient )
+{
+	if (RowFillCoefficients.Num() <= RowId)
+	{
+		RowFillCoefficients.AddZeroed( RowId - RowFillCoefficients.Num() + 1 );
+	}
+	RowFillCoefficients[RowId] = Coefficient;
+}
+
+void SGridPanel::ComputePartialSums( TArray<float>& TurnMeIntoPartialSums )
+{
+	// We assume there is a 0-valued item already at the end of this array.
+	// We need it so that we can  compute the original values
+	// by doing Array[N] - Array[N-1];
+	
+
+	float LastValue = 0;
+	float SumSoFar = 0;
+	for(int32 Index=0; Index < TurnMeIntoPartialSums.Num(); ++Index)
+	{
+		LastValue = TurnMeIntoPartialSums[Index];
+		TurnMeIntoPartialSums[Index] = SumSoFar;
+		SumSoFar += LastValue;
+	}
+}
+
+
+void SGridPanel::DistributeSizeContributions( float SizeContribution, TArray<float>& DistributeOverMe, int32 StartIndex, int32 UpperBound )
+{
+	for ( int32 Index = StartIndex; Index < UpperBound; ++Index )
+	{
+		// Each column or row only needs to get bigger if its current size does not already accommodate it.
+		DistributeOverMe[Index] = FMath::Max( SizeContribution, DistributeOverMe[Index] );
+	}
+}
+
+
+SGridPanel::FSlot& SGridPanel::InsertSlot( SGridPanel::FSlot* InSlot )
+{
+	bool bInserted = false;
+
+	InSlot->Panel = SharedThis(this);
+
+	// Insert the slot in the list such that slots are sorted by LayerOffset.
+	for( int32 SlotIndex=0; !bInserted && SlotIndex < Slots.Num(); ++SlotIndex )
+	{
+		if ( InSlot->LayerParam < this->Slots[SlotIndex].LayerParam )
+		{
+			Slots.Insert( InSlot, SlotIndex );
+			bInserted = true;
+		}
+	}
+
+	// We haven't inserted yet, so add to the end of the list.
+	if ( !bInserted )
+	{
+		Slots.Add( InSlot );
+	}
+
+	NotifySlotChanged(InSlot);
+
+	return *InSlot;
+}
+
+void SGridPanel::NotifySlotChanged(SGridPanel::FSlot* InSlot)
+{
+	// Keep the size of the grid up to date.
+	// We need an extra cell at the end for easily figuring out the size across any number of cells
+	// by doing Columns[End] - Columns[Start] or Rows[End] - Rows[Start].
+	// The first Columns[]/Rows[] entry will be 0.
+	const int32 NumColumnsRequiredForThisSlot = InSlot->ColumnParam + InSlot->ColumnSpanParam + 1;
+	if ( NumColumnsRequiredForThisSlot > Columns.Num() )
+	{
+		Columns.AddZeroed( NumColumnsRequiredForThisSlot - Columns.Num() );
+	}
+
+	const int32 NumRowsRequiredForThisSlot = InSlot->RowParam + InSlot->RowSpanParam + 1;
+	if ( NumRowsRequiredForThisSlot > Rows.Num() )
+	{
+		Rows.AddZeroed( NumRowsRequiredForThisSlot - Rows.Num() );
+	}
+}
+
+
+void SGridPanel::ComputeDesiredCellSizes( TArray<float>& OutColumns, TArray<float>& OutRows ) const
+{
+	FMemory::Memzero( OutColumns.GetData(), OutColumns.Num() * sizeof(float) );
+	FMemory::Memzero( OutRows.GetData(), OutRows.Num() * sizeof(float) );
+
+	for (int32 SlotIndex=0; SlotIndex<Slots.Num(); ++SlotIndex)
+	{
+		const FSlot& CurSlot = Slots[SlotIndex];
+		if (CurSlot.Widget->GetVisibility() != EVisibility::Collapsed)
+		{
+			// The slots wants to be as big as its content along with the required padding.
+			const FVector2D SlotDesiredSize = CurSlot.Widget->GetDesiredSize() + CurSlot.SlotPadding.Get().GetDesiredSize();
+
+			// If the slot has a (colspan,rowspan) of (1,1) it will only affect that slot.
+			// For larger spans, the slots size will be evenly distributed across all the affected slots.
+			const FVector2D SizeContribution( SlotDesiredSize.X / CurSlot.ColumnSpanParam, SlotDesiredSize.Y / CurSlot.RowSpanParam );
+
+			// Distribute the size contributions over all the columns and rows that this slot spans
+			DistributeSizeContributions( SizeContribution.X, OutColumns, CurSlot.ColumnParam, CurSlot.ColumnParam + CurSlot.ColumnSpanParam );
+			DistributeSizeContributions( SizeContribution.Y, OutRows, CurSlot.RowParam, CurSlot.RowParam + CurSlot.RowSpanParam );
+		}
+	}
+}
+
+
+int32 SGridPanel::LayoutDebugPaint(const FGeometry& AllottedGeometry, const FSlateRect& MyClippingRect, FSlateWindowElementList& OutDrawElements, int32 LayerId ) const
+{
+	float XOffset = 0;
+	for (int32 Column=0; Column<Columns.Num(); ++Column)
+	{
+		float YOffset = 0;
+		for (int32 Row=0; Row<Rows.Num(); ++Row)
+		{
+			FSlateDrawElement::MakeDebugQuad
+			(
+				OutDrawElements, 
+				LayerId,
+				AllottedGeometry.ToPaintGeometry( FVector2D(XOffset, YOffset), FVector2D( Columns[Column], Rows[Row] ) ),
+				MyClippingRect
+			);
+
+			YOffset += Rows[Row];
+		}
+		XOffset += Columns[Column];
+	}
+
+	return LayerId;
+}
+
diff --git a/Engine/Source/Runtime/Slate/Private/Widgets/Layout/SHeader.cpp b/Engine/Source/Runtime/Slate/Private/Widgets/Layout/SHeader.cpp
new file mode 100644
--- /dev/null
+++ b/Engine/Source/Runtime/Slate/Private/Widgets/Layout/SHeader.cpp
@@ -0,0 +1,53 @@
+// Copyright 1998-2014 Epic Games, Inc. All Rights Reserved.
+
+#include "SlatePrivatePCH.h"
+
+
+void SHeader::Construct( const FArguments& InArgs )
+{
+	SHorizontalBox::FSlot* FirstSlot = NULL;
+	SHorizontalBox::FSlot* LastSlot = NULL;
+
+	SHorizontalBox::Construct( SHorizontalBox::FArguments()
+		+SHorizontalBox::Slot()
+		.Expose( FirstSlot )
+		.FillWidth(1)
+		.VAlign(VAlign_Center)
+		[
+			SNew(SSeparator) .SeparatorImage(FCoreStyle::Get().GetBrush("Header.Pre")) .Orientation(Orient_Horizontal)
+		]
+		+SHorizontalBox::Slot()
+		.AutoWidth()
+		.VAlign(VAlign_Center)
+		.Padding(5, 0)
+		[
+			InArgs._Content.Widget
+		]
+		+SHorizontalBox::Slot()
+		.Expose( LastSlot )
+		.FillWidth(1)
+		.VAlign(VAlign_Center)
+		[
+			SNew(SSeparator) .SeparatorImage(FCoreStyle::Get().GetBrush("Header.Post")) .Orientation(Orient_Horizontal)
+		]	
+	);
+
+	// Either the left or right side of the header needs to be auto-sized based on the alignment of the content.
+	// @see the HAlign argument
+	switch( InArgs._HAlign )
+	{
+		default:
+		case HAlign_Center:
+		case HAlign_Fill:
+		break;
+
+		case HAlign_Left:
+			FirstSlot->AutoWidth();
+		break;
+
+		case HAlign_Right:
+			LastSlot->AutoWidth();
+		break;
+
+	};
+}
diff --git a/Engine/Source/Runtime/Slate/Private/Widgets/Layout/SMissingWidget.cpp b/Engine/Source/Runtime/Slate/Private/Widgets/Layout/SMissingWidget.cpp
new file mode 100644
--- /dev/null
+++ b/Engine/Source/Runtime/Slate/Private/Widgets/Layout/SMissingWidget.cpp
@@ -0,0 +1,26 @@
+// Copyright 1998-2014 Epic Games, Inc. All Rights Reserved.
+
+/*=============================================================================
+	SMissingWidget.h: Declares the SMissingWidget class.
+=============================================================================*/
+
+#include "SlatePrivatePCH.h"
+
+
+#define LOCTEXT_NAMESPACE "SMissingWidget"
+
+
+TSharedRef<class SWidget> SMissingWidget::MakeMissingWidget( )
+{
+	return SNew(SBorder)
+		.HAlign(HAlign_Center)
+		.VAlign(VAlign_Center)
+		[
+			SNew(STextBlock)
+				.Text(LOCTEXT("MissingWidget", "Missing Widget"))
+				.TextStyle(FCoreStyle::Get(), "EmbossedText")
+		];
+}
+
+
+#undef LOCTEXT_NAMESPACE
diff --git a/Engine/Source/Runtime/Slate/Private/Widgets/Layout/SScrollBar.cpp b/Engine/Source/Runtime/Slate/Private/Widgets/Layout/SScrollBar.cpp
new file mode 100644
--- /dev/null
+++ b/Engine/Source/Runtime/Slate/Private/Widgets/Layout/SScrollBar.cpp
@@ -0,0 +1,523 @@
+// Copyright 1998-2014 Epic Games, Inc. All Rights Reserved.
+ 
+#include "SlatePrivatePCH.h"
+
+
+static const int32 NUM_SCROLLBAR_SLOTS = 3;
+
+
+/** Arrange 3 widgets: the top track, bottom track, and thumb. */
+class SScrollBarTrack : public SPanel
+{
+public:
+
+	/** A ListPanel slot is very simple - it just stores a widget. */
+	class FSlot
+	{
+	public:		
+		FSlot()
+			: Widget( SNullWidget::NullWidget )
+		{
+		}
+
+		FSlot& operator[]( TSharedRef<SWidget> InWidget )
+		{
+			Widget = InWidget;
+			return *this;
+		}
+
+		TSharedRef<SWidget> Widget;
+	};
+
+	SLATE_BEGIN_ARGS( SScrollBarTrack )
+		: _TopSlot()
+		, _ThumbSlot()
+		, _BottomSlot()
+		, _Orientation(Orient_Vertical)
+		{
+			_Visibility = EVisibility::SelfHitTestInvisible;
+		}
+
+		SLATE_NAMED_SLOT( FArguments, TopSlot)
+		SLATE_NAMED_SLOT( FArguments, ThumbSlot)
+		SLATE_NAMED_SLOT( FArguments, BottomSlot)
+		SLATE_ARGUMENT( EOrientation, Orientation )
+
+	SLATE_END_ARGS()
+
+	/**
+	 * Construct the widget from a Declaration
+	 *
+	 * @param InArgs  Declaration from which to construct the widget.
+	 */
+	void Construct( const FArguments& InArgs)
+	{
+		OffsetFraction = 0;
+		ThumbSizeFraction = 0;
+		MinThumbSize = 35;
+		Orientation = InArgs._Orientation;
+
+		// This panel only supports 3 children!
+		for (int32 CurChild=0; CurChild < NUM_SCROLLBAR_SLOTS; ++CurChild)
+		{
+			Children.Add( new FSlot() );
+		}
+
+		Children[ TOP_SLOT_INDEX ]
+		[
+			InArgs._TopSlot.Widget
+		];
+		
+		Children[ BOTTOM_SLOT_INDEX ]
+		[
+			InArgs._BottomSlot.Widget
+		];
+		
+		Children[ THUMB_SLOT_INDEX ]
+		[
+			InArgs._ThumbSlot.Widget
+		];
+	}
+
+	struct FTrackSizeInfo
+	{
+		FTrackSizeInfo( const FGeometry& TrackGeometry, EOrientation InOrientation, float InMinThumbSize, float ThumbSizeAsFractionOfTrack, float ThumbOffsetAsFractionOfTrack )
+		{
+			BiasedTrackSize = ((InOrientation == Orient_Horizontal) ? TrackGeometry.Size.X : TrackGeometry.Size.Y) - InMinThumbSize;
+			const float AccurateThumbSize = ThumbSizeAsFractionOfTrack * (BiasedTrackSize);
+			ThumbStart = BiasedTrackSize * ThumbOffsetAsFractionOfTrack;
+			ThumbSize = InMinThumbSize + AccurateThumbSize;
+			
+		}
+
+		float BiasedTrackSize;
+		float ThumbStart;
+		float ThumbSize;
+		float GetThumbEnd()
+		{
+			return ThumbStart + ThumbSize;
+		}
+	};
+
+	FTrackSizeInfo GetTrackSizeInfo( const FGeometry& InTrackGeometry ) const
+	{
+		return FTrackSizeInfo( InTrackGeometry, Orientation, MinThumbSize, this->ThumbSizeFraction, this->OffsetFraction );
+	}
+
+	/**
+	 * Panels arrange their children in a space described by the AllottedGeometry parameter. The results of the arrangement
+	 * should be returned by appending a FArrangedWidget pair for every child widget. See StackPanel for an example
+	 *
+	 * @param AllottedGeometry    The geometry allotted for this widget by its parent.
+	 * @param ArrangedChildren    The array to which to add the WidgetGeometries that represent the arranged children.
+	 */
+	virtual void ArrangeChildren( const FGeometry& AllottedGeometry, FArrangedChildren& ArrangedChildren ) const
+	{
+		const float Width = AllottedGeometry.Size.X;
+		const float Height = AllottedGeometry.Size.Y;
+
+		FTrackSizeInfo TrackSizeInfo = this->GetTrackSizeInfo(AllottedGeometry);
+
+		// Arrange top half of the track
+		FVector2D ChildSize = (Orientation == Orient_Horizontal)
+			? FVector2D(TrackSizeInfo.ThumbStart, Height)
+			: FVector2D(Width, TrackSizeInfo.ThumbStart);
+
+		FVector2D ChildPos(0, 0);
+		ArrangedChildren.AddWidget(			
+			AllottedGeometry.MakeChild( Children[TOP_SLOT_INDEX].Widget, ChildPos, ChildSize )
+		);
+
+		// Arrange bottom half of the track
+		ChildPos = (Orientation == Orient_Horizontal)
+			? FVector2D(TrackSizeInfo.GetThumbEnd(), 0)
+			: FVector2D(0, TrackSizeInfo.GetThumbEnd());
+
+		ChildSize = (Orientation == Orient_Horizontal)
+			? FVector2D(AllottedGeometry.Size.X - TrackSizeInfo.GetThumbEnd(), Height)
+			: FVector2D(Width, AllottedGeometry.Size.Y - TrackSizeInfo.GetThumbEnd());
+
+		ArrangedChildren.AddWidget( 
+			AllottedGeometry.MakeChild( Children[BOTTOM_SLOT_INDEX].Widget, ChildPos, ChildSize )
+		);
+		
+		// Arrange the thumb; only show it in enabled scrollbars
+		if ( IsNeeded() )
+		{
+			ChildPos = (Orientation == Orient_Horizontal)
+				? FVector2D(TrackSizeInfo.ThumbStart, 0)
+				: FVector2D(0, TrackSizeInfo.ThumbStart);
+
+			ChildSize = (Orientation == Orient_Horizontal)
+				? FVector2D(TrackSizeInfo.ThumbSize, Height)
+				: FVector2D(Width, TrackSizeInfo.ThumbSize);
+
+			ArrangedChildren.AddWidget(
+				AllottedGeometry.MakeChild( Children[THUMB_SLOT_INDEX].Widget, ChildPos, ChildSize )
+			);
+		}
+
+	}
+	
+	/**
+	 * A Panel's desired size in the space required to arrange of its children on the screen while respecting all of
+	 * the children's desired sizes and any layout-related options specified by the user. See StackPanel for an example.
+	 *
+	 * @return The desired size.
+	 */
+	virtual FVector2D ComputeDesiredSize() const
+	{
+		if ( Orientation == Orient_Horizontal )
+		{
+			const float DesiredHeight = FMath::Max3( Children[0].Widget->GetDesiredSize().Y, Children[1].Widget->GetDesiredSize().Y, Children[2].Widget->GetDesiredSize().Y );
+			return FVector2D( MinThumbSize, DesiredHeight );
+		}
+		else
+		{
+			const float DesiredWidth = FMath::Max3( Children[0].Widget->GetDesiredSize().X, Children[1].Widget->GetDesiredSize().X, Children[2].Widget->GetDesiredSize().X );
+			return FVector2D( DesiredWidth, MinThumbSize );
+		}
+	}
+	
+	/** @return  The children of a panel in a slot-agnostic way. */
+	virtual FChildren* GetChildren()
+	{
+		return &Children;
+	}
+
+	void SetSizes( float InThumbOffsetFraction, float InThumbSizeFraction )
+	{
+		OffsetFraction = InThumbOffsetFraction;
+		ThumbSizeFraction = InThumbSizeFraction;
+	}
+
+	bool IsNeeded() const
+	{
+		// We use a small epsilon here to avoid the scroll bar showing up when all of the content is already in view, due to
+		// floating point precision when the scroll bar state is set
+		return ThumbSizeFraction < ( 1.0f - KINDA_SMALL_NUMBER );
+	}
+
+	float DistanceFromTop() const
+	{
+		return OffsetFraction;
+	}
+
+	float DistanceFromBottom() const
+	{
+		return 1.0f - (OffsetFraction + ThumbSizeFraction);
+	}
+
+	float GetMinThumbSize() const
+	{
+		return MinThumbSize;
+	}
+
+	float GetThumbSizeFraction() const
+	{
+		return ThumbSizeFraction;
+	}
+
+protected:
+
+	static const int32 TOP_SLOT_INDEX = 0;
+	static const int32 BOTTOM_SLOT_INDEX = 1;
+	static const int32 THUMB_SLOT_INDEX = 2;
+	
+	TPanelChildren< FSlot > Children;
+
+	float OffsetFraction;
+	float ThumbSizeFraction;
+	float MinThumbSize;
+	EOrientation Orientation;
+};
+
+
+/**
+ * Construct this widget
+ *
+ * @param	InArgs	The declaration data for this widget
+ */
+void SScrollBar::Construct(const FArguments& InArgs)
+{
+	check(InArgs._Style);
+
+	OnUserScrolled = InArgs._OnUserScrolled;
+	Orientation = InArgs._Orientation;
+	UserVisibility = InArgs._Visibility;
+
+	NormalThumbImage = &InArgs._Style->NormalThumbImage;
+	HoveredThumbImage = &InArgs._Style->HoveredThumbImage;
+	DraggedThumbImage = &InArgs._Style->DraggedThumbImage;
+
+	EHorizontalAlignment HorizontalAlignment;
+	EVerticalAlignment VerticalAlignment;
+	const FSlateBrush* Brush;
+
+	if(Orientation == Orient_Vertical)
+	{
+		HorizontalAlignment = HAlign_Center;
+		VerticalAlignment = VAlign_Fill;
+		Brush = &InArgs._Style->VerticalBackgroundImage;
+	}
+	else
+	{
+		HorizontalAlignment = HAlign_Fill;
+		VerticalAlignment = VAlign_Center;
+		Brush = &InArgs._Style->HorizontalBackgroundImage;
+	}
+
+	SBorder::Construct( SBorder::FArguments()
+		.BorderImage(FCoreStyle::Get().GetBrush("NoBorder"))
+		.BorderBackgroundColor( this, &SScrollBar::GetTrackOpacity )
+		.ColorAndOpacity( this, &SScrollBar::GetThumbOpacity )
+		[
+			SNew(SVerticalBox)
+			+SVerticalBox::Slot() .FillHeight( 1 )
+			[
+				SAssignNew(Track, SScrollBarTrack)
+				.Orientation(InArgs._Orientation)
+				.TopSlot()
+				[
+					SNew(SBox)
+					.HAlign(HorizontalAlignment)
+					.VAlign(VerticalAlignment)
+					[
+						SNew(SImage)
+						.Image(Brush)
+					]
+				]
+				.ThumbSlot()
+				[
+					SAssignNew(DragThumb , SBorder)
+					.BorderImage( this, &SScrollBar::GetDragThumbImage )
+					.HAlign(HAlign_Center)
+					.VAlign(VAlign_Center)
+					[
+						SNew(SSpacer)
+						.Size(InArgs._Thickness)
+					]
+				]
+				.BottomSlot()
+				[
+					SNew(SBox)
+					.HAlign(HorizontalAlignment)
+					.VAlign(VerticalAlignment)
+					[
+						SNew(SImage)
+						.Image(Brush)
+					]
+				]
+			]
+		]
+	);
+
+	this->EnabledState = TAttribute<bool>( Track.ToSharedRef(), &SScrollBarTrack::IsNeeded );
+	if ( ! InArgs._AlwaysShowScrollbar )
+	{
+		this->Visibility = TAttribute<EVisibility>( SharedThis(this), &SScrollBar::ShouldBeVisible );
+	}
+}
+
+void SScrollBar::SetOnUserScrolled( const FOnUserScrolled& InHandler )
+{
+	OnUserScrolled = InHandler;
+}
+
+void SScrollBar::SetState( float InOffsetFraction, float InThumbSizeFraction )
+{
+	// Note that the maximum offset depends on how many items fit per screen
+	// It is 1.0f-InThumbSizeFraction.
+	Track->SetSizes( InOffsetFraction, InThumbSizeFraction );
+}
+
+/**
+ * The system calls this method to notify the widget that a mouse button was pressed within it. This event is bubbled.
+ *
+ * @param MyGeometry The Geometry of the widget receiving the event
+ * @param MouseEvent Information about the input event
+ *
+ * @return Whether the event was handled along with possible requests for the system to take action.
+ */
+FReply SScrollBar::OnMouseButtonDown( const FGeometry& MyGeometry, const FPointerEvent& MouseEvent )
+{
+	if ( MouseEvent.GetEffectingButton() == EKeys::LeftMouseButton )
+	{
+		FGeometry ThumbGeometry = this->FindChildGeometry(MyGeometry, DragThumb.ToSharedRef());
+
+		if( DragThumb->IsHovered() )
+		{
+			// Clicking on the scrollbar drag thumb
+			if( Orientation == Orient_Horizontal )
+			{
+				DragGrabOffset = ThumbGeometry.AbsoluteToLocal( MouseEvent.GetScreenSpacePosition() ).X;
+			}
+			else
+			{
+				DragGrabOffset = ThumbGeometry.AbsoluteToLocal( MouseEvent.GetScreenSpacePosition() ).Y;
+			}
+
+			bDraggingThumb = true;
+		}
+		else if( OnUserScrolled.IsBound() )
+		{
+			// Clicking in the non drag thumb area of the scrollbar
+			DragGrabOffset = Orientation == Orient_Horizontal ? (ThumbGeometry.Size.X * 0.5f) : (ThumbGeometry.Size.Y * 0.5f);
+			bDraggingThumb = true;
+
+			ExecuteOnUserScrolled( MyGeometry, MouseEvent );
+		}
+	}
+
+	if( bDraggingThumb )
+	{
+		return FReply::Handled().CaptureMouse(AsShared()).SetKeyboardFocus( AsShared(), EKeyboardFocusCause::Mouse );
+	}
+	else
+	{
+		return FReply::Unhandled();
+	}
+}
+
+/**
+ * The system calls this method to notify the widget that a mouse button was release within it. This event is bubbled.
+ *
+ * @param MyGeometry The Geometry of the widget receiving the event
+ * @param MouseEvent Information about the input event
+ *
+ * @return Whether the event was handled along with possible requests for the system to take action.
+ */
+FReply SScrollBar::OnMouseButtonUp( const FGeometry& MyGeometry, const FPointerEvent& MouseEvent )
+{
+	if ( MouseEvent.GetEffectingButton() == EKeys::LeftMouseButton )
+	{
+		bDraggingThumb = false;
+		return FReply::Handled().ReleaseMouseCapture();
+	}
+	else
+	{
+		return FReply::Unhandled();
+	}
+}
+
+/**
+ * The system calls this method to notify the widget that a mouse moved within it. This event is bubbled.
+ *
+ * @param MyGeometry The Geometry of the widget receiving the event
+ * @param MouseEvent Information about the input event
+ *
+ * @return Whether the event was handled along with possible requests for the system to take action.
+ */
+FReply SScrollBar::OnMouseMove( const FGeometry& MyGeometry, const FPointerEvent& MouseEvent )
+{
+	if ( this->HasMouseCapture() && !MouseEvent.GetCursorDelta().IsZero() )
+	{
+		if ( OnUserScrolled.IsBound() )
+		{
+			ExecuteOnUserScrolled( MyGeometry, MouseEvent );
+		}
+		return FReply::Handled();
+	}
+	else
+	{
+		return FReply::Unhandled();
+	}
+}
+
+void SScrollBar::ExecuteOnUserScrolled( const FGeometry& MyGeometry, const FPointerEvent& MouseEvent )
+{
+	const int32 AxisId = (Orientation == Orient_Horizontal) ? 0 : 1;
+	const FGeometry TrackGeometry = FindChildGeometry( MyGeometry, Track.ToSharedRef() );
+	const float UnclampedOffsetInTrack = TrackGeometry.AbsoluteToLocal( MouseEvent.GetScreenSpacePosition() )[ AxisId ] - DragGrabOffset;
+	const float TrackSizeBiasedByMinThumbSize = TrackGeometry.Size[ AxisId ] - Track->GetMinThumbSize();
+	const float ThumbOffsetInTrack = FMath::Clamp( UnclampedOffsetInTrack, 0.0f, TrackSizeBiasedByMinThumbSize );
+	const float ThumbOffset = ThumbOffsetInTrack / TrackSizeBiasedByMinThumbSize;
+	OnUserScrolled.ExecuteIfBound( ThumbOffset );
+}
+
+bool SScrollBar::IsNeeded() const
+{
+	return Track->IsNeeded();
+}
+
+float SScrollBar::DistanceFromTop() const
+{
+	return Track->DistanceFromTop();
+}
+
+float SScrollBar::DistanceFromBottom() const
+{
+	return Track->DistanceFromBottom();
+}
+
+SScrollBar::SScrollBar()
+: bDraggingThumb(false)
+, DragGrabOffset( 0.0f )
+{
+
+}
+
+
+FSlateColor SScrollBar::GetTrackOpacity() const
+{
+	if ( bDraggingThumb || this->IsHovered() )
+	{
+		return FLinearColor(1,1,1,1);
+	}
+	else
+	{
+		return FLinearColor(1,1,1,0.0f);
+	}
+}
+
+
+FLinearColor SScrollBar::GetThumbOpacity() const
+{
+	if ( bDraggingThumb || this->IsHovered() )
+	{
+		return FLinearColor(1,1,1,1);
+	}
+	else
+	{
+		return FLinearColor(1,1,1,0.75f);
+	}
+}
+
+
+const FSlateBrush* SScrollBar::GetDragThumbImage() const
+{
+	if ( bDraggingThumb )
+	{
+		return DraggedThumbImage;
+	}
+	else if (DragThumb->IsHovered())
+	{
+		return HoveredThumbImage;
+	}
+	else
+	{
+		return NormalThumbImage;
+	}
+}
+
+
+EVisibility SScrollBar::ShouldBeVisible() const
+{
+	if ( this->HasMouseCapture() )
+	{
+		return EVisibility::Visible;
+	}
+	else if( Track->IsNeeded() )
+	{
+		return UserVisibility.Get();
+	}
+	else
+	{
+		return EVisibility::Collapsed;
+	}
+}
+
+bool SScrollBar::IsScrolling() const
+{
+	return bDraggingThumb;
+}
diff --git a/Engine/Source/Runtime/Slate/Private/Widgets/Layout/SScrollBox.cpp b/Engine/Source/Runtime/Slate/Private/Widgets/Layout/SScrollBox.cpp
new file mode 100644
--- /dev/null
+++ b/Engine/Source/Runtime/Slate/Private/Widgets/Layout/SScrollBox.cpp
@@ -0,0 +1,425 @@
+// Copyright 1998-2014 Epic Games, Inc. All Rights Reserved.
+
+#include "SlatePrivatePCH.h"
+#include "LayoutUtils.h"
+
+
+SScrollBox::FSlot& SScrollBox::Slot()
+{
+	return *(new SScrollBox::FSlot());
+}
+
+
+class SScrollPanel : public SPanel
+{
+public:
+	
+	SLATE_BEGIN_ARGS(SScrollPanel)
+		{
+			_Visibility = EVisibility::SelfHitTestInvisible;
+		}
+	SLATE_END_ARGS()
+
+
+
+	void Construct( const FArguments& InArgs, const TArray<SScrollBox::FSlot*>& InSlots )
+	{
+		PhysicalOffset = 0;
+		Children.Reserve( InSlots.Num() );
+		for ( int32 SlotIndex=0; SlotIndex < InSlots.Num(); ++SlotIndex )
+		{
+			Children.Add( InSlots[ SlotIndex ] );
+		}
+	}
+
+
+
+public:
+
+	virtual void ArrangeChildren( const FGeometry& AllottedGeometry, FArrangedChildren& ArrangedChildren ) const OVERRIDE
+	{
+		
+		float CurChildOffset = -PhysicalOffset;
+
+		for( int32 SlotIndex=0; SlotIndex < Children.Num(); ++SlotIndex )
+		{
+			const SScrollBox::FSlot& ThisSlot = Children[SlotIndex];
+			const EVisibility ChildVisibility = ThisSlot.Widget->GetVisibility();
+
+			if ( ChildVisibility != EVisibility::Collapsed )
+			{
+				const FMargin& ThisPadding = ThisSlot.SlotPadding.Get();
+				const FVector2D& WidgetDesiredSize = ThisSlot.Widget->GetDesiredSize();
+				const float ThisSlotDesiredHeight = WidgetDesiredSize.Y + ThisSlot.SlotPadding.Get().GetTotalSpaceAlong<Orient_Vertical>();
+
+				// Figure out the size and local position of the child within the slot
+				// There is no vertical alignment, because it does not make sense in a panel where items are stacked end-to-end
+				AlignmentArrangeResult XAlignmentResult = AlignChild<Orient_Horizontal>( AllottedGeometry.Size.X, ThisSlot, ThisPadding );
+
+				ArrangedChildren.AddWidget( AllottedGeometry.MakeChild( ThisSlot.Widget, FVector2D(XAlignmentResult.Offset, CurChildOffset + ThisPadding.Top), FVector2D(XAlignmentResult.Size, WidgetDesiredSize.Y) ) );
+				CurChildOffset += ThisSlotDesiredHeight;
+			}
+		}
+
+
+	}
+
+
+
+	virtual FVector2D ComputeDesiredSize() const OVERRIDE
+	{
+		FVector2D ThisDesiredSize = FVector2D::ZeroVector;
+		for(int32 SlotIndex=0; SlotIndex < Children.Num(); ++SlotIndex )
+		{
+			const SScrollBox::FSlot& ThisSlot = Children[SlotIndex];
+			const FVector2D ChildDesiredSize = ThisSlot.Widget->GetDesiredSize();
+			ThisDesiredSize.X  = FMath::Max(ChildDesiredSize.X, ThisDesiredSize.X);
+			ThisDesiredSize.Y += ChildDesiredSize.Y + ThisSlot.SlotPadding.Get().GetTotalSpaceAlong<Orient_Vertical>();
+		}
+
+		return ThisDesiredSize;
+	}
+
+
+
+	virtual FChildren* GetChildren() OVERRIDE
+	{
+		return &Children;
+	}
+
+
+	float PhysicalOffset;
+	TPanelChildren<SScrollBox::FSlot> Children;	
+	
+};
+
+
+
+void SScrollBox::Construct( const FArguments& InArgs )
+{
+	check(InArgs._Style);
+
+	DesiredScrollOffset = 0;
+	bIsScrolling = false;
+	bAnimateScroll = false;
+	AmountScrolledWhileRightMouseDown = 0;
+	bShowSoftwareCursor = false;
+	SoftwareCursorPosition = FVector2D::ZeroVector;
+
+	this->ChildSlot
+	[
+		SNew(SHorizontalBox)
+		+SHorizontalBox::Slot()
+		.FillWidth(1)
+		[
+			SNew(SOverlay)
+			+SOverlay::Slot()
+			.Padding( FMargin(0.0f, 0.0f, 0.0f, 1.0f) )
+			[
+				// Scroll panel that presents the scrolled content
+				SAssignNew( ScrollPanel, SScrollPanel, InArgs.Slots )
+			]
+			+SOverlay::Slot()
+			.HAlign(HAlign_Fill)
+			.VAlign(VAlign_Top)
+			[
+				// Shadow: Hint to scroll up
+				SNew(SImage)
+				.Visibility(EVisibility::HitTestInvisible)
+				.ColorAndOpacity(this, &SScrollBox::GetTopShadowOpacity)
+				.Image( &InArgs._Style->TopShadowBrush )	
+			]
+			+SOverlay::Slot()
+			.HAlign(HAlign_Fill)
+			.VAlign(VAlign_Bottom)
+			[
+				// Shadow: a hint to scroll down
+				SNew(SImage)
+				.Visibility(EVisibility::HitTestInvisible)
+				.ColorAndOpacity(this, &SScrollBox::GetBottomShadowOpacity)
+				.Image( &InArgs._Style->BottomShadowBrush )	
+			]
+		]
+		+SHorizontalBox::Slot()
+		.AutoWidth()
+		[
+			SAssignNew( ScrollBar, SScrollBar )
+			.Thickness( FVector2D(5.0f, 5.0f) )
+			.OnUserScrolled( this, &SScrollBox::ScrollBar_OnUserScrolled )
+		]
+	];
+}
+
+
+
+/** Adds a slot to SScrollBox */
+SScrollBox::FSlot& SScrollBox::AddSlot()
+{
+	SScrollBox::FSlot& NewSlot = SScrollBox::Slot();
+	ScrollPanel->Children.Add( &NewSlot );
+
+	return NewSlot;
+}
+
+
+
+/** Removes a slot at the specified location */
+void SScrollBox::RemoveSlot( const TSharedRef<SWidget>& WidgetToRemove )
+{
+	TPanelChildren<SScrollBox::FSlot>& Children = ScrollPanel->Children;
+	for( int32 SlotIndex=0; SlotIndex < Children.Num(); ++SlotIndex )
+	{
+		if ( Children[SlotIndex].Widget == WidgetToRemove )
+		{
+			Children.RemoveAt(SlotIndex);
+			return;
+		}
+	}
+}
+
+void SScrollBox::ClearChildren()
+{
+	ScrollPanel->Children.Empty();
+}
+
+bool SScrollBox::IsRightClickScrolling() const
+{
+	return AmountScrolledWhileRightMouseDown >= SlatePanTriggerDistance && this->ScrollBar->IsNeeded();
+}
+
+void SScrollBox::SetScrollOffset( float NewScrollOffset )
+{
+	DesiredScrollOffset = NewScrollOffset;
+}
+
+void SScrollBox::Tick( const FGeometry& AllottedGeometry, const double InCurrentTime, const float InDeltaTime )
+{
+	if ( !IsRightClickScrolling() )
+	{
+		this->InertialScrollManager.UpdateScrollVelocity(InDeltaTime);
+
+		const float ScrollVelocity = this->InertialScrollManager.GetScrollVelocity();
+		if ( ScrollVelocity != 0.f )
+		{
+			this->ScrollBy(AllottedGeometry, ScrollVelocity * InDeltaTime, true);
+		}
+	}
+
+	const FGeometry ScrollPanelGeometry = FindChildGeometry( AllottedGeometry, ScrollPanel.ToSharedRef() );
+	
+	const float ContentHeight = ScrollPanel->GetDesiredSize().Y;
+	const float ViewFraction = ScrollPanelGeometry.Size.Y / ContentHeight;
+	const float ViewOffset = FMath::Clamp<float>( DesiredScrollOffset/ContentHeight, 0.0, 1.0 - ViewFraction );
+	
+	// Update the scrollbar with the clamped version of the offset
+	const float TargetPhysicalOffset = ViewOffset*ScrollPanel->GetDesiredSize().Y;
+	bIsScrolling = !FMath::IsNearlyEqual(TargetPhysicalOffset, ScrollPanel->PhysicalOffset, 0.001f);	
+	ScrollPanel->PhysicalOffset = (bAnimateScroll)
+		? FMath::FInterpTo(ScrollPanel->PhysicalOffset, TargetPhysicalOffset, InDeltaTime, 15.f)
+		: TargetPhysicalOffset;
+	
+	ScrollBar->SetState(ViewOffset, ViewFraction);
+	if (!ScrollBar->IsNeeded())
+	{
+		// We cannot scroll, so ensure that there is no offset.
+		ScrollPanel->PhysicalOffset = 0.0f;
+	}
+
+}
+
+
+FReply SScrollBox::OnMouseButtonDown( const FGeometry& MyGeometry, const FPointerEvent& MouseEvent )
+{
+	// Zero the scroll velocity so the panel stops immediately on mouse down, even if the user does not drag
+	this->InertialScrollManager.ClearScrollVelocity();
+
+	if ( MouseEvent.GetEffectingButton() == EKeys::RightMouseButton && ScrollBar->IsNeeded() )
+	{
+		AmountScrolledWhileRightMouseDown = 0;
+
+		// NOTE: We don't bother capturing the mouse, unless the user starts dragging a few pixels (see the
+		// mouse move handling here.)  This is important so that the item row has a chance to select
+		// items when the right mouse button is released.  Just keep in mind that you might not get
+		// an OnMouseButtonUp event for the right mouse button if the user moves off of the table before
+		// they reach our scroll threshold
+		return FReply::Handled();
+	}
+	return FReply::Unhandled();	
+}
+
+
+FReply SScrollBox::OnMouseButtonUp( const FGeometry& MyGeometry, const FPointerEvent& MouseEvent )
+{
+	if ( MouseEvent.GetEffectingButton() == EKeys::RightMouseButton )
+	{
+		AmountScrolledWhileRightMouseDown = 0;
+
+		FReply Reply = FReply::Handled().ReleaseMouseCapture();
+		bShowSoftwareCursor = false;
+
+		// If we have mouse capture, snap the mouse back to the closest location that is within the panel's bounds
+		if ( HasMouseCapture() )
+		{
+			FSlateRect PanelScreenSpaceRect = MyGeometry.GetClippingRect();
+			FVector2D CursorPosition = MyGeometry.LocalToAbsolute( SoftwareCursorPosition );
+
+			FIntPoint BestPositionInPanel(
+				FMath::Round( FMath::Clamp( CursorPosition.X, PanelScreenSpaceRect.Left, PanelScreenSpaceRect.Right ) ),
+				FMath::Round( FMath::Clamp( CursorPosition.Y, PanelScreenSpaceRect.Top, PanelScreenSpaceRect.Bottom ) )
+				);
+
+			Reply.SetMousePos(BestPositionInPanel);
+		}
+
+		return Reply;
+	}
+	return FReply::Unhandled();
+}
+
+
+FReply SScrollBox::OnMouseMove( const FGeometry& MyGeometry, const FPointerEvent& MouseEvent )
+{	
+	if( MouseEvent.IsMouseButtonDown( EKeys::RightMouseButton ) )
+	{
+		const float ScrollByAmount = MouseEvent.GetCursorDelta().Y;
+		// If scrolling with the right mouse button, we need to remember how much we scrolled.
+		// If we did not scroll at all, we will bring up the context menu when the mouse is released.
+		AmountScrolledWhileRightMouseDown += FMath::Abs( ScrollByAmount );
+
+		// Has the mouse moved far enough with the right mouse button held down to start capturing
+		// the mouse and dragging the view?
+		if( IsRightClickScrolling() )
+		{
+			this->InertialScrollManager.AddScrollSample(-ScrollByAmount, FPlatformTime::Seconds());
+			const bool bDidScroll = this->ScrollBy( MyGeometry, -ScrollByAmount );
+
+			FReply Reply = FReply::Handled();
+
+			// Capture the mouse if we need to
+			if( FSlateApplication::Get().GetMouseCaptor().Get() != this )
+			{
+				Reply.CaptureMouse( AsShared() ).UseHighPrecisionMouseMovement( AsShared() );
+				SoftwareCursorPosition = MyGeometry.AbsoluteToLocal( MouseEvent.GetScreenSpacePosition() );
+				bShowSoftwareCursor = true;
+			}
+
+			// Check if the mouse has moved.
+			if( bDidScroll )
+			{
+				SoftwareCursorPosition.Y += MouseEvent.GetCursorDelta().Y;
+			}
+
+			return Reply;
+		}
+	}
+
+	return FReply::Unhandled();
+}
+
+void SScrollBox::OnMouseLeave( const FPointerEvent& MouseEvent )
+{
+	if( FSlateApplication::Get().GetMouseCaptor().Get() != this )
+	{
+		// No longer scrolling (unless we have mouse capture)
+		AmountScrolledWhileRightMouseDown = 0;
+	}
+}
+
+FReply SScrollBox::OnMouseWheel( const FGeometry& MyGeometry, const FPointerEvent& MouseEvent )
+{
+	if (ScrollBar->IsNeeded())
+	{
+		// Make sure scroll velocity is cleared so it doesn't fight with the mouse wheel input
+		this->InertialScrollManager.ClearScrollVelocity();
+
+		const bool bScrollWasHandled = this->ScrollBy( MyGeometry, -MouseEvent.GetWheelDelta()*WheelScrollAmount );
+
+		return bScrollWasHandled ? FReply::Handled() : FReply::Unhandled();
+	}
+	else
+	{
+		return FReply::Unhandled();
+	}
+}
+
+bool SScrollBox::ScrollBy( const FGeometry& AllottedGeometry, float ScrollAmount, bool InAnimateScroll )
+{
+	bAnimateScroll = InAnimateScroll;
+
+	const float ContentSize = ScrollPanel->GetDesiredSize().Y;
+	const FGeometry ScrollPanelGeometry = FindChildGeometry( AllottedGeometry, ScrollPanel.ToSharedRef() );
+
+	DesiredScrollOffset += ScrollAmount;
+
+	const float ScrollMin = 0.0f, ScrollMax = ContentSize - ScrollPanelGeometry.Size.Y;
+	if (DesiredScrollOffset < ScrollMin || DesiredScrollOffset > ScrollMax)
+	{
+		DesiredScrollOffset = FMath::Clamp( DesiredScrollOffset, ScrollMin, ScrollMax );
+		return false;
+	}
+	return true;
+}
+
+
+FReply SScrollBox::OnDragDetected( const FGeometry& MyGeometry, const FPointerEvent& MouseEvent )
+{
+	return FReply::Handled().CaptureMouse( SharedThis(this) );
+}
+
+FCursorReply SScrollBox::OnCursorQuery( const FGeometry& MyGeometry, const FPointerEvent& CursorEvent ) const
+{
+	if ( IsRightClickScrolling() )
+	{
+		// We hide the native cursor as we'll be drawing the software EMouseCursor::GrabHandClosed cursor
+		return FCursorReply::Cursor( EMouseCursor::None );
+	}
+	else
+	{
+		return FCursorReply::Unhandled();
+	}
+}
+
+int32 SScrollBox::OnPaint( const FGeometry& AllottedGeometry, const FSlateRect& MyClippingRect, FSlateWindowElementList& OutDrawElements, int32 LayerId, const FWidgetStyle& InWidgetStyle, bool bParentEnabled ) const
+{
+	int32 NewLayerId = SCompoundWidget::OnPaint( AllottedGeometry, MyClippingRect, OutDrawElements, LayerId, InWidgetStyle, bParentEnabled );
+
+	if( !bShowSoftwareCursor )
+	{
+		return NewLayerId;
+	}
+
+	const FSlateBrush* Brush = FCoreStyle::Get().GetBrush(TEXT("SoftwareCursor_Grab"));
+
+	FSlateDrawElement::MakeBox(
+		OutDrawElements,
+		++NewLayerId,
+		AllottedGeometry.ToPaintGeometry( SoftwareCursorPosition - ( Brush->ImageSize / 2 ), Brush->ImageSize ),
+		Brush,
+		MyClippingRect
+		);
+
+	return NewLayerId;
+}
+
+void SScrollBox::ScrollBar_OnUserScrolled( float InScrollOffsetFraction )
+{
+	DesiredScrollOffset = InScrollOffsetFraction * ScrollPanel->GetDesiredSize().Y;
+}
+
+const float ShadowFadeDistance = 32.0f;
+
+FSlateColor SScrollBox::GetTopShadowOpacity() const
+{
+	// The shadow should only be visible when the user needs a hint that they can scroll up.
+	const float ShadowOpacity = FMath::Clamp( ScrollPanel->PhysicalOffset/ShadowFadeDistance, 0.0f, 1.0f);
+	
+	return FLinearColor(1.0f, 1.0f, 1.0f, ShadowOpacity);
+}
+
+FSlateColor SScrollBox::GetBottomShadowOpacity() const
+{
+	// The shadow should only be visible when the user needs a hint that they can scroll down.
+	const float ShadowOpacity = (ScrollBar->DistanceFromBottom() * ScrollPanel->GetDesiredSize().Y) / ShadowFadeDistance;
+	
+	return FLinearColor(1.0f, 1.0f, 1.0f, ShadowOpacity);
+}
diff --git a/Engine/Source/Runtime/Slate/Private/Widgets/Layout/SSpacer.cpp b/Engine/Source/Runtime/Slate/Private/Widgets/Layout/SSpacer.cpp
new file mode 100644
--- /dev/null
+++ b/Engine/Source/Runtime/Slate/Private/Widgets/Layout/SSpacer.cpp
@@ -0,0 +1,47 @@
+// Copyright 1998-2014 Epic Games, Inc. All Rights Reserved.
+
+#include "SlatePrivatePCH.h"
+
+
+/**
+ * Construct this widget
+ *
+ * @param	InArgs	The declaration data for this widget
+ */
+void SSpacer::Construct( const FArguments& InArgs )
+{
+	SpacerSize = InArgs._Size;
+
+	SWidget::Construct( 
+		InArgs._ToolTipText,
+		InArgs._ToolTip,
+		InArgs._Cursor, 
+		InArgs._IsEnabled,
+		InArgs._Visibility,
+		InArgs._Tag);
+}
+
+
+/**
+ * The widget should respond by populating the OutDrawElements array with FDrawElements 
+ * that represent it and any of its children.
+ *
+ * @param AllottedGeometry  The FGeometry that describes an area in which the widget should appear.
+ * @param MyClippingRect    The clipping rectangle allocated for this widget and its children.
+ * @param OutDrawElements   A list of FDrawElements to populate with the output.
+ * @param LayerId           The Layer onto which this widget should be rendered.
+ * @param InColorAndOpacity Color and Opacity to be applied to all the descendants of the widget being painted
+ * @param bParentEnabled	True if the parent of this widget is enabled.
+ *
+ * @return The maximum layer ID attained by this widget or any of its children.
+ */
+int32 SSpacer::OnPaint( const FGeometry& AllottedGeometry, const FSlateRect& MyClippingRect, FSlateWindowElementList& OutDrawElements, int32 LayerId, const FWidgetStyle& InWidgetStyle, bool bParentEnabled ) const
+{
+	// We did not paint anything. The parent's current LayerId is probably the max we were looking for.
+	return LayerId;
+}
+
+FVector2D SSpacer::ComputeDesiredSize() const
+{
+	return SpacerSize.Get();
+}
