Slate: Reorganized private Slate files; proper pre-compiled header for all files.

diff --git a/Engine/Source/Runtime/Slate/Private/Widgets/Layout/SSplitter.cpp b/Engine/Source/Runtime/Slate/Private/Widgets/Layout/SSplitter.cpp
new file mode 100644
--- /dev/null
+++ b/Engine/Source/Runtime/Slate/Private/Widgets/Layout/SSplitter.cpp
@@ -0,0 +1,981 @@
+// Copyright 1998-2014 Epic Games, Inc. All Rights Reserved.
+
+#include "SlatePrivatePCH.h"
+
+
+/** The user is not allowed to make any of the splitter's children smaller than this. */
+const float MinSplitterChildLength = 20.0f;
+
+SSplitter::FSlot& SSplitter::Slot()
+{
+	return (*new FSlot());
+}
+
+
+SSplitter::FSlot& SSplitter::AddSlot( int32 AtIndex )
+{
+	FSlot& NewSlot = SSplitter::Slot();
+	if ( AtIndex == INDEX_NONE )
+	{
+		// No index was specified; just add to the end of the list.
+		this->Children.Add( &NewSlot );
+	}
+	else
+	{
+		// Add a slot at the desired location
+		this->Children.Insert( &NewSlot, AtIndex );
+	}
+
+	return NewSlot;
+}
+
+SSplitter::FSlot& SSplitter::SlotAt( int32 SlotIndex )
+{
+	return Children[SlotIndex];
+}
+
+void SSplitter::RemoveAt( int32 IndexToRemove )
+{
+	Children.RemoveAt( IndexToRemove );
+}
+
+/**
+* Construct this widget
+*
+* @param	InArgs	The declaration data for this widget
+*/
+void SSplitter::Construct( const SSplitter::FArguments& InArgs )
+{
+	check(InArgs._Style);
+
+	OnSplitterFinishedResizing = InArgs._OnSplitterFinishedResizing;
+	ResizeMode = InArgs._ResizeMode;
+	PhysicalSplitterHandleSize = InArgs._PhysicalSplitterHandleSize;
+	HitDetectionSplitterHandleSize = InArgs._HitDetectionSplitterHandleSize;
+	Orientation = InArgs._Orientation;
+	HoveredHandleIndex = INDEX_NONE;
+	bIsResizing = false;
+	Style = InArgs._Style;
+
+	for (int32 SlotIndex = 0; SlotIndex < InArgs.Slots.Num(); ++SlotIndex )
+	{
+		Children.Add( InArgs.Slots[SlotIndex] );
+	}
+}
+
+/**
+* Panels arrange their children in a space described by the AllottedGeometry parameter. The results of the arrangement
+* should be returned by appending a FArrangedWidget pair for every child widget. See StackPanel for an example
+*/
+void SSplitter::ArrangeChildren( const FGeometry& AllottedGeometry, FArrangedChildren& ArrangedChildren ) const
+{
+	const int32 AxisIndex = (Orientation == Orient_Horizontal) ? 0 : 1;
+
+	// Splitters divide the space between their children proportionately based on size coefficients.
+	// The size coefficients are usually determined by a user, who grabs the handled between the child elements
+	// and moves them to resize the space available to the children.
+	// Some children are sized automatically based on their content; those children cannot be resized.
+	//
+	// e.g.   _____________________________________ Children
+	//       /              /                  /
+	//      v              v                  v
+	//   + - - - - - + + - - - + + - - - - - - - - - - - - - - +
+	//   |           | |       | |                             |
+	//   | Child 0   | |Child1 | |  Child2                     |
+	//   + - - - - - + + - - - + + - - - - - - - - - - - - - - +
+	//                ^         ^
+	//                 \_________\___________ Resize handles.
+
+	int32 NumNonCollapsedChildren = 0;
+	int32 NumResizeableChildren = 0;
+	float CoefficientTotal = 0;
+	// Some space is claimed by non-resizeable elements (auto-sized elements)
+	float NonResizeableSpace = 0;
+	{
+		for (int32 ChildIndex=0; ChildIndex < Children.Num(); ++ChildIndex)	
+		{
+			if (Children[ChildIndex].Widget->GetVisibility() != EVisibility::Collapsed)
+			{
+				++NumNonCollapsedChildren;
+
+				if ( Children[ChildIndex].SizingRule == SSplitter::SizeToContent )
+				{
+					NonResizeableSpace += Children[ChildIndex].Widget->GetDesiredSize()[AxisIndex];
+				}
+				else // SizingRule == SSplitter::FractionOfParent
+				{
+					CoefficientTotal += Children[ChildIndex].SizeValue.Get();
+				}
+			}			
+		}
+	}
+
+
+
+	// The user-sizeable children must make room for the resize handles and for auto-sized children.
+	const float SpaceNeededForHandles = FMath::Max(0, NumNonCollapsedChildren-1) * PhysicalSplitterHandleSize;
+	const float ResizeableSpace = AllottedGeometry.Size.Component(AxisIndex) - SpaceNeededForHandles - NonResizeableSpace;
+
+	// Arrange the children horizontally or vertically.
+	float XOffset = 0;
+	for (int32 ChildIndex=0; ChildIndex < Children.Num(); ++ChildIndex)
+	{
+		const FSlot& CurSlot = Children[ChildIndex];
+
+		const float ChildSpace = ( CurSlot.SizingRule == SSplitter::SizeToContent )
+			? CurSlot.Widget->GetDesiredSize()[AxisIndex]
+		: ResizeableSpace * CurSlot.SizeValue.Get() / CoefficientTotal;
+
+		const EVisibility ChildVisibility = CurSlot.Widget->GetVisibility();
+
+		// If the output array wants arranged children
+		if ( ArrangedChildren.Accepts(ChildVisibility) )
+		{
+			if (Orientation == Orient_Horizontal)
+			{
+				ArrangedChildren.AddWidget( ChildVisibility, AllottedGeometry.MakeChild( CurSlot.Widget, FVector2D(XOffset, 0), FVector2D( ChildSpace, AllottedGeometry.Size.Y ) ) );
+			}
+			else
+			{
+				ArrangedChildren.AddWidget( ChildVisibility, AllottedGeometry.MakeChild( CurSlot.Widget, FVector2D(0, XOffset), FVector2D( AllottedGeometry.Size.X, ChildSpace ) ) );
+			}
+		}
+
+		// Advance to the next slot. If the child is collapsed, it takes up no room and does not need a splitter
+		if ( ChildVisibility != EVisibility::Collapsed)
+		{
+			XOffset += ChildSpace + PhysicalSplitterHandleSize;
+		}
+	}
+}
+
+
+int32 SSplitter::OnPaint( const FGeometry& AllottedGeometry, const FSlateRect& MyClippingRect, FSlateWindowElementList& OutDrawElements, int32 LayerId, const FWidgetStyle& InWidgetStyle, bool bParentEnabled ) const
+{
+	FArrangedChildren ArrangedChildren( EVisibility::Visible );
+	ArrangeChildren( AllottedGeometry, ArrangedChildren );
+
+	int32 MaxLayerId = PaintArrangedChildren( ArrangedChildren, MyClippingRect, OutDrawElements, LayerId, InWidgetStyle, bParentEnabled );
+
+	const FSlateBrush* NormalHandleBrush = &Style->HandleNormalBrush;
+
+	// Draw the splitter above any children
+	MaxLayerId += 1;
+
+	for( int32 ChildIndex = 0; ChildIndex < ArrangedChildren.Num(); ++ChildIndex )
+	{
+		const FGeometry& GeometryAfterSplitter = ArrangedChildren( FMath::Clamp(ChildIndex + 1, 0, ArrangedChildren.Num()-1) ).Geometry;
+
+		const float HalfHitDetectionSplitterHandleSize = ( HitDetectionSplitterHandleSize / 2 );
+		const float HalfPhysicalSplitterHandleSize = ( PhysicalSplitterHandleSize / 2 );
+
+		FVector2D HandleSize;		
+		FVector2D HandlePosition;
+		if ( Orientation == Orient_Horizontal )
+		{
+			HandleSize.Set( PhysicalSplitterHandleSize, GeometryAfterSplitter.Size.Y );
+			HandlePosition.Set( -(HalfHitDetectionSplitterHandleSize + HalfPhysicalSplitterHandleSize), 0 );
+		}
+		else
+		{
+			HandleSize.Set( GeometryAfterSplitter.Size.X, PhysicalSplitterHandleSize );
+			HandlePosition.Set( 0, -(HalfHitDetectionSplitterHandleSize + HalfPhysicalSplitterHandleSize) );
+		}
+
+		if (HoveredHandleIndex != ChildIndex)
+		{
+			FSlateDrawElement::MakeBox(
+				OutDrawElements,
+				MaxLayerId,
+				GeometryAfterSplitter.ToPaintGeometry( HandlePosition, HandleSize, 1.0f ),
+				NormalHandleBrush,
+				MyClippingRect,
+				ShouldBeEnabled( bParentEnabled ),
+				InWidgetStyle.GetColorAndOpacityTint()
+			);
+		}
+		else
+		{
+			FSlateDrawElement::MakeBox(
+				OutDrawElements,
+				MaxLayerId,
+				GeometryAfterSplitter.ToPaintGeometry( HandlePosition, HandleSize, 1.0f ),
+				&Style->HandleHighlightBrush,
+				MyClippingRect,
+				ShouldBeEnabled( bParentEnabled ),
+				InWidgetStyle.GetColorAndOpacityTint()
+			);	
+		}
+	}
+
+	return MaxLayerId;
+}
+
+template<EOrientation Orientation>
+static FVector2D ComputeDesiredSizeForSplitter( const float PhysicalSplitterHandleSize, const TPanelChildren<SSplitter::FSlot>& Children )
+{
+	FVector2D MyDesiredSize(0,0);
+
+	int32 NumNonCollapsed = 0;
+	for (int32 ChildIndex=0; ChildIndex < Children.Num(); ++ChildIndex)
+	{
+		const SSplitter::FSlot& CurSlot = Children[ChildIndex];
+		const EVisibility ChildVisibility = CurSlot.Widget->GetVisibility();
+		if ( ChildVisibility != EVisibility::Collapsed )
+		{
+			++NumNonCollapsed;
+
+			FVector2D ChildDesiredSize( CurSlot.Widget->GetDesiredSize() );
+			if ( Orientation == Orient_Horizontal )
+			{
+				MyDesiredSize.X += ChildDesiredSize.X;
+				MyDesiredSize.Y = FMath::Max(ChildDesiredSize.Y, MyDesiredSize.Y);
+			}
+			else
+			{
+				MyDesiredSize.X = FMath::Max(ChildDesiredSize.X, MyDesiredSize.X);				
+				MyDesiredSize.Y += ChildDesiredSize.Y;
+			}
+		}
+	}
+
+	float SpaceNeededForHandles = FMath::Max(0, NumNonCollapsed-1) * PhysicalSplitterHandleSize;
+	if (Orientation == Orient_Horizontal)
+	{
+		MyDesiredSize.X += SpaceNeededForHandles;
+	}
+	else
+	{
+		MyDesiredSize.Y += SpaceNeededForHandles;
+	}
+
+	return MyDesiredSize;
+}
+
+/**
+* A Panel's desired size in the space required to arrange of its children on the screen while respecting all of
+* the children's desired sizes and any layout-related options specified by the user. See StackPanel for an example.
+*/
+FVector2D SSplitter::ComputeDesiredSize() const
+{
+	FVector2D MyDesiredSize = (Orientation == Orient_Horizontal)
+		? ComputeDesiredSizeForSplitter<Orient_Horizontal>( PhysicalSplitterHandleSize, Children )
+		: ComputeDesiredSizeForSplitter<Orient_Vertical>( PhysicalSplitterHandleSize, Children );
+
+	return MyDesiredSize;
+}
+
+/**
+* All widgets must provide a way to access their children in a layout-agnostic way.
+* Panels store their children in Slots, which creates a dilemma. Most panels
+* can store their children in a TPanelChildren<Slot>, where the Slot class
+* provides layout information about the child it stores. In that case
+* GetChildren should simply return the TPanelChildren<Slot>. See StackPanel for an example.
+*/
+FChildren* SSplitter::GetChildren()
+{
+	return &Children;
+}
+
+/**
+* The system calls this method to notify the widget that a mouse button was pressed within it. This event is bubbled.
+*
+* @param MyGeometry The Geometry of the widget receiving the event
+* @param MouseEvent Information about the input event
+*
+* @return Whether the event was handled along with possible requests for the system to take action.
+*/
+FReply SSplitter::OnMouseButtonDown( const FGeometry& MyGeometry, const FPointerEvent& MouseEvent )
+{
+	if ( MouseEvent.GetEffectingButton() == EKeys::LeftMouseButton && HoveredHandleIndex != INDEX_NONE )
+	{
+		bIsResizing = true;
+		return FReply::Handled().CaptureMouse( SharedThis(this) );
+	}
+	else
+	{
+		return FReply::Unhandled();
+	}
+}
+
+/**
+* The system calls this method to notify the widget that a mouse button was release within it. This event is bubbled.
+*
+* @param MyGeometry The Geometry of the widget receiving the event
+* @param MouseEvent Information about the input event
+*
+* @return Whether the event was handled along with possible requests for the system to take action.
+*/
+FReply SSplitter::OnMouseButtonUp( const FGeometry& MyGeometry, const FPointerEvent& MouseEvent )
+{
+	if ( MouseEvent.GetEffectingButton() == EKeys::LeftMouseButton && bIsResizing == true )
+	{
+		OnSplitterFinishedResizing.ExecuteIfBound();
+
+		bIsResizing = false;
+		return FReply::Handled().ReleaseMouseCapture();
+	}
+	return FReply::Unhandled();
+}
+
+/**
+* The system calls this method to notify the widget that a mouse moved within it. This event is bubbled.
+*
+* @param MyGeometry The Geometry of the widget receiving the event
+* @param MouseEvent Information about the input event
+*
+* @return Whether the event was handled along with possible requests for the system to take action.
+*/
+FReply SSplitter::OnMouseMove( const FGeometry& MyGeometry, const FPointerEvent& MouseEvent )
+{
+	const FVector2D LocalMousePosition = MyGeometry.AbsoluteToLocal( MouseEvent.GetScreenSpacePosition() );
+
+	FArrangedChildren ArrangedChildren(EVisibility::All);
+	ArrangeChildren( MyGeometry, ArrangedChildren );
+
+	if ( bIsResizing )
+	{
+		if ( !MouseEvent.GetCursorDelta().IsZero() )
+		{
+			if (Orientation == Orient_Horizontal)
+			{
+				HandleResizing<Orient_Horizontal>( PhysicalSplitterHandleSize, ResizeMode, HoveredHandleIndex, LocalMousePosition, Children, ArrangedChildren );
+			}
+			else
+			{
+				HandleResizing<Orient_Vertical>( PhysicalSplitterHandleSize, ResizeMode, HoveredHandleIndex, LocalMousePosition, Children, ArrangedChildren );
+			}
+		}
+
+		return FReply::Handled();
+	}
+	else
+	{	
+		// Hit test which handle we are hovering over.		
+		HoveredHandleIndex = (Orientation == Orient_Horizontal)
+			? GetHandleBeingResizedFromMousePosition<Orient_Horizontal>( PhysicalSplitterHandleSize, HitDetectionSplitterHandleSize, LocalMousePosition, ArrangedChildren )
+			: GetHandleBeingResizedFromMousePosition<Orient_Vertical>( PhysicalSplitterHandleSize, HitDetectionSplitterHandleSize, LocalMousePosition, ArrangedChildren );
+
+		if (HoveredHandleIndex != INDEX_NONE)
+		{
+			if ( FindResizeableSlotBeforeHandle(HoveredHandleIndex, Children) <= INDEX_NONE || FindResizeableSlotAfterHandle(HoveredHandleIndex, Children) >= Children.Num() )
+			{
+				HoveredHandleIndex = INDEX_NONE;
+			}
+		}
+
+		return FReply::Unhandled();
+	}
+
+}
+
+void SSplitter::OnMouseLeave( const FPointerEvent& MouseEvent )
+{
+	if ( !bIsResizing )
+	{
+		HoveredHandleIndex = INDEX_NONE;
+	}
+}
+
+/**
+* The system asks each widget under the mouse to provide a cursor. This event is bubbled.
+* 
+* @return FCursorReply::Unhandled() if the event is not handled; return FCursorReply::Cursor() otherwise.
+*/
+FCursorReply SSplitter::OnCursorQuery( const FGeometry& MyGeometry, const FPointerEvent& CursorEvent ) const
+{
+	const FVector2D LocalMousePosition = MyGeometry.AbsoluteToLocal( CursorEvent.GetScreenSpacePosition() );
+
+	FArrangedChildren ArrangedChildren(EVisibility::All);
+	ArrangeChildren( MyGeometry, ArrangedChildren );
+
+	// Hit test which handle we are hovering over.		
+	const int32 CurrentHoveredHandleIndex = (Orientation == Orient_Horizontal)
+		? GetHandleBeingResizedFromMousePosition<Orient_Horizontal>( PhysicalSplitterHandleSize, HitDetectionSplitterHandleSize, LocalMousePosition, ArrangedChildren )
+		: GetHandleBeingResizedFromMousePosition<Orient_Vertical>( PhysicalSplitterHandleSize, HitDetectionSplitterHandleSize, LocalMousePosition, ArrangedChildren );
+
+	if (CurrentHoveredHandleIndex != INDEX_NONE)
+	{
+		return ( Orientation == Orient_Horizontal )
+			? FCursorReply::Cursor( EMouseCursor::ResizeLeftRight )
+			: FCursorReply::Cursor( EMouseCursor::ResizeUpDown );
+	}
+	else
+	{
+		return FCursorReply::Unhandled();
+	}
+}
+
+/**
+* Change the orientation of the splitter
+*
+* @param NewOrientation  Should the splitter be horizontal or vertical
+*/
+void SSplitter::SetOrientation( EOrientation NewOrientation )
+{
+	Orientation = NewOrientation;
+}
+
+/**
+* @return the current orientation of the splitter.
+*/
+EOrientation SSplitter::GetOrientation() const
+{
+	return Orientation;
+}
+
+
+SSplitter::SSplitter()
+	: HoveredHandleIndex( INDEX_NONE )
+	, bIsResizing( false )
+	, Orientation( Orient_Horizontal )
+	, Style( nullptr )
+{
+}
+
+
+int32 SSplitter::FindResizeableSlotBeforeHandle( int32 DraggedHandle, const TPanelChildren<FSlot>& Children )
+{
+	// Resizing collapsed or autosizing slots does not make sense (their size is predetermined).
+	// Search out from the DraggedHandle to find the first non-collapsed, non-autosizing slot we can resize.
+	int32 SlotBeforeDragHandle = DraggedHandle;
+	while ( SlotBeforeDragHandle >= 0 && (Children[SlotBeforeDragHandle].Widget->GetVisibility() == EVisibility::Collapsed || Children[SlotBeforeDragHandle].SizingRule.Get() == SSplitter::SizeToContent ) )
+	{
+		--SlotBeforeDragHandle;
+	}
+
+	return SlotBeforeDragHandle;
+}
+
+
+int32 SSplitter::FindResizeableSlotAfterHandle( int32 DraggedHandle, const TPanelChildren<FSlot>& Children )
+{
+	const int32 NumChildren = Children.Num();
+
+	// But the slots list does contain collapsed children! Make sure that we are not resizing a collapsed slot.
+	// We also cannot resize auto-sizing slots.
+	int32 SlotAfterDragHandle = DraggedHandle+1;
+	{
+		while( SlotAfterDragHandle < NumChildren && (Children[SlotAfterDragHandle].Widget->GetVisibility() == EVisibility::Collapsed || Children[SlotAfterDragHandle].SizingRule.Get() == SSplitter::SizeToContent) )
+		{
+			++SlotAfterDragHandle;
+		}
+	}
+
+	return SlotAfterDragHandle;
+}
+
+
+void SSplitter::FindAllResizeableSlotsAfterHandle( int32 DraggedHandle, const TPanelChildren<FSlot>& Children, TArray< int32 >& OutSlotIndicies )
+{
+	const int32 NumChildren = Children.Num();
+
+	for (int SlotIndex = DraggedHandle+1; SlotIndex < NumChildren; SlotIndex++)
+	{
+		if (Children[ SlotIndex ].Widget->GetVisibility() == EVisibility::Collapsed || Children[ SlotIndex ].SizingRule.Get() == SSplitter::SizeToContent)
+		{
+			continue;
+		}
+
+		OutSlotIndicies.Add( SlotIndex );
+	}
+};
+
+
+template<EOrientation SplitterOrientation>
+void SSplitter::HandleResizing( const float PhysicalSplitterHandleSize, const ESplitterResizeMode::Type ResizeMode, int32 DraggedHandle, const FVector2D& LocalMousePos, TPanelChildren<FSlot>& Children, FArrangedChildren& ChildGeometries )
+{
+	const int32 NumChildren = Children.Num();
+
+	const int32 AxisIndex = (SplitterOrientation == Orient_Horizontal) ? 0 : 1;
+
+	// Note:
+	//  - Prev vs. Next refers to the widgets in the order they are laid out (left->right, top->bottom).
+	//  - New vs. Old refers to the Old values for width/height vs. the post-resize values.
+
+	const float HandlePos = ChildGeometries(DraggedHandle+1).Geometry.Position.Component(AxisIndex) - PhysicalSplitterHandleSize / 2;
+	float Delta = LocalMousePos.Component(AxisIndex) - HandlePos;
+
+	const int32 SlotBeforeDragHandle = FindResizeableSlotBeforeHandle( DraggedHandle, Children );
+
+	TArray< int32 > SlotsAfterDragHandleIndicies;
+	if ( ResizeMode == ESplitterResizeMode::Fixed )
+	{
+		const int32 SlotAfterDragHandle = FindResizeableSlotAfterHandle( DraggedHandle, Children );
+
+		if ( SlotAfterDragHandle < NumChildren )
+		{
+			SlotsAfterDragHandleIndicies.Add( SlotAfterDragHandle );
+		}
+	}
+	else if ( ResizeMode == ESplitterResizeMode::Fill )
+	{
+		FindAllResizeableSlotsAfterHandle( DraggedHandle, Children, /*OUT*/ SlotsAfterDragHandleIndicies );
+	}
+
+	if ( SlotBeforeDragHandle >= 0 && SlotsAfterDragHandleIndicies.Num() > 0 )
+	{
+		struct FSlotInfo 
+		{
+			FSlot* Slot;
+			FGeometry* Geometry;
+			float NewSize;
+		};
+
+		TArray< FSlotInfo > SlotsAfterDragHandle;
+		for (int SlotIndex = 0; SlotIndex < SlotsAfterDragHandleIndicies.Num(); SlotIndex++)
+		{
+			FSlotInfo SlotInfo;
+
+			SlotInfo.Slot = &Children[ SlotsAfterDragHandleIndicies[ SlotIndex ] ];
+			SlotInfo.Geometry = &ChildGeometries( SlotsAfterDragHandleIndicies[ SlotIndex ] ).Geometry;
+			SlotInfo.NewSize = SlotInfo.Geometry->Size.Component( AxisIndex );
+
+			SlotsAfterDragHandle.Add( SlotInfo );
+		}
+
+		// Get references the prev and next children and their layout settings so that we can modify them.
+		FSlot& PrevChild = Children[SlotBeforeDragHandle];
+		FGeometry& PrevChildGeom = ChildGeometries(SlotBeforeDragHandle).Geometry;
+
+		// Compute the new sizes of the children
+		float NewPrevChildLength = ClampChild( PrevChildGeom.Size.Component(AxisIndex) + Delta );
+		Delta = NewPrevChildLength - PrevChildGeom.Size.Component(AxisIndex);
+
+		// Distribute the Delta across the affected slots after the drag handle
+		float UnusedDelta = Delta;
+		for (int DistributionCount = 0; DistributionCount < SlotsAfterDragHandle.Num() && UnusedDelta != 0; DistributionCount++)
+		{
+			float DividedDelta = UnusedDelta / SlotsAfterDragHandle.Num();
+			UnusedDelta = 0;
+			for (int SlotIndex = 0; SlotIndex < SlotsAfterDragHandle.Num(); SlotIndex++)
+			{
+				FSlotInfo& SlotInfo = SlotsAfterDragHandle[ SlotIndex ];
+
+				float CurrentSize = ClampChild( SlotInfo.Geometry->Size.Component(AxisIndex) );
+				SlotInfo.NewSize = ClampChild( CurrentSize - DividedDelta );
+
+				// If one of the slots couldn't be fully adjusted by the delta due to min/max constraints then
+				// the leftover delta needs to be evenly distributed to all of the other slots
+				UnusedDelta += SlotInfo.NewSize - ( CurrentSize - DividedDelta );
+			}
+		}
+
+		Delta = Delta - UnusedDelta;
+
+		// PrevChildLength needs to be updated: it's value has to take into account the next child's min/max restrictions
+		NewPrevChildLength = ClampChild( PrevChildGeom.Size.Component(AxisIndex) + Delta );
+
+		// Cells being resized are both stretch values -> redistribute the stretch coefficients proportionately
+		// to match the new child sizes on the screen.
+		{
+			float TotalLength = NewPrevChildLength;
+			float TotalStretchCoefficients = PrevChild.SizeValue.Get();
+
+			for (int SlotIndex = 0; SlotIndex < SlotsAfterDragHandle.Num(); SlotIndex++)
+			{
+				FSlotInfo SlotInfo = SlotsAfterDragHandle[ SlotIndex ];
+
+				TotalLength += SlotInfo.NewSize;
+				TotalStretchCoefficients += SlotInfo.Slot->SizeValue.Get();
+			}
+
+			const float NewPrevChildSize = ( TotalStretchCoefficients * NewPrevChildLength / TotalLength );
+
+			if (PrevChild.OnSlotResized_Handler.IsBound())
+			{
+				PrevChild.OnSlotResized_Handler.Execute( NewPrevChildSize );
+			}
+			else
+			{
+				PrevChild.SizeValue = NewPrevChildSize;
+			}
+
+			for (int SlotIndex = 0; SlotIndex < SlotsAfterDragHandle.Num(); SlotIndex++)
+			{
+				FSlotInfo SlotInfo = SlotsAfterDragHandle[ SlotIndex ];
+
+				const float NewNextChildSize = ( TotalStretchCoefficients * SlotInfo.NewSize / TotalLength );
+
+				if (SlotInfo.Slot->OnSlotResized_Handler.IsBound())
+				{
+					SlotInfo.Slot->OnSlotResized_Handler.Execute(NewNextChildSize);
+				}
+				else
+				{
+					SlotInfo.Slot->SizeValue = NewNextChildSize;
+				}
+			}
+		}
+	}
+}
+
+/**
+* @param ProposedSize  A size that a child would like to be
+*
+* @return A size that is clamped against the minimum size allowed for children.
+*/
+float SSplitter::ClampChild( float ProposedSize )
+{
+	return FMath::Max( MinSplitterChildLength, ProposedSize );
+}
+
+
+template<EOrientation SplitterOrientation>
+int32 SSplitter::GetHandleBeingResizedFromMousePosition( float PhysicalSplitterHandleSize, float HitDetectionSplitterHandleSize, FVector2D LocalMousePos, FArrangedChildren& ChildGeometries )
+{
+	const int32 AxisIndex = (SplitterOrientation == Orient_Horizontal) ? 0 : 1;
+	const float HalfHitDetectionSplitterHandleSize = ( HitDetectionSplitterHandleSize / 2 );
+	const float HalfPhysicalSplitterHandleSize = ( PhysicalSplitterHandleSize / 2 );
+
+	// Search for the two widgets between which the cursor currently resides.
+	for ( int32 ChildIndex = 1; ChildIndex < ChildGeometries.Num(); ++ChildIndex )
+	{
+		FGeometry& PrevChild = ChildGeometries(ChildIndex - 1).Geometry;
+		FGeometry& NextChild = ChildGeometries(ChildIndex).Geometry;
+		float PrevBound = PrevChild.Position.Component(AxisIndex) + PrevChild.Size.Component(AxisIndex) - HalfHitDetectionSplitterHandleSize + HalfPhysicalSplitterHandleSize;
+		float NextBound = NextChild.Position.Component(AxisIndex) + HalfHitDetectionSplitterHandleSize - HalfPhysicalSplitterHandleSize;
+
+		if ( LocalMousePos.Component(AxisIndex) > PrevBound && LocalMousePos.Component(AxisIndex) < NextBound )
+		{
+			return ChildIndex - 1;
+		}
+	}
+
+	return INDEX_NONE;
+}
+
+/********************************************************************
+* SSplitter2x2														*
+* A splitter which has exactly 4 children and allows simultaneous	*		
+* of all children along an axis as well as resizing all children	* 
+* by dragging the center of the splitter.							*
+********************************************************************/
+
+void SSplitter2x2::Construct( const FArguments& InArgs )
+{
+	Children.Add( new FSlot(InArgs._TopLeft.Widget) );
+	Children.Add( new FSlot(InArgs._BottomLeft.Widget) );
+	Children.Add( new FSlot(InArgs._TopRight.Widget) );
+	Children.Add( new FSlot(InArgs._BottomRight.Widget) );
+
+	SplitterHandleSize = 5.0f;
+	bIsResizing = false;
+	ResizingAxis = INDEX_NONE;
+}
+
+
+void SSplitter2x2::ArrangeChildren( const FGeometry& AllottedGeometry, FArrangedChildren& ArrangedChildren ) const
+{
+	int32 NumNonCollapsedChildren = 0;
+	FVector2D CoefficientTotal(0,0);
+
+	check( Children.Num() == 4 );
+
+	// The allotted space for our children is our geometry minus a little space to show splitter handles
+	const FVector2D SpaceAllottedForChildren = AllottedGeometry.Size - FVector2D(SplitterHandleSize,SplitterHandleSize);
+
+	// The current offset that the next child should be positioned at.
+	FVector2D Offset(0,0);
+
+	for (int32 ChildIndex=0; ChildIndex < Children.Num(); ++ChildIndex)
+	{
+		const FSlot& CurSlot = Children[ChildIndex];
+
+		// Calculate the amount of space that this child should take up.  
+		// It is based on the current percentage of space it should take up which is defined by a user moving the splitters
+		FVector2D ChildSpace = SpaceAllottedForChildren * CurSlot.PercentageAttribute.Get();
+
+		const EVisibility ChildVisibility = CurSlot.Widget->GetVisibility();
+
+		// If the child is visible, put them in their spot
+		if ( ArrangedChildren.Accepts( ChildVisibility ) )
+		{
+			ArrangedChildren.AddWidget( ChildVisibility, AllottedGeometry.MakeChild( CurSlot.Widget, Offset, ChildSpace ) );
+		}
+
+		// Advance to the next slot. If the child is collapsed, it takes up no room and does not need a splitter
+		if( ChildIndex == 1 )
+		{
+			// ChildIndex of 1 means we are starting the next column so reset the Y offset.
+			Offset.Y = 0.0f;
+			Offset += FVector2D( ChildSpace.X + SplitterHandleSize, 0);
+		}
+		else
+		{
+			Offset += FVector2D( 0, ChildSpace.Y + SplitterHandleSize );
+		}
+	}
+}
+
+
+
+FVector2D SSplitter2x2::ComputeDesiredSize() const
+{
+	return FVector2D(100,100);
+}
+
+
+FChildren* SSplitter2x2::GetChildren()
+{
+	return &Children;
+}
+
+
+FReply SSplitter2x2::OnMouseButtonDown( const FGeometry& MyGeometry, const FPointerEvent& MouseEvent )
+{
+	FReply Reply = FReply::Unhandled();
+	if ( MouseEvent.GetEffectingButton() == EKeys::LeftMouseButton )
+	{
+		const FVector2D LocalMousePos = MyGeometry.AbsoluteToLocal( MouseEvent.GetScreenSpacePosition() );
+		ResizingAxis = CalculateResizingAxis( MyGeometry, LocalMousePos );
+		if ( ResizingAxis != INDEX_NONE )
+		{
+			bIsResizing = true;
+			Reply = FReply::Handled().CaptureMouse( SharedThis(this) );
+		}
+	}
+
+	return Reply;
+}
+
+
+FReply SSplitter2x2::OnMouseButtonUp( const FGeometry& MyGeometry, const FPointerEvent& MouseEvent )
+{
+	if ( MouseEvent.GetEffectingButton() == EKeys::LeftMouseButton && bIsResizing == true )
+	{
+		bIsResizing = false;
+		return FReply::Handled().ReleaseMouseCapture();
+	}
+	return FReply::Unhandled();
+}
+
+
+FReply SSplitter2x2::OnMouseMove( const FGeometry& MyGeometry, const FPointerEvent& MouseEvent )
+{
+	const FVector2D LocalMousePos = MyGeometry.AbsoluteToLocal( MouseEvent.GetScreenSpacePosition() );
+
+	if ( bIsResizing && FSlateApplication::Get().GetMouseCaptor() == AsShared() )
+	{	
+		FArrangedChildren ArrangedChildren(EVisibility::Visible);
+		ArrangeChildren( MyGeometry, ArrangedChildren );
+
+		ResizeChildren( ArrangedChildren, LocalMousePos );
+		return FReply::Handled();
+	}
+	else
+	{	
+		ResizingAxis = CalculateResizingAxis( MyGeometry, LocalMousePos );
+		return FReply::Unhandled();
+	}
+}
+
+
+FCursorReply SSplitter2x2::OnCursorQuery( const FGeometry& MyGeometry, const FPointerEvent& CursorEvent ) const 
+{
+	if( ResizingAxis == 0 )
+	{
+		return FCursorReply::Cursor( EMouseCursor::ResizeLeftRight );
+	}
+	else if( ResizingAxis == 1 )
+	{
+		return FCursorReply::Cursor( EMouseCursor::ResizeUpDown );
+	}
+	else if( ResizingAxis == 2)
+	{
+		return FCursorReply::Cursor( EMouseCursor::CardinalCross );
+	}
+
+	return FCursorReply::Unhandled();
+}
+
+void SSplitter2x2::ResizeChildren( FArrangedChildren& ArrangedChildren, FVector2D LocalMousePos )
+{
+	// Compute the handle position.  The last child is used because it is always the furthest away from the origin
+	const FVector2D HandlePos = ArrangedChildren(3).Geometry.Position - (FVector2D( SplitterHandleSize , SplitterHandleSize ) * .5f);
+	FVector2D Delta = LocalMousePos - HandlePos;
+
+	FGeometry& TopLeftGeom = ArrangedChildren(0).Geometry;
+	FGeometry& BotLeftGeom = ArrangedChildren(1).Geometry;
+	FGeometry& TopRightGeom = ArrangedChildren(2).Geometry;
+	FGeometry& BotRightGeom = ArrangedChildren(3).Geometry;
+
+	FSlot& TopLeft = Children[0];
+	FSlot& BotLeft = Children[1];
+	FSlot& TopRight = Children[2];
+	FSlot& BotRight = Children[3];
+
+	if( ResizingAxis == 0 )
+	{
+		// Ensure deltas along the Y axis are not taken into account
+		Delta.Y = 0;
+	}
+	else if( ResizingAxis == 1 )
+	{
+		// Ensure deltas along the X axis are not taken into account
+		Delta.X = 0;
+	}
+
+	// The new size of each child
+	FVector2D NewSizeTL;
+	FVector2D NewSizeBL;
+	FVector2D NewSizeTR;
+	FVector2D NewSizeBR;
+
+	if( ResizingAxis == 0)
+	{
+		NewSizeTL = TopLeftGeom.Size + Delta;
+		NewSizeBL = BotLeftGeom.Size + Delta;
+		NewSizeTR = TopRightGeom.Size - Delta;
+		NewSizeBR = BotRightGeom.Size - Delta;
+	}
+	else if( ResizingAxis == 1 )
+	{
+		NewSizeTL = TopLeftGeom.Size + Delta;
+		NewSizeBL = BotLeftGeom.Size - Delta;
+		NewSizeTR = TopRightGeom.Size + Delta;
+		NewSizeBR = BotRightGeom.Size - Delta;
+	}
+	else
+	{
+		// Resize X and Y independently as they have different rules for X and Y
+		NewSizeTL.X = TopLeftGeom.Size.X + Delta.X;
+		NewSizeBL.X = BotLeftGeom.Size.X + Delta.X;
+		NewSizeTR.X = TopRightGeom.Size.X - Delta.X;
+		NewSizeBR.X = BotRightGeom.Size.X - Delta.X;
+
+
+		NewSizeTL.Y = TopLeftGeom.Size.Y + Delta.Y;
+		NewSizeBL.Y = BotLeftGeom.Size.Y - Delta.Y;
+		NewSizeTR.Y = TopRightGeom.Size.Y + Delta.Y;
+		NewSizeBR.Y = BotRightGeom.Size.Y - Delta.Y;
+	}
+
+	// Clamp all values so they cant be too small
+	// Must be done independently on each axis because of how FVector2D handles greater than
+	NewSizeTL.X = FMath::Max(MinSplitterChildLength, NewSizeTL.X);
+	NewSizeBL.X = FMath::Max(MinSplitterChildLength, NewSizeBL.X);
+	NewSizeTR.X = FMath::Max(MinSplitterChildLength, NewSizeTR.X);
+	NewSizeBR.X = FMath::Max(MinSplitterChildLength, NewSizeBR.X);
+
+	NewSizeTL.Y = FMath::Max(MinSplitterChildLength, NewSizeTL.Y);
+	NewSizeBL.Y = FMath::Max(MinSplitterChildLength, NewSizeBL.Y);
+	NewSizeTR.Y = FMath::Max(MinSplitterChildLength, NewSizeTR.Y);
+	NewSizeBR.Y = FMath::Max(MinSplitterChildLength, NewSizeBR.Y);
+
+	// Set the percentage space within the allotted area that each child should take up
+	FVector2D TotalLength = NewSizeTL+NewSizeBR;
+	TopLeft.SetPercentage( NewSizeTL / TotalLength );
+	TopRight.SetPercentage( NewSizeTR / TotalLength );
+	BotLeft.SetPercentage( NewSizeBL / TotalLength );
+	BotRight.SetPercentage( NewSizeBR / TotalLength );
+
+}
+
+int32 SSplitter2x2::CalculateResizingAxis( const FGeometry& MyGeometry, const FVector2D& LocalMousePos ) const
+{
+	int32 Axis = INDEX_NONE;
+
+	FArrangedChildren ArrangedChildren(EVisibility::Visible);
+	ArrangeChildren( MyGeometry, ArrangedChildren );
+
+	// The axis is in the center if it passes all hit tests
+	bool bInCenter = true;
+	// Search for the two widgets between which the cursor currently resides.
+	for ( int32 ChildIndex = 1; ChildIndex < ArrangedChildren.Num(); ++ChildIndex )
+	{
+		FGeometry& PrevChild = ArrangedChildren(ChildIndex - 1).Geometry;
+		FGeometry& NextChild = ArrangedChildren(ChildIndex).Geometry;
+		FVector2D PrevBound = PrevChild.Position + PrevChild.Size;
+		FVector2D NextBound = NextChild.Position;
+
+		if( LocalMousePos.X > PrevBound.X && LocalMousePos.X < NextBound.X )
+		{
+			// The mouse is in between two viewports vertically
+			// Resizing axis is the X axis
+			Axis = 0;
+		}
+		else if( LocalMousePos.Y > PrevBound.Y && LocalMousePos.Y < NextBound.Y )
+		{
+			// The mouse is in between two viewports horizontally
+			// Resizing axis is the Y axis
+			Axis = 1;
+		}
+		else
+		{
+			// Failed a hit test
+			bInCenter = false;
+		}
+	}
+
+	if( bInCenter )
+	{
+		Axis = 2;
+	}
+
+	return Axis;
+}
+
+TSharedRef< SWidget > SSplitter2x2::GetTopLeftContent()
+{
+	return Children[ 0 ].Widget;
+}
+
+
+TSharedRef< SWidget > SSplitter2x2::GetBottomLeftContent()
+{
+	return Children[ 1 ].Widget;
+}
+
+
+TSharedRef< SWidget > SSplitter2x2::GetTopRightContent()
+{
+	return Children[ 2 ].Widget;
+}
+
+
+TSharedRef< SWidget > SSplitter2x2::GetBottomRightContent()
+{
+	return Children[ 3 ].Widget;
+}
+
+
+void SSplitter2x2::SetTopLeftContent( TSharedRef< SWidget > TopLeftContent )
+{
+	Children[ 0 ].Widget = TopLeftContent;
+}
+
+
+void SSplitter2x2::SetBottomLeftContent( TSharedRef< SWidget > BottomLeftContent )
+{
+	Children[ 1 ].Widget = BottomLeftContent;
+}
+
+
+void SSplitter2x2::SetTopRightContent( TSharedRef< SWidget > TopRightContent )
+{
+	Children[ 2 ].Widget = TopRightContent;
+}
+
+
+void SSplitter2x2::SetBottomRightContent( TSharedRef< SWidget > BottomRightContent )
+{
+	Children[ 3 ].Widget = BottomRightContent;
+}
+
+void SSplitter2x2::GetSplitterPercentages( TArray< FVector2D >& OutPercentages ) const
+{
+	OutPercentages.Empty();
+	for (int32 i = 0; i < 4; ++i)
+	{
+		OutPercentages.Add(Children[i].PercentageAttribute.Get());
+	}
+}
+
+void SSplitter2x2::SetSplitterPercentages( const TArray< FVector2D >& InPercentages )
+{
+	for (int32 i = 0; i < 4; ++i)
+	{
+		Children[i].SetPercentage(InPercentages[i]);
+	}
+}
diff --git a/Engine/Source/Runtime/Slate/Private/Widgets/Layout/SUniformGridPanel.cpp b/Engine/Source/Runtime/Slate/Private/Widgets/Layout/SUniformGridPanel.cpp
new file mode 100644
--- /dev/null
+++ b/Engine/Source/Runtime/Slate/Private/Widgets/Layout/SUniformGridPanel.cpp
@@ -0,0 +1,105 @@
+// Copyright 1998-2014 Epic Games, Inc. All Rights Reserved.
+
+#include "SlatePrivatePCH.h"
+#include "LayoutUtils.h"
+
+
+void SUniformGridPanel::Construct( const FArguments& InArgs )
+{
+	SlotPadding = InArgs._SlotPadding;
+	NumColumns = 0;
+	NumRows = 0;
+	MinDesiredSlotWidth = InArgs._MinDesiredSlotWidth.Get();
+	MinDesiredSlotHeight = InArgs._MinDesiredSlotHeight.Get();
+
+	Children.Reserve( InArgs.Slots.Num() );
+	for (int32 ChildIndex=0; ChildIndex < InArgs.Slots.Num(); ChildIndex++)
+	{
+		FSlot* ChildSlot = InArgs.Slots[ChildIndex];
+		Children.Add( ChildSlot );
+		// A single cell at (N,M) means our grid size is (N+1, M+1)
+		NumColumns = FMath::Max( ChildSlot->Column+1, NumColumns );
+		NumRows = FMath::Max( ChildSlot->Row+1, NumRows );
+	}
+}
+
+
+void SUniformGridPanel::ArrangeChildren( const FGeometry& AllottedGeometry, FArrangedChildren& ArrangedChildren ) const
+{
+	const FVector2D CellSize(AllottedGeometry.Size.X / NumColumns, AllottedGeometry.Size.Y / NumRows);
+	const FMargin& CurrentSlotPadding(SlotPadding.Get());
+	for (int32 ChildIndex=0; ChildIndex < Children.Num(); ++ChildIndex)
+	{
+		const FSlot& Child = Children[ChildIndex];
+		const EVisibility ChildVisibility = Child.Widget->GetVisibility();
+		if ( ArrangedChildren.Accepts(ChildVisibility) )
+		{
+			// Do the standard arrangement of elements within a slot
+			// Takes care of alignment and padding.
+			AlignmentArrangeResult XAxisResult = AlignChild<Orient_Horizontal>( CellSize.X, Child, CurrentSlotPadding );
+			AlignmentArrangeResult YAxisResult = AlignChild<Orient_Vertical>( CellSize.Y, Child, CurrentSlotPadding );
+
+			ArrangedChildren.AddWidget( ChildVisibility,
+				AllottedGeometry.MakeChild(Child.Widget,
+				FVector2D(CellSize.X*Child.Column + XAxisResult.Offset, CellSize.Y*Child.Row + YAxisResult.Offset),
+				FVector2D(XAxisResult.Size, YAxisResult.Size)
+			));
+		}
+		
+	}
+}
+
+
+FVector2D SUniformGridPanel::ComputeDesiredSize() const
+{
+	FVector2D MaxChildDesiredSize = FVector2D::ZeroVector;
+	const FVector2D SlotPaddingDesiredSize = SlotPadding.Get().GetDesiredSize();
+	
+	const float CachedMinDesiredSlotWidth = MinDesiredSlotWidth.Get();
+	const float CachedMinDesiredSlotHeight = MinDesiredSlotHeight.Get();
+	
+	for ( int32 ChildIndex=0; ChildIndex < Children.Num(); ++ChildIndex )
+	{
+		const FSlot& Child = Children[ ChildIndex ];
+		if (Child.Widget->GetVisibility() != EVisibility::Collapsed)
+		{
+			FVector2D ChildDesiredSize = Child.Widget->GetDesiredSize() + SlotPaddingDesiredSize;
+
+			ChildDesiredSize.X = FMath::Max( ChildDesiredSize.X, CachedMinDesiredSlotWidth);
+			ChildDesiredSize.Y = FMath::Max( ChildDesiredSize.Y, CachedMinDesiredSlotHeight);
+
+			MaxChildDesiredSize.X = FMath::Max( MaxChildDesiredSize.X, ChildDesiredSize.X );
+			MaxChildDesiredSize.Y = FMath::Max( MaxChildDesiredSize.Y, ChildDesiredSize.Y );
+		}
+	}
+
+	return FVector2D( NumColumns*MaxChildDesiredSize.X, NumRows*MaxChildDesiredSize.Y );
+}
+
+
+FChildren* SUniformGridPanel::GetChildren()
+{
+	return &Children;
+}
+
+
+SUniformGridPanel::FSlot& SUniformGridPanel::AddSlot( int32 Column, int32 Row )
+{
+	FSlot& NewSlot = SUniformGridPanel::Slot(Column, Row);
+
+	// A single cell at (N,M) means our grid size is (N+1, M+1)
+	NumColumns = FMath::Max( Column+1, NumColumns );
+	NumRows = FMath::Max( Row+1, NumRows );
+
+	Children.Add( &NewSlot );
+
+	return NewSlot;
+}
+
+
+void SUniformGridPanel::ClearChildren()
+{
+	NumColumns = 0;
+	NumRows = 0;
+	Children.Empty();
+}
\ No newline at end of file
diff --git a/Engine/Source/Runtime/Slate/Private/Widgets/Layout/SWidgetSwitcher.cpp b/Engine/Source/Runtime/Slate/Private/Widgets/Layout/SWidgetSwitcher.cpp
new file mode 100644
--- /dev/null
+++ b/Engine/Source/Runtime/Slate/Private/Widgets/Layout/SWidgetSwitcher.cpp
@@ -0,0 +1,111 @@
+// Copyright 1998-2014 Epic Games, Inc. All Rights Reserved.
+
+/*=============================================================================
+	SWidgetSwitcher.cpp: Implements the SWidgetSwitcher class.
+=============================================================================*/
+
+#include "SlatePrivatePCH.h"
+#include "LayoutUtils.h"
+
+
+/* SWidgetSwitcher interface
+ *****************************************************************************/
+
+FSimpleSlot& SWidgetSwitcher::AddSlot( int32 SlotIndex )
+{
+	FSimpleSlot& NewSlot = SWidgetSwitcher::Slot();
+
+	if (!AllChildren.IsValidIndex(SlotIndex))
+	{
+		AllChildren.Add(&NewSlot);
+	}
+	else
+	{
+		AllChildren.Insert(&NewSlot, SlotIndex);
+	}
+
+	return NewSlot;
+}
+
+
+void SWidgetSwitcher::Construct( const FArguments& InArgs )
+{
+	OneDynamicChild = FOneDynamicChild( &AllChildren, &WidgetIndex );
+
+	for (int32 Index = 0; Index < InArgs.Slots.Num(); ++Index)
+	{
+		AllChildren.Add(InArgs.Slots[Index]);
+	}
+
+	WidgetIndex = InArgs._WidgetIndex;
+}
+
+
+TSharedPtr<SWidget> SWidgetSwitcher::GetActiveWidget( ) const
+{
+	const int32 ActiveWidgetIndex = WidgetIndex.Get();
+	if (ActiveWidgetIndex >= 0)
+	{
+		return AllChildren[ActiveWidgetIndex].Widget;
+	}
+
+	return NULL;
+}
+
+
+TSharedPtr<SWidget> SWidgetSwitcher::GetWidget( int32 SlotIndex ) const
+{
+	if (AllChildren.IsValidIndex(SlotIndex))
+	{
+		return AllChildren[SlotIndex].Widget;
+	}
+
+	return NULL;
+}
+
+
+int32 SWidgetSwitcher::GetWidgetIndex( TSharedRef<SWidget> Widget ) const
+{
+	for (int32 Index = 0; Index < AllChildren.Num(); ++Index)
+	{
+		const FSimpleSlot& Slot = AllChildren[Index];
+
+		if (Slot.Widget == Widget)
+		{
+			return Index;
+		}
+	}
+
+	return INDEX_NONE;
+}
+
+void SWidgetSwitcher::SetActiveWidgetIndex( int32 Index )
+{
+	WidgetIndex = Index;
+}
+
+TAttribute<FVector2D> ContentScale = FVector2D::UnitVector;
+
+
+/* SCompoundWidget interface
+ *****************************************************************************/
+
+void SWidgetSwitcher::ArrangeChildren( const FGeometry& AllottedGeometry, FArrangedChildren& ArrangedChildren ) const
+{
+	if (AllChildren.Num() > 0)
+	{
+		ArrangeSingleChild(AllottedGeometry, ArrangedChildren, AllChildren[WidgetIndex.Get()], ContentScale);
+	}
+}
+	
+FVector2D SWidgetSwitcher::ComputeDesiredSize( ) const
+{
+	return AllChildren.Num() > 0
+		? AllChildren[WidgetIndex.Get()].Widget->GetDesiredSize()
+		: FVector2D::ZeroVector;
+}
+
+FChildren* SWidgetSwitcher::GetChildren( )
+{
+	return &OneDynamicChild;
+}
diff --git a/Engine/Source/Runtime/Slate/Private/Widgets/Layout/SWrapBox.cpp b/Engine/Source/Runtime/Slate/Private/Widgets/Layout/SWrapBox.cpp
new file mode 100644
--- /dev/null
+++ b/Engine/Source/Runtime/Slate/Private/Widgets/Layout/SWrapBox.cpp
@@ -0,0 +1,221 @@
+// Copyright 1998-2014 Epic Games, Inc. All Rights Reserved.
+
+#include "SlatePrivatePCH.h"
+#include "LayoutUtils.h"
+
+
+FWrapSlot& SWrapBox::Slot()
+{
+	return *( new FWrapSlot() );
+}
+
+FWrapSlot& SWrapBox::AddSlot()
+{
+	FWrapSlot* NewSlot = &Slot();
+	Slots.Add(NewSlot);
+	return *NewSlot;
+}
+
+int32 SWrapBox::RemoveSlot( const TSharedRef<SWidget>& SlotWidget )
+{
+	for (int32 SlotIdx = 0; SlotIdx < Slots.Num(); ++SlotIdx)
+	{
+		if ( SlotWidget == Slots[SlotIdx].Widget )
+		{
+			Slots.RemoveAt(SlotIdx);
+			return SlotIdx;
+		}
+	}
+
+	return -1;
+}
+
+void SWrapBox::Construct( const FArguments & InArgs )
+{
+	PreferredWidth = InArgs._PreferredWidth;
+	InnerSlotPadding = InArgs._InnerSlotPadding;
+	bUseAllottedWidth = InArgs._UseAllottedWidth;
+
+	// Copy the children from the declaration to the widget
+	for ( int32 ChildIndex=0; ChildIndex < InArgs.Slots.Num(); ++ChildIndex )
+	{
+		Slots.Add( InArgs.Slots[ChildIndex] );
+	}
+}
+
+void SWrapBox::Tick( const FGeometry& AllottedGeometry, const double InCurrentTime, const float InDeltaTime )
+{
+	if (bUseAllottedWidth)
+	{
+		PreferredWidth = AllottedGeometry.Size.X;
+	}
+}
+
+void SWrapBox::ArrangeChildren( const FGeometry& AllottedGeometry, FArrangedChildren& ArrangedChildren ) const
+{
+	const float MyPreferredWidth = PreferredWidth.Get();
+
+	// PREPASS
+	// Figure out when we start wrapping and how tall each line (aka. row) needs to be
+	TArray<int32, TInlineAllocator<3> > WidgetIndexStartsNewLine;
+	TArray<float, TInlineAllocator<3> > LineHeights;
+	{
+		float WidthSoFar = 0;
+		float LineHeightSoFar = 0;
+		for ( int32 ChildIndex=0; ChildIndex < Slots.Num(); ++ChildIndex )
+		{
+			const TSharedRef<SWidget>& ThisChild = Slots[ChildIndex].Widget;
+
+			// If slot isn't the first one on the line, we need to add the inner slot padding.
+			const float ConditionPaddingX = ( WidthSoFar == 0 ) ? 0 : InnerSlotPadding.X;
+			
+			FVector2D ThisSlotSize = ThisChild->GetDesiredSize() + Slots[ChildIndex].SlotPadding.Get().GetDesiredSize();
+
+			// If the width is smaller than this slots fill limit, we need to account for it filling the line
+			// and forcing all controls that follow it to wrap.
+			if ( Slots[ChildIndex].SlotFillLineWhenWidthLessThan.IsSet() )
+			{
+				if ( MyPreferredWidth < Slots[ChildIndex].SlotFillLineWhenWidthLessThan.GetValue() )
+				{
+					ThisSlotSize.X = MyPreferredWidth;
+				}
+			}
+
+			if ( ( WidthSoFar + ThisSlotSize.X + ConditionPaddingX ) > MyPreferredWidth )
+			{
+				WidgetIndexStartsNewLine.Add( ChildIndex );
+				LineHeights.Add( LineHeightSoFar );
+				WidthSoFar = ThisSlotSize.X + ConditionPaddingX;
+				LineHeightSoFar = ThisSlotSize.Y;
+			}
+			else
+			{
+				WidthSoFar += ThisSlotSize.X;
+				LineHeightSoFar = FMath::Max( ThisSlotSize.Y + InnerSlotPadding.Y, LineHeightSoFar );
+			}
+		}
+		// Add trailing values so that we won't out of bound (easier than adding extra conditionals)
+		LineHeights.Add(LineHeightSoFar);
+		WidgetIndexStartsNewLine.Add(INDEX_NONE);
+	}
+
+	// ACTUALLY ARRANGE
+	float XOffsetSoFar = 0;
+	float YOffsetSoFar = 0;
+	for ( int32 CurrentLineIndex=0, ChildIndex=0; ChildIndex < Slots.Num(); ++ChildIndex )
+	{
+		// If slot isn't on the first line, we need to add the inner slot padding.
+		const float ConditionPaddingY = ( YOffsetSoFar == 0 ) ? 0 : InnerSlotPadding.Y;
+
+		// Do we need to start a new line?
+		const bool bStartNewLine = ( WidgetIndexStartsNewLine[CurrentLineIndex] == ChildIndex );
+		if (bStartNewLine)
+		{
+			// New line starts at the very left
+			XOffsetSoFar = 0;
+			// 
+			YOffsetSoFar += LineHeights[CurrentLineIndex] + InnerSlotPadding.Y;
+
+			// Move on to the next line.
+			CurrentLineIndex++;
+		}
+
+		// If slot isn't the first one on the line, we need to add the inner slot padding.
+		const float ConditionPaddingX = ( XOffsetSoFar == 0 ) ? 0 : InnerSlotPadding.X;
+
+		// Is this the last widget on the line?
+		bool bLastWidgetOnLine = ((ChildIndex + 1) >= Slots.Num()) || WidgetIndexStartsNewLine[CurrentLineIndex] == ( ChildIndex + 1 );
+
+		const float CurrentLineHeight = LineHeights[ CurrentLineIndex ];
+		const FMargin& SlotPadding( Slots[ChildIndex].SlotPadding.Get() );
+
+		const TSharedRef<SWidget>& ThisChild = Slots[ChildIndex].Widget;
+		float ThisSlotSizeX = ThisChild->GetDesiredSize().X + SlotPadding.GetTotalSpaceAlong<Orient_Horizontal>();
+
+		// If the width is smaller than this slots fill limit, we need to account for it filling the line
+		// and forcing all controls that follow it to wrap.
+		if ( Slots[ChildIndex].SlotFillLineWhenWidthLessThan.IsSet() )
+		{
+			if ( MyPreferredWidth < Slots[ChildIndex].SlotFillLineWhenWidthLessThan.GetValue() )
+			{
+				ThisSlotSizeX = MyPreferredWidth;
+			}
+		}
+
+		// If this slot fills empty space, we need to have it fill the remaining space on the line.
+		if ( bLastWidgetOnLine && Slots[ChildIndex].bSlotFillEmptySpace )
+		{
+			ThisSlotSizeX = MyPreferredWidth - (XOffsetSoFar - ConditionPaddingX);
+		}
+
+		// Horizontal alignment will always be equivalent to HAlign_Fill due to the nature of the widget.
+		// But we still need to account for padding.		
+		AlignmentArrangeResult XAlignmentResult = AlignChild<Orient_Horizontal>( ThisSlotSizeX, Slots[ChildIndex], SlotPadding + ConditionPaddingX );
+		AlignmentArrangeResult YAlignmentResult = AlignChild<Orient_Vertical>( CurrentLineHeight, Slots[ChildIndex], SlotPadding + ConditionPaddingY );
+
+		ArrangedChildren.AddWidget( AllottedGeometry.MakeChild(
+			ThisChild,
+			FVector2D( XOffsetSoFar + XAlignmentResult.Offset, YOffsetSoFar + YAlignmentResult.Offset),
+			FVector2D( XAlignmentResult.Size, YAlignmentResult.Size )
+		) );
+
+		XOffsetSoFar += ThisSlotSizeX;		
+	}
+}
+
+void SWrapBox::ClearChildren()
+{
+	Slots.Empty();
+}
+
+FVector2D SWrapBox::ComputeDesiredSize() const
+{
+	FVector2D MyDesiredSize = FVector2D::ZeroVector;
+	FVector2D RowDesiredSize = FVector2D::ZeroVector;
+
+	const float MyPreferredWidth = PreferredWidth.Get();
+	
+	
+	// Find the size of every row. Once the size of the row exceeds the preferred width,
+	// start a new row. Whenever we start a new row, update the desired size of the entire widget.
+	for ( int32 ChildIndex=0; ChildIndex<Slots.Num(); ++ChildIndex )
+	{
+		// If slot isn't the first one on the line, we need to add the inner slot padding.
+		const float ConditionPaddingX = ( RowDesiredSize.X == 0 ) ? 0 : InnerSlotPadding.X;
+
+		FVector2D ThisChildDesiredSize = Slots[ChildIndex].Widget->GetDesiredSize() + Slots[ChildIndex].SlotPadding.Get().GetDesiredSize() + ConditionPaddingX;
+
+		// If the width is smaller than this slots fill limit, we need to account for it filling the line
+		// and forcing all controls that follow it to wrap.
+		if ( Slots[ChildIndex].SlotFillLineWhenWidthLessThan.IsSet() )
+		{
+			if ( MyPreferredWidth < Slots[ChildIndex].SlotFillLineWhenWidthLessThan.GetValue() )
+			{
+				ThisChildDesiredSize.X = MyPreferredWidth;
+			}
+		}
+
+		if ( (RowDesiredSize.X + ThisChildDesiredSize.X) > MyPreferredWidth)
+		{
+			// We are starting a new line; update the total desired size.
+			MyDesiredSize.X = FMath::Max(RowDesiredSize.X, MyDesiredSize.X);
+			MyDesiredSize.Y += RowDesiredSize.Y + InnerSlotPadding.Y;
+
+			RowDesiredSize = FVector2D::ZeroVector;
+		}
+
+		RowDesiredSize.X += ThisChildDesiredSize.X;
+		RowDesiredSize.Y = FMath::Max( RowDesiredSize.Y, ThisChildDesiredSize.Y );
+	}
+
+	// Update the desired size with the current unfinished line.
+	MyDesiredSize.X = FMath::Max(RowDesiredSize.X, MyDesiredSize.X);
+	MyDesiredSize.Y += RowDesiredSize.Y;
+
+	return MyDesiredSize;
+}
+
+FChildren* SWrapBox::GetChildren()
+{
+	return &Slots;	
+}
diff --git a/Engine/Source/Runtime/Slate/Private/Widgets/Notifications/SNotificationList.cpp b/Engine/Source/Runtime/Slate/Private/Widgets/Notifications/SNotificationList.cpp
new file mode 100644
--- /dev/null
+++ b/Engine/Source/Runtime/Slate/Private/Widgets/Notifications/SNotificationList.cpp
@@ -0,0 +1,580 @@
+// Copyright 1998-2014 Epic Games, Inc. All Rights Reserved.
+
+#include "SlatePrivatePCH.h"
+
+
+/////////////////////////////////////////////////
+// SNotificationExtendable
+
+/** Contains the standard functionality for a notification to inherit from*/
+class SNotificationExtendable : public SNotificationItem
+{
+public:
+	virtual ~SNotificationExtendable()
+	{
+		// Just in case, make sure we have left responsive mode when getting cleaned up
+		if ( ThrottleHandle.IsValid() )
+		{
+			FSlateThrottleManager::Get().LeaveResponsiveMode( ThrottleHandle );
+		}
+	}
+	
+	
+	/** Sets the text for message element */
+	virtual void SetText( const TAttribute< FText >& InText ) OVERRIDE
+	{
+		Text = InText;
+		MyTextBlock->SetText( Text );
+	}
+
+	virtual ECompletionState GetCompletionState() const OVERRIDE
+	{
+		return CompletionState;
+	}
+
+	virtual void SetCompletionState(ECompletionState State) OVERRIDE
+	{
+		CompletionState = State;
+
+		if (State == CS_Success || State == CS_Fail)
+		{
+			CompletionStateAnimation = FCurveSequence();
+			GlowCurve = CompletionStateAnimation.AddCurve(0.f, 0.75f);
+			CompletionStateAnimation.Play();
+		}
+	}
+
+	virtual void ExpireAndFadeout() OVERRIDE
+	{
+		FadeAnimation = FCurveSequence();
+		// Add some space for the expire time
+		FadeAnimation.AddCurve(FadeOutDuration.Get(), ExpireDuration.Get());
+		// Add the actual fade curve
+		FadeCurve = FadeAnimation.AddCurve(0.f, FadeOutDuration.Get());
+		FadeAnimation.PlayReverse();
+	}
+
+	/** Begins the fadein of this message */
+	void Fadein( const bool bAllowThrottleWhenFrameRateIsLow )
+	{
+		// Make visible
+		SetVisibility(EVisibility::Visible);
+
+		// Play Fadein animation
+		FadeAnimation = FCurveSequence();
+		FadeCurve = FadeAnimation.AddCurve(0.f, FadeInDuration.Get());
+		FadeAnimation.Play();
+
+		// Scale up/flash animation
+		IntroAnimation = FCurveSequence();
+		ScaleCurveX = IntroAnimation.AddCurve(0.2f, 0.3f, ECurveEaseFunction::QuadOut);
+		ScaleCurveY = IntroAnimation.AddCurve(0.f, 0.2f);
+		GlowCurve = IntroAnimation.AddCurve(0.5f, 0.55f, ECurveEaseFunction::QuadOut);
+		IntroAnimation.Play();
+
+		// When a fade in occurs, we need a high framerate for the animation to look good
+		if( FadeInDuration.Get() > KINDA_SMALL_NUMBER && bAllowThrottleWhenFrameRateIsLow && !ThrottleHandle.IsValid() )
+		{
+			if( !FSlateApplication::Get().IsRunningAtTargetFrameRate() )
+			{
+				ThrottleHandle = FSlateThrottleManager::Get().EnterResponsiveMode();
+			}
+		}
+	}
+
+	/** Begins the fadeout of this message */
+	virtual void Fadeout() OVERRIDE
+	{
+		// Start fade animation
+		FadeAnimation = FCurveSequence();
+		FadeCurve = FadeAnimation.AddCurve(0.f, FadeOutDuration.Get());
+		FadeAnimation.PlayReverse();
+	}
+
+	void Tick( const FGeometry& AllottedGeometry, const double InCurrentTime, const float InDeltaTime ) OVERRIDE
+	{
+		SCompoundWidget::Tick(AllottedGeometry, InCurrentTime, InDeltaTime);
+
+		const bool bIsFadingOut = FadeAnimation.IsInReverse();
+		const bool bIsCurrentlyPlaying = FadeAnimation.IsPlaying();
+		const bool bIsIntroPlaying = IntroAnimation.IsPlaying();
+
+		if ( !bIsCurrentlyPlaying && bIsFadingOut )
+		{
+			// Reset the Animation
+			FadeoutComplete();
+		}
+
+		if ( !bIsIntroPlaying && ThrottleHandle.IsValid() )
+		{
+			// Leave responsive mode once the intro finishes playing
+			FSlateThrottleManager::Get().LeaveResponsiveMode( ThrottleHandle );
+		}
+	}
+
+protected:
+	/** A fadeout has completed */
+	void FadeoutComplete()
+	{
+		// Make sure we are no longer fading
+		FadeAnimation = FCurveSequence();
+		FadeCurve = FCurveHandle();
+
+		// Clear the complete state to hide all the images/throbber
+		SetCompletionState(CS_None);
+
+		// Make sure we have left responsive mode
+		if ( ThrottleHandle.IsValid() )
+		{
+			FSlateThrottleManager::Get().LeaveResponsiveMode( ThrottleHandle );
+		}
+
+		// Clear reference
+		if( MyList.IsValid() )
+		{
+			MyList.Pin()->NotificationItemFadedOut(SharedThis(this));
+		}
+	}
+
+	/** Gets the current color along the fadeout curve */
+	FSlateColor GetContentColor() const
+	{
+		return GetContentColorRaw();
+	}
+
+	/** Gets the current color along the fadeout curve */
+	FLinearColor GetContentColorRaw() const
+	{
+		// if we have a parent window, we need to make that transparent, rather than
+		// this widget
+		if(MyList.IsValid() && MyList.Pin()->ParentWindowPtr.IsValid())
+		{
+			MyList.Pin()->ParentWindowPtr.Pin()->SetOpacity( FadeCurve.GetLerp() );
+			return FLinearColor(1,1,1,1);
+		}
+		else
+		{
+			return FMath::Lerp( FLinearColor(1,1,1,0), FLinearColor(1,1,1,1), FadeCurve.GetLerp() );
+		}
+	}
+
+	/** Gets the color of the glow effect */
+	FSlateColor GetGlowColor() const
+	{
+		float GlowAlpha = 1.0f - GlowCurve.GetLerp();
+		
+		if (GlowAlpha == 1.0f)
+		{
+			GlowAlpha = 0.0f;
+		}
+
+		switch (CompletionState)
+		{
+		case CS_Success: return FLinearColor(0,1,0,GlowAlpha);
+		case CS_Fail: return FLinearColor(1,0,0,GlowAlpha);
+		default: return FLinearColor(1,1,1,GlowAlpha);
+		}
+	}
+
+	/** Gets the scale for the entire item */
+	FVector2D GetItemScale() const
+	{
+		return FVector2D(  ScaleCurveX.GetLerp(), ScaleCurveY.GetLerp() );
+	}
+
+	/** Gets the visibility for the throbber */
+	EVisibility GetThrobberVisibility() const
+	{
+		return CompletionState == CS_Pending ? EVisibility::Visible : EVisibility::Collapsed;
+	}
+
+	EVisibility GetSuccessImageVisibility() const
+	{
+		return CompletionState == CS_Success ? EVisibility::Visible : EVisibility::Collapsed;
+	}
+
+	EVisibility GetFailImageVisibility() const
+	{
+		return CompletionState == CS_Fail ? EVisibility::Visible : EVisibility::Collapsed;
+	}
+
+public:
+	/** The SNotificationList that is displaying this Item */
+	TWeakPtr<SNotificationList> MyList;
+
+protected:
+
+	/** The text displayed in this text block */
+	TAttribute< FText > Text;
+
+	/** The fade in duration for this element */
+	TAttribute< float > FadeInDuration;
+
+	/** The fade out duration for this element */
+	TAttribute< float > FadeOutDuration;
+
+	/** The duration before a fadeout for this element */
+	TAttribute< float > ExpireDuration;
+
+	/** The text displayed in this element */
+	TSharedPtr<STextBlock> MyTextBlock;
+
+	/** The completion state of this message */	
+	ECompletionState CompletionState;
+
+	/** The fading animation */
+	FCurveSequence FadeAnimation;
+	FCurveHandle FadeCurve;
+
+	/** The intro animation */
+	FCurveSequence IntroAnimation;
+	FCurveHandle ScaleCurveX;
+	FCurveHandle ScaleCurveY;
+	FCurveHandle GlowCurve;
+
+	/** The completion state change animation */
+	FCurveSequence CompletionStateAnimation;
+
+	/** Handle to a throttle request made to ensure the intro animation is smooth in low FPS situations */
+	FThrottleRequest ThrottleHandle;
+};
+
+
+/////////////////////////////////////////////////
+// SNotificationItemImpl
+
+/** A single line in the event message list with additional buttons*/
+class SNotificationItemImpl : public SNotificationExtendable
+{
+public:
+	SLATE_BEGIN_ARGS( SNotificationItemImpl )
+		: _Text()
+		, _Image()	
+		, _FadeInDuration(0.5f)
+		, _FadeOutDuration(2.f)
+		, _ExpireDuration(1.f)
+		, _WidthOverride(FOptionalSize())
+	{}
+
+		/** The text displayed in this text block */
+		SLATE_ATTRIBUTE( FText, Text )
+		/** Setup information for the buttons on the notification */ 
+		SLATE_ATTRIBUTE(TArray<FNotificationButtonInfo>, ButtonDetails)
+		/** The icon image to display next to the text */
+		SLATE_ATTRIBUTE( const FSlateBrush*, Image )
+		/** The fade in duration for this element */
+		SLATE_ATTRIBUTE( float, FadeInDuration )
+		/** The fade out duration for this element */
+		SLATE_ATTRIBUTE( float, FadeOutDuration )
+		/** The duration before a fadeout for this element */
+		SLATE_ATTRIBUTE( float, ExpireDuration )
+		/** Controls whether or not to add the animated throbber */
+		SLATE_ATTRIBUTE( bool, bUseThrobber)
+		/** Controls whether or not to display the success and fail icons */
+		SLATE_ATTRIBUTE( bool, bUseSuccessFailIcons)
+		/** When true the larger bolder font will be used to display the message */
+		SLATE_ATTRIBUTE( bool, bUseLargeFont)
+		/** When set this forces the width of the box, used to stop resizeing on text change */
+		SLATE_ATTRIBUTE( FOptionalSize, WidthOverride )
+		/** When set this will display as a hyperlink on the right side of the notification. */
+		SLATE_EVENT( FSimpleDelegate, Hyperlink)
+		/** Text to display for the hyperlink (if Hyperlink is valid) */
+		SLATE_ATTRIBUTE( FText, HyperlinkText );
+
+	SLATE_END_ARGS()
+	/**
+	 * Constructs this widget
+	 *
+	 * @param InArgs    Declaration from which to construct the widget
+	*/
+	void Construct( const FArguments& InArgs )
+	{
+		CompletionState = CS_None;
+
+		Text = InArgs._Text;
+		FadeInDuration = InArgs._FadeInDuration;
+		FadeOutDuration = InArgs._FadeOutDuration;
+		ExpireDuration = InArgs._ExpireDuration;
+
+		ChildSlot
+		[
+			SNew(SBorder)
+			.BorderImage(FCoreStyle::Get().GetBrush("NotificationList.ItemBackground"))
+			.BorderBackgroundColor( this, &SNotificationItemImpl::GetContentColor )
+			.ColorAndOpacity(this, &SNotificationItemImpl::GetContentColorRaw )
+			.DesiredSizeScale( this, &SNotificationItemImpl::GetItemScale )
+			[
+				SNew(SBorder)
+				.Padding( FMargin(5) )
+				.BorderImage(FCoreStyle::Get().GetBrush("NotificationList.ItemBackground_Border"))
+				.BorderBackgroundColor( this, &SNotificationItemImpl::GetGlowColor)
+				[
+					ConstructInternals(InArgs)
+				]
+			]
+		];
+	}
+	
+	/*
+	 * Returns the internals of the notification
+	 */
+	TSharedRef<SHorizontalBox> ConstructInternals( const FArguments& InArgs ) 
+	{
+		Hyperlink = InArgs._Hyperlink;
+		HyperlinkText = InArgs._HyperlinkText;
+
+		TSharedRef<SHorizontalBox> HorizontalBox = SNew(SHorizontalBox);
+
+		// Notification image
+		HorizontalBox->AddSlot()
+		.AutoWidth()
+		.Padding(10.f, 0.f, 0.f, 0.f)
+		.VAlign(VAlign_Center)
+		.HAlign(HAlign_Left)
+		[
+			SNew(SImage)
+			.Image( InArgs._Image )
+		];
+
+		// Set font to default case.
+		FSlateFontInfo font = FCoreStyle::Get().GetFontStyle(TEXT("NotificationList.FontBold"));
+
+		// only override if requested
+		if (!InArgs._bUseLargeFont.Get())
+		{
+			font = FCoreStyle::Get().GetFontStyle(TEXT("NotificationList.FontLight"));
+		}
+
+		// Build Text box
+		HorizontalBox->AddSlot()
+		.AutoWidth()
+		.Padding(10.f, 0.f, 25.f, 0.f)
+		.VAlign(VAlign_Center)
+		.HAlign(HAlign_Left)
+		[
+			SNew(SBox)
+			.WidthOverride(InArgs._WidthOverride)
+			[
+				SAssignNew(MyTextBlock, STextBlock)
+				.Text(Text)
+				.Font(font)
+			]
+		];
+
+		if (InArgs._bUseThrobber.Get())
+		{
+			// Build pending throbber
+			HorizontalBox->AddSlot()
+			.AutoWidth()
+			.Padding(0.f, 0.f, 25.f, 0.f)
+			.HAlign(HAlign_Center)
+			.VAlign(VAlign_Center)
+			[
+				SNew(SThrobber)
+				.Visibility( this, &SNotificationItemImpl::GetThrobberVisibility )
+			];
+		}
+
+		if (InArgs._bUseSuccessFailIcons.Get())
+		{
+			// Build success image
+			HorizontalBox->AddSlot()
+			.AutoWidth()
+			.Padding(0.f, 0.f, 25.f, 0.f)
+			.HAlign(HAlign_Center)
+			.VAlign(VAlign_Center)
+			[
+				SNew(SImage)
+				.Image( FCoreStyle::Get().GetBrush("NotificationList.SuccessImage") )
+				.Visibility( this, &SNotificationItemImpl::GetSuccessImageVisibility )
+			];
+
+			// Build fail image
+			HorizontalBox->AddSlot()
+			.AutoWidth()
+			.Padding(0.f, 0.f, 25.f, 0.f )
+			.HAlign(HAlign_Center)
+			.VAlign(VAlign_Center)
+			[
+				SNew(SImage)
+				.Image( FCoreStyle::Get().GetBrush("NotificationList.FailImage") )
+				.Visibility( this, &SNotificationItemImpl::GetFailImageVisibility )
+			];
+		}
+
+		// Adds any buttons that were passed in.
+		for (int32 idx = 0; idx < InArgs._ButtonDetails.Get().Num(); idx++)
+		{
+			FNotificationButtonInfo button = InArgs._ButtonDetails.Get()[idx];
+
+			HorizontalBox->AddSlot()
+			.AutoWidth()
+			.HAlign(HAlign_Right)
+			.VAlign(VAlign_Center)
+			.Padding(2)
+			[
+				SNew(SButton)
+				.Text(button.Text)
+				.ToolTipText(button.ToolTip)
+				.OnClicked(this, &SNotificationItemImpl::OnButtonClicked, button.Callback) 
+				.Visibility( this, &SNotificationItemImpl::GetButtonVisibility, button.VisibilityOnNone, button.VisibilityOnPending, button.VisibilityOnSuccess, button.VisibilityOnFail )
+			];
+		}
+
+		// Adds a hyperlink, but only visible when bound
+		HorizontalBox->AddSlot()
+			.AutoWidth()
+			.HAlign(HAlign_Right)
+			.VAlign(VAlign_Center)
+			.Padding(2)
+			[
+				SNew(SHyperlink)
+					.Text(this, &SNotificationItemImpl::GetHyperlinkText)
+					.OnNavigate(this, &SNotificationItemImpl::OnHyperlinkClicked)
+					.Visibility(this, &SNotificationItemImpl::GetHyperlinkVisibility)
+			];	
+
+		return HorizontalBox;
+	}
+
+	/** Sets the text and delegate for the hyperlink */
+	virtual void SetHyperlink( const FSimpleDelegate& InHyperlink, const TAttribute< FText >& InHyperlinkText = TAttribute< FText >() ) OVERRIDE
+	{
+		Hyperlink = InHyperlink;
+
+		// Only replace the text if specified
+		if ( InHyperlinkText.IsBound() )
+		{
+			HyperlinkText = InHyperlinkText;
+		}
+	}
+
+protected:
+
+	/* Used to determine whether the button is visible */
+	EVisibility GetButtonVisibility( const EVisibility VisibilityOnNone, const EVisibility VisibilityOnPending, const EVisibility VisibilityOnSuccess, const EVisibility VisibilityOnFail ) const
+	{
+		switch ( CompletionState )
+		{
+		case SNotificationItem::CS_None: return VisibilityOnNone;
+		case SNotificationItem::CS_Pending: return VisibilityOnPending;
+		case SNotificationItem::CS_Success: return VisibilityOnSuccess;
+		case SNotificationItem::CS_Fail: return VisibilityOnFail;
+		default:
+			check( false );
+			return EVisibility::Visible;
+		}
+	}
+
+	/* Used to determine whether the hyperlink is visible */
+	EVisibility GetHyperlinkVisibility() const
+	{
+		return Hyperlink.IsBound() ? EVisibility::Visible : EVisibility::Collapsed;
+	}
+
+	/* Used to fetch the text to display in the hyperlink */
+	FText GetHyperlinkText() const
+	{
+		return HyperlinkText.Get();
+	}
+
+private:
+
+	/* Used as a wrapper for the callback, this means any code calling it does not require access to FReply type */
+	FReply OnButtonClicked(FSimpleDelegate InCallback)
+	{
+		InCallback.ExecuteIfBound();
+		return FReply::Handled();
+	}
+
+	/* Execute the delegate for the hyperlink, if bound */
+	void OnHyperlinkClicked() const
+	{
+		Hyperlink.ExecuteIfBound();
+	}
+
+protected:
+
+	/** When set this will display as a hyperlink on the right side of the notification. */
+	FSimpleDelegate Hyperlink;
+
+	/** Text to display for the hyperlink message */
+	TAttribute< FText > HyperlinkText;
+};
+
+///////////////////////////////////////////////////
+//// SNotificationList
+TSharedRef<SNotificationItem> SNotificationList::AddNotification(const FNotificationInfo& Info)
+{
+	EVisibility visibility = EVisibility::HitTestInvisible;
+	static const FSlateBrush* CachedImage = FCoreStyle::Get().GetBrush("NotificationList.DefaultMessage");
+
+	// Notifications with buttons need to be hit test enabled.
+	if (Info.ButtonDetails.Num() > 0)
+	{
+		visibility = EVisibility::Collapsed;
+	}
+
+	// Create notification.
+	TSharedRef<SNotificationExtendable> NewItem =
+		SNew(SNotificationItemImpl)
+		.Text(Info.Text)
+		.ButtonDetails(Info.ButtonDetails)
+		.Image((Info.Image != NULL) ? Info.Image : CachedImage)
+		.FadeInDuration(Info.FadeInDuration)
+		.ExpireDuration(Info.ExpireDuration)
+		.FadeOutDuration(Info.FadeOutDuration)
+		.bUseThrobber(Info.bUseThrobber)
+		.bUseSuccessFailIcons(Info.bUseSuccessFailIcons)
+		.bUseLargeFont(Info.bUseLargeFont)
+		.WidthOverride(Info.WidthOverride)
+		.Visibility(visibility)
+		.Hyperlink(Info.Hyperlink)
+		.HyperlinkText(Info.HyperlinkText);
+		
+	NewItem->MyList = SharedThis(this);
+
+	MessageItemBoxPtr->AddSlot()
+		.AutoHeight()
+		.HAlign(HAlign_Right)
+		[
+			NewItem
+		];
+
+	NewItem->Fadein( Info.bAllowThrottleWhenFrameRateIsLow );
+
+	if (Info.bFireAndForget)
+	{
+		NewItem->ExpireAndFadeout();
+	}
+
+	return NewItem;
+}
+
+void SNotificationList::NotificationItemFadedOut (const TSharedRef<SNotificationItem>& NotificationItem)
+{
+	if( ParentWindowPtr.IsValid() )
+	{
+		// If we are in a single-window per notification situation, we dont want to 
+		// remove the NotificationItem straight away, rather we will flag us as done
+		// and wait for the FSlateNotificationManager to release the parent window.
+		bDone = true;
+	}
+	else
+	{
+		// This should remove the last non-local reference to this SNotificationItem.
+		// Since there may be many local references on the call stack we are not checking
+		// if the reference is unique.
+		MessageItemBoxPtr->RemoveSlot(NotificationItem);
+	}
+}
+
+void SNotificationList::Construct(const FArguments& InArgs)
+{
+	bDone = false;
+
+	ChildSlot
+	[
+		SAssignNew(MessageItemBoxPtr, SVerticalBox)
+	];
+}
diff --git a/Engine/Source/Runtime/Slate/Private/Widgets/Notifications/SProgressBar.cpp b/Engine/Source/Runtime/Slate/Private/Widgets/Notifications/SProgressBar.cpp
new file mode 100644
--- /dev/null
+++ b/Engine/Source/Runtime/Slate/Private/Widgets/Notifications/SProgressBar.cpp
@@ -0,0 +1,101 @@
+// Copyright 1998-2014 Epic Games, Inc. All Rights Reserved.
+
+#include "SlatePrivatePCH.h"
+
+
+void SProgressBar::Construct( const FArguments& InArgs )
+{
+	check(InArgs._Style);
+
+	Percent = InArgs._Percent;
+
+	BackgroundImage = &InArgs._Style->BackgroundImage;
+	FillImage = &InArgs._Style->FillImage;
+	MarqueeImage = &InArgs._Style->MarqueeImage;
+	FillColorAndOpacity = InArgs._FillColorAndOpacity;
+	BorderPadding = InArgs._BorderPadding;
+
+	CurveSequence = FCurveSequence(0.0f, 0.5f);
+	CurveSequence.Play();
+}
+
+
+int32 SProgressBar::OnPaint( const FGeometry& AllottedGeometry, const FSlateRect& MyClippingRect, FSlateWindowElementList& OutDrawElements, int32 LayerId, const FWidgetStyle& InWidgetStyle, bool bParentEnabled ) const
+{
+	// Used to track the layer ID we will return.
+	int32 RetLayerId = LayerId;
+
+	bool bEnabled = ShouldBeEnabled( bParentEnabled );
+	const ESlateDrawEffect::Type DrawEffects = bEnabled ? ESlateDrawEffect::None : ESlateDrawEffect::DisabledEffect;
+	
+	const FColor FillColorAndOpacitySRGB( InWidgetStyle.GetColorAndOpacityTint() * FillColorAndOpacity.Get().GetColor(InWidgetStyle) * FillImage->GetTint( InWidgetStyle ) );
+	const FColor ColorAndOpacitySRGB = InWidgetStyle.GetColorAndOpacityTint();
+
+	TOptional<float> ProgressFraction = Percent.Get();	
+
+	// Paint inside the border only. 
+	FPaintGeometry ForegroundPaintGeometry = AllottedGeometry.ToInflatedPaintGeometry( -BorderPadding.Get() );
+	const FSlateRect ForegroundClippingRect = ForegroundPaintGeometry.ToSlateRect().IntersectionWith( MyClippingRect );
+
+	FSlateDrawElement::MakeBox(
+		OutDrawElements,
+		RetLayerId++,
+		AllottedGeometry.ToPaintGeometry(),
+		BackgroundImage,
+		MyClippingRect,
+		DrawEffects,
+		InWidgetStyle.GetColorAndOpacityTint() * BackgroundImage->GetTint( InWidgetStyle )
+	);	
+	
+	if( ProgressFraction.IsSet() )
+	{
+		const float ClampedFraction = FMath::Clamp( ProgressFraction.GetValue(), 0.0f, 1.0f  );
+
+		// Draw Fill
+		FSlateDrawElement::MakeBox(
+			OutDrawElements,
+			RetLayerId++,
+			AllottedGeometry.ToPaintGeometry(
+				FVector2D::ZeroVector,
+				FVector2D( AllottedGeometry.Size.X * ( ClampedFraction ) , AllottedGeometry.Size.Y )),
+			FillImage,
+			ForegroundClippingRect,
+			DrawEffects,
+			FillColorAndOpacitySRGB
+			);
+		
+	}
+	else
+	{
+		// Draw Marquee
+		const float MarqueeAnimOffset = MarqueeImage->ImageSize.X * CurveSequence.GetLerpLooping();
+		const float MarqueeImageSize = MarqueeImage->ImageSize.X;
+
+		FSlateDrawElement::MakeBox(
+			OutDrawElements,
+			RetLayerId++,
+			AllottedGeometry.ToPaintGeometry(
+				FVector2D( MarqueeAnimOffset - MarqueeImageSize, 0.0f ),
+				FVector2D( AllottedGeometry.Size.X + MarqueeImageSize, AllottedGeometry.Size.Y )),
+			MarqueeImage,
+			ForegroundClippingRect,
+			DrawEffects,
+			ColorAndOpacitySRGB
+			);
+
+	}
+
+	return RetLayerId - 1;
+}
+
+/**
+ * Computes the desired size of this widget (SWidget)
+ *
+ * @return  The widget's desired size
+ */
+FVector2D SProgressBar::ComputeDesiredSize() const
+{
+	return MarqueeImage->ImageSize;
+}
+
+
diff --git a/Engine/Source/Runtime/Slate/Private/Widgets/SCanvas.cpp b/Engine/Source/Runtime/Slate/Private/Widgets/SCanvas.cpp
new file mode 100644
--- /dev/null
+++ b/Engine/Source/Runtime/Slate/Private/Widgets/SCanvas.cpp
@@ -0,0 +1,127 @@
+// Copyright 1998-2014 Epic Games, Inc. All Rights Reserved.
+
+/*=============================================================================
+	SCanvas.cpp: Implements the SCanvas class.
+=============================================================================*/
+
+#include "SlatePrivatePCH.h"
+#include "LayoutUtils.h"
+
+
+/* SCanvas interface
+ *****************************************************************************/
+
+void SCanvas::Construct( const SCanvas::FArguments& InArgs )
+{
+	const int32 NumSlots = InArgs.Slots.Num();
+	for ( int32 SlotIndex = 0; SlotIndex < NumSlots; ++SlotIndex )
+	{
+		Children.Add( InArgs.Slots[SlotIndex] );
+	}
+}
+
+
+void SCanvas::ClearChildren( )
+{
+	Children.Empty();
+}
+
+
+int32 SCanvas::RemoveSlot( const TSharedRef<SWidget>& SlotWidget )
+{
+	for (int32 SlotIdx = 0; SlotIdx < Children.Num(); ++SlotIdx)
+	{
+		if (SlotWidget == Children[SlotIdx].Widget)
+		{
+			Children.RemoveAt(SlotIdx);
+			return SlotIdx;
+		}
+	}
+
+	return -1;
+}
+
+
+/* SWidget overrides
+ *****************************************************************************/
+
+void SCanvas::ArrangeChildren( const FGeometry& AllottedGeometry, FArrangedChildren& ArrangedChildren ) const
+{
+	if (Children.Num() > 0)
+	{
+		for (int32 ChildIndex = 0; ChildIndex < Children.Num(); ++ChildIndex)
+		{
+			const SCanvas::FSlot& CurChild = Children[ChildIndex];
+			const FVector2D Size = CurChild.SizeAttr.Get();
+
+			//Handle HAlignment
+			FVector2D Offset(0.0f, 0.0f);
+
+			switch (CurChild.HAlignment)
+			{
+			case HAlign_Center:
+				Offset.X = -Size.X / 2.0f;
+				break;
+			case HAlign_Right:
+				Offset.X = -Size.X;
+				break;
+			}
+
+			//handle VAlignment
+			switch (CurChild.VAlignment)
+			{
+			case VAlign_Bottom:
+				Offset.Y = -Size.Y;
+				break;
+			case VAlign_Center:
+				Offset.Y = -Size.Y / 2.0f;
+				break;
+			}
+
+			// Add the information about this child to the output list (ArrangedChildren)
+			ArrangedChildren.AddWidget( AllottedGeometry.MakeChild(
+				// The child widget being arranged
+				CurChild.Widget,
+				// Child's local position (i.e. position within parent)
+				CurChild.PositionAttr.Get() + Offset,
+				// Child's size
+				Size
+			));
+		}
+	}
+}
+
+
+int32 SCanvas::OnPaint( const FGeometry& AllottedGeometry, const FSlateRect& MyClippingRect, FSlateWindowElementList& OutDrawElements, int32 LayerId, const FWidgetStyle& InWidgetStyle, bool bParentEnabled ) const
+{
+	FArrangedChildren ArrangedChildren(EVisibility::Visible);
+	this->ArrangeChildren(AllottedGeometry, ArrangedChildren);
+
+	// Because we paint multiple children, we must track the maximum layer id that they produced in case one of our parents
+	// wants to an overlay for all of its contents.
+	int32 MaxLayerId = LayerId;
+
+	for (int32 ChildIndex = 0; ChildIndex < ArrangedChildren.Num(); ++ChildIndex)
+	{
+		FArrangedWidget& CurWidget = ArrangedChildren(ChildIndex);
+		FSlateRect ChildClipRect = MyClippingRect.IntersectionWith(CurWidget.Geometry.GetClippingRect());
+		const int32 CurWidgetsMaxLayerId = CurWidget.Widget->OnPaint(CurWidget.Geometry, ChildClipRect, OutDrawElements, MaxLayerId + 1, InWidgetStyle, ShouldBeEnabled(bParentEnabled));
+
+		MaxLayerId = FMath::Max(MaxLayerId, CurWidgetsMaxLayerId);
+	}
+
+	return MaxLayerId;
+}
+
+
+FVector2D SCanvas::ComputeDesiredSize() const
+{
+	// Canvas widgets have no desired size -- their size is always determined by their container
+	return FVector2D::ZeroVector;
+}
+
+
+FChildren* SCanvas::GetChildren()
+{
+	return &Children;
+}
diff --git a/Engine/Source/Runtime/Slate/Private/Widgets/SToolTip.cpp b/Engine/Source/Runtime/Slate/Private/Widgets/SToolTip.cpp
new file mode 100644
--- /dev/null
+++ b/Engine/Source/Runtime/Slate/Private/Widgets/SToolTip.cpp
@@ -0,0 +1,71 @@
+// Copyright 1998-2014 Epic Games, Inc. All Rights Reserved.
+
+/*=============================================================================
+	SToolTip.cpp: Implements the SToolTip class.
+=============================================================================*/
+
+#include "SlatePrivatePCH.h"
+
+
+static TAutoConsoleVariable<float> StaticToolTipWrapWidth(
+	TEXT( "Slate.ToolTipWrapWidth" ),
+	400.0f,
+	TEXT( "Width of Slate tool-tips before we wrap the tool-tip text" ) );
+
+
+/**
+ * Construct this widget
+ *
+ * @param	InArgs	The declaration data for this widget
+ */
+void SToolTip::Construct( const FArguments& InArgs )
+{
+	TextContent = InArgs._Text;
+	bIsInteractive = InArgs._IsInteractive;
+	if( InArgs._Content.Widget != SNullWidget::NullWidget )
+	{
+		// Widget content argument takes precedence
+		// overrides the text content.
+		WidgetContent = InArgs._Content.Widget;
+	}
+
+	struct Local
+	{
+		static float GetToolTipWrapWidth()
+		{
+			return StaticToolTipWrapWidth.GetValueOnGameThread();
+		}
+	};
+
+	TSharedRef<SWidget> ToolTipContent = ( !WidgetContent.IsValid() )
+		? StaticCastSharedRef<SWidget>(
+			SNew( STextBlock )
+			.Text( TextContent )
+			.Font( InArgs._Font )
+			.ColorAndOpacity( InArgs._ColorAndOpacity )
+			.WrapTextAt_Static( &Local::GetToolTipWrapWidth )
+		)
+		: InArgs._Content.Widget;
+
+	SBorder::Construct(
+		SBorder::FArguments()
+		.BorderImage(InArgs._BorderImage)
+		.Padding(InArgs._TextMargin)
+		[
+			ToolTipContent
+		]	
+	);
+
+	
+}
+
+/** @return True if the tool tip has no content to display at this particular moment (remember, attribute callbacks can change this at any time!) */
+bool SToolTip::IsEmpty() const
+{
+	return !WidgetContent.IsValid() && TextContent.Get().IsEmpty();
+}
+
+bool SToolTip::IsInteractive() const
+{
+	return bIsInteractive.Get();
+}
diff --git a/Engine/Source/Runtime/Slate/Private/Widgets/SViewport.cpp b/Engine/Source/Runtime/Slate/Private/Widgets/SViewport.cpp
new file mode 100644
--- /dev/null
+++ b/Engine/Source/Runtime/Slate/Private/Widgets/SViewport.cpp
@@ -0,0 +1,275 @@
+// Copyright 1998-2014 Epic Games, Inc. All Rights Reserved.
+
+/*=============================================================================
+	SViewport.cpp: Implements the SViewport class.
+=============================================================================*/
+
+#include "SlatePrivatePCH.h"
+
+
+/* SViewport structors
+ *****************************************************************************/
+
+SViewport::SViewport( )
+	: ViewportInterface(NULL)
+	, bRenderDirectlyToWindow(false)
+	, bEnableGammaCorrection(true)
+{ }
+
+
+/* SViewport interface
+ *****************************************************************************/
+
+void SViewport::Construct( const FArguments& InArgs )
+{
+	ShowDisabledEffect = InArgs._ShowEffectWhenDisabled;
+	bRenderDirectlyToWindow = InArgs._RenderDirectlyToWindow;
+	bEnableGammaCorrection = InArgs._EnableGammaCorrection;
+	bEnableBlending = InArgs._EnableBlending;
+	bIgnoreTextureAlpha = InArgs._IgnoreTextureAlpha;
+
+	this->ChildSlot
+	[
+		InArgs._Content.Widget
+	];
+}
+
+
+int32 SViewport::OnPaint( const FGeometry& AllottedGeometry, const FSlateRect& MyClippingRect, FSlateWindowElementList& OutDrawElements, int32 LayerId, const FWidgetStyle& InWidgetStyle, bool bParentEnabled ) const
+{
+#if SLATE_HD_STATS
+	SCOPE_CYCLE_COUNTER( STAT_SlateOnPaint_SViewport );
+#endif
+	bool bEnabled = ShouldBeEnabled( bParentEnabled );
+	bool bShowDisabledEffect = ShowDisabledEffect.Get();
+	ESlateDrawEffect::Type DrawEffects = bShowDisabledEffect && !bEnabled ? ESlateDrawEffect::DisabledEffect : ESlateDrawEffect::None;
+
+	// Viewport texture alpha channels are often in an indeterminate state, even after the resolve,
+	// so we'll tell the shader to not use the alpha channel when blending
+	if( bIgnoreTextureAlpha )
+	{
+		DrawEffects |= ESlateDrawEffect::IgnoreTextureAlpha;
+	}
+
+	TSharedPtr<ISlateViewport> ViewportInterfacePin = ViewportInterface.Pin();
+
+	// Tell the interface that we are drawing.
+	if (ViewportInterfacePin.IsValid())
+	{
+		ViewportInterfacePin->OnDrawViewport( AllottedGeometry, MyClippingRect, OutDrawElements, LayerId, InWidgetStyle, bParentEnabled );
+	}
+
+	// Only draw a quad if not rendering directly to the backbuffer
+	if( !ShouldRenderDirectly() )
+	{
+		if( ViewportInterfacePin.IsValid() && ViewportInterfacePin->GetViewportRenderTargetTexture() != NULL )
+		{
+			FSlateDrawElement::MakeViewport( OutDrawElements, LayerId, AllottedGeometry.ToPaintGeometry(), ViewportInterfacePin, MyClippingRect, bEnableGammaCorrection, bEnableBlending, DrawEffects, InWidgetStyle.GetColorAndOpacityTint() );
+		}
+		else
+		{
+			// Viewport isn't ready yet, so just draw a black box
+			static FSlateColorBrush BlackBrush( FColor::Black );
+			FSlateDrawElement::MakeBox( OutDrawElements, LayerId, AllottedGeometry.ToPaintGeometry(), &BlackBrush, MyClippingRect, DrawEffects, BlackBrush.GetTint( InWidgetStyle ) );
+		}
+	}
+
+	int32 Layer = SCompoundWidget::OnPaint(AllottedGeometry, MyClippingRect, OutDrawElements, LayerId, InWidgetStyle, bEnabled );
+
+	if( ViewportInterfacePin.IsValid() && ViewportInterfacePin->IsSoftwareCursorVisible() )
+	{
+		const FVector2D CursorPosScreenSpace = FSlateApplication::Get().GetCursorPos();		
+		// @todo Slate: why are we calling OnCursorQuery in here?
+		FCursorReply Reply = ViewportInterfacePin->OnCursorQuery( AllottedGeometry,
+			FPointerEvent(
+				CursorPosScreenSpace,
+				CursorPosScreenSpace,
+				FVector2D::ZeroVector,
+				TSet<FKey>(),
+				FModifierKeysState(false, false, false, false, false, false) )
+		 );
+		EMouseCursor::Type CursorType = Reply.GetCursor();
+
+		const FSlateBrush* Brush = FCoreStyle::Get().GetBrush(TEXT("SoftwareCursor_Grab"));
+		if( CursorType == EMouseCursor::CardinalCross )
+		{
+			Brush = FCoreStyle::Get().GetBrush(TEXT("SoftwareCursor_CardinalCross"));
+		}
+
+		LayerId++;
+		FSlateDrawElement::MakeBox(
+			OutDrawElements,
+			LayerId,
+			AllottedGeometry.ToPaintGeometry( ViewportInterfacePin->GetSoftwareCursorPosition() - ( Brush->ImageSize / 2 ), Brush->ImageSize ),
+			Brush,
+			MyClippingRect
+		);
+	}
+
+	return Layer;
+}
+
+
+void SViewport::Tick( const FGeometry& AllottedGeometry, const double InCurrentTime, const float InDeltaTime )
+{
+	if(ViewportInterface.IsValid())
+	{
+		ViewportInterface.Pin()->Tick(InDeltaTime);
+	}
+}
+
+
+/* SWidget interface
+ *****************************************************************************/
+
+FCursorReply SViewport::OnCursorQuery( const FGeometry& MyGeometry, const FPointerEvent& CursorEvent ) const
+{
+	return ViewportInterface.IsValid() ? ViewportInterface.Pin()->OnCursorQuery(MyGeometry, CursorEvent) : FCursorReply::Unhandled();
+}
+
+
+FReply SViewport::OnMouseButtonDown( const FGeometry& MyGeometry, const FPointerEvent& MouseEvent )
+{
+	return ViewportInterface.IsValid() ? ViewportInterface.Pin()->OnMouseButtonDown(MyGeometry, MouseEvent) : FReply::Unhandled();
+}
+
+
+FReply SViewport::OnMouseButtonUp( const FGeometry& MyGeometry, const FPointerEvent& MouseEvent )
+{
+	return ViewportInterface.IsValid() ? ViewportInterface.Pin()->OnMouseButtonUp(MyGeometry, MouseEvent) : FReply::Unhandled();
+}
+
+
+void SViewport::OnMouseEnter( const FGeometry& MyGeometry, const FPointerEvent& MouseEvent )
+{
+	SCompoundWidget::OnMouseEnter(MyGeometry, MouseEvent);
+
+	if (ViewportInterface.IsValid())
+	{
+		ViewportInterface.Pin()->OnMouseEnter(MyGeometry, MouseEvent);
+	}
+}
+
+
+void SViewport::OnMouseLeave( const FPointerEvent& MouseEvent )
+{
+	SCompoundWidget::OnMouseLeave(MouseEvent);
+
+	if (ViewportInterface.IsValid())
+	{
+		ViewportInterface.Pin()->OnMouseLeave(MouseEvent);
+	}
+}
+
+
+FReply SViewport::OnMouseMove( const FGeometry& MyGeometry, const FPointerEvent& MouseEvent )
+{
+	return ViewportInterface.IsValid() ? ViewportInterface.Pin()->OnMouseMove(MyGeometry, MouseEvent) : FReply::Unhandled();
+}
+
+
+FReply SViewport::OnMouseWheel( const FGeometry& MyGeometry, const FPointerEvent& MouseEvent )
+{
+	return ViewportInterface.IsValid() ? ViewportInterface.Pin()->OnMouseWheel(MyGeometry, MouseEvent) : FReply::Unhandled();
+}
+
+
+FReply SViewport::OnMouseButtonDoubleClick( const FGeometry& MyGeometry, const FPointerEvent& MouseEvent )
+{
+	return ViewportInterface.IsValid() ? ViewportInterface.Pin()->OnMouseButtonDoubleClick(MyGeometry, MouseEvent) : FReply::Unhandled();
+}
+
+
+FReply SViewport::OnKeyDown( const FGeometry& MyGeometry, const FKeyboardEvent& KeyboardEvent )
+{
+	return ViewportInterface.IsValid() ? ViewportInterface.Pin()->OnKeyDown(MyGeometry, KeyboardEvent) : FReply::Unhandled();
+}
+
+ 
+FReply SViewport::OnKeyUp( const FGeometry& MyGeometry, const FKeyboardEvent& KeyboardEvent )
+{
+	return ViewportInterface.IsValid() ? ViewportInterface.Pin()->OnKeyUp(MyGeometry, KeyboardEvent) : FReply::Unhandled();
+}
+
+
+FReply SViewport::OnKeyChar( const FGeometry& MyGeometry, const FCharacterEvent& CharacterEvent )
+{
+	return ViewportInterface.IsValid() ? ViewportInterface.Pin()->OnKeyChar(MyGeometry, CharacterEvent) : FReply::Unhandled();
+}
+
+
+FReply SViewport::OnKeyboardFocusReceived( const FGeometry& MyGeometry, const FKeyboardFocusEvent& InKeyboardFocusEvent )
+{
+	if (WidgetToFocusOnActivate.IsValid())
+	{
+		return FReply::Handled().SetKeyboardFocus(WidgetToFocusOnActivate.Pin().ToSharedRef(), InKeyboardFocusEvent.GetCause());
+	}
+
+	return ViewportInterface.IsValid() ? ViewportInterface.Pin()->OnKeyboardFocusReceived(InKeyboardFocusEvent) : FReply::Unhandled();
+}
+
+
+void SViewport::OnKeyboardFocusLost( const FKeyboardFocusEvent& InKeyboardFocusEvent )
+{
+	if (ViewportInterface.IsValid())
+	{
+		ViewportInterface.Pin()->OnKeyboardFocusLost(InKeyboardFocusEvent);
+	}
+}
+
+
+void SViewport::OnWindowClosed( const TSharedRef<SWindow>& WindowBeingClosed )
+{
+	if (ViewportInterface.IsValid())
+	{
+		ViewportInterface.Pin()->OnViewportClosed();
+	}
+}
+
+
+FReply SViewport::OnControllerButtonPressed( const FGeometry& MyGeometry, const FControllerEvent& ControllerEvent )
+{
+	return ViewportInterface.IsValid() ? ViewportInterface.Pin()->OnControllerButtonPressed(MyGeometry, ControllerEvent) : FReply::Unhandled();
+}
+
+
+FReply SViewport::OnControllerButtonReleased( const FGeometry& MyGeometry, const FControllerEvent& ControllerEvent )
+{
+	return ViewportInterface.IsValid() ? ViewportInterface.Pin()->OnControllerButtonReleased(MyGeometry, ControllerEvent) : FReply::Unhandled();
+}
+
+
+FReply SViewport::OnControllerAnalogValueChanged( const FGeometry& MyGeometry, const FControllerEvent& ControllerEvent )
+{
+	return ViewportInterface.IsValid() ? ViewportInterface.Pin()->OnControllerAnalogValueChanged(MyGeometry, ControllerEvent) : FReply::Unhandled();
+}
+
+
+FReply SViewport::OnTouchStarted( const FGeometry& MyGeometry, const FPointerEvent& TouchEvent )
+{
+	return ViewportInterface.IsValid() ? ViewportInterface.Pin()->OnTouchStarted(MyGeometry, TouchEvent) : FReply::Unhandled();
+}
+
+
+FReply SViewport::OnTouchMoved( const FGeometry& MyGeometry, const FPointerEvent& TouchEvent )
+{
+	return ViewportInterface.IsValid() ? ViewportInterface.Pin()->OnTouchMoved(MyGeometry, TouchEvent) : FReply::Unhandled();
+}
+
+
+FReply SViewport::OnTouchEnded( const FGeometry& MyGeometry, const FPointerEvent& TouchEvent )
+{
+	return ViewportInterface.IsValid() ? ViewportInterface.Pin()->OnTouchEnded(MyGeometry, TouchEvent) : FReply::Unhandled();
+}
+
+
+FReply SViewport::OnTouchGesture( const FGeometry& MyGeometry, const FPointerEvent& GestureEvent )
+{
+	return ViewportInterface.IsValid() ? ViewportInterface.Pin()->OnTouchGesture(MyGeometry, GestureEvent) : FReply::Unhandled();
+}
+
+
+FReply SViewport::OnMotionDetected( const FGeometry& MyGeometry, const FMotionEvent& MotionEvent )
+{
+	return ViewportInterface.IsValid() ? ViewportInterface.Pin()->OnMotionDetected(MyGeometry, MotionEvent) : FReply::Unhandled();
+}
diff --git a/Engine/Source/Runtime/Slate/Private/Widgets/SWeakWidget.cpp b/Engine/Source/Runtime/Slate/Private/Widgets/SWeakWidget.cpp
new file mode 100644
--- /dev/null
+++ b/Engine/Source/Runtime/Slate/Private/Widgets/SWeakWidget.cpp
@@ -0,0 +1,82 @@
+// Copyright 1998-2014 Epic Games, Inc. All Rights Reserved.
+
+/*=============================================================================
+	SWeakWidget.cpp: Implements the SWeakWidget class.
+=============================================================================*/
+
+#include "SlatePrivatePCH.h"
+
+
+void SWeakWidget::Construct(const FArguments& InArgs)
+{
+	WeakChild.Widget = InArgs._PossiblyNullContent;
+}
+
+FVector2D SWeakWidget::ComputeDesiredSize() const
+{	
+	TSharedPtr<SWidget> ReferencedWidget = WeakChild.Widget.Pin();
+
+	if (ReferencedWidget.IsValid() && ReferencedWidget->GetVisibility() != EVisibility::Collapsed)
+	{
+		return ReferencedWidget->GetDesiredSize();
+	}
+
+	return FVector2D::ZeroVector;
+}
+
+void SWeakWidget::ArrangeChildren( const FGeometry& AllottedGeometry, FArrangedChildren& ArrangedChildren ) const
+{
+	// We just want to show the child that we are presenting. Always stretch it to occupy all of the space.
+	TSharedPtr<SWidget> MyContent = WeakChild.Widget.Pin();
+
+	if( MyContent.IsValid() && ArrangedChildren.Accepts(MyContent->GetVisibility()))
+	{
+		ArrangedChildren.AddWidget( AllottedGeometry.MakeChild(
+			MyContent.ToSharedRef(),
+			FVector2D(0,0),
+			AllottedGeometry.Size
+			) );
+	}
+}
+		
+FChildren* SWeakWidget::GetChildren()
+{
+	return &WeakChild;
+}
+
+int32 SWeakWidget::OnPaint( const FGeometry& AllottedGeometry, const FSlateRect& MyClippingRect, FSlateWindowElementList& OutDrawElements, int32 LayerId, const FWidgetStyle& InWidgetStyle, bool bParentEnabled ) const
+{
+	// Just draw the children.
+	FArrangedChildren ArrangedChildren(EVisibility::Visible);
+	this->ArrangeChildren(AllottedGeometry, ArrangedChildren);
+
+	// There may be zero elements in this array if our child collapsed/hidden
+	if( ArrangedChildren.Num() > 0 )
+	{
+		check( ArrangedChildren.Num() == 1 );
+		FArrangedWidget& TheChild = ArrangedChildren(0);
+
+		return TheChild.Widget->OnPaint( TheChild.Geometry, 
+			MyClippingRect, 
+			OutDrawElements, 
+			LayerId + 1,
+			InWidgetStyle, 
+			ShouldBeEnabled( bParentEnabled ) );
+	}
+	return LayerId;
+}
+
+void SWeakWidget::SetContent(TWeakPtr<SWidget> InWidget)
+{
+	WeakChild.Widget = InWidget;
+}
+
+bool SWeakWidget::ChildWidgetIsValid() const
+{
+	return WeakChild.Widget.IsValid();
+}
+
+TWeakPtr<SWidget> SWeakWidget::GetChildWidget() const
+{
+	return WeakChild.Widget;
+}
diff --git a/Engine/Source/Runtime/Slate/Private/Widgets/Text/SInlineEditableTextBlock.cpp b/Engine/Source/Runtime/Slate/Private/Widgets/Text/SInlineEditableTextBlock.cpp
new file mode 100644
--- /dev/null
+++ b/Engine/Source/Runtime/Slate/Private/Widgets/Text/SInlineEditableTextBlock.cpp
@@ -0,0 +1,259 @@
+// Copyright 1998-2014 Epic Games, Inc. All Rights Reserved.
+
+#include "SlatePrivatePCH.h"
+
+
+void SInlineEditableTextBlock::Construct( const FArguments& InArgs )
+{
+	check(InArgs._Style);
+
+	OnBeginTextEditDelegate = InArgs._OnBeginTextEdit;
+	OnTextCommittedDelegate = InArgs._OnTextCommitted;
+	IsSelected = InArgs._IsSelected;
+	OnVerifyTextChanged= InArgs._OnVerifyTextChanged;
+	Text = InArgs._Text;
+	bIsReadOnly = InArgs._IsReadOnly;
+	DoubleSelectDelay = 0.0f;
+
+	ChildSlot
+	[
+		SAssignNew(HorizontalBox, SHorizontalBox)
+			
+		+SHorizontalBox::Slot()
+		.VAlign(VAlign_Center)
+		[
+			SAssignNew(TextBlock, STextBlock)
+			.Text(Text)
+			.TextStyle( &InArgs._Style->TextStyle )
+			.Font(InArgs._Font)
+			.ColorAndOpacity( InArgs._ColorAndOpacity )
+			.ShadowColorAndOpacity( InArgs._ShadowColorAndOpacity )
+			.ShadowOffset( InArgs._ShadowOffset )
+			.HighlightText( InArgs._HighlightText )
+			.ToolTipText( InArgs._ToolTipText )
+			.WrapTextAt( InArgs._WrapTextAt )
+		]
+	];
+
+	SAssignNew(TextBox, SEditableTextBox)
+		.Text(InArgs._Text)
+		.Style(&InArgs._Style->EditableTextBoxStyle)
+		.Font(InArgs._Font)
+		.ToolTipText( InArgs._ToolTipText )
+		.OnTextChanged( this, &SInlineEditableTextBlock::OnTextChanged )
+		.OnTextCommitted(this, &SInlineEditableTextBlock::OnTextBoxCommitted)
+		.SelectAllTextWhenFocused(true)
+		.ClearKeyboardFocusOnCommit(false);
+}
+
+SInlineEditableTextBlock::~SInlineEditableTextBlock()
+{
+	if(IsInEditMode())
+	{
+		// Clear the error so it will vanish.
+		TextBox->SetError( FText::GetEmpty() );
+	}
+}
+
+void SInlineEditableTextBlock::CancelEditMode()
+{
+	HorizontalBox->RemoveSlot(TextBox.ToSharedRef());
+	TextBlock->SetVisibility(EVisibility::Visible);
+
+	// Clear the error so it will vanish.
+	TextBox->SetError( FText::GetEmpty() );
+
+	// Get the text from source again.
+	TextBox->SetText(Text);
+}
+
+bool SInlineEditableTextBlock::SupportsKeyboardFocus() const
+{
+	//Can not have keyboard focus if it's status of being selected is managed by another widget.
+	return !IsSelected.IsBound();
+}
+
+void SInlineEditableTextBlock::EnterEditingMode()
+{
+	if(!bIsReadOnly.Get() && !FSlateApplication::Get().GetMouseCaptor().IsValid() )
+	{
+		if(TextBlock->GetVisibility() == EVisibility::Visible)
+		{
+			HorizontalBox->AddSlot()
+				[
+					TextBox.ToSharedRef()
+				];
+
+			FSlateApplication::Get().SetKeyboardFocus(TextBox, EKeyboardFocusCause::SetDirectly);
+
+			TextBlock->SetVisibility(EVisibility::Collapsed);
+
+			OnBeginTextEditDelegate.ExecuteIfBound(TextBox->GetText());
+		}
+	}
+}
+
+void SInlineEditableTextBlock::ExitEditingMode()
+{
+	FSlateApplication::Get().ClearKeyboardFocus(EKeyboardFocusCause::SetDirectly);	
+}
+
+bool SInlineEditableTextBlock::IsInEditMode() const
+{
+	return TextBlock->GetVisibility() == EVisibility::Collapsed;
+}
+
+void SInlineEditableTextBlock::SetText( const TAttribute< FText >& InText )
+{
+	Text = InText;
+	TextBlock->SetText( Text );
+	TextBox->SetText( Text );
+}
+
+void SInlineEditableTextBlock::SetText( const FString& InText )
+{
+	Text = FText::FromString( InText );
+	TextBlock->SetText( Text );
+	TextBox->SetText( Text );
+}
+
+FReply SInlineEditableTextBlock::OnMouseButtonDown( const FGeometry& MyGeometry, const FPointerEvent& MouseEvent )
+{
+	if( !MouseEvent.IsMouseButtonDown(EKeys::LeftMouseButton) || MouseEvent.IsControlDown() || MouseEvent.IsShiftDown())
+	{
+		return FReply::Unhandled();
+	}
+
+	if(IsSelected.IsBound())
+	{
+		if(IsSelected.Execute() && !bIsReadOnly.Get())
+		{
+			DoubleSelectDelay = 0.5f;
+		}
+	}
+	else
+	{
+		// The widget is not managed by another widget, so handle the mouse input and enter edit mode if ready.
+		if(HasKeyboardFocus())
+		{
+			EnterEditingMode();
+		}
+		return FReply::Handled();
+	}
+
+	// Do not handle the mouse input, this will allow for drag and dropping events to trigger.
+	return FReply::Unhandled();
+}
+
+FReply SInlineEditableTextBlock::OnDragOver( const FGeometry& MyGeometry, const FDragDropEvent& DragDropEvent )
+{
+	// Cancel during a drag over event, otherwise the widget may enter edit mode.
+	DoubleSelectDelay = 0.0f;
+	return FReply::Unhandled();
+}
+
+FReply SInlineEditableTextBlock::OnMouseButtonDoubleClick( const FGeometry& InMyGeometry, const FPointerEvent& InMouseEvent )
+{
+	DoubleSelectDelay = 0.0f;
+	return FReply::Unhandled();
+}
+
+void SInlineEditableTextBlock::Tick( const FGeometry& AllottedGeometry, const double InCurrentTime, const float InDeltaTime )
+{
+	if(DoubleSelectDelay > 0.0f && IsSelected.IsBound())
+	{
+		DoubleSelectDelay -= InDeltaTime;
+
+		if(DoubleSelectDelay <= 0.0f && IsSelected.Execute())
+		{
+			EnterEditingMode();
+		}
+	}
+}
+
+FReply SInlineEditableTextBlock::OnKeyDown( const FGeometry& MyGeometry, const FKeyboardEvent& InKeyboardEvent )
+{
+	if(InKeyboardEvent.GetKey() == EKeys::F2)
+	{
+		EnterEditingMode();
+		return FReply::Handled();
+	}
+	return FReply::Unhandled();
+}
+
+void SInlineEditableTextBlock::OnTextChanged(const FText& InText)
+{
+	if(IsInEditMode())
+	{
+		FText OutErrorMessage;
+
+		if(OnVerifyTextChanged.IsBound() && !OnVerifyTextChanged.Execute(InText, OutErrorMessage))
+		{
+			TextBox->SetError( OutErrorMessage );
+		}
+		else
+		{
+			TextBox->SetError( FText::GetEmpty() );
+		}
+	}
+}
+
+void SInlineEditableTextBlock::OnTextBoxCommitted(const FText& InText, ETextCommit::Type InCommitType)
+{
+	if(InCommitType == ETextCommit::OnCleared)
+	{
+		CancelEditMode();
+		// Commit the name, certain actions might need to be taken by the bound function
+		OnTextCommittedDelegate.ExecuteIfBound(Text.Get(), InCommitType);
+	}
+	else if(IsInEditMode())
+	{
+		if(OnVerifyTextChanged.IsBound())
+		{
+			if(InCommitType == ETextCommit::OnEnter)
+			{
+				FText OutErrorMessage;
+				if(!OnVerifyTextChanged.Execute(InText, OutErrorMessage))
+				{
+					// Display as an error.
+					TextBox->SetError( OutErrorMessage );
+					return;
+				}
+			}
+			else if(InCommitType == ETextCommit::OnUserMovedFocus)
+			{
+				FText OutErrorMessage;
+				if(!OnVerifyTextChanged.Execute(InText, OutErrorMessage))
+				{
+					CancelEditMode();
+
+					// Commit the name, certain actions might need to be taken by the bound function
+					OnTextCommittedDelegate.ExecuteIfBound(Text.Get(), InCommitType);
+				
+					return;
+				}
+			}
+			else // When the user removes all focus from the window, revert the name
+			{
+				CancelEditMode();
+
+				// Commit the name, certain actions might need to be taken by the bound function
+				OnTextCommittedDelegate.ExecuteIfBound(Text.Get(), InCommitType);
+				return;
+			}
+		}
+		
+		HorizontalBox->RemoveSlot(TextBox.ToSharedRef());
+		TextBlock->SetVisibility(EVisibility::Visible);
+
+		OnTextCommittedDelegate.ExecuteIfBound(InText, InCommitType);
+
+		// Clear the error so it will vanish.
+		TextBox->SetError( FText::GetEmpty() );
+
+		if ( !Text.IsBound() )
+		{
+			TextBlock->SetText( Text );
+		}
+	}
+}
