Slate: Reorganized private Slate files; proper pre-compiled header for all files.

diff --git a/Engine/Source/Runtime/Slate/Private/Widgets/Text/SMultiLineEditableText.cpp b/Engine/Source/Runtime/Slate/Private/Widgets/Text/SMultiLineEditableText.cpp
new file mode 100644
--- /dev/null
+++ b/Engine/Source/Runtime/Slate/Private/Widgets/Text/SMultiLineEditableText.cpp
@@ -0,0 +1,562 @@
+// Copyright 1998-2014 Epic Games, Inc. All Rights Reserved.
+
+#include "SlatePrivatePCH.h"
+#include "TextEditHelper.h"
+#include "SlateWordWrapper.h"
+
+SMultiLineEditableText::SMultiLineEditableText()
+: CursorPosition( FTextLocation() )
+, FontInfo( FPaths::EngineContentDir() / TEXT("Slate/Fonts/Roboto-Regular.ttf"), 10 ) 
+, NumLinesScrollOffset( 0.0f )
+, PreferredCursorOffsetInLine( 0 )
+, LastCursorInteractionTime( 0 )
+, PreferredSize( FVector2D::ZeroVector )
+{
+}
+
+void SMultiLineEditableText::Construct( const FArguments& InArgs )
+{
+	NumLinesScrollOffset = 0.0f;
+	OptionalScrollBarPtr = InArgs._ScrollBar;
+	PreferredSize = InArgs._PreferredSize;
+	if (InArgs._ScrollBar.IsValid())
+	{
+		InArgs._ScrollBar->SetOnUserScrolled( FOnUserScrolled::CreateSP( this, &SMultiLineEditableText::OnScrollbarScrolled ) );
+	}
+
+	TArray<FString> TempTextLines;
+	InArgs._Text.ParseIntoArray( &TempTextLines, TEXT("\n"), false );
+	TextLines.Reserve( TempTextLines.Num() );
+
+	for ( int32 LineIndex=0; LineIndex < TempTextLines.Num(); ++LineIndex )
+	{
+		TextLines.Add( MakeShareable( new FString( TempTextLines[LineIndex] ) ) );
+	}
+}
+
+void SMultiLineEditableText::StartChangingText()
+{
+}
+
+void SMultiLineEditableText::FinishChangingText()
+{
+}
+
+bool SMultiLineEditableText::GetIsReadOnly() const
+{
+	return false;
+}
+
+void SMultiLineEditableText::BackspaceChar()
+{
+}
+
+void SMultiLineEditableText::DeleteChar()
+{
+}
+
+bool SMultiLineEditableText::CanTypeCharacter(const TCHAR CharInQuestion) const
+{
+	return true;
+}
+
+void SMultiLineEditableText::TypeChar( const int32 Character )
+{
+}
+
+FReply SMultiLineEditableText::MoveCursor( ECursorMoveMethod::Type Method, int8 Direction, ECursorAction::Type Action )
+{
+	FTextLocation NewCursorPosition;
+	if (Method == ECursorMoveMethod::CharacterHorizontal)
+	{
+		NewCursorPosition = TranslatedLocation(CursorPosition, Direction);
+		PreferredCursorOffsetInLine = AbsoluteToWrapped(NewCursorPosition).GetOffset();
+	}
+	else if (Method == ECursorMoveMethod::Word)
+	{
+		NewCursorPosition = ScanForWordBoundary( CursorPosition, Direction );
+		PreferredCursorOffsetInLine = AbsoluteToWrapped(NewCursorPosition).GetOffset();
+	}
+	else
+	{
+		ensure( Method == ECursorMoveMethod::CharacterVertical );
+		NewCursorPosition = TranslateLocationVertical(CursorPosition, Direction);
+	}
+
+	if (Action == ECursorAction::SelectText)
+	{
+		// We are selecting text. Just remember where the selection started.
+		// The cursor is implicitly the other endpoint.
+		if (!SelectionStart.IsSet())
+		{
+			this->SelectionStart = CursorPosition;
+		}
+	}
+	else
+	{
+		// No longer selection text; clear the selection!
+		this->ClearSelection();
+	}
+
+	if( WrappedText.Num() > 0 )
+	{
+		const int32 NewWrappedCursorPosition = AbsoluteToWrapped( NewCursorPosition ).GetLineIndex();
+		const float LineHeight = WrappedText[NewCursorPosition.GetLineIndex()].Size.Y;
+		const float NumLinesOnScreen = CachedLastFrameGeometry.Size.Y / LineHeight;
+		const float FractionOfTotalLinesOnScreen = FMath::Clamp(NumLinesOnScreen / WrappedText.Num(), 0.0f, 1.0f);
+		if ( NewWrappedCursorPosition < NumLinesScrollOffset )
+		{
+			NumLinesScrollOffset = NewWrappedCursorPosition;
+			UpdateScrollbar( CachedLastFrameGeometry );
+		}
+		else if ( FMath::Floor(NumLinesScrollOffset + NumLinesOnScreen - 1) < NewWrappedCursorPosition )
+		{
+			NumLinesScrollOffset = FMath::Clamp( FMath::Floor(NewWrappedCursorPosition - NumLinesOnScreen + 1 ), 0, WrappedText.Num());
+			UpdateScrollbar( CachedLastFrameGeometry );
+		}
+	}
+
+	CursorPosition = NewCursorPosition;
+
+	LastCursorInteractionTime = FSlateApplication::Get().GetCurrentTime();
+
+	return FReply::Handled();
+
+}
+void SMultiLineEditableText::JumpTo(ETextLocation::Type JumpLocation, ECursorAction::Type Action)
+{
+	
+}
+
+void SMultiLineEditableText::ClearSelection()
+{
+	SelectionStart = TOptional<FTextLocation>();
+}
+
+void SMultiLineEditableText::SelectAllText()
+{
+	SelectionStart = FTextLocation(0,0);
+	CursorPosition = FTextLocation( TextLines.Num() - 1, TextLines[TextLines.Num()-1]->Len() );
+}
+
+FReply SMultiLineEditableText::OnEscape()
+{
+	return FReply::Unhandled();
+}
+
+void SMultiLineEditableText::OnEnter()
+{
+}
+
+bool SMultiLineEditableText::CanExecuteCut() const
+{
+	return true;
+}
+
+void SMultiLineEditableText::CutSelectedTextToClipboard()
+{
+}
+
+bool SMultiLineEditableText::CanExecuteCopy() const
+{
+	return true;
+}
+
+void SMultiLineEditableText::CopySelectedTextToClipboard()
+{
+}
+
+bool SMultiLineEditableText::CanExecutePaste() const
+{
+	return true;
+}
+
+void SMultiLineEditableText::PasteTextFromClipboard()
+{
+}
+
+bool SMultiLineEditableText::CanExecuteUndo() const
+{
+	return true;
+}
+
+void SMultiLineEditableText::Undo()
+{
+}
+
+void SMultiLineEditableText::Redo()
+{
+}
+
+void SMultiLineEditableText::Tick( const FGeometry& AllottedGeometry, const double InCurrentTime, const float InDeltaTime )
+{
+	if ( CachedLastFrameGeometry.Size != AllottedGeometry.Size )
+	{
+		CachedLastFrameGeometry = AllottedGeometry;
+		ReWrap( AllottedGeometry );
+		UpdateScrollbar(AllottedGeometry);
+
+	}
+}
+
+int32 SMultiLineEditableText::OnPaint( const FGeometry& AllottedGeometry, const FSlateRect& MyClippingRect, FSlateWindowElementList& OutDrawElements, int32 LayerId, const FWidgetStyle& InWidgetStyle, bool bParentEnabled ) const
+{
+	const int32 TextLayer = LayerId;
+	const int32 SelectionLayer = LayerId + 1;
+	const TSharedRef< FSlateFontMeasure > FontMeasureService = FSlateApplication::Get().GetRenderer()->GetFontMeasureService();
+
+	const float CursorWidth = 2.0f;
+
+	const double CurrentTime = FSlateApplication::Get().GetCurrentTime();
+
+	const bool bEnabled = ShouldBeEnabled( bParentEnabled );
+	const ESlateDrawEffect::Type DrawEffects = (bEnabled) ? ESlateDrawEffect::None : ESlateDrawEffect::DisabledEffect;
+
+	const int32 WholeLinesOffset = FMath::Floor( NumLinesScrollOffset );
+	const float PartialLineOffset = FMath::Fractional( NumLinesScrollOffset );
+
+
+	if (WrappedText.Num() > 0)
+	{
+		float HeightOffsetSoFar = - PartialLineOffset *  WrappedText[FMath::Clamp(WholeLinesOffset, 0, WrappedText.Num()-1)].Size.Y;
+
+		bool FilledVerticalSpace = false;
+
+		const FTextRange Selection = FTextRange(SelectionStart.GetValue(), CursorPosition);
+		FTextRange SelectionInWrappedText = (SelectionStart.IsSet())
+			? AbsoluteToWrapped( Selection )
+			: FTextRange( FTextLocation(), FTextLocation() );
+
+		bool bInSelection = SelectionInWrappedText.First.GetLineIndex() < WholeLinesOffset && SelectionInWrappedText.First.GetLineIndex() > WholeLinesOffset;
+
+		for ( int32 LineIndex=WholeLinesOffset; !FilledVerticalSpace && LineIndex < WrappedText.Num(); ++LineIndex )
+		{
+			const FWrappedStringSlice& ThisWrappedLine = WrappedText[LineIndex];
+			// @todo FontService: Copy is undesirable?
+			const FString StringToPaint = ThisWrappedLine.SourceString.Mid( ThisWrappedLine.FirstCharIndex, ThisWrappedLine.LastCharIndex-ThisWrappedLine.FirstCharIndex );
+
+			// Draw this wrapped line of text
+			FSlateDrawElement::MakeText(
+				OutDrawElements,
+				LayerId,
+				AllottedGeometry.ToPaintGeometry(FVector2D(0, HeightOffsetSoFar), FVector2D::UnitVector),
+				StringToPaint,
+				FontInfo,
+				MyClippingRect,
+				ShouldBeEnabled(bParentEnabled),
+				InWidgetStyle.GetColorAndOpacityTint()
+			);
+
+			static const FLinearColor SelectionColor = FLinearColor(0,0,1.0f,0.25f);
+
+			// DRAW SELECTION
+			if ( SelectionStart.IsSet() )
+			{
+				const bool bIsFirstSelectionLine =
+					(&(TextLines[Selection.First.GetLineIndex()].Get()) == &ThisWrappedLine.SourceString) &&
+					Selection.First.GetOffset() >= ThisWrappedLine.FirstCharIndex &&
+					Selection.First.GetOffset() <= ThisWrappedLine.LastCharIndex;
+
+				const bool bIsLastSelectionLine =
+				 	(&(TextLines[Selection.Last.GetLineIndex()].Get()) == &ThisWrappedLine.SourceString) &&
+					Selection.Last.GetOffset() >= ThisWrappedLine.FirstCharIndex &&
+					Selection.Last.GetOffset() <= ThisWrappedLine.LastCharIndex;
+
+				if ( bIsFirstSelectionLine && bIsLastSelectionLine )
+				{
+					// Selection contained within a single line.
+					const FString StringBeforeSelection = ThisWrappedLine.SourceString.Left( SelectionInWrappedText.First.GetOffset() );
+					const FVector2D SizeBeforeSelection = FontMeasureService->Measure(StringBeforeSelection, FontInfo, AllottedGeometry.Scale);					
+					
+					const FString SelectedString = ThisWrappedLine.SourceString.Mid( SelectionInWrappedText.First.GetOffset(), SelectionInWrappedText.Last.GetOffset() - SelectionInWrappedText.First.GetOffset() );		
+					const FVector2D SelectionSize = FontMeasureService->Measure(SelectedString, FontInfo, AllottedGeometry.Scale);
+
+					FSlateDrawElement::MakeBox(
+						OutDrawElements,
+						SelectionLayer,
+						AllottedGeometry.ToPaintGeometry(FVector2D(SizeBeforeSelection.X, HeightOffsetSoFar), SelectionSize),
+						FCoreStyle::Get().GetBrush( "WhiteBrush" ),
+						MyClippingRect,
+						DrawEffects,
+						SelectionColor
+					);
+
+				}
+				else if ( bIsFirstSelectionLine )
+				{
+					// Selection starts on this line, ends on another.
+					bInSelection = true;
+
+					// @todo FontService: Copy is unnecessary
+					const FString StringBeforeSelection = StringToPaint.Left(SelectionInWrappedText.First.GetOffset());
+					const FVector2D UnselectedPortionSize = FontMeasureService->Measure(StringBeforeSelection, FontInfo, AllottedGeometry.Scale);
+					FSlateDrawElement::MakeBox(
+						OutDrawElements,
+						SelectionLayer,
+						AllottedGeometry.ToPaintGeometry(FVector2D(UnselectedPortionSize.X, HeightOffsetSoFar), ThisWrappedLine.Size - FVector2D(UnselectedPortionSize.X, 0)),
+						FCoreStyle::Get().GetBrush( "WhiteBrush" ),
+						MyClippingRect,
+						DrawEffects,
+						SelectionColor
+					);
+				}
+				else if ( bIsLastSelectionLine )
+				{
+					// Selection starts on some line above, ends on this one.
+					bInSelection = false;
+
+					// @todo FontService: Copy is unnecessary
+					const FString StringBeforeSelection = StringToPaint.Left(SelectionInWrappedText.Last.GetOffset());
+					const FVector2D SelectedPortionSize = FontMeasureService->Measure(StringBeforeSelection, FontInfo, AllottedGeometry.Scale);
+					FSlateDrawElement::MakeBox(
+						OutDrawElements,
+						SelectionLayer,
+						AllottedGeometry.ToPaintGeometry(FVector2D(0, HeightOffsetSoFar), FVector2D(SelectedPortionSize.X, ThisWrappedLine.Size.Y)),
+						FCoreStyle::Get().GetBrush( "WhiteBrush" ),
+						MyClippingRect,
+						DrawEffects,
+						SelectionColor
+					);
+				}
+				else if (bInSelection)
+				{
+					// The entirety of this line is selected
+					FSlateDrawElement::MakeBox(
+						OutDrawElements,
+						SelectionLayer,
+						AllottedGeometry.ToPaintGeometry(FVector2D(0, HeightOffsetSoFar), ThisWrappedLine.Size),
+						FCoreStyle::Get().GetBrush( "WhiteBrush" ),
+						MyClippingRect,
+						DrawEffects,
+						SelectionColor
+					);
+				}
+
+			}
+
+			// DRAW CURSOR
+			if (this->HasKeyboardFocus())
+			{
+				const bool CursorIsOnThisTextLine = ( &ThisWrappedLine.SourceString == &(TextLines[CursorPosition.GetLineIndex()].Get()) );
+				const bool bIsStartOfNewLine = ThisWrappedLine.FirstCharIndex == 0;
+				if ( CursorIsOnThisTextLine &&
+					((bIsStartOfNewLine && CursorPosition.GetOffset() == 0) || CursorPosition.GetOffset() > ThisWrappedLine.FirstCharIndex) &&
+					CursorPosition.GetOffset() <= ThisWrappedLine.LastCharIndex )
+				{
+					// The cursor is always visible (i.e. not blinking) when we're interacting with it; otherwise it might get lost.
+					const bool bForceCursorVisible = (CurrentTime - LastCursorInteractionTime) < EditableTextDefs::CaretBlinkPauseTime;
+					const float CursorOpacity = (bForceCursorVisible)
+						? 1.0f
+						: FMath::Round( FMath::MakePulsatingValue( CurrentTime, EditableTextDefs::BlinksPerSecond ));
+
+					const int32 CursorOffsetInWrappedLine = CursorPosition.GetOffset() - ThisWrappedLine.FirstCharIndex;
+					const FVector2D SlateUnitsCursorOffsetInLine = FontMeasureService->Measure( StringToPaint.Left(CursorOffsetInWrappedLine), FontInfo, AllottedGeometry.Scale );
+
+					const FVector2D DrawPosition = FVector2D( SlateUnitsCursorOffsetInLine.X, HeightOffsetSoFar );
+					const FVector2D DrawSize = FVector2D( CursorWidth, ThisWrappedLine.Size.Y );
+					// @todo: State Styles - make this brush part of the widget style
+					const FSlateBrush* StyleInfo = FCoreStyle::Get().GetBrush("EditableText.SelectionBackground");
+
+					FSlateDrawElement::MakeBox(
+						OutDrawElements,
+						LayerId,
+						AllottedGeometry.ToPaintGeometry(DrawPosition, DrawSize),
+						StyleInfo,
+						MyClippingRect,
+						DrawEffects,
+						FLinearColor(1,1,1,CursorOpacity) );
+
+				}
+			}
+
+			HeightOffsetSoFar += WrappedText[LineIndex].Size.Y;
+		}
+	}
+
+	return SelectionLayer;
+}
+
+
+FVector2D SMultiLineEditableText::ComputeDesiredSize() const
+{
+	return PreferredSize;
+}
+
+bool SMultiLineEditableText::SupportsKeyboardFocus() const
+{
+	return true;
+}
+
+FReply SMultiLineEditableText::OnKeyChar( const FGeometry& MyGeometry,const FCharacterEvent& InCharacterEvent )
+{
+	return FTextEditHelper::OnKeyChar( MyGeometry, InCharacterEvent, *this );
+}
+
+FReply SMultiLineEditableText::OnKeyDown( const FGeometry& MyGeometry, const FKeyboardEvent& InKeyboardEvent )
+{
+	FReply Reply = FTextEditHelper::OnKeyDown( MyGeometry, InKeyboardEvent, *this );
+
+	//// Process keybindings if the event wasn't already handled
+	//if( !Reply.IsEventHandled() && UICommandList->ProcessCommandBindings( InKeyboardEvent ) )
+	//{
+	//	Reply = FReply::Handled();
+	//}
+
+	return Reply;
+}
+
+FReply SMultiLineEditableText::OnKeyUp( const FGeometry& MyGeometry, const FKeyboardEvent& InKeyboardEvent )
+{
+	return FReply::Unhandled();
+}
+
+
+void SMultiLineEditableText::ReWrap( const FGeometry& AllottedGeometry )
+{
+	SlateWordWrapper::WrapText(TextLines, FontInfo, AllottedGeometry.Scale, AllottedGeometry.Size.X , WrappedText);
+}
+
+void SMultiLineEditableText::OnScrollbarScrolled( float OffsetFraction )
+{
+	NumLinesScrollOffset = FMath::Clamp(OffsetFraction * WrappedText.Num(), 0.0f, WrappedText.Num()-1.0f);
+	UpdateScrollbar( CachedLastFrameGeometry );
+}
+
+/** Remember where the cursor was when we started selectin. */
+void SMultiLineEditableText::BeginSelecting( const FTextLocation& Endpoint )
+{
+	SelectionStart = Endpoint;
+}
+
+
+void SMultiLineEditableText::UpdateScrollbar( const FGeometry &AllottedGeometry ) 
+{
+	TSharedPtr<SScrollBar> OptionalScrollBar = OptionalScrollBarPtr.Pin();
+	if (OptionalScrollBar.IsValid() && WrappedText.Num() > 0)
+	{
+		const float LineHeight = WrappedText[0].Size.Y;
+		const float NumLinesOnScreen = AllottedGeometry.Size.Y / LineHeight;
+		const float FractionOfTotalLinesOnScreen = FMath::Clamp(NumLinesOnScreen / WrappedText.Num(), 0.0f, 1.0f);
+		const float MaxVerticalOffsetInLines = FMath::Max(0.0f, WrappedText.Num() - NumLinesOnScreen);
+		NumLinesScrollOffset = FMath::Clamp( NumLinesScrollOffset, 0.0f, MaxVerticalOffsetInLines );
+		OptionalScrollBar->SetState(NumLinesScrollOffset / WrappedText.Num(), FractionOfTotalLinesOnScreen);
+	}
+}
+
+SMultiLineEditableText::FTextLocation SMultiLineEditableText::TranslatedLocation( const FTextLocation& Location, int8 Direction ) const
+{
+	const int32 OffsetInLine = Location.GetOffset() + Direction;
+		
+	if ( OffsetInLine > TextLines[Location.GetLineIndex()]->Len() && Location.GetLineIndex() < TextLines.Num()-1 )
+	{
+		// We're going over the end of the line and we aren't on the last line
+		return FTextLocation( Location.GetLineIndex() + 1, 0 );
+	}
+	else if ( OffsetInLine < 0 && Location.GetLineIndex() > 0 )
+	{
+		// We're stepping before the beginning of the line, and we're not on the first line.
+		return FTextLocation( Location.GetLineIndex() - 1, TextLines[Location.GetLineIndex() - 1]->Len() );	
+	}
+	else
+	{
+		return FTextLocation( Location.GetLineIndex(), FMath::Clamp(OffsetInLine, 0, TextLines[Location.GetLineIndex()]->Len()) );
+	}
+}
+
+/** Find wrapped line on which the cursor is currently located. */
+struct FMatchWrappedLine
+{
+	FMatchWrappedLine( const FString& InStringToMatch, int32 InOffset )
+		: StringToMatch( InStringToMatch )
+		, Offset( InOffset )
+		{}
+	bool Matches( const FWrappedStringSlice& CandidateLine  ) const
+	{
+		return
+			(&StringToMatch == &(CandidateLine.SourceString)) &&
+			Offset >= CandidateLine.FirstCharIndex &&
+			Offset <= CandidateLine.LastCharIndex;
+	}
+	const FString& StringToMatch;
+	const int32 Offset;
+};
+
+SMultiLineEditableText::FTextLocation SMultiLineEditableText::TranslateLocationVertical( const FTextLocation& Location, int8 Direction ) const
+{
+	const int32 CurrentWrappedLineIndex = WrappedText.FindMatch( FMatchWrappedLine( TextLines[Location.GetLineIndex()].Get(), Location.GetOffset() ) );
+	const FWrappedStringSlice& CurrentWrappedLine = WrappedText[CurrentWrappedLineIndex];
+	ensure(CurrentWrappedLineIndex != INDEX_NONE);
+
+	const int32 NewWrappedLinedIndex = FMath::Clamp(CurrentWrappedLineIndex + Direction, 0, WrappedText.Num()-1);
+	const FWrappedStringSlice& NewWrappedLine = WrappedText[NewWrappedLinedIndex];
+	
+	// We moved wrapped lines. Did we actually move to a different source line?
+	const bool bMovedToNewSourceLine = &CurrentWrappedLine.SourceString != &NewWrappedLine.SourceString;
+	const int32 NewSourceLineIndex = Location.GetLineIndex() + ( bMovedToNewSourceLine ? Direction : 0 );
+	// Our horizontal position is the clamped version of whatever the user explicitly set with horizontal movement.
+	const int32 NewOffsetInLine = FMath::Clamp(NewWrappedLine.FirstCharIndex + PreferredCursorOffsetInLine, 0, NewWrappedLine.LastCharIndex);
+	
+	return FTextLocation( NewSourceLineIndex, NewOffsetInLine );	
+}
+
+SMultiLineEditableText::FTextLocation SMultiLineEditableText::ScanForWordBoundary( const FTextLocation& CurrentLocation, int8 Direction ) const
+{
+	FTextLocation Location = TranslatedLocation(CurrentLocation, Direction);
+
+	while ( !IsAtStartOfDoc(Location) && !IsAtStartOfLine(Location) && !IsAtEndOfDoc(Location) && !IsAtEndOfLine(Location) && !IsAtWordStart(Location) )
+	{
+		Location = TranslatedLocation(Location, Direction);
+	}	
+
+	return Location;
+}
+
+TCHAR SMultiLineEditableText::GetCharacterAt( const FTextLocation& Location ) const
+{
+	const bool bIsLineEmpty = TextLines[Location.GetLineIndex()]->IsEmpty();
+	return (bIsLineEmpty)
+		? '\n'
+		: (*TextLines[Location.GetLineIndex()])[Location.GetOffset()];
+}
+
+
+bool SMultiLineEditableText::IsAtStartOfDoc( const FTextLocation& Location ) const
+{
+	return Location.GetLineIndex() == 0 && Location.GetOffset() == 0;
+}
+	
+
+bool SMultiLineEditableText::IsAtEndOfDoc( const FTextLocation& Location ) const
+{
+	return TextLines.Num() == 0 || ( TextLines.Num()-1 == Location.GetLineIndex() && TextLines[Location.GetLineIndex()]->Len()-1 == Location.GetOffset());
+}
+
+bool SMultiLineEditableText::IsAtStartOfLine( const FTextLocation& Location ) const
+{
+	return Location.GetOffset() == 0;
+}
+
+bool SMultiLineEditableText::IsAtEndOfLine( const FTextLocation& Location ) const
+{
+	return TextLines[Location.GetLineIndex()]->Len()-1 == Location.GetOffset();
+}
+
+bool SMultiLineEditableText::IsAtWordStart( const FTextLocation& Location ) const
+{
+	const TCHAR CharBeforeCursor = GetCharacterAt(TranslatedLocation(Location, -1));
+	const TCHAR CharAfterCursor = GetCharacterAt(TranslatedLocation(Location, +1));
+
+	return FChar::IsWhitespace(CharBeforeCursor) && !FChar::IsWhitespace(CharAfterCursor);
+}
+
+SMultiLineEditableText::FTextLocation SMultiLineEditableText::AbsoluteToWrapped( const FTextLocation& Location ) const
+{
+	const int32 CurrentWrappedLineIndex = WrappedText.FindMatch( FMatchWrappedLine( TextLines[Location.GetLineIndex()].Get(), Location.GetOffset() ) );
+	const int32 OffsetInWrappedLine = Location.GetOffset() - WrappedText[CurrentWrappedLineIndex].FirstCharIndex;
+	return FTextLocation(CurrentWrappedLineIndex, OffsetInWrappedLine);
+}
+
+SMultiLineEditableText::FTextRange SMultiLineEditableText::AbsoluteToWrapped( const FTextRange& Range ) const
+{
+	return FTextRange( AbsoluteToWrapped(Range.First), AbsoluteToWrapped(Range.Last) );
+}
diff --git a/Engine/Source/Runtime/Slate/Private/Widgets/Text/SRichTextBlock.cpp b/Engine/Source/Runtime/Slate/Private/Widgets/Text/SRichTextBlock.cpp
new file mode 100644
--- /dev/null
+++ b/Engine/Source/Runtime/Slate/Private/Widgets/Text/SRichTextBlock.cpp
@@ -0,0 +1,234 @@
+// Copyright 1998-2014 Epic Games, Inc. All Rights Reserved.
+
+#include "SlatePrivatePCH.h"
+#include "TextLayoutEngine.h"
+#include "SlateTextLayout.h"
+#include "SlateSimpleRunHighlighter.h"
+#include "IRichTextMarkupParser.h"
+#include "RichTextMarkupProcessing.h"
+
+
+void SRichTextBlock::Construct( const FArguments& InArgs )
+{
+	WrapTextAt = InArgs._WrapTextAt;
+	AutoWrapText = InArgs._AutoWrapText;
+	Margin = InArgs._Margin;
+	LineHeightPercentage = InArgs._LineHeightPercentage;
+	Justification = InArgs._Justification;
+	TagStyleSet = InArgs._DecoratorStyleSet;
+	TextStyle = *InArgs._TextStyle;
+
+	CachedAutoWrapTextWidth = 0.0f;
+
+	Parser = InArgs._Parser;
+	if ( !Parser.IsValid() )
+	{
+		Parser = FRichTextMarkupProcessing::Create();
+	}
+
+	TextHighlighter = FSlateSimpleRunHighlighter::Create();
+
+	Decorators.Append( InArgs._Decorators );
+	Decorators.Append( InArgs.InlineDecorators );
+
+	TextLayout = FSlateTextLayout::Create();
+
+	SetText( InArgs._Text );
+
+	TextLayout->SetWrappingWidth( WrapTextAt.Get() );
+	TextLayout->SetMargin( Margin.Get() );
+	TextLayout->SetLineHeightPercentage( LineHeightPercentage.Get() );
+	TextLayout->SetJustification( Justification.Get() );
+	TextLayout->UpdateIfNeeded();
+}
+
+void SRichTextBlock::Tick( const FGeometry& AllottedGeometry, const double InCurrentTime, const float InDeltaTime )
+{
+	TextLayout->SetScale( AllottedGeometry.Scale );
+}
+
+int32 SRichTextBlock::OnPaint( const FGeometry& AllottedGeometry, const FSlateRect& MyClippingRect, FSlateWindowElementList& OutDrawElements, int32 LayerId, const FWidgetStyle& InWidgetStyle, bool bParentEnabled ) const
+{
+	SCOPE_CYCLE_COUNTER( STAT_SlateOnPaint_SRichTextBlock );
+
+	// Cache the last known size in case we need it for performing auto-wrap.
+	CachedAutoWrapTextWidth = AllottedGeometry.Size.X;
+
+	const ETextJustify::Type TextJustification = TextLayout->GetJustification();
+
+	if ( TextJustification == ETextJustify::Right )
+	{
+		const FVector2D TextLayoutSize = TextLayout->GetSize();
+		const FVector2D Offset( ( AllottedGeometry.Size.X - TextLayoutSize.X ), 0 );
+		LayerId = TextLayout->OnPaint( TextStyle, AllottedGeometry.MakeChild( Offset, TextLayoutSize ), MyClippingRect, OutDrawElements, LayerId, InWidgetStyle, ShouldBeEnabled( bParentEnabled ) );
+	}
+	else if ( TextJustification == ETextJustify::Center )
+	{
+		const FVector2D TextLayoutSize = TextLayout->GetSize();
+		const FVector2D Offset( ( AllottedGeometry.Size.X - TextLayoutSize.X ) / 2, 0 );
+		LayerId = TextLayout->OnPaint( TextStyle, AllottedGeometry.MakeChild( Offset, TextLayoutSize ), MyClippingRect, OutDrawElements, LayerId, InWidgetStyle, ShouldBeEnabled( bParentEnabled ) );
+	}
+	else
+	{
+		LayerId = TextLayout->OnPaint( TextStyle, AllottedGeometry, MyClippingRect, OutDrawElements, LayerId, InWidgetStyle, ShouldBeEnabled( bParentEnabled ) );
+	}
+
+	return LayerId;
+}
+
+void SRichTextBlock::CacheDesiredSize()
+{
+	float WrappingWidth = WrapTextAt.Get();
+
+	// Text wrapping can either be used defined (WrapTextAt), automatic (AutoWrapText and CachedAutoWrapTextWidth), 
+	// or a mixture of both.  Take whichever has the smallest value (>1)
+	if ( AutoWrapText.Get() && CachedAutoWrapTextWidth >= 1.0f )
+	{
+		WrappingWidth = (WrappingWidth >= 1.0f) ? FMath::Min(WrappingWidth, CachedAutoWrapTextWidth) : CachedAutoWrapTextWidth;
+	}
+
+	TextLayout->SetWrappingWidth( WrappingWidth );
+	TextLayout->SetMargin( Margin.Get() );
+	TextLayout->SetLineHeightPercentage( LineHeightPercentage.Get() );
+	TextLayout->SetJustification( Justification.Get() );
+	TextLayout->UpdateIfNeeded();
+
+	SWidget::CacheDesiredSize();
+}
+
+FVector2D SRichTextBlock::ComputeDesiredSize() const
+{
+	//The layouts current margin size. We should not report a size smaller then the margins.
+	const FMargin LayoutMargin = TextLayout->GetMargin();
+	const FVector2D TextLayoutSize = TextLayout->GetSize();
+
+	//If a wrapping width has been provided that should be reported as the desired width.
+	const float Width = TextLayout->GetWrappingWidth() > 0 ? TextLayout->GetWrappingWidth() : TextLayoutSize.X;
+
+	return FVector2D(FMath::Max(LayoutMargin.GetTotalSpaceAlong<Orient_Horizontal>(), Width), FMath::Max(LayoutMargin.GetTotalSpaceAlong<Orient_Vertical>(), TextLayoutSize.Y));
+}
+
+FChildren* SRichTextBlock::GetChildren()
+{
+	return TextLayout->GetChildren();
+}
+
+void SRichTextBlock::ArrangeChildren( const FGeometry& AllottedGeometry, FArrangedChildren& ArrangedChildren ) const
+{
+	TextLayout->ArrangeChildren( AllottedGeometry, ArrangedChildren );
+}
+
+TSharedPtr< ITextDecorator > SRichTextBlock::TryGetDecorator( const TArray< TSharedRef< ITextDecorator > >& InDecorators, const FString& InText, const FTextRunParseResults& TextRun ) const
+{
+	if (InDecorators.Num() > 0)
+	{
+		for (auto DecoratorIter = InDecorators.CreateConstIterator(); DecoratorIter; ++DecoratorIter)
+		{
+			if ((*DecoratorIter)->Supports(TextRun, InText))
+			{
+				return *DecoratorIter;
+			}
+		}
+	}
+
+	return NULL;
+}
+
+void SRichTextBlock::SetText( const FText& InText )
+{
+	Text = InText;
+
+	const FString& InString = InText.ToString();
+	TArray<FTextLineParseResults> LineParseResultsArray;
+	FString ProcessedString;
+	Parser->Process(LineParseResultsArray, InString, ProcessedString);
+
+	TextLayout->ClearLines();
+
+	// Iterate through parsed line results and create processed lines with runs.
+	for (int LineIndex = 0; LineIndex < LineParseResultsArray.Num(); LineIndex++)
+	{
+		const FTextLineParseResults& LineParseResults = LineParseResultsArray[ LineIndex ];
+
+		TSharedRef<FString> ModelString = MakeShareable(new FString());
+		TArray< TSharedRef< IRun > > Runs;
+
+		for (int RunIndex = 0; RunIndex < LineParseResults.Runs.Num(); RunIndex++)
+		{
+			const FTextRunParseResults& RunParseResult = LineParseResults.Runs[ RunIndex ];
+			TSharedPtr< ISlateRun > Run;
+
+			TSharedPtr< ITextDecorator > Decorator = TryGetDecorator( Decorators, ProcessedString, RunParseResult );
+
+			if ( Decorator.IsValid() )
+			{
+				// Create run and update model string.
+				Run = Decorator->Create( RunParseResult, ProcessedString, ModelString, TagStyleSet );
+			}
+			else
+			{
+				const FTextBlockStyle* TextBlockStyle;
+				FTextRange ModelRange;
+				ModelRange.BeginIndex = ModelString->Len();
+				if(!(RunParseResult.Name.IsEmpty()) && TagStyleSet->HasWidgetStyle< FTextBlockStyle >( FName(*RunParseResult.Name) ))
+				{
+					*ModelString += ProcessedString.Mid(RunParseResult.ContentRange.BeginIndex, RunParseResult.ContentRange.EndIndex - RunParseResult.ContentRange.BeginIndex);
+					TextBlockStyle = &(TagStyleSet->GetWidgetStyle< FTextBlockStyle >( FName(*RunParseResult.Name) ));
+				}
+				else
+				{
+					*ModelString += ProcessedString.Mid(RunParseResult.OriginalRange.BeginIndex, RunParseResult.OriginalRange.EndIndex - RunParseResult.OriginalRange.BeginIndex);
+					TextBlockStyle = &(TextStyle);
+				}
+				ModelRange.EndIndex = ModelString->Len();
+
+				// Create run.
+				Run = FSlateTextRun::Create( ModelString, *TextBlockStyle, ModelRange );
+			}
+
+			Runs.Add( Run.ToSharedRef() );
+		}
+
+		TextLayout->AddLine( ModelString, Runs );
+	}
+}
+
+void SRichTextBlock::SetHighlightText( const FText& InHighlightText )
+{
+	if ( !TextHighlighter.IsValid() )
+	{
+		return;
+	}
+
+	Highlights.Empty();
+	HighlightText = InHighlightText;
+
+	const FString HighlightTextString = HighlightText.ToString();
+	const int32 HighlightTextLength = HighlightTextString.Len();
+
+	const TArray< FTextLayout::FLineModel >& LineModels = TextLayout->GetLineModels();
+
+	for (int LineIndex = 0; LineIndex < LineModels.Num(); LineIndex++)
+	{
+		const FTextLayout::FLineModel& LineModel = LineModels[ LineIndex ];
+
+		int32 FindBegin = 0;
+		int32 CurrentHighlightBegin;
+		const int32 TextLength = LineModel.Text->Len();
+		while( FindBegin < TextLength && (CurrentHighlightBegin = LineModel.Text->Find(HighlightTextString, ESearchCase::IgnoreCase, ESearchDir::FromStart, FindBegin)) != INDEX_NONE )
+		{
+			FindBegin = CurrentHighlightBegin + HighlightTextLength;
+
+			if ( Highlights.Num() > 0 && Highlights.Last().LineIndex == LineIndex && Highlights.Last().Range.EndIndex == CurrentHighlightBegin )
+			{
+				Highlights[ Highlights.Num() - 1 ] = FTextHighlight( LineIndex, FTextRange( Highlights.Last().Range.BeginIndex, FindBegin ), TextHighlighter.ToSharedRef() );
+			}
+			else
+			{
+				Highlights.Add( FTextHighlight( LineIndex, FTextRange( CurrentHighlightBegin, FindBegin ), TextHighlighter.ToSharedRef() ) );
+			}
+		}
+	}
+
+	TextLayout->SetHighlights( Highlights );
+}
diff --git a/Engine/Source/Runtime/Slate/Private/Widgets/Text/STextBlock.cpp b/Engine/Source/Runtime/Slate/Private/Widgets/Text/STextBlock.cpp
new file mode 100644
--- /dev/null
+++ b/Engine/Source/Runtime/Slate/Private/Widgets/Text/STextBlock.cpp
@@ -0,0 +1,281 @@
+﻿// Copyright 1998-2014 Epic Games, Inc. All Rights Reserved.
+
+#include "SlatePrivatePCH.h"
+#include "SlateWordWrapper.h"
+
+
+/**
+ * Construct this widget
+ *
+ * @param	InArgs	The declaration data for this widget
+ */
+void STextBlock::Construct( const FArguments& InArgs )
+{
+	Text = InArgs._Text;
+
+	const FTextBlockStyle* MyStyle = InArgs._TextStyle;
+
+	HighlightText = InArgs._HighlightText;
+	WrapTextAt = InArgs._WrapTextAt;
+	AutoWrapText = InArgs._AutoWrapText;
+
+	CachedWrapTextWidth = 0.0f;
+	CachedAutoWrapTextWidth = 0.0f;
+
+	Font = (InArgs._Font.IsSet())
+		? InArgs._Font
+		: MyStyle->Font;
+
+	ForegroundColor = InArgs._ColorAndOpacity.IsSet()
+		? InArgs._ColorAndOpacity
+		: MyStyle->ColorAndOpacity;
+
+	ShadowOffset = InArgs._ShadowOffset.IsSet()
+		? InArgs._ShadowOffset
+		: MyStyle->ShadowOffset;
+
+	ShadowColorAndOpacity = InArgs._ShadowColorAndOpacity.IsSet()
+		? InArgs._ShadowColorAndOpacity
+		: MyStyle->ShadowColorAndOpacity;
+
+	HighlightColor = InArgs._HighlightColor.IsSet()
+		? InArgs._HighlightColor
+		: MyStyle->HighlightColor;
+
+	HighlightShape = InArgs._HighlightShape.IsSet()
+		? InArgs._HighlightShape
+		: &MyStyle->HighlightShape;
+
+	OnDoubleClicked = InArgs._OnDoubleClicked;
+
+	// Request text size be cached
+	bRequestCache = true;
+}
+
+void STextBlock::SetText( const TAttribute< FText >& InText )
+{
+	struct Local
+	{
+		static FString PassThroughAttribute( TAttribute< FText > TextAttribute )
+		{
+			return TextAttribute.Get( FText::GetEmpty() ).ToString();
+		}
+	};
+
+	Text = TAttribute< FString >::Create(TAttribute<FString>::FGetter::CreateStatic( &Local::PassThroughAttribute, InText));
+	bRequestCache = true;
+}
+
+void STextBlock::SetForegroundColor( const TAttribute<FSlateColor>& InSlateColor )
+{
+	ForegroundColor = InSlateColor;
+}
+
+int32 STextBlock::OnPaint( const FGeometry& AllottedGeometry, const FSlateRect& MyClippingRect, FSlateWindowElementList& OutDrawElements, int32 LayerId, const FWidgetStyle& InWidgetStyle, bool bParentEnabled ) const
+{
+#if SLATE_HD_STATS
+	SCOPE_CYCLE_COUNTER( STAT_SlateOnPaint_STextBlock );
+#endif
+
+	const TSharedRef< FSlateFontMeasure > FontMeasureService = FSlateApplication::Get().GetRenderer()->GetFontMeasureService();
+
+	const FSlateRect ClippingRect = AllottedGeometry.GetClippingRect().IntersectionWith(MyClippingRect);
+
+	bool bEnabled = ShouldBeEnabled( bParentEnabled );
+	ESlateDrawEffect::Type DrawEffects = bEnabled ? ESlateDrawEffect::None : ESlateDrawEffect::DisabledEffect;
+
+	const FVector2D CurShadowOffset = ShadowOffset.Get();
+	const bool ShouldDropShadow = CurShadowOffset.Size() > 0;
+
+	const FSlateFontInfo& FontInfo = Font.Get();
+
+	// Perform text auto-wrapping if that was enabled.
+	if( AutoWrapText.Get() )
+	{
+		const float OldWrapTextAt = CachedAutoWrapTextWidth;
+		const float NewWrapTextAt = AllottedGeometry.Size.X;
+		if( OldWrapTextAt != NewWrapTextAt )
+		{
+			// Available space has changed, so make sure that we recompute wrapping
+			bRequestCache = true;
+			CachedAutoWrapTextWidth = NewWrapTextAt;
+		}
+	}
+
+	// Draw the text highlight
+	{
+		const FString& StringToHighlight = HighlightText.Get().ToString();
+		const int32 CurrentHighlightLength = StringToHighlight.Len();
+
+		// Do we have text to highlight?
+		if( CurrentHighlightLength > 0 )
+		{
+			// Check to see if the original string contains text which needs to be highlighted
+			// We need to use the original as the wrapped string has had newlines added to it
+			const int32 CurrentHighlightStart = CachedOriginalString.Find(StringToHighlight, ESearchCase::IgnoreCase);
+			if ( INDEX_NONE != CurrentHighlightStart )
+			{
+				const int32 CurrentHighlightEnd = CurrentHighlightStart + CurrentHighlightLength;
+
+				TSharedPtr<FSlateRenderer> SlateRenderer = FSlateApplication::Get().GetRenderer();
+				const int32 LineHeight = FontMeasureService->GetMaxCharacterHeight(FontInfo);
+
+				// We might have to highlight multiple lines if the range we've found spans multiple entries in CachedWrappedLineData
+				for(auto It = CachedWrappedLineData.CreateConstIterator(); It; ++It)
+				{
+					const int32 LineStart	= It->Key;
+					const int32 LineEnd		= It->Value;
+
+					// We need to highlight this line if it started after a break from a previous line, or has highlighted characters before its end
+					if(CurrentHighlightStart <= LineStart || CurrentHighlightStart <= LineEnd)
+					{
+						// Clamp the highlight indices to this line so we can measure the highlighted text for just this line
+						const int32 LineHighlightStart  = FMath::Max(CurrentHighlightStart, LineStart);
+						const int32 LineHighlightEnd	= FMath::Min(CurrentHighlightEnd, LineEnd);
+
+						// Figure out screen space locations to start drawing the highlight rectangle
+						const FString TextUpToHighlightStart = CachedOriginalString.Mid(LineStart, LineHighlightStart - LineStart);
+						const FVector2D HighlightStartOffset(FontMeasureService->Measure(TextUpToHighlightStart, FontInfo).X, LineHeight * It.GetIndex());
+
+						// Figure out the actual text being highlighted
+						// Cannot just use the filter string; it might be different case and therefore different width characters
+						const FString TextToHighlight = CachedOriginalString.Mid(LineHighlightStart, LineHighlightEnd - LineHighlightStart);
+						const FVector2D HighlightSize = FontMeasureService->Measure(TextToHighlight, FontInfo);
+
+						// Draw the actual highlight rectangle
+						FSlateDrawElement::MakeBox(
+							OutDrawElements,
+							++LayerId,
+							AllottedGeometry.ToPaintGeometry(HighlightStartOffset, HighlightSize),
+							HighlightShape.Get(),
+							ClippingRect,
+							DrawEffects,
+							InWidgetStyle.GetColorAndOpacityTint() * HighlightColor.Get()
+							);
+					}
+				}
+			}
+		}
+	}
+
+	// Draw the optional shadow
+	if ( ShouldDropShadow )
+	{
+		FSlateDrawElement::MakeText(
+			OutDrawElements,
+			LayerId,
+			AllottedGeometry.ToPaintGeometry( CurShadowOffset, AllottedGeometry.Size ),
+			CachedWrappedString,
+			FontInfo,
+			ClippingRect,
+			DrawEffects,
+			ShadowColorAndOpacity.Get()*InWidgetStyle.GetColorAndOpacityTint()
+		);
+	}
+
+	// Draw the text itself
+	FSlateDrawElement::MakeText(
+		OutDrawElements,
+		++LayerId,
+		AllottedGeometry.ToPaintGeometry(),
+		CachedWrappedString,
+		FontInfo,
+		ClippingRect,
+		DrawEffects,
+		InWidgetStyle.GetColorAndOpacityTint() * ForegroundColor.Get().GetColor(InWidgetStyle)
+	);
+
+	return LayerId;
+}
+
+/**
+ * See SWidget::OnMouseButtonDoubleClick.
+ *
+ * @param MyGeometry The Geometry of the widget receiving the event
+ * @param MouseEvent Information about the input event
+ *
+ * @return Whether the event was handled along with possible requests for the system to take action.
+ */
+FReply STextBlock::OnMouseButtonDoubleClick( const FGeometry& InMyGeometry, const FPointerEvent& InMouseEvent )
+{
+	if ( InMouseEvent.GetEffectingButton() == EKeys::LeftMouseButton )
+	{
+		if( OnDoubleClicked.IsBound() )
+		{
+			OnDoubleClicked.Execute();
+
+			return FReply::Handled();
+		}
+	}
+
+	return FReply::Unhandled();
+}
+
+void STextBlock::CacheDesiredSize()
+{
+#if SLATE_HD_STATS
+	SCOPE_CYCLE_COUNTER( STAT_SlateCacheDesiredSize_STextBlock );
+#endif
+
+	// Get the wrapping width and font to see if they have changed
+	float WrappingWidth = WrapTextAt.Get();
+	const FSlateFontInfo& FontInfo = Font.Get();
+
+	// Text wrapping can either be used defined (WrapTextAt), automatic (AutoWrapText and CachedAutoWrapTextWidth), or a mixture of both
+	// Take whichever has the smallest value (>1)
+	if(AutoWrapText.Get() && CachedAutoWrapTextWidth >= 1.0f)
+	{
+		WrappingWidth = (WrappingWidth >= 1.0f) ? FMath::Min(WrappingWidth, CachedAutoWrapTextWidth) : CachedAutoWrapTextWidth;
+	}
+
+	// We should recache text size, width, font, and wrapped string if...
+	// A cache was requested, the text is dynamic, the wrapping width has changed, or the font has changed
+	const bool bShouldCacheText = bRequestCache || Text.IsBound() || (WrappingWidth != CachedWrapTextWidth) || (!(FontInfo == CachedFont));
+
+	if(bShouldCacheText)
+	{
+		CachedOriginalString = Text.Get();
+
+		const bool bWrapText = (WrappingWidth >= 1.0f);
+
+		// Handle optional text wrapping, caching the appropriate result
+		if(bWrapText)
+		{
+			// WrapText takes care of clearing CachedWrappedLineData
+			CachedWrappedString = SlateWordWrapper::WrapText(CachedOriginalString, FontInfo, FMath::Trunc(WrappingWidth), 1.0f, &CachedWrappedLineData);
+		}
+		else
+		{
+			CachedWrappedString = CachedOriginalString;
+
+			CachedWrappedLineData.Empty();
+			CachedWrappedLineData.Add(TPairInitializer<int32, int32>(0, CachedOriginalString.Len()));
+		}
+
+		const TSharedRef< FSlateFontMeasure > FontMeasureService = FSlateApplication::Get().GetRenderer()->GetFontMeasureService();
+		FVector2D TextMeasurement = FontMeasureService->Measure(CachedWrappedString, FontInfo);
+
+		// Force the result to exactly the wrapping width
+		if(bWrapText)
+		{
+			TextMeasurement.X = FMath::Min(TextMeasurement.X, WrappingWidth);
+		}
+
+		this->Advanced_SetDesiredSize(TextMeasurement + ShadowOffset.Get());
+
+		// Update cached values
+		CachedWrapTextWidth = WrappingWidth;
+		CachedFont = FontInfo;
+		bRequestCache = false;
+	}
+}
+
+
+FVector2D STextBlock::ComputeDesiredSize() const
+{
+	// Usually widgets just override ComputeDesiredSize(), but STextBlock
+	// overrides CacheDesiredSize() and does all the work in there.
+
+	// Dummy implementation.
+	return FVector2D::ZeroVector;
+}
diff --git a/Engine/Source/Runtime/Slate/Private/Widgets/Tutorials/STutorialWrapper.cpp b/Engine/Source/Runtime/Slate/Private/Widgets/Tutorials/STutorialWrapper.cpp
new file mode 100644
--- /dev/null
+++ b/Engine/Source/Runtime/Slate/Private/Widgets/Tutorials/STutorialWrapper.cpp
@@ -0,0 +1,119 @@
+// Copyright 1998-2014 Epic Games, Inc. All Rights Reserved.
+
+#include "SlatePrivatePCH.h"
+
+
+FOnWidgetHighlight STutorialWrapper::OnWidgetHighlightDelegate;
+
+namespace TutorialConstants
+{
+	const float BorderPulseAnimationLength = 1.0f;
+	const float BorderIntroAnimationLength = 0.5f;
+	const float MinBorderOpacity = 0.1f;
+	const float ShadowScale = 8.0f;
+	const float MaxBorderOffset = 8.0f;
+}
+
+void STutorialWrapper::Construct( const FArguments& InArgs, const FName& InName )
+{
+	Name = InName;
+
+	BorderIntroAnimation.AddCurve(0.0f, TutorialConstants::BorderIntroAnimationLength, ECurveEaseFunction::CubicOut);
+	BorderPulseAnimation.AddCurve(0.0f, TutorialConstants::BorderPulseAnimationLength, ECurveEaseFunction::Linear);
+
+	bIsPlaying = false;
+
+	SBorder::Construct(SBorder::FArguments()
+		.BorderImage(FStyleDefaults::GetNoBrush())
+		.Padding(0.0f)
+		[
+			InArgs._Content.Widget
+		]);
+}
+
+void STutorialWrapper::Tick( const FGeometry& AllottedGeometry, const double InCurrentTime, const float InDeltaTime )
+{
+	if(OnWidgetHighlightDelegate.IsBound() && OnWidgetHighlightDelegate.Execute(Name))
+	{
+		if(!bIsPlaying)
+		{
+			BorderIntroAnimation.Play();
+			BorderPulseAnimation.Play();
+		}
+
+		bIsPlaying = true;
+	}
+	else
+	{
+		bIsPlaying = false;
+	}
+
+	CachedGeometry = AllottedGeometry;
+}
+
+int32 STutorialWrapper::OnPaint(const FGeometry& AllottedGeometry, const FSlateRect& MyClippingRect, FSlateWindowElementList& OutDrawElements, int32 LayerId, const FWidgetStyle& InWidgetStyle, bool bParentEnabled ) const
+{
+	bool IsVisible = !( CachedGeometry.Size.X == 0 || CachedGeometry.Size.Y == 0 );
+
+	if ( OnWidgetHighlightDelegate.IsBound() && OnWidgetHighlightDelegate.Execute(Name) && IsVisible )
+	{
+		float AlphaFactor;
+		float PulseFactor;
+		FLinearColor ShadowTint;
+		FLinearColor BorderTint;
+		GetAnimationValues(AlphaFactor, PulseFactor, ShadowTint, BorderTint);
+
+		const FSlateBrush* ShadowBrush = FCoreStyle::Get().GetBrush(TEXT("Tutorials.Shadow"));
+		const FSlateBrush* BorderBrush = FCoreStyle::Get().GetBrush(TEXT("Tutorials.Border"));
+					
+		const FGeometry& WidgetGeometry = CachedGeometry;
+		const FVector2D WindowPos = OutDrawElements.GetWindow()->GetPositionInScreen();
+		const FVector2D WindowSize = OutDrawElements.GetWindow()->GetSizeInScreen();
+
+		// We should be clipped by the window size, not our containing widget, as we want to draw outside the widget
+		FSlateRect WindowClippingRect(0.0f, 0.0f, WindowSize.X, WindowSize.Y);
+		// We want to draw on-top of everything else
+		int32 OverlayLayerId = MAX_int32 - 2;
+
+		FPaintGeometry ShadowGeometry((WidgetGeometry.AbsolutePosition - FVector2D(ShadowBrush->Margin.Left, ShadowBrush->Margin.Top) * ShadowBrush->ImageSize * WidgetGeometry.Scale * TutorialConstants::ShadowScale) - WindowPos,
+										((WidgetGeometry.Size * WidgetGeometry.Scale) + (FVector2D(ShadowBrush->Margin.Right * 2.0f, ShadowBrush->Margin.Bottom * 2.0f) * ShadowBrush->ImageSize * WidgetGeometry.Scale * TutorialConstants::ShadowScale)),
+										WidgetGeometry.Scale * TutorialConstants::ShadowScale);
+		// draw highlight shadow
+		FSlateDrawElement::MakeBox(OutDrawElements, OverlayLayerId + 0, ShadowGeometry, ShadowBrush, WindowClippingRect, ESlateDrawEffect::None, ShadowTint);
+
+		FVector2D PulseOffset = FVector2D(PulseFactor * TutorialConstants::MaxBorderOffset, PulseFactor * TutorialConstants::MaxBorderOffset);
+
+		FVector2D BorderPosition = (WidgetGeometry.AbsolutePosition - ((FVector2D(BorderBrush->Margin.Left, BorderBrush->Margin.Top) * BorderBrush->ImageSize * WidgetGeometry.Scale) + PulseOffset)) - WindowPos;
+		FVector2D BorderSize = ((WidgetGeometry.Size * WidgetGeometry.Scale) + (PulseOffset * 2.0f) + (FVector2D(BorderBrush->Margin.Right * 2.0f, BorderBrush->Margin.Bottom * 2.0f) * BorderBrush->ImageSize * WidgetGeometry.Scale));
+
+		FPaintGeometry BorderGeometry(BorderPosition, BorderSize, WidgetGeometry.Scale);
+
+		// draw highlight border
+		FSlateDrawElement::MakeBox(OutDrawElements, OverlayLayerId + 1, BorderGeometry, BorderBrush, WindowClippingRect, ESlateDrawEffect::None, BorderTint);
+	}
+
+	return SBorder::OnPaint(AllottedGeometry, MyClippingRect, OutDrawElements, LayerId, InWidgetStyle, bParentEnabled);
+}
+
+void STutorialWrapper::GetAnimationValues(float& OutAlphaFactor, float& OutPulseFactor, FLinearColor& OutShadowTint, FLinearColor& OutBorderTint) const
+{
+	if(BorderIntroAnimation.IsPlaying())
+	{
+		OutAlphaFactor = BorderIntroAnimation.GetLerp();
+		OutPulseFactor = (1.0f - OutAlphaFactor) * 50.0f;
+		OutShadowTint = FLinearColor(1.0f, 1.0f, 0.0f, OutAlphaFactor);
+		OutBorderTint = FLinearColor(1.0f, 1.0f, 0.0f, OutAlphaFactor * OutAlphaFactor);
+	}
+	else
+	{
+		OutAlphaFactor = 1.0f - (0.5f + (FMath::Cos(2.0f * PI * BorderPulseAnimation.GetLerpLooping()) * 0.5f));
+		OutPulseFactor = 0.5f + (FMath::Cos(2.0f * PI * BorderPulseAnimation.GetLerpLooping()) * 0.5f);
+		OutShadowTint = FLinearColor(1.0f, 1.0f, 0.0f, 1.0f);
+		OutBorderTint = FLinearColor(1.0f, 1.0f, 0.0f, TutorialConstants::MinBorderOpacity + ((1.0f - TutorialConstants::MinBorderOpacity) * OutAlphaFactor));
+	}
+}
+
+FOnWidgetHighlight& STutorialWrapper::OnWidgetHighlight()
+{
+	return OnWidgetHighlightDelegate;
+}
diff --git a/Engine/Source/Runtime/Slate/Private/Widgets/Views/SHeaderRow.cpp b/Engine/Source/Runtime/Slate/Private/Widgets/Views/SHeaderRow.cpp
new file mode 100644
--- /dev/null
+++ b/Engine/Source/Runtime/Slate/Private/Widgets/Views/SHeaderRow.cpp
@@ -0,0 +1,545 @@
+// Copyright 1998-2014 Epic Games, Inc. All Rights Reserved.
+ 
+#include "SlatePrivatePCH.h"
+
+
+class STableColumnHeader : public SCompoundWidget
+{
+public:
+	 
+	SLATE_BEGIN_ARGS(STableColumnHeader)
+		: _Style( &FCoreStyle::Get().GetWidgetStyle<FTableColumnHeaderStyle>("TableView.Header.Column") )
+		{}
+		SLATE_STYLE_ARGUMENT( FTableColumnHeaderStyle, Style )
+
+	SLATE_END_ARGS()
+
+	STableColumnHeader()
+		: SortMode( EColumnSortMode::None )
+		, OnSortModeChanged()
+		, ContextMenuContent( SNullWidget::NullWidget )
+		, ColumnId( NAME_None )
+		, Style( nullptr )
+	{
+
+	}
+
+	/**
+	 * Construct the widget
+	 * 
+	 * @param InDeclaration   A declaration from which to construct the widget
+	 */
+	void Construct( const STableColumnHeader::FArguments& InArgs, const SHeaderRow::FColumn& Column, const FMargin DefaultHeaderContentPadding )
+	{
+		check(InArgs._Style);
+
+		SWidget::Construct( InArgs._ToolTipText, InArgs._ToolTip, InArgs._Cursor, InArgs._IsEnabled, InArgs._Visibility, InArgs._Tag );
+
+		Style = InArgs._Style;
+		ColumnId = Column.ColumnId;
+		SortMode = Column.SortMode;
+
+		SortMode = Column.SortMode;
+		OnSortModeChanged = Column.OnSortModeChanged;
+		ContextMenuContent = Column.HeaderMenuContent.Widget;
+
+		ComboVisibility = Column.HeaderComboVisibility;
+
+		FMargin AdjustedDefaultHeaderContentPadding = DefaultHeaderContentPadding;
+
+		TAttribute< FString > LabelText = Column.DefaultText;
+
+		if (Column.HeaderContent.Widget == SNullWidget::NullWidget && !Column.DefaultText.IsBound() && Column.DefaultText.Get().IsEmpty())
+		{
+			LabelText = Column.ColumnId.ToString() + TEXT("[LabelMissing]");
+		}
+
+		TSharedPtr< SHorizontalBox > Box;
+		TSharedRef< SOverlay > Overlay = SNew( SOverlay );
+
+		Overlay->AddSlot( 0 )
+			[
+				SAssignNew( Box, SHorizontalBox )
+			];
+
+		TSharedRef< SWidget > PrimaryContent = Column.HeaderContent.Widget;
+		if ( PrimaryContent == SNullWidget::NullWidget )
+		{
+			PrimaryContent = 
+				SNew( SBox )
+				.Padding( OnSortModeChanged.IsBound() ? FMargin( 0, 2, 0, 2 ) : FMargin( 0, 4, 0, 4 ) )
+				.VAlign( VAlign_Center )
+				[
+					SNew(STextBlock)
+					.Text( LabelText )
+					.ToolTipText( LabelText )
+				];
+		}
+
+		if ( OnSortModeChanged.IsBound() )
+		{
+			//optional main button with the column's title. Used to toggle sorting modes.
+			PrimaryContent = SNew(SButton)
+			.ButtonStyle( FCoreStyle::Get(), "NoBorder" )
+			.ForegroundColor( FSlateColor::UseForeground() )
+			.ContentPadding( FMargin( 0, 2, 0, 2 ) )
+			.OnClicked(this, &STableColumnHeader::OnTitleClicked)
+			[
+				PrimaryContent
+			];
+		}
+		
+		Box->AddSlot()
+		.FillWidth(1.0f)
+		[
+			PrimaryContent
+		];
+
+		if( Column.HeaderMenuContent.Widget != SNullWidget::NullWidget )
+		{
+			// Add Drop down menu button (only if menu content has been specified)
+			Box->AddSlot()
+			.AutoWidth()
+			[
+				SAssignNew( MenuOverlay, SOverlay )
+				.Visibility( this, &STableColumnHeader::GetMenuOverlayVisibility )
+				+SOverlay::Slot()
+				[
+					SNew( SSpacer )
+					.Size( FVector2D( 12.0f, 0 ) )
+				]
+
+				+SOverlay::Slot()
+				.Padding(FMargin(0,1,0,1))
+				[
+					SNew( SBorder )
+					.Padding( FMargin( 0, 0, AdjustedDefaultHeaderContentPadding.Right, 0 ) )
+					.BorderImage( this, &STableColumnHeader::GetComboButtonBorderBrush )
+					[
+						SAssignNew( ComboButton, SComboButton)
+						.HasDownArrow(false)
+						.ButtonStyle( FCoreStyle::Get(), "NoBorder" )
+						.ContentPadding( FMargin(0) )
+						.ButtonContent()
+						[
+							SNew( SSpacer )
+							.Size( FVector2D( 14.0f, 0 ) )
+						]
+						.MenuContent()
+						[
+							ContextMenuContent
+						]
+					]
+				]
+
+				+SOverlay::Slot()
+				.Padding(FMargin(0,0,0,2))
+				.HAlign( HAlign_Center )
+				.VAlign( VAlign_Bottom )
+				[
+					SNew(SImage)
+					.Image( &Style->MenuDropdownImage )
+					.ColorAndOpacity( this, &STableColumnHeader::GetComboButtonTint )
+					.Visibility( EVisibility::HitTestInvisible )
+				]
+			];		
+
+			AdjustedDefaultHeaderContentPadding.Right = 0;
+		}
+
+		Overlay->AddSlot( 1 )
+			.HAlign(HAlign_Center)
+			.VAlign( VAlign_Top )
+			.Padding( FMargin( 0, 2, 0, 0 ) )
+			[
+				SNew(SImage)
+				.Image( this, &STableColumnHeader::GetSortingBrush )
+				.Visibility( this, &STableColumnHeader::GetSortModeVisibility )
+			];
+
+		this->ChildSlot
+		[
+			SNew( SBorder )
+			.BorderImage( this, &STableColumnHeader::GetHeaderBackgroundBrush )
+			.HAlign( Column.HeaderHAlignment )
+			.VAlign( Column.HeaderVAlignment )
+			.Padding( Column.HeaderContentPadding.Get( AdjustedDefaultHeaderContentPadding ) )
+			[
+				Overlay
+			]
+		];
+	}
+
+	/** Gets sorting mode */
+	EColumnSortMode::Type GetSortMode() const
+	{
+		return SortMode.Get();
+	}
+
+	/** Sets sorting mode */
+	void SetSortMode( EColumnSortMode::Type NewMode )
+	{
+		SortMode = NewMode;
+	}
+
+	virtual FReply OnMouseButtonUp( const FGeometry& MyGeometry, const FPointerEvent& MouseEvent ) OVERRIDE
+	{
+		if ( MouseEvent.GetEffectingButton() == EKeys::RightMouseButton && ContextMenuContent != SNullWidget::NullWidget )
+		{
+			OpenContextMenu( MouseEvent.GetScreenSpacePosition() );
+			return FReply::Handled();
+		}
+
+		return FReply::Unhandled();
+	}
+
+
+private:
+
+
+	const FSlateBrush* GetHeaderBackgroundBrush() const
+	{
+		if ( IsHovered() && SortMode.IsBound() )
+		{
+			return &Style->HoveredBrush;
+		}
+
+		return &Style->NormalBrush;
+	}
+
+	EVisibility GetMenuOverlayVisibility() const
+	{
+		if (ComboVisibility == EHeaderComboVisibility::OnHover)
+		{
+			if (!ComboButton.IsValid() || !(IsHovered() || ComboButton->IsOpen()))
+			{
+				return EVisibility::Collapsed;
+			}
+		}
+
+		return EVisibility::Visible;
+	}
+
+	const FSlateBrush* GetComboButtonBorderBrush() const
+	{
+		if ( ComboButton.IsValid() && ( ComboButton->IsHovered() || ComboButton->IsOpen() ) )
+		{
+			return &Style->MenuDropdownHoveredBorderBrush;
+		}
+
+		if ( IsHovered() || ComboVisibility == EHeaderComboVisibility::Always )
+		{
+			return &Style->MenuDropdownNormalBorderBrush;
+		}
+
+		return FStyleDefaults::GetNoBrush();
+	}
+
+	FSlateColor GetComboButtonTint() const
+	{
+		if (!ComboButton.IsValid())
+		{
+			return FLinearColor::White;
+		}
+
+		switch (ComboVisibility)
+		{
+		case EHeaderComboVisibility::Always:
+			{
+				return FLinearColor::White;
+			}
+
+		case EHeaderComboVisibility::Ghosted:
+			if ( ComboButton->IsHovered() || ComboButton->IsOpen() )
+			{
+				return FLinearColor::White;
+			}
+			else
+			{
+				return FLinearColor::White.CopyWithNewOpacity(0.5f);
+			}
+
+		case EHeaderComboVisibility::OnHover:
+			if ( IsHovered() || ComboButton->IsHovered() || ComboButton->IsOpen() )
+			{
+				return FLinearColor::White;
+			}
+			break;
+
+		default:
+			break;
+		}
+		
+		return FLinearColor::Transparent;
+	}
+
+	/** Gets the icon associated with the current sorting mode */
+	const FSlateBrush* GetSortingBrush() const
+	{
+		return (SortMode == EColumnSortMode::Ascending ? &Style->SortAscendingImage : &Style->SortDescendingImage);
+	}
+
+	/** Checks if sorting mode has been selected */
+	EVisibility GetSortModeVisibility() const
+	{
+		return (SortMode.Get() != EColumnSortMode::None) ? EVisibility::HitTestInvisible : EVisibility::Hidden;
+	}
+	
+	/** Called when the column title has been clicked to change sorting mode */
+	FReply OnTitleClicked()
+	{
+		if ( OnSortModeChanged.IsBound() )
+		{
+			EColumnSortMode::Type ColumnSortMode = SortMode.Get();
+			if ( ColumnSortMode == EColumnSortMode::None || ColumnSortMode == EColumnSortMode::Descending )
+			{
+				ColumnSortMode = EColumnSortMode::Ascending;
+			}
+			else
+			{
+				ColumnSortMode = EColumnSortMode::Descending;
+			}
+
+			OnSortModeChanged.Execute( ColumnId, ColumnSortMode );
+		}
+
+		return FReply::Handled();
+	}
+
+	void OpenContextMenu(const FVector2D& SummonLocation)
+	{
+		if ( ContextMenuContent != SNullWidget::NullWidget )
+		{
+			FSlateApplication::Get().PushMenu( AsShared(), ContextMenuContent, SummonLocation, FPopupTransitionEffect( FPopupTransitionEffect::ContextMenu ) );
+		}
+	}
+
+
+private:
+
+	/** Current sorting mode */
+	TAttribute< EColumnSortMode::Type > SortMode;
+
+	/** Callback triggered when sorting mode changes */
+	FOnSortModeChanged OnSortModeChanged;
+
+	TSharedRef< SWidget > ContextMenuContent;
+
+	TSharedPtr< SComboButton > ComboButton;
+	
+	/** The visibility method of the combo button */
+	EHeaderComboVisibility ComboVisibility;
+
+	TSharedPtr< SOverlay > MenuOverlay;
+
+	FName ColumnId;
+
+	const FTableColumnHeaderStyle* Style;
+};
+
+void SHeaderRow::Construct( const FArguments& InArgs )
+{
+	check(InArgs._Style);
+
+	SWidget::Construct( InArgs._ToolTipText, InArgs._ToolTip, InArgs._Cursor, InArgs._IsEnabled, InArgs._Visibility, InArgs._Tag );
+
+	ScrollBarThickness = FVector2D::ZeroVector;
+	ScrollBarVisibility = EVisibility::Collapsed;
+	Style = InArgs._Style;
+
+	if ( InArgs._OnColumnsChanged.IsBound() )
+	{
+		ColumnsChanged.Add( InArgs._OnColumnsChanged );
+	}
+
+	SBorder::Construct( SBorder::FArguments()
+		.Padding( 0 )
+		.BorderImage( &Style->BackgroundBrush )
+		.ForegroundColor( Style->ForegroundColor )
+	);
+
+	// Copy all the column info from the declaration
+	bool bHaveFillerColumn = false;
+	for ( int32 SlotIndex=0; SlotIndex < InArgs.Slots.Num(); ++SlotIndex )
+	{
+		FColumn* const Column = InArgs.Slots[SlotIndex];
+		Columns.Add( Column );
+	}
+
+	// Generate widgets for all columns
+	RegenerateWidgets();
+}
+
+void SHeaderRow::ResetColumnWidths()
+{
+	for ( int32 iColumn = 0; iColumn < Columns.Num(); iColumn++ )
+	{
+		FColumn& Column = Columns[iColumn];
+		Column.SetWidth( Column.DefaultWidth );
+	}
+}
+
+const TIndirectArray<SHeaderRow::FColumn>& SHeaderRow::GetColumns() const
+{
+	return Columns;
+}
+
+void SHeaderRow::AddColumn( const FColumn::FArguments& NewColumnArgs )
+{
+	SHeaderRow::FColumn* NewColumn = new SHeaderRow::FColumn( NewColumnArgs );
+	AddColumn( *NewColumn );
+}
+
+void SHeaderRow::AddColumn( SHeaderRow::FColumn& NewColumn )
+{
+	int32 InsertIdx = Columns.Num();
+	InsertColumn( NewColumn, InsertIdx );
+}
+
+void SHeaderRow::InsertColumn( const FColumn::FArguments& NewColumnArgs, int32 InsertIdx )
+{
+	SHeaderRow::FColumn* NewColumn = new SHeaderRow::FColumn( NewColumnArgs );
+	InsertColumn( *NewColumn, InsertIdx );
+}
+
+void SHeaderRow::InsertColumn( FColumn& NewColumn, int32 InsertIdx )
+{
+	check(NewColumn.ColumnId != NAME_None);
+
+	if ( Columns.Num() > 0 && Columns[Columns.Num() - 1].ColumnId == NAME_None )
+	{
+		// Insert before the filler column, or where the filler column used to be if we replaced it.
+		InsertIdx--;
+	}
+
+	Columns.Insert( &NewColumn, InsertIdx );
+	ColumnsChanged.Broadcast( SharedThis( this ) );
+
+	RegenerateWidgets();
+}
+
+void SHeaderRow::RemoveColumn( const FName& InColumnId )
+{
+	check(InColumnId != NAME_None);
+
+	bool bHaveFillerColumn = false;
+	for ( int32 SlotIndex=Columns.Num() - 1; SlotIndex >= 0; --SlotIndex )
+	{
+		FColumn& Column = Columns[SlotIndex];
+		if ( Column.ColumnId == InColumnId )
+		{
+			Columns.RemoveAt(SlotIndex);
+		}
+	}
+
+	ColumnsChanged.Broadcast( SharedThis( this ) );
+	RegenerateWidgets();
+}
+
+void SHeaderRow::ClearColumns()
+{
+	Columns.Empty();
+	ColumnsChanged.Broadcast( SharedThis( this ) );
+
+	RegenerateWidgets();
+}
+
+void SHeaderRow::SetAssociatedVerticalScrollBar( const TSharedRef< SScrollBar >& ScrollBar, const float ScrollBarSize )
+{
+	ScrollBarThickness.X = ScrollBarSize;
+	ScrollBarVisibility.Bind( TAttribute< EVisibility >::FGetter::CreateSP( ScrollBar, &SScrollBar::ShouldBeVisible ) );
+	RegenerateWidgets();
+}
+
+void SHeaderRow::SetColumnWidth( const FName& InColumnId, float InWidth )
+{
+	check(InColumnId != NAME_None);
+
+	for ( int32 SlotIndex=Columns.Num() - 1; SlotIndex >= 0; --SlotIndex )
+	{
+		FColumn& Column = Columns[SlotIndex];
+		if ( Column.ColumnId == InColumnId )
+		{
+			Column.SetWidth( InWidth );
+		}
+	}
+}
+
+void SHeaderRow::RegenerateWidgets()
+{
+	const float SplitterHandleDetectionSize = 5.0f;
+	TSharedPtr<SSplitter> Splitter;
+
+	TSharedRef< SHorizontalBox > Box = 
+		SNew(SHorizontalBox)
+		+SHorizontalBox::Slot()
+		.FillWidth( 1.0f )
+		[
+			SAssignNew(Splitter, SSplitter)
+			.Style( &Style->ColumnSplitterStyle )
+			.ResizeMode( ESplitterResizeMode::Fill )
+			.PhysicalSplitterHandleSize( 0.0f )
+			.HitDetectionSplitterHandleSize( SplitterHandleDetectionSize )
+		]
+
+		+SHorizontalBox::Slot()
+		.AutoWidth()
+		.Padding( 0 )
+		[
+			SNew( SSpacer )
+			.Size( ScrollBarThickness )
+			.Visibility( ScrollBarVisibility )
+		];
+
+	// Construct widgets for all columns
+	{
+		const float HalfSplitterDetectionSize = ( SplitterHandleDetectionSize + 2 ) / 2;
+
+		// Populate the slot with widgets that represent the columns.
+		TSharedPtr<STableColumnHeader> NewlyMadeHeader;
+		for ( int32 SlotIndex=0; SlotIndex < Columns.Num(); ++SlotIndex )
+		{
+			FColumn& SomeColumn = Columns[SlotIndex];
+
+			// Keep track of the last header we created.
+			TSharedPtr<STableColumnHeader> PrecedingHeader = NewlyMadeHeader;
+			NewlyMadeHeader.Reset();
+
+			FMargin DefaultPadding = FMargin( HalfSplitterDetectionSize, 0, HalfSplitterDetectionSize, 0 );
+			TSharedRef<STableColumnHeader> NewHeader =
+				SAssignNew(NewlyMadeHeader, STableColumnHeader, SomeColumn, DefaultPadding)
+				.Style( (SlotIndex + 1 == Columns.Num()) ? &Style->LastColumnStyle : &Style->ColumnStyle );
+
+			if ( SomeColumn.SizeRule == EColumnSizeMode::Fixed )
+			{
+				// Add resizable cell
+				Splitter->AddSlot()
+				.SizeRule( SSplitter::SizeToContent )
+				[
+					SNew( SBox )
+					.WidthOverride( SomeColumn.Width.Get() )
+					[
+						NewHeader
+					]
+				];
+			}
+			else 
+			{
+				TAttribute<float> WidthBinding;
+				WidthBinding.BindRaw( &SomeColumn, &FColumn::GetWidth );
+
+				// Add resizable cell
+				Splitter->AddSlot()
+				.Value( WidthBinding )
+				.SizeRule( SSplitter::FractionOfParent )
+				.OnSlotResized( SSplitter::FOnSlotResized::CreateRaw( &SomeColumn, &FColumn::SetWidth ) )
+				[
+					NewHeader
+				];
+			}
+		}
+	}
+
+	// Create a box to contain widgets for each column
+	SetContent( Box );
+}
diff --git a/Engine/Source/Runtime/Slate/Private/Widgets/Views/SListPanel.cpp b/Engine/Source/Runtime/Slate/Private/Widgets/Views/SListPanel.cpp
new file mode 100644
--- /dev/null
+++ b/Engine/Source/Runtime/Slate/Private/Widgets/Views/SListPanel.cpp
@@ -0,0 +1,230 @@
+// Copyright 1998-2014 Epic Games, Inc. All Rights Reserved.
+ 
+#include "SlatePrivatePCH.h"
+#include "SListPanel.h"
+
+
+FNoChildren SListPanel::NoChildren = FNoChildren();
+
+/**
+ * Construct the widget
+ *
+ * @param InArgs   A declaration from which to construct the widget
+ */
+void SListPanel::Construct( const FArguments& InArgs )
+{
+	PreferredRowNum = 0;
+	SmoothScrollOffsetInItems = 0;
+	ItemWidth = InArgs._ItemWidth;
+	ItemHeight = InArgs._ItemHeight;
+	NumDesiredItems = InArgs._NumDesiredItems;
+	bIsRefreshPending = false;
+}
+
+/** Make a new ListPanel::Slot  */
+SListPanel::FSlot& SListPanel::Slot()
+{
+	return *(new FSlot());
+}
+	
+/** Add a slot to the ListPanel */
+SListPanel::FSlot& SListPanel::AddSlot(int32 InsertAtIndex)
+{
+	FSlot& NewSlot = SListPanel::Slot();
+	if (InsertAtIndex == INDEX_NONE)
+	{
+		this->Children.Add( &NewSlot );
+	}
+	else
+	{
+		this->Children.Insert( &NewSlot, InsertAtIndex );
+	}
+	
+	return NewSlot;
+}
+	
+/**
+ * Arrange the children top-to-bottom with not additional layout info.
+ */
+void SListPanel::ArrangeChildren( const FGeometry& AllottedGeometry, FArrangedChildren& ArrangedChildren ) const
+{
+	if ( ShouldArrangeHorizontally() )
+	{
+		// This is a tile view list, arrange items horizontally until there is no more room then create a new row.
+		const float AllottedWidth = AllottedGeometry.Size.X;
+		const float ItemPadding = GetItemPadding(AllottedGeometry);
+		const float HalfItemPadding = ItemPadding * 0.5;
+		float WidthSoFar = 0;
+		float LocalItemWidth = ItemWidth.Get();
+		float LocalItemHeight = ItemHeight.Get();
+		float HeightSoFar = -FMath::Floor(SmoothScrollOffsetInItems * LocalItemHeight);
+
+		for( int32 ItemIndex = 0; ItemIndex < Children.Num(); ++ItemIndex )
+		{
+			ArrangedChildren.AddWidget(
+				AllottedGeometry.MakeChild( Children[ItemIndex].Widget, FVector2D(WidthSoFar + HalfItemPadding, HeightSoFar), FVector2D(LocalItemWidth, LocalItemHeight) )
+				);
+
+			WidthSoFar += LocalItemWidth + ItemPadding;
+
+			if ( WidthSoFar + LocalItemWidth + ItemPadding > AllottedWidth )
+			{
+				WidthSoFar = 0;
+				HeightSoFar += LocalItemHeight;
+			}
+		}
+	}
+	else
+	{
+		if (Children.Num() > 0)
+		{
+			// This is a normal list, arrange items vertically
+			float HeightSoFar = -FMath::Floor(SmoothScrollOffsetInItems * Children[0].Widget->GetDesiredSize().Y);
+			for( int32 ItemIndex = 0; ItemIndex < Children.Num(); ++ItemIndex )
+			{
+					const FVector2D ItemDesiredSize = Children[ItemIndex].Widget->GetDesiredSize();
+					const float LocalItemHeight = ItemDesiredSize.Y;
+
+				// Note that ListPanel does not respect child Visibility.
+				// It is simply not useful for ListPanels.
+				ArrangedChildren.AddWidget(
+					AllottedGeometry.MakeChild( Children[ItemIndex].Widget, FVector2D(0, HeightSoFar), FVector2D(AllottedGeometry.Size.X, LocalItemHeight) )
+					);
+
+				HeightSoFar += LocalItemHeight;
+			}
+		}
+	}
+}
+	
+void SListPanel::Tick( const FGeometry& AllottedGeometry, const double InCurrentTime, const float InDeltaTime )
+{
+	if (ShouldArrangeHorizontally())
+	{
+		const float AllottedWidth = AllottedGeometry.Size.X;
+		const float ItemPadding = GetItemPadding(AllottedGeometry);
+		const float LocalItemWidth = ItemWidth.Get();
+		const float TotalItemSize = LocalItemWidth + ItemPadding;
+		const int32 NumChildren = Children.Num();
+
+		if (TotalItemSize > 0.0f && NumChildren > 0)
+		{
+			int32 NumColumns = FMath::Clamp(FMath::Ceil(AllottedWidth / TotalItemSize) - 1, 1, NumChildren);
+			PreferredRowNum = FMath::Ceil(NumChildren / (float)NumColumns);
+		}
+		else
+		{
+			PreferredRowNum = 1;
+		}
+	}
+	else
+	{
+		PreferredRowNum = NumDesiredItems.Get();
+	}
+
+}
+	
+/**
+ * Simply the sum of all the children (vertically), and the largest width (horizontally).
+ */
+FVector2D SListPanel::ComputeDesiredSize() const
+{
+	float MaxWidth = 0;
+	float TotalHeight = 0;
+	for( int32 ItemIndex = 0; ItemIndex < Children.Num(); ++ItemIndex )
+	{
+		// Notice that we do not respect child Visibility.
+		// It is simply not useful for ListPanels.
+		FVector2D ChildDesiredSize = Children[ItemIndex].Widget->GetDesiredSize();
+		MaxWidth = FMath::Max( ChildDesiredSize.X, MaxWidth );
+		TotalHeight += ChildDesiredSize.Y;
+	}
+
+	if (ShouldArrangeHorizontally())
+	{
+		return FVector2D( MaxWidth, ItemHeight.Get() * PreferredRowNum );
+	}
+	else
+	{
+		return (Children.Num() > 0)
+			? FVector2D( MaxWidth, TotalHeight/Children.Num()*PreferredRowNum )
+			: FVector2D::ZeroVector;
+	}
+	
+}
+
+/**
+ * @return A slot-agnostic representation of this panel's children
+ */
+FChildren* SListPanel::GetChildren()
+{
+	if (bIsRefreshPending)
+	{
+		// When a refresh is pending it is unsafe to cache the desired sizes of our children because
+		// they may be representing unsound data structures. Any delegates/attributes accessing unsound
+		// data will cause a crash.
+		return &NoChildren;
+	}
+	else
+	{
+		return &Children;
+	}
+	
+}
+
+/** Set the offset of the view area from the top of the list. */
+void SListPanel::SmoothScrollOffset(float InOffsetInItems)
+{
+	SmoothScrollOffsetInItems = InOffsetInItems;
+}
+
+/** Remove all the children from this panel */
+void SListPanel::ClearItems()
+{
+	Children.Empty();
+}
+
+float SListPanel::GetItemWidth() const
+{
+	return ItemWidth.Get();
+}
+
+float SListPanel::GetItemPadding(const FGeometry& AllottedGeometry) const
+{
+	float LocalItemWidth = ItemWidth.Get();
+	const int32 NumItemsWide = LocalItemWidth > 0 ? FMath::Floor(AllottedGeometry.Size.X / LocalItemWidth) : 0;
+	
+	// Only add padding between items if we have more total items that we can fit on a single row.  Otherwise,
+	// the padding around items would continue to change proportionately with the (ample) free horizontal space
+	float Padding = 0.0f;
+	if( NumItemsWide > 0 && Children.Num() > NumItemsWide )
+	{
+		// Subtract a tiny amount from the available width to avoid floating point precision problems when arranging children
+		static const float FloatingPointPrecisionOffset = 0.001f;
+
+		Padding = (AllottedGeometry.Size.X - FloatingPointPrecisionOffset - NumItemsWide * LocalItemWidth) / NumItemsWide;
+	}
+	return Padding;
+}
+
+/** @return the uniform item height used when arranging children. */
+float SListPanel::GetItemHeight() const
+{
+	return ItemHeight.Get();
+}
+
+void SListPanel::SetRefreshPending( bool IsPendingRefresh )
+{
+	bIsRefreshPending = IsPendingRefresh;
+}
+
+bool SListPanel::IsRefreshPending() const
+{
+	return bIsRefreshPending;
+}
+
+bool SListPanel::ShouldArrangeHorizontally() const
+{
+	return ItemWidth.Get() > 0;
+}
+
diff --git a/Engine/Source/Runtime/Slate/Private/Widgets/Views/STableViewBase.cpp b/Engine/Source/Runtime/Slate/Private/Widgets/Views/STableViewBase.cpp
new file mode 100644
--- /dev/null
+++ b/Engine/Source/Runtime/Slate/Private/Widgets/Views/STableViewBase.cpp
@@ -0,0 +1,603 @@
+// Copyright 1998-2014 Epic Games, Inc. All Rights Reserved.
+
+#include "SlatePrivatePCH.h"
+#include "SListPanel.h"
+
+
+/** Create the child widgets that comprise the list */
+void STableViewBase::ConstructChildren( const TAttribute<float>& InItemWidth, const TAttribute<float>& InItemHeight, const TSharedPtr<SHeaderRow>& InHeaderRow, const TSharedPtr<SScrollBar>& InScrollBar )
+{
+	HeaderRow = InHeaderRow;
+
+	// If the user provided a scrollbar, we do not need to make one of our own.
+	if (InScrollBar.IsValid())
+	{
+		ScrollBar = InScrollBar;
+		ScrollBar->SetOnUserScrolled( FOnUserScrolled::CreateSP(this, &STableViewBase::ScrollBar_OnUserScrolled) );
+	}
+	
+	
+	TSharedRef<SWidget> ListAndScrollbar = (!ScrollBar.IsValid())
+		// We need to make our own scrollbar
+		? StaticCastSharedRef<SWidget>
+		(
+			SNew(SHorizontalBox)
+			+SHorizontalBox::Slot()
+			.FillWidth(1)
+			[
+				SAssignNew( ItemsPanel, SListPanel )
+				.ItemWidth(InItemWidth)
+				.ItemHeight( InItemHeight )
+				.NumDesiredItems( this, &STableViewBase::GetNumItemsBeingObserved )
+			]
+			+SHorizontalBox::Slot()
+			.AutoWidth()
+			[
+				SNew( SBox )
+				.WidthOverride( FOptionalSize( 16 ) )
+				[
+					SAssignNew(ScrollBar, SScrollBar)
+					.OnUserScrolled( this, &STableViewBase::ScrollBar_OnUserScrolled )
+				]
+			]
+		)
+		// The user provided us with a scrollbar; we will rely on it.
+		: StaticCastSharedRef<SWidget>
+		(
+			SAssignNew( ItemsPanel, SListPanel )
+			.ItemWidth(InItemWidth)
+			.ItemHeight( InItemHeight )
+			.NumDesiredItems( this, &STableViewBase::GetNumItemsBeingObserved )
+		);
+
+	if (InHeaderRow.IsValid())
+	{
+		InHeaderRow->SetAssociatedVerticalScrollBar( ScrollBar.ToSharedRef(), 16 );
+
+		this->ChildSlot
+		[
+			SNew(SVerticalBox)
+			+SVerticalBox::Slot()
+			.AutoHeight()
+			[
+				InHeaderRow.ToSharedRef()
+			]
+			+SVerticalBox::Slot()
+			.FillHeight(1)
+			[
+				ListAndScrollbar
+			]
+		];
+	}
+	else
+	{
+		this->ChildSlot
+		[
+			ListAndScrollbar
+		];
+	}
+	
+}
+
+/**
+ * See SWidget::SupportsKeyboardFocus().
+ *
+ * @return  True if this widget can take keyboard focus
+ */
+bool STableViewBase::SupportsKeyboardFocus() const
+{
+	// The ListView is focusable.
+	return true;
+}
+
+
+void STableViewBase::OnKeyboardFocusLost( const FKeyboardFocusEvent& InKeyboardFocusEvent )
+{
+	bShowSoftwareCursor = false;
+}
+
+
+void STableViewBase::OnMouseCaptureLost()
+{
+	bShowSoftwareCursor = false;
+}
+
+
+struct FEndOfListResult
+{
+	FEndOfListResult( float InOffset, float InItemsAboveView )
+	: OffsetFromEndOfList(InOffset)
+	, ItemsAboveView(InItemsAboveView)
+	{
+
+	}
+
+	float OffsetFromEndOfList;
+	float ItemsAboveView;
+};
+
+static FEndOfListResult ComputeOffsetForEndOfList( const FGeometry& ListPanelGeometry, const FChildren& ListPanelChildren )
+{
+	float OffsetFromEndOfList = 0.0f;
+	float AvailableSpace = ListPanelGeometry.Size.Y;
+	float ItemsAboveView = 0.0f;
+	for ( int ChildIndex=ListPanelChildren.Num()-1; ChildIndex >= 0; --ChildIndex )
+	{
+		const float CurChildHeight = ListPanelChildren.GetChildAt(ChildIndex)->GetDesiredSize().Y;
+		if (AvailableSpace == 0)
+		{
+			ItemsAboveView ++;
+		}
+
+		if ( CurChildHeight < AvailableSpace )
+		{
+			// This whole child fits
+			OffsetFromEndOfList += 1;
+			AvailableSpace -= CurChildHeight;
+		}
+		else
+		{
+			OffsetFromEndOfList += AvailableSpace / CurChildHeight;
+			ItemsAboveView += (CurChildHeight - AvailableSpace)/CurChildHeight;
+			AvailableSpace = 0;
+		}
+	}
+
+	return FEndOfListResult( OffsetFromEndOfList, ItemsAboveView );
+}
+
+/**
+ * See SWidget::Tick()
+ *
+ * @param  AllottedGeometry The space allotted for this widget
+ * @param  InCurrentTime  Current absolute real time
+ * @param  InDeltaTime  Real time passed since last tick
+ */
+void STableViewBase::Tick( const FGeometry& AllottedGeometry, const double InCurrentTime, const float InDeltaTime )
+{
+	if (ItemsPanel.IsValid())
+	{
+		if ( !IsRightClickScrolling() )
+		{
+			this->InertialScrollManager.UpdateScrollVelocity(InDeltaTime);
+
+			const float ScrollVelocity = this->InertialScrollManager.GetScrollVelocity();
+			if ( ScrollVelocity != 0.f )
+			{
+				this->ScrollBy(AllottedGeometry, ScrollVelocity * InDeltaTime);
+			}
+		}
+
+		FGeometry PanelGeometry = FindChildGeometry( AllottedGeometry, ItemsPanel.ToSharedRef() );
+		if ( bItemsNeedRefresh || PanelGeometryLastTick.Size != PanelGeometry.Size)
+		{
+			PanelGeometryLastTick = PanelGeometry;
+
+			// We never create the ItemsPanel if the user did not specify all the parameters required to successfully make a list.
+			ScrollIntoView( PanelGeometry );
+
+			const FReGenerateResults ReGenerateResults = ReGenerateItems( PanelGeometry );
+			LastGenerateResults = ReGenerateResults;
+
+			const int32 NumItemsBeingObserved = GetNumItemsBeingObserved();
+
+			const bool bEnoughRoomForAllItems = ReGenerateResults.ExactNumWidgetsOnScreen >= NumItemsBeingObserved;
+			if (bEnoughRoomForAllItems)
+			{
+				// We can show all the items, so make sure there is no scrolling.
+				ScrollOffset = 0;
+			}
+			else if ( ReGenerateResults.bGeneratedPastLastItem )
+			{
+				ScrollOffset = ReGenerateResults.NewScrollOffset;
+			}
+			
+			ScrollOffset = FMath::Max(0.0, ScrollOffset);
+			ItemsPanel->SmoothScrollOffset( FMath::Fractional(ScrollOffset / GetNumItemsWide()) );
+
+			UpdateSelectionSet();
+
+			// Update scrollbar
+			{
+				// The thumb size is whatever fraction of the items we are currently seeing (including partially seen items).
+				// e.g. if we are seeing 0.5 of the first generated widget and 0.75 of the last widget, that's 1.25 widgets.
+				const double ThumbSizeFraction = ReGenerateResults.ExactNumWidgetsOnScreen / NumItemsBeingObserved;
+				const double OffsetFraction = ScrollOffset / NumItemsBeingObserved;
+				ScrollBar->SetState( OffsetFraction, ThumbSizeFraction );
+			}
+
+			NotifyItemScrolledIntoView();
+
+			bWasAtEndOfList = ScrollBar->DistanceFromBottom() <= 0.f ? true : false;
+
+			bItemsNeedRefresh = false;
+			ItemsPanel->SetRefreshPending(false);
+		}
+	}
+}
+
+/**
+ * Invoked by the scrollbar when the user scrolls.
+ *
+ * @param InScrollOffsetFraction  The location to which the user scrolled as a fraction (between 0 and 1) of total height of the content.
+ */
+void STableViewBase::ScrollBar_OnUserScrolled( float InScrollOffsetFraction )
+{
+	const double ClampedScrollOffsetInItems = FMath::Clamp<double>( InScrollOffsetFraction, 0.0, 1.0 )* GetNumItemsBeingObserved();
+	ScrollTo( ClampedScrollOffsetInItems );
+}
+
+/**
+ * See SWidget::OnMouseButtonDown()
+ *
+ * @param MyGeometry The Geometry of the widget receiving the event
+ * @param MouseEvent Information about the input event
+ *
+ * @return Whether the event was handled along with possible requests for the system to take action.
+ */
+FReply STableViewBase::OnMouseButtonDown( const FGeometry& MyGeometry, const FPointerEvent& MouseEvent )
+{
+	// Zero the scroll velocity so the list stops immediately on mouse down, even if the user does not drag
+	this->InertialScrollManager.ClearScrollVelocity();
+
+	if ( MouseEvent.GetEffectingButton() == EKeys::RightMouseButton && ScrollBar->IsNeeded() )
+	{
+		AmountScrolledWhileRightMouseDown = 0;
+
+		// NOTE: We don't bother capturing the mouse, unless the user starts dragging a few pixels (see the
+		// mouse move handling here.)  This is important so that the item row has a chance to select
+		// items when the right mouse button is released.  Just keep in mind that you might not get
+		// an OnMouseButtonUp event for the right mouse button if the user moves off of the table before
+		// they reach our scroll threshold
+		return FReply::Handled();
+	}
+	return FReply::Unhandled();			
+}
+
+/**
+ * See SWidget::OnMouseButtonUp()
+ *
+ * @param MyGeometry The Geometry of the widget receiving the event
+ * @param MouseEvent Information about the input event
+ *
+ * @return Whether the event was handled along with possible requests for the system to take action.
+ */
+FReply STableViewBase::OnMouseButtonUp( const FGeometry& MyGeometry, const FPointerEvent& MouseEvent )
+{
+	if ( MouseEvent.GetEffectingButton() == EKeys::RightMouseButton )
+	{
+		OnRightMouseButtonUp( MouseEvent.GetScreenSpacePosition() );
+
+		FReply Reply = FReply::Handled().ReleaseMouseCapture();
+		bShowSoftwareCursor = false;
+
+		// If we have mouse capture, snap the mouse back to the closest location that is within the list's bounds
+		if ( HasMouseCapture() )
+		{
+			FSlateRect ListScreenSpaceRect = MyGeometry.GetClippingRect();
+			FVector2D CursorPosition = MyGeometry.LocalToAbsolute( SoftwareCursorPosition );
+
+			FIntPoint BestPositionInList(
+				FMath::Round( FMath::Clamp( CursorPosition.X, ListScreenSpaceRect.Left, ListScreenSpaceRect.Right ) ),
+				FMath::Round( FMath::Clamp( CursorPosition.Y, ListScreenSpaceRect.Top, ListScreenSpaceRect.Bottom ) )
+				);
+
+			Reply.SetMousePos(BestPositionInList);
+		}
+
+		return Reply;
+	}
+	return FReply::Unhandled();
+}
+
+
+/**
+ * See SWidget::OnMouseMove()
+ *
+ * @param MyGeometry The Geometry of the widget receiving the event
+ * @param MouseEvent Information about the input event
+ *
+ * @return Whether the event was handled along with possible requests for the system to take action.
+ */
+FReply STableViewBase::OnMouseMove( const FGeometry& MyGeometry, const FPointerEvent& MouseEvent )
+{	
+	if( MouseEvent.IsMouseButtonDown( EKeys::RightMouseButton ) )
+	{
+		const float ScrollByAmount = MouseEvent.GetCursorDelta().Y;
+		// If scrolling with the right mouse button, we need to remember how much we scrolled.
+		// If we did not scroll at all, we will bring up the context menu when the mouse is released.
+		AmountScrolledWhileRightMouseDown += FMath::Abs( ScrollByAmount );
+
+		// Has the mouse moved far enough with the right mouse button held down to start capturing
+		// the mouse and dragging the view?
+		if( IsRightClickScrolling() )
+		{
+			this->InertialScrollManager.AddScrollSample(-ScrollByAmount, FPlatformTime::Seconds());
+			const float AmountScrolled = this->ScrollBy( MyGeometry, -ScrollByAmount );
+
+			FReply Reply = FReply::Handled();
+
+			// Capture the mouse if we need to
+			if( FSlateApplication::Get().GetMouseCaptor().Get() != this )
+			{
+				Reply.CaptureMouse( AsShared() ).UseHighPrecisionMouseMovement( AsShared() );
+				SoftwareCursorPosition = MyGeometry.AbsoluteToLocal( MouseEvent.GetScreenSpacePosition() );
+				bShowSoftwareCursor = true;
+			}
+
+			// Check if the mouse has moved.
+			if( AmountScrolled != 0 )
+			{
+				SoftwareCursorPosition.Y += MouseEvent.GetCursorDelta().Y;
+			}
+
+			return Reply;
+		}
+	}
+
+	return FReply::Unhandled();
+}
+
+
+void STableViewBase::OnMouseLeave( const FPointerEvent& MouseEvent )
+{
+	if( FSlateApplication::Get().GetMouseCaptor().Get() != this )
+	{
+		// No longer scrolling (unless we have mouse capture)
+		AmountScrolledWhileRightMouseDown = 0;
+	}
+}
+
+
+/**
+ * See SWidget::OnMouseWheel()
+ *
+ * @param MyGeometry The Geometry of the widget receiving the event
+ * @param MouseEvent Information about the input event
+ *
+ * @return Whether the event was handled along with possible requests for the system to take action.
+ */
+FReply STableViewBase::OnMouseWheel( const FGeometry& MyGeometry, const FPointerEvent& MouseEvent )
+{
+	if( !MouseEvent.IsControlDown() )
+	{
+		// Make sure scroll velocity is cleared so it doesn't fight with the mouse wheel input
+		this->InertialScrollManager.ClearScrollVelocity();
+
+		const float AmountScrolledInItems = this->ScrollBy( MyGeometry, -MouseEvent.GetWheelDelta()*WheelScrollAmount );
+		if (FMath::Abs(AmountScrolledInItems) > 0.0f)
+		{
+			return FReply::Handled();
+		}
+	}
+	return FReply::Unhandled();
+}
+
+/**
+ * See SWidget::OnKeyDown().
+ *
+ * @param MyGeometry The Geometry of the widget receiving the event
+ * @param  InKeyboardEvent  Keyboard event
+ *
+ * @return  Returns whether the event was handled, along with other possible actions
+ */
+FReply STableViewBase::OnKeyDown( const FGeometry& MyGeometry, const FKeyboardEvent& InKeyboardEvent )
+{
+	if ( InKeyboardEvent.IsControlDown() && InKeyboardEvent.GetKey() == EKeys::End )
+	{
+		ScrollOffset = GetNumItemsBeingObserved();
+		RequestListRefresh();
+		return FReply::Handled();
+	}
+
+
+	return FReply::Unhandled();
+}
+
+FCursorReply STableViewBase::OnCursorQuery( const FGeometry& MyGeometry, const FPointerEvent& CursorEvent ) const
+{
+	if ( IsRightClickScrolling() )
+	{
+		// We hide the native cursor as we'll be drawing the software EMouseCursor::GrabHandClosed cursor
+		return FCursorReply::Cursor( EMouseCursor::None );
+	}
+	else
+	{
+		return FCursorReply::Unhandled();
+	}
+}
+
+
+/** @return The number of Widgets we currently have generated. */
+int32 STableViewBase::GetNumGeneratedChildren() const
+{
+	return (ItemsPanel.IsValid())
+		? ItemsPanel->GetChildren()->Num()
+		: 0;
+}
+
+TSharedPtr<SHeaderRow> STableViewBase::GetHeaderRow() const
+{
+	return HeaderRow;
+}
+
+bool STableViewBase::IsRightClickScrolling() const
+{
+	return AmountScrolledWhileRightMouseDown >= SlatePanTriggerDistance && this->ScrollBar->IsNeeded();
+}
+
+bool STableViewBase::IsUserScrolling() const
+{
+	bool bUserScroll = this->ScrollBar->IsNeeded() && this->ScrollBar->IsScrolling();
+	return bUserScroll || IsRightClickScrolling();
+}
+
+void STableViewBase::RequestListRefresh()
+{
+	bItemsNeedRefresh = true;
+	ItemsPanel->SetRefreshPending(true);
+}
+
+bool STableViewBase::IsPendingRefresh() const
+{
+	return bItemsNeedRefresh || ItemsPanel->IsRefreshPending();
+}
+
+int32 STableViewBase::OnPaint( const FGeometry& AllottedGeometry, const FSlateRect& MyClippingRect, FSlateWindowElementList& OutDrawElements, int32 LayerId, const FWidgetStyle& InWidgetStyle, bool bParentEnabled ) const
+{
+	int32 NewLayerId = SCompoundWidget::OnPaint( AllottedGeometry, MyClippingRect, OutDrawElements, LayerId, InWidgetStyle, bParentEnabled );
+
+	if( !bShowSoftwareCursor )
+	{
+		return NewLayerId;
+	}
+
+	const FSlateBrush* Brush = FCoreStyle::Get().GetBrush(TEXT("SoftwareCursor_Grab"));
+
+	FSlateDrawElement::MakeBox(
+		OutDrawElements,
+		++NewLayerId,
+		AllottedGeometry.ToPaintGeometry( SoftwareCursorPosition - ( Brush->ImageSize / 2 ), Brush->ImageSize ),
+		Brush,
+		MyClippingRect
+		);
+
+	return NewLayerId;
+}
+
+STableViewBase::STableViewBase( ETableViewMode::Type InTableViewMode )
+	: TableViewMode( InTableViewMode )
+	, ScrollOffset( 0 )
+	, AmountScrolledWhileRightMouseDown( 0 )
+	, LastGenerateResults( 0,0,0,false )
+	, bWasAtEndOfList(false)
+	, SelectionMode( ESelectionMode::Multi )
+	, SoftwareCursorPosition( ForceInitToZero )
+	, bShowSoftwareCursor( false )
+	, bItemsNeedRefresh( true )
+{
+}
+
+float STableViewBase::ScrollBy( const FGeometry& MyGeometry, float ScrollByAmountInSlateUnits )
+{
+	const int32 NumItemsBeingObserved = GetNumItemsBeingObserved();
+	const float FractionalScrollOffsetInItems = (ScrollOffset + GetScrollRateInItems() * ScrollByAmountInSlateUnits) / NumItemsBeingObserved;
+	const double ClampedScrollOffsetInItems = FMath::Clamp<double>( FractionalScrollOffsetInItems, 0.0, 1.0 ) * NumItemsBeingObserved;
+	return ScrollTo( ClampedScrollOffsetInItems );
+}
+
+float STableViewBase::ScrollTo( float InScrollOffset )
+{
+	const float NewScrollOffset = FMath::Clamp( InScrollOffset, 0.0f, (float)GetNumItemsBeingObserved() );
+	float AmountScrolled = FMath::Abs( ScrollOffset - NewScrollOffset );
+	ScrollOffset = NewScrollOffset;
+
+	RequestListRefresh();	
+
+	if ( bWasAtEndOfList && NewScrollOffset >= ScrollOffset )
+	{
+		AmountScrolled = 0;
+	}
+
+	return AmountScrolled;
+}
+
+void STableViewBase::InsertWidget( const TSharedRef<ITableRow> & WidgetToInset )
+{
+	ItemsPanel->AddSlot(0)
+	[
+		WidgetToInset->AsWidget()
+	];
+}
+
+/**
+ * Add a widget to the view.
+ *
+ * @param WidgetToAppend   Widget to append to the view.
+ */
+void STableViewBase::AppendWidget( const TSharedRef<ITableRow>& WidgetToAppend )
+{
+	ItemsPanel->AddSlot()
+	[
+		WidgetToAppend->AsWidget()
+	];
+}
+
+/**
+ * Remove all the widgets from the view.
+ */
+void STableViewBase::ClearWidgets()
+{
+	ItemsPanel->ClearItems();
+}
+
+/**
+ * Get the uniform item width.
+ */
+float STableViewBase::GetItemWidth() const
+{
+	// Casting to a STilePanel here because we constructed it in Construct()
+	return ItemsPanel->GetItemWidth() + ItemsPanel->GetItemPadding(PanelGeometryLastTick);
+}
+
+/**
+ * Get the uniform item height that is enforced by ListViews.
+ */
+float STableViewBase::GetItemHeight() const
+{
+	return ItemsPanel->GetItemHeight();
+}
+
+float STableViewBase::GetNumLiveWidgets() const
+{
+	return ItemsPanel->GetChildren()->Num();
+}
+
+int32 STableViewBase::GetNumItemsWide() const
+{
+	return 1;
+}
+
+void STableViewBase::OnRightMouseButtonUp(const FVector2D& SummonLocation)
+{
+	const bool bShouldOpenContextMenu = !IsRightClickScrolling();
+	const bool bContextMenuOpeningBound = OnContextMenuOpening.IsBound();
+
+	if ( bShouldOpenContextMenu && bContextMenuOpeningBound )
+	{
+		// Get the context menu content. If NULL, don't open a menu.
+		TSharedPtr<SWidget> MenuContent = OnContextMenuOpening.Execute();
+
+		if( MenuContent.IsValid() )
+		{
+			bShowSoftwareCursor = false;
+			FSlateApplication::Get().PushMenu( AsShared(), MenuContent.ToSharedRef(), SummonLocation, FPopupTransitionEffect( FPopupTransitionEffect::ContextMenu ) );
+		}
+	}
+
+	AmountScrolledWhileRightMouseDown = 0;
+}
+
+float STableViewBase::GetScrollRateInItems() const
+{
+	return ( LastGenerateResults.HeightOfGeneratedItems != 0 && LastGenerateResults.ExactNumWidgetsOnScreen != 0)
+		// Approximate a consistent scrolling rate based on the average item height.
+		? LastGenerateResults.ExactNumWidgetsOnScreen / LastGenerateResults.HeightOfGeneratedItems
+		// Scroll 1/2 an item at a time as a default.
+		: 0.5f;
+}
+
+FVector2D STableViewBase::GetScrollDistance()
+{
+	return FVector2D( 0, ScrollBar->DistanceFromTop() );
+}
+
+FVector2D STableViewBase::GetScrollDistanceRemaining()
+{
+	return FVector2D( 0, ScrollBar->DistanceFromBottom() );
+}
+
+TSharedRef<class SWidget> STableViewBase::GetScrollWidget()
+{
+	return SharedThis(this);
+}
