Slate: Reorganized private Slate files; proper pre-compiled header for all files.

diff --git a/Engine/Source/Runtime/Slate/Public/Framework/MultiBox/MultiBox.h b/Engine/Source/Runtime/Slate/Public/Framework/MultiBox/MultiBox.h
new file mode 100644
--- /dev/null
+++ b/Engine/Source/Runtime/Slate/Public/Framework/MultiBox/MultiBox.h
@@ -0,0 +1,637 @@
+// Copyright 1998-2014 Epic Games, Inc. All Rights Reserved.
+
+#pragma once
+
+
+class FUICommandInfo;
+class FUICommandList;
+
+
+namespace MultiBoxConstants
+{	
+	// @todo Slate MultiBox: Hard coded dimensions
+	const float MenuIconSize = 16.0f;
+	const float MenuCheckBoxSize = 16.0f;
+
+	/** The time that a mouse should be hovered over a sub-menu before it automatically opens */
+	const float SubMenuOpenTime = 0.0f;
+
+	/** When a sub-menu is already open, the time that a mouse should be hovered over a sub-menu entry before
+	    dismissing the other menu and opening this one */
+	const float SubMenuClobberTime = 0.5f;
+
+	//const FName MenuItemFont = "MenuItem.Font";
+	//const FName MenuBindingFont = "MenuItem.BindingFont";
+
+	//const FName EditableTextFont = "MenuItem.Font";
+
+	/** Minimum widget of an editable text box within a multi-box */
+	const float EditableTextMinWidth = 30.0f;
+}
+
+/**
+ * MultiBlock (abstract).  Wraps a "block" of useful UI functionality that can be added to a MultiBox.
+ */
+class SLATE_API FMultiBlock
+	: public TSharedFromThis< FMultiBlock >		// Enables this->AsShared()
+{
+
+public:
+
+	/**
+	 * Constructor
+ 	 *
+	 * @param InCommand		The command info that describes what action to take when this block is activated
+	 * @param InCommandList	The list of mappings from command info to delegates so we can find the delegates to process for the provided action
+	 */
+	FMultiBlock( const TSharedPtr< const FUICommandInfo > InCommand, TSharedPtr< const FUICommandList > InCommandList, FName InExtensionHook = NAME_None )
+		: Action( InCommand )
+		, ActionList( InCommandList )
+		, ExtensionHook( InExtensionHook )
+		, TutorialHighlightName( NAME_None )
+	{
+	}
+
+	/**
+	 * Constructor
+	 *
+	 * @InAction UI action delegates that should be used in place of UI commands (dynamic menu items)
+	 */
+	FMultiBlock( const FUIAction& InAction,  FName InExtensionHook = NAME_None )
+		: DirectActions( InAction )
+		, ExtensionHook( InExtensionHook )
+		, TutorialHighlightName( NAME_None )
+	{
+	}
+
+	virtual ~FMultiBlock()
+	{
+	}
+
+	/**
+	 * Returns the action list associated with this block
+	 * 
+	 * @return The action list or null if one does not exist for this block
+	 */
+	TSharedPtr< const FUICommandList > GetActionList() const { return ActionList; }
+
+	/**
+	 * Returns the action associated with this block
+	 * 
+	 * @return The action for this block or null if one does not exist
+	 */
+	TSharedPtr< const FUICommandInfo> GetAction() const { return Action; }
+
+	/**
+	 * Returns the direct actions for this block.  Delegates may be unbound if this block has a UICommand
+	 *
+	 * @return DirectActions for this block
+	 */
+	const FUIAction& GetDirectActions() const { return DirectActions; }
+
+	/** Creates a menu entry that is representative of this block */
+	virtual void CreateMenuEntry(class FMenuBuilder& MenuBuilder) const { }
+
+	/** Group blocks interface */
+	virtual bool IsGroupStartBlock()	const	{ return false; };
+	virtual bool IsGroupEndBlock()		const	{ return false; };
+
+	/** Set the tutorial highlight name for this menu entry */
+	void SetTutorialHightlightName(FName InTutorialName)
+	{
+		TutorialHighlightName = InTutorialName;
+	}
+
+	/** Get the tutorial highlight name for this menu entry */
+	FName GetTutorialHighlightName() const
+	{
+		return TutorialHighlightName;
+	}
+
+	/**
+	 * Creates a MultiBlock widget for this MultiBlock
+	 *
+	 * @param	InOwnerMultiBoxWidget	The widget that will own the new MultiBlock widget
+	 * @param	InLocation				The location information for the MultiBlock widget
+	 *
+	 * @return  MultiBlock widget object
+	 */
+	TSharedRef< class IMultiBlockBaseWidget > MakeWidget( TSharedRef< class SMultiBoxWidget > InOwnerMultiBoxWidget, EMultiBlockLocation::Type InLocation ) const;
+
+private:
+	/**
+	 * Allocates a widget for this type of MultiBlock.  Override this in derived classes.
+	 *
+	 * @return  MultiBlock widget object
+	 */
+	virtual TSharedRef< class IMultiBlockBaseWidget > ConstructWidget() const = 0;
+	
+	/** Gets the extension hook so users can see what hooks are where */
+	FName GetExtensionHook() const {return ExtensionHook;}
+
+private:
+
+	// We're friends with SMultiBoxWidget so that it can call MakeWidget() directly
+	friend class SMultiBoxWidget;	
+
+	/** Direct processing of actions. Will use these actions if there is not UICommand associated with this block that handles actions*/
+	FUIAction DirectActions;
+
+	/** The action associated with this block (can be null for some actions) */
+	const TSharedPtr< const FUICommandInfo > Action;
+
+	/** The list of mappings from command info to delegates that should be called. This is here for quick access. Can be null for some widgets*/
+	const TSharedPtr< const FUICommandList > ActionList;
+
+	/** Optional extension hook which is used for debug display purposes, so users can see what hooks are where */
+	FName ExtensionHook;
+
+	/** Name to identify a widget for tutorials */
+	FName TutorialHighlightName;
+};
+
+
+
+
+/**
+ * MultiBox.  Contains a list of MultiBlocks that provide various functionality.
+ */
+class SLATE_API FMultiBox
+	: public TSharedFromThis< FMultiBox >		// Enables this->AsShared()
+{
+
+public:
+	virtual ~FMultiBox();
+
+	/**
+	 * Creates a new multibox instance
+	 */
+	static TSharedRef<FMultiBox> Create( const EMultiBoxType::Type InType,  FMultiBoxCustomization InCustomization, const bool bInShouldCloseWindowAfterMenuSelection );
+
+
+	/**
+	 * Gets the type of this MultiBox
+	 *
+	 * @return	The MultiBox's type
+	 */
+	const EMultiBoxType::Type GetType() const
+	{
+		return Type;
+	}
+
+
+	/**
+	 * Gets whether or not the window that contains this multibox should be destroyed after the user clicks on a menu item in this box
+	 *
+	 * @return	True if window should be closed automatically when the user clicks on a menu item, otherwise false
+	 */
+	bool ShouldCloseWindowAfterMenuSelection() const
+	{
+		return bShouldCloseWindowAfterMenuSelection;
+	}
+
+
+
+	/**
+	 * Adds a MultiBlock to this MultiBox, to the end of the list
+	 */
+	void AddMultiBlock( TSharedRef< const FMultiBlock > InBlock );
+
+	/**
+	 * Removes a MultiBlock from the list for user customization
+	 */
+	void RemoveCustomMultiBlock( TSharedRef< const FMultiBlock> InBlock );
+
+	/**
+	 * Inserts a MultiBlock to the list for user customization
+	 */
+	void InsertCustomMultiBlock( TSharedRef<const FMultiBlock> InBlock, int32 Index );
+
+	/**
+	 * Creates a MultiBox widget for this MultiBox
+	 *
+	 * @return  MultiBox widget object
+	 */
+	TSharedRef< class SMultiBoxWidget > MakeWidget();
+
+
+	/**
+	 * Access this MultiBox's list of blocks
+	 *
+	 * @return	Our list of MultiBlocks
+	 */
+	const TArray< TSharedRef< const FMultiBlock > >& GetBlocks() const
+	{
+		return Blocks;
+	}
+	
+	/** @return The style set used by the multibox widgets */
+	const ISlateStyle* GetStyleSet() const { return StyleSet; }
+
+	/** @return The style name used by the multibox widgets */
+	const FName& GetStyleName() const { return StyleName; }
+
+	/** Sets the style to use on the multibox widgets */
+	void SetStyle( const ISlateStyle* InStyleSet, const FName& InStyleName )
+	{
+		StyleSet  = InStyleSet;
+		StyleName = InStyleName;
+	}
+
+	/** @return The customization name for this box */
+	FName GetCustomizationName() const;
+
+	/**
+	 * Creates a block from the provided command that is compatible with this box 
+	 *
+	 * @param Command	The UI command to create the block from
+	 * @return The created multi block.  If null, this command could not be placed in this block
+	 */
+	TSharedPtr<FMultiBlock> MakeMultiBlockFromCommand( TSharedPtr<const FUICommandInfo> Command, bool bCommandMustBeBound ) const;
+
+	/**
+	 * Finds an existing block that handles the provided command
+	 *
+	 * @param Command The command to check for
+	 */
+	TSharedPtr<const FMultiBlock> FindBlockFromCommand( TSharedPtr<const FUICommandInfo> Command ) const;
+
+	bool IsInEditMode() const { return FMultiBoxSettings::IsInToolbarEditMode() && IsCustomizable(); }
+private:
+	
+	/**
+	 * Constructor
+	 *
+	 * @param	InType	Type of MultiBox
+	 * @param	bInShouldCloseWindowAfterMenuSelection	Sets whether or not the window that contains this multibox should be destroyed after the user clicks on a menu item in this box
+	 */
+	FMultiBox( const EMultiBoxType::Type InType,  FMultiBoxCustomization InCustomization, const bool bInShouldCloseWindowAfterMenuSelection );
+
+	/**
+	 * @return true if this box can be customized by a user
+	 */
+	bool IsCustomizable() const;
+
+	/**
+	 * Applies user customized blocks to the multibox
+	 */
+	void ApplyCustomizedBlocks();
+private:
+
+	/** Saved customization data */
+	TSharedRef< class FMultiBoxCustomizationData > CustomizationData;
+
+	/** All command lists in this box */
+	TArray< TSharedPtr<const FUICommandList> > CommandLists;
+
+	/** Ordered list of blocks */
+	TArray< TSharedRef< const FMultiBlock > > Blocks;
+
+	/** The style set to use with the widgets in the MultiBox */
+	const ISlateStyle* StyleSet;
+
+	/** The style name to use with the widgets in the MultiBox */
+	FName StyleName;
+
+	/** Type of MultiBox */
+	EMultiBoxType::Type Type;
+
+	/** True if window that owns any widgets created from this multibox should be closed automatically after the user commits to a menu choice */
+	bool bShouldCloseWindowAfterMenuSelection;
+};
+
+
+
+/**
+ * MultiBlock Slate widget interface
+ */
+class SLATE_API IMultiBlockBaseWidget
+{
+
+public:
+
+	/**
+	 * Interprets this object as a SWidget
+	 *
+	 * @return  Widget reference
+	 */
+	virtual TSharedRef< SWidget > AsWidget() = 0;
+
+
+	/**
+	 * Interprets this object as a SWidget
+	 *
+	 * @return  Widget reference
+	 */
+	virtual TSharedRef< const SWidget > AsWidget() const = 0;
+
+
+	/**
+	 * Associates the owner MultiBox widget with this widget
+	 *
+	 * @param	InOwnerMultiBoxWidget		The MultiBox widget that owns us
+	 */
+	virtual void SetOwnerMultiBoxWidget( TSharedRef< SMultiBoxWidget > InOwnerMultiBoxWidget ) = 0;
+
+	
+	/**
+	 * Associates this widget with a MultiBlock
+	 *
+	 * @param	InMultiBlock	The MultiBlock we'll be associated with
+	 */
+	virtual void SetMultiBlock( TSharedRef< const FMultiBlock > InMultiBlock ) = 0;
+
+
+	/**
+	 * Builds this MultiBlock widget up from the MultiBlock associated with it
+	 *
+	 * @param	StyleSet	The Slate style to use to build the widget
+	 * @param	StyleName	The style name to use from the StyleSet
+	 */
+	virtual void BuildMultiBlockWidget(const ISlateStyle* StyleSet, const FName& StyleName) = 0;
+
+	/**
+	 * Sets the blocks location relative to the other blocks
+	 * 
+	 * @param InLocation	The MultiBlocks location
+	 */
+	virtual void SetMultiBlockLocation(EMultiBlockLocation::Type InLocation) = 0;
+
+	/**
+	 * Returns this MultiBlocks location
+	 */
+	virtual EMultiBlockLocation::Type GetMultiBlockLocation() = 0;
+
+};
+
+
+
+/**
+ * MultiBlock Slate base widget (pure virtual).  You'll derive your own MultiBlock class from this base class.
+ */
+class SLATE_API SMultiBlockBaseWidget
+	: public IMultiBlockBaseWidget,
+	  public SCompoundWidget
+{
+
+public:
+
+	/**
+	 * Interprets this object as a SWidget
+	 *
+	 * @return  Widget reference
+	 */
+	virtual TSharedRef< SWidget > AsWidget()
+	{
+		return this->AsShared();
+	}
+
+
+	/**
+	 * Interprets this object as a SWidget
+	 *
+	 * @return  Widget reference
+	 */
+	virtual TSharedRef< const SWidget > AsWidget() const
+	{
+		return this->AsShared();
+	}
+
+
+	/**
+	 * Associates the owner MultiBox widget with this widget
+	 *
+	 * @param	InOwnerMultiBoxWidget		The MultiBox widget that owns us
+	 */
+	void SetOwnerMultiBoxWidget( TSharedRef< SMultiBoxWidget > InOwnerMultiBoxWidget )
+	{
+		OwnerMultiBoxWidget = InOwnerMultiBoxWidget;
+	}
+
+	
+	/**
+	 * Associates this widget with a MultiBlock
+	 *
+	 * @param	InMultiBlock	The MultiBlock we'll be associated with
+	 */
+	void SetMultiBlock( TSharedRef< const FMultiBlock > InMultiBlock )
+	{
+		MultiBlock = InMultiBlock;
+	}
+
+	/**
+	 * Sets the blocks location relative to the other blocks
+	 * 
+	 * @param InLocation	The MultiBlocks location
+	 */
+	virtual void SetMultiBlockLocation(EMultiBlockLocation::Type InLocation)
+	{
+		Location = InLocation;
+	}
+
+	/**
+	 * Returns this MultiBlocks location
+	 */
+	virtual EMultiBlockLocation::Type GetMultiBlockLocation()
+	{
+		return Location;
+	}
+
+	/**
+	 * Builds this MultiBlock widget up from the MultiBlock associated with it
+	 */
+	virtual void BuildMultiBlockWidget(const ISlateStyle* StyleSet, const FName& StyleName) = 0;
+
+	/** SWidget Interface */
+	virtual void OnDragEnter( const FGeometry& MyGeometry, const FDragDropEvent& DragDropEvent ) OVERRIDE;
+	virtual FReply OnDragOver( const FGeometry& MyGeometry, const FDragDropEvent& DragDropEvent ) OVERRIDE;
+	virtual FReply OnDrop( const FGeometry& MyGeometry, const FDragDropEvent& DragDropEvent ) OVERRIDE;
+protected:
+
+	/** Weak reference back to the MultiBox widget that owns us */
+	TWeakPtr< SMultiBoxWidget > OwnerMultiBoxWidget;
+
+	/** The MultiBlock we're associated with */
+	TSharedPtr< const FMultiBlock > MultiBlock;
+
+	/** The MultiBlocks location relative to the other blocks in the set */
+	EMultiBlockLocation::Type Location;
+};
+
+/**
+ * MultiBox Slate widget
+ */
+class SLATE_API SMultiBoxWidget
+	: public SCompoundWidget
+{
+
+public:
+
+	SLATE_BEGIN_ARGS( SMultiBoxWidget )
+		: _ContentScale( FVector2D::UnitVector )
+		{}
+
+		/** Content scaling factor */
+		SLATE_ATTRIBUTE( FVector2D, ContentScale )
+
+	SLATE_END_ARGS()
+
+	/**
+	 * Construct this widget
+	 *
+	 * @param	InArgs	The declaration data for this widget
+	 */
+	void Construct( const FArguments& InArgs );
+
+	/**
+	 * Associates this widget with a MultiBox
+	 *
+	 * @param	InMultiBox	The MultiBox we'll be associated with
+	 */
+	void SetMultiBox( TSharedRef< FMultiBox > InMultiBox )
+	{
+		MultiBox = InMultiBox;
+	}
+
+
+	/**
+	 * Access the MultiBox associated with this widget
+	 */
+	TSharedRef< const FMultiBox > GetMultiBox() const
+	{
+		return MultiBox.ToSharedRef();
+	}
+
+	
+	/**
+	 * Builds this MultiBox widget up from the MultiBox associated with it
+	 */
+	void BuildMultiBoxWidget();
+
+	
+	/**
+	 * For menu bar multibox widgets, tells the multibox widget about a currently active pull-down menu
+	 *
+	 * @param	InMenuAnchor	Menu anchor for active pull-down menu or sub-menu
+	 */
+	void SetSummonedMenu( TSharedRef< SMenuAnchor > InMenuAnchor );
+	
+
+	/**
+	 * For menu bar or sub-menu multibox widgets, returns the currently open menu, if there is one open
+	 *
+	 * @return	Menu anchor, or null pointer
+	 */
+	TSharedPtr< const SMenuAnchor > GetOpenMenu() const;
+
+	/**
+	 * For menu bar multibox widget, closes any open pull-down or sub menus
+	 */
+	void CloseSummonedMenus();
+	
+
+	/** Generates the tiles for an STileView for button rows */
+	TSharedRef<ITableRow> GenerateTiles(TSharedPtr<SWidget> Item, const TSharedRef<STableViewBase>& OwnerTable);
+
+	/** Gets the maximum item width and height of the consituent widgets */
+	float GetItemWidth() const;
+	float GetItemHeight() const;
+
+	/** Event handler for clicking the wrap button */
+	TSharedRef<SWidget> OnWrapButtonClicked();
+
+	const ISlateStyle* GetStyleSet() const { return MultiBox->GetStyleSet(); }
+	const FName& GetStyleName() const { return MultiBox->GetStyleName(); }
+
+	/**
+	 * Called when a user drags a UI command into a multiblock in this widget
+	 */
+	void OnCustomCommandDragEnter( TSharedRef<const FMultiBlock> MultiBlock, const FGeometry& MyGeometry, const FDragDropEvent& DragDropEvent );
+
+	/**
+	 * Called when a user drags a UI command within multiblock in this widget
+	 */
+	void OnCustomCommandDragged( TSharedRef<const FMultiBlock> MultiBlock, const FGeometry& MyGeometry, const FDragDropEvent& DragDropEvent );
+
+	/**
+	 * Called when a user drops a UI command in this widget
+	 */
+	void OnCustomCommandDropped();
+	
+	/**
+	 * Called after a drag was initiated from this box but was dropped elsewhere
+	 */
+	void OnDropExternal();
+
+	/** Helper function used to transfer focus to the next/previous widget */
+	static FReply FocusNextWidget( EFocusMoveDirection::Type MoveDirection );
+
+	/** SWidget interface */
+	virtual FReply OnDragOver( const FGeometry& MyGeometry, const FDragDropEvent& DragDropEvent ) OVERRIDE;
+	virtual FReply OnDrop( const FGeometry& MyGeometry, const FDragDropEvent& DragDropEvent ) OVERRIDE;
+	virtual bool SupportsKeyboardFocus() const OVERRIDE;
+	virtual FReply OnKeyboardFocusReceived( const FGeometry& MyGeometry, const FKeyboardFocusEvent& InKeyboardFocusEvent ) OVERRIDE;
+	virtual FReply OnKeyDown( const FGeometry& MyGeometry, const FKeyboardEvent& KeyboardEvent ) OVERRIDE;
+private:
+	/** Adds a block Widget to this widget */
+	void AddBlockWidget( const FMultiBlock& Block, TSharedPtr<SHorizontalBox> HorizontalBox, TSharedPtr<SVerticalBox> VerticalBox, EMultiBlockLocation::Type InLocation );
+
+	/**
+	 * @return True if the passed in block is being dragged
+	 */
+	bool IsBlockBeingDragged( TSharedPtr<const FMultiBlock> Block ) const;
+
+	/**
+	 * Updates the preview block being dragged.  The drag area is where the users dragged block will be dropped
+	 */
+	void UpdateDropAreaPreviewBlock( TSharedRef<const FMultiBlock> MultiBlock, TSharedPtr<FUICommandDragDropOp> DragDropContent, const FGeometry& DragArea, const FVector2D& DragPos );
+
+	/**
+   	 * @return The visibility of customization widgets for a block
+	 */
+	EVisibility GetCustomizationVisibility( TWeakPtr<const FMultiBlock> BlockWeakPtr, TWeakPtr<SWidget> BlockWidgetWeakPtr ) const;
+	
+	/** Called when a user clicks the delete button on a block */
+	FReply OnDeleteBlockClicked( TWeakPtr<const FMultiBlock> BlockWeakPtr );
+private:
+	/** A preview of a block being dragged */
+	struct FDraggedMultiBlockPreview
+	{
+		/** Command being dragged */
+		TSharedPtr<const FUICommandInfo> UICommand;
+		/** Preview block for the command */
+		TSharedPtr<class FDropPreviewBlock> PreviewBlock;
+		/** Index into the block list where the block will be added*/
+		int32 InsertIndex;
+		// Vertical for menus and vertical toolbars, horizontally otherwise
+		EOrientation InsertOrientation;
+
+		FDraggedMultiBlockPreview()
+			: InsertIndex( INDEX_NONE )
+		{}
+
+		void Reset()
+		{
+			UICommand.Reset();
+			PreviewBlock.Reset();
+			InsertIndex = INDEX_NONE;
+		}
+
+		bool IsValid() const { return UICommand.IsValid() && PreviewBlock.IsValid() && InsertIndex != INDEX_NONE; }
+	};
+	
+	/** The MultiBox we're associated with */
+	TSharedPtr< FMultiBox > MultiBox;
+
+	/** For menu bar multibox widgets, this stores a weak reference to the last pull-down or sub-menu that was summoned. */
+	TWeakPtr< SMenuAnchor > SummonedMenuAnchor;
+
+	/** An array of widgets used for an STileView if used */
+	TArray< TSharedPtr<SWidget> > TileViewWidgets;
+
+	/** Specialized box widget to handle clipping of toolbars and menubars */
+	TSharedPtr<class SClippingHorizontalBox> ClippedHorizontalBox;
+
+	/** A preview of a block being dragged inside this box */
+	FDraggedMultiBlockPreview DragPreview;
+};
diff --git a/Engine/Source/Runtime/Slate/Public/Framework/MultiBox/MultiBoxBuilder.h b/Engine/Source/Runtime/Slate/Public/Framework/MultiBox/MultiBoxBuilder.h
--- a/Engine/Source/Runtime/Slate/Public/Framework/MultiBox/MultiBoxBuilder.h
+++ b/Engine/Source/Runtime/Slate/Public/Framework/MultiBox/MultiBoxBuilder.h
@@ -2,7 +2,6 @@
 
 #pragma once
 
-#include "MultiBoxDefs.h"
 #include "MultiBoxExtender.h"
 
 
diff --git a/Engine/Source/Runtime/Slate/Public/Framework/MultiBox/MultiBoxExtender.h b/Engine/Source/Runtime/Slate/Public/Framework/MultiBox/MultiBoxExtender.h
--- a/Engine/Source/Runtime/Slate/Public/Framework/MultiBox/MultiBoxExtender.h
+++ b/Engine/Source/Runtime/Slate/Public/Framework/MultiBox/MultiBoxExtender.h
@@ -2,7 +2,6 @@
 
 #pragma once
 
-#include "MultiBoxDefs.h"
 
 /** Called on your extension to add new pull-down menus to a menu bar */
 DECLARE_DELEGATE_OneParam( FMenuBarExtensionDelegate, class FMenuBarBuilder& );
diff --git a/Engine/Source/Runtime/Slate/Public/Framework/MultiBox/SToolBarButtonBlock.h b/Engine/Source/Runtime/Slate/Public/Framework/MultiBox/SToolBarButtonBlock.h
new file mode 100644
--- /dev/null
+++ b/Engine/Source/Runtime/Slate/Public/Framework/MultiBox/SToolBarButtonBlock.h
@@ -0,0 +1,174 @@
+// Copyright 1998-2014 Epic Games, Inc. All Rights Reserved.
+
+#pragma once
+
+/**
+ * Tool bar button MultiBlock
+ */
+class SLATE_API FToolBarButtonBlock
+	: public FMultiBlock
+{
+
+public:
+
+	/**
+	 * Constructor
+	 *
+	 * @param	InCommand			The command associated with this tool bar button
+	 * @param	InCommandList		The list of commands that are mapped to delegates so that we know what to execute for this button
+	 * @param	InLabelOverride		Optional label override.  If omitted, then the action's label will be used instead.
+	 * @param	InToolTipOverride	Optional tool tip override.	 If omitted, then the action's label will be used instead.
+	 * @param	InIconOverride		Optional icon to use for the tool bar image.  If omitted, then the action's icon will be used instead.
+	 */
+	FToolBarButtonBlock( const TSharedPtr< const FUICommandInfo > InCommand, TSharedPtr< const FUICommandList > InCommandList, const TAttribute<FText>& InLabelOverride = TAttribute<FText>(), const TAttribute<FText>& InToolTipOverride = TAttribute<FText>(), const TAttribute<FSlateIcon>& InIconOverride = TAttribute<FSlateIcon>() );
+	
+	/**
+	 * Constructor
+	 *
+	 * @param	InLabel				The label to display in the menu
+	 * @param	InToolTip			The tool tip to display when the menu entry is hovered over
+	 * @param	InIcon				The icon to display to the left of the label
+	 * @param	InUIAction			UI action to take when this menu item is clicked as well as to determine if the menu entry can be executed or appears "checked"
+	 * @param	InUserInterfaceActionType	Type of interface action
+	 */
+	FToolBarButtonBlock( const TAttribute<FText>& InLabel, const TAttribute<FText>& InToolTip, const TAttribute<FSlateIcon>& InIcon, const FUIAction& InUIAction, const EUserInterfaceActionType::Type InUserInterfaceActionType );
+
+	void SetLabelVisibility( EVisibility InLabelVisibility ) { LabelVisibility = InLabelVisibility ; }
+
+	void SetIsFocusable( bool bInIsFocusable ) { bIsFocusable = bInIsFocusable; }
+
+	/** Set whether this toolbar should always use small icons, regardless of the current settings */
+	void SetForceSmallIcons( const bool InForceSmallIcons ) { bForceSmallIcons = InForceSmallIcons; }
+
+	/** FMultiBlock interface */
+	virtual void CreateMenuEntry(class FMenuBuilder& MenuBuilder) const OVERRIDE;
+
+private:
+
+	/**
+	 * Allocates a widget for this type of MultiBlock.  Override this in derived classes.
+	 *
+	 * @return  MultiBlock widget object
+	 */
+	virtual TSharedRef< class IMultiBlockBaseWidget > ConstructWidget() const;
+
+private:
+
+	// Friend our corresponding widget class
+	friend class SToolBarButtonBlock;
+
+	/** Optional overridden text label for this tool bar button.  If not set, then the action's label will be used instead. */
+	TAttribute<FText> LabelOverride;
+
+	/** Optional overridden tool tip for this tool bar button.  If not set, then the action's tool tip will be used instead. */
+	TAttribute<FText> ToolTipOverride;
+
+	/** Optional overridden icon for this tool bar button.  IF not set, then the action's icon will be used instead. */
+	TAttribute<FSlateIcon> IconOverride;
+
+	TOptional< EVisibility > LabelVisibility;
+	
+	/** In the case where a command is not bound, the user interface action type to use.  If a command is bound, we
+		simply use the action type associated with that command. */
+	EUserInterfaceActionType::Type UserInterfaceActionType;
+
+	/** Whether ToolBar will have Focusable buttons */
+	bool bIsFocusable;
+
+	/** Whether this toolbar should always use small icons, regardless of the current settings */
+	bool bForceSmallIcons;
+};
+
+
+
+/**
+ * Tool bar button MultiBlock widget
+ */
+class SLATE_API SToolBarButtonBlock
+	: public SMultiBlockBaseWidget
+{
+
+public:
+
+	SLATE_BEGIN_ARGS( SToolBarButtonBlock )
+		: _IsFocusable(false)
+		, _ForceSmallIcons(false)
+		, _TutorialHighlightName(NAME_None)
+		{}
+
+		SLATE_ARGUMENT( TOptional< EVisibility >, LabelVisibility )
+		SLATE_ARGUMENT( bool, IsFocusable )
+		SLATE_ARGUMENT( bool, ForceSmallIcons )
+		SLATE_ARGUMENT( FName, TutorialHighlightName )
+	SLATE_END_ARGS()
+
+
+	/**
+	 * Builds this MultiBlock widget up from the MultiBlock associated with it
+	 */
+	virtual void BuildMultiBlockWidget(const ISlateStyle* StyleSet, const FName& StyleName) OVERRIDE;
+
+
+	/**
+	 * Construct this widget
+	 *
+	 * @param	InArgs	The declaration data for this widget
+	 */
+	void Construct( const FArguments& InArgs );
+
+protected:
+
+	/**
+	 * Called by Slate when this tool bar button's button is clicked
+	 */
+	FReply OnClicked();
+
+
+	/**
+	 * Called by Slate when this tool bar check box button is toggled
+	 */
+	void OnCheckStateChanged( const ESlateCheckBoxState::Type NewCheckedState );
+
+	/**
+	 * Called by slate to determine if this button should appear checked
+	 *
+	 * @return ESlateCheckBoxState::Checked if it should be checked, ESlateCheckBoxState::Unchecked if not.
+	 */
+	ESlateCheckBoxState::Type OnIsChecked() const;
+
+	/**
+	 * Called by Slate to determine if this button is enabled
+	 * 
+	 * @return True if the menu entry is enabled, false otherwise
+	 */
+	bool IsEnabled() const;
+
+	/**
+	 * Called by Slate to determine if this button is visible
+	 *
+	 * @return EVisibility::Visible or EVisibility::Collapsed, depending on if the button should be displayed
+	 */
+	EVisibility GetBlockVisibility() const;
+
+private:
+
+	/** Called by Slate to determine whether icons/labels are visible */
+	EVisibility GetIconVisibility(bool bIsASmallIcon) const;
+	
+	/** Gets the icon brush for the toolbar block widget */
+	const FSlateBrush* GetIconBrush() const;
+	const FSlateBrush* GetSmallIconBrush() const;
+
+private:
+
+	TAttribute< EVisibility > LabelVisibility;
+
+	/** Whether ToolBar will have Focusable buttons. */
+	bool bIsFocusable;
+
+	/** Whether this toolbar should always use small icons, regardless of the current settings */
+	bool bForceSmallIcons;
+
+	/** Name to identify a widget for tutorials */
+	FName TutorialHighlightName;
+};
diff --git a/Engine/Source/Runtime/Slate/Public/Framework/Text/DefaultLayoutBlock.h b/Engine/Source/Runtime/Slate/Public/Framework/Text/DefaultLayoutBlock.h
new file mode 100644
--- /dev/null
+++ b/Engine/Source/Runtime/Slate/Public/Framework/Text/DefaultLayoutBlock.h
@@ -0,0 +1,58 @@
+// Copyright 1998-2014 Epic Games, Inc. All Rights Reserved.
+#pragma once
+
+class SLATE_API FDefaultLayoutBlock : public ILayoutBlock
+{
+public:
+
+	static TSharedRef< FDefaultLayoutBlock > Create( const TSharedRef< IRun >& InRun, const FTextRange& InRange, FVector2D InSize, const TSharedPtr< IRunHighlighter >& InHighlighter )
+	{
+		return MakeShareable( new FDefaultLayoutBlock( InRun, InRange, InSize, InHighlighter ) );
+	}
+
+	virtual ~FDefaultLayoutBlock() {}
+
+	virtual TSharedRef< IRun > GetRun() const OVERRIDE { return Run; }
+	virtual FTextRange GetTextRange() const OVERRIDE { return Range; }
+	virtual FVector2D GetSize() const OVERRIDE { return Size; }
+	virtual TSharedPtr< IRunHighlighter > GetHighlighter() const OVERRIDE { return Highlighter; }
+
+	virtual void SetLocationOffset( const FVector2D& InLocationOffset ) OVERRIDE { LocationOffset = InLocationOffset; }
+	virtual FVector2D GetLocationOffset() const OVERRIDE { return LocationOffset; }
+
+private:
+
+	static TSharedRef< FDefaultLayoutBlock > Create( const FDefaultLayoutBlock& Block )
+	{
+		return MakeShareable( new FDefaultLayoutBlock( Block ) );
+	}
+
+	FDefaultLayoutBlock( const TSharedRef< IRun >& InRun, const FTextRange& InRange, FVector2D InSize, const TSharedPtr< IRunHighlighter >& InHighlighter )
+		: Run( InRun )
+		, Range( InRange )
+		, Size( InSize )
+		, LocationOffset( ForceInitToZero )
+		, Highlighter( InHighlighter )
+	{
+
+	}
+
+	FDefaultLayoutBlock( const FDefaultLayoutBlock& Block )
+		: Run( Block.Run )
+		, Range( Block.Range )
+		, Size( Block.Size )
+		, LocationOffset( ForceInitToZero )
+		, Highlighter( Block.Highlighter )
+	{
+
+	}
+
+private:
+
+	TSharedRef< IRun > Run;
+
+	FTextRange Range;
+	FVector2D Size;
+	FVector2D LocationOffset;
+	TSharedPtr< IRunHighlighter > Highlighter;
+};
\ No newline at end of file
diff --git a/Engine/Source/Runtime/Slate/Public/Framework/Text/ILayoutBlock.h b/Engine/Source/Runtime/Slate/Public/Framework/Text/ILayoutBlock.h
new file mode 100644
--- /dev/null
+++ b/Engine/Source/Runtime/Slate/Public/Framework/Text/ILayoutBlock.h
@@ -0,0 +1,15 @@
+// Copyright 1998-2014 Epic Games, Inc. All Rights Reserved.
+#pragma once
+
+class SLATE_API ILayoutBlock
+{
+public:
+
+	virtual TSharedRef< class IRun > GetRun() const = 0;
+	virtual FTextRange GetTextRange() const = 0;
+	virtual FVector2D GetSize() const = 0;
+	virtual TSharedPtr< class IRunHighlighter > GetHighlighter() const = 0;
+
+	virtual void SetLocationOffset( const FVector2D& InLocationOffset ) = 0;
+	virtual FVector2D GetLocationOffset() const = 0;
+};
\ No newline at end of file
diff --git a/Engine/Source/Runtime/Slate/Public/Framework/Text/IRichTextMarkupParser.h b/Engine/Source/Runtime/Slate/Public/Framework/Text/IRichTextMarkupParser.h
new file mode 100644
--- /dev/null
+++ b/Engine/Source/Runtime/Slate/Public/Framework/Text/IRichTextMarkupParser.h
@@ -0,0 +1,20 @@
+// Copyright 1998-2014 Epic Games, Inc. All Rights Reserved.
+#pragma once
+
+/**
+ */
+class SLATE_API IRichTextMarkupParser
+{
+public:
+
+	/**
+	 * Virtual destructor
+	 */
+	virtual ~IRichTextMarkupParser() {}
+
+	/**
+	 * Processes the provided Input string producing a set of FTextLineParseResults and a output string stripped of any markup.
+	 */
+	virtual void Process(TArray<FTextLineParseResults>& Results, const FString& Input, FString& Output) = 0;
+
+};
diff --git a/Engine/Source/Runtime/Slate/Public/Framework/Text/IRun.h b/Engine/Source/Runtime/Slate/Public/Framework/Text/IRun.h
new file mode 100644
--- /dev/null
+++ b/Engine/Source/Runtime/Slate/Public/Framework/Text/IRun.h
@@ -0,0 +1,19 @@
+// Copyright 1998-2014 Epic Games, Inc. All Rights Reserved.
+#pragma once
+
+class SLATE_API IRun
+{
+public:
+
+	virtual FTextRange GetTextRange() const = 0;
+	virtual int16 GetBaseLine( float Scale ) const = 0;
+	virtual int16 GetMaxHeight( float Scale ) const = 0;
+
+	virtual FVector2D Measure( int32 StartIndex, int32 EndIndex, float Scale) const = 0;
+	virtual uint8 GetKerning( int32 CurrentIndex, float Scale ) const = 0;
+
+	virtual TSharedRef< class ILayoutBlock > CreateBlock( int32 StartIndex, int32 EndIndex, FVector2D Size, const TSharedPtr< class IRunHighlighter >& Highlighter ) = 0;
+
+	virtual void BeginLayout() = 0;
+	virtual void EndLayout() = 0;
+};
diff --git a/Engine/Source/Runtime/Slate/Public/Framework/Text/IRunHighlighter.h b/Engine/Source/Runtime/Slate/Public/Framework/Text/IRunHighlighter.h
new file mode 100644
--- /dev/null
+++ b/Engine/Source/Runtime/Slate/Public/Framework/Text/IRunHighlighter.h
@@ -0,0 +1,10 @@
+// Copyright 1998-2014 Epic Games, Inc. All Rights Reserved.
+#pragma once
+
+class SLATE_API IRunHighlighter
+{
+public:
+
+	virtual ~IRunHighlighter() {}
+
+};
\ No newline at end of file
diff --git a/Engine/Source/Runtime/Slate/Public/Framework/Text/ISlateRun.h b/Engine/Source/Runtime/Slate/Public/Framework/Text/ISlateRun.h
new file mode 100644
--- /dev/null
+++ b/Engine/Source/Runtime/Slate/Public/Framework/Text/ISlateRun.h
@@ -0,0 +1,15 @@
+// Copyright 1998-2014 Epic Games, Inc. All Rights Reserved.
+#pragma once
+
+#include "TextLayout.h"
+
+class SLATE_API ISlateRun : public IRun
+{
+public:
+
+	virtual int32 OnPaint( const FTextLayout::FLineView& Line, const TSharedRef< ILayoutBlock >& Block, const FTextBlockStyle& DefaultStyle, const FGeometry& AllottedGeometry, const FSlateRect& MyClippingRect, FSlateWindowElementList& OutDrawElements, int32 LayerId, const FWidgetStyle& InWidgetStyle, bool bParentEnabled ) const = 0;
+
+	virtual FChildren* GetChildren() = 0;
+
+	virtual void ArrangeChildren( const TSharedRef< ILayoutBlock >& Block, const FGeometry& AllottedGeometry, FArrangedChildren& ArrangedChildren ) const = 0;
+};
\ No newline at end of file
diff --git a/Engine/Source/Runtime/Slate/Public/Framework/Text/ISlateRunHighlighter.h b/Engine/Source/Runtime/Slate/Public/Framework/Text/ISlateRunHighlighter.h
new file mode 100644
--- /dev/null
+++ b/Engine/Source/Runtime/Slate/Public/Framework/Text/ISlateRunHighlighter.h
@@ -0,0 +1,17 @@
+// Copyright 1998-2014 Epic Games, Inc. All Rights Reserved.
+#pragma once
+
+class ISlateRun;
+class ILayoutBlock;
+
+class SLATE_API ISlateRunHighlighter : public IRunHighlighter
+{
+public:
+	virtual ~ISlateRunHighlighter() {}
+
+	virtual int32 OnPaint( const FTextLayout::FLineView& Line, const TSharedRef< ISlateRun >& Run, const TSharedRef< ILayoutBlock >& Block, const FTextBlockStyle& DefaultStyle, const FGeometry& AllottedGeometry, const FSlateRect& MyClippingRect, FSlateWindowElementList& OutDrawElements, int32 LayerId, const FWidgetStyle& InWidgetStyle, bool bParentEnabled ) const = 0;
+
+	virtual FChildren* GetChildren() = 0;
+
+	virtual void ArrangeChildren( const TSharedRef< ILayoutBlock >& Block, const FGeometry& AllottedGeometry, FArrangedChildren& ArrangedChildren ) const = 0;
+};
\ No newline at end of file
diff --git a/Engine/Source/Runtime/Slate/Public/Framework/Text/ITextDecorator.h b/Engine/Source/Runtime/Slate/Public/Framework/Text/ITextDecorator.h
new file mode 100644
--- /dev/null
+++ b/Engine/Source/Runtime/Slate/Public/Framework/Text/ITextDecorator.h
@@ -0,0 +1,73 @@
+// Copyright 1998-2014 Epic Games, Inc. All Rights Reserved.
+#pragma once
+
+struct SLATE_API FTextRunParseResults
+{
+	FTextRunParseResults( FString InName, const FTextRange& InOriginalRange)
+		: Name( InName )
+		, OriginalRange( InOriginalRange )
+		, MetaData()
+	{
+
+	}
+
+	FTextRunParseResults( FString InName, const FTextRange& InOriginalRange, const FTextRange& InContentRange)
+		: Name( InName )
+		, OriginalRange( InOriginalRange )
+		, ContentRange( InContentRange )
+		, MetaData()
+	{
+
+	}
+
+	FString Name;
+	FTextRange OriginalRange;
+	FTextRange ContentRange;
+	TMap< FString, FTextRange > MetaData; 
+};
+
+struct SLATE_API FTextLineParseResults
+{
+public:
+
+	FTextLineParseResults()
+		: Range( )
+		, Runs()
+	{
+
+	}
+
+	FTextLineParseResults(const FTextRange& InRange)
+		: Range( InRange )
+		, Runs()
+	{
+
+	}
+
+	FTextRange Range;
+	TArray< FTextRunParseResults > Runs;
+};
+
+struct SLATE_API FTextRunInfo
+{
+	FTextRunInfo( FString InName, const FText& InContent )
+		: Name( InName )
+		, Content( InContent )
+		, MetaData()
+	{
+
+	}
+
+	FString Name;
+	FText Content;
+	TMap< FString, FString > MetaData;
+};
+
+class SLATE_API ITextDecorator
+{
+public:
+
+	virtual bool Supports( const FTextRunParseResults& RunInfo, const FString& Text ) const = 0;
+
+	virtual TSharedRef< ISlateRun > Create( const FTextRunParseResults& RunInfo, const FString& OriginalText, const TSharedRef< FString >& ModelText, const ISlateStyle* Style ) = 0;
+};
\ No newline at end of file
diff --git a/Engine/Source/Runtime/Slate/Public/Framework/Text/SlateHyperlinkRun.h b/Engine/Source/Runtime/Slate/Public/Framework/Text/SlateHyperlinkRun.h
new file mode 100644
--- /dev/null
+++ b/Engine/Source/Runtime/Slate/Public/Framework/Text/SlateHyperlinkRun.h
@@ -0,0 +1,82 @@
+// Copyright 1998-2014 Epic Games, Inc. All Rights Reserved.
+#pragma once
+
+class SLATE_API FSlateHyperlinkRun : public ISlateRun, public TSharedFromThis< FSlateHyperlinkRun >
+{
+public:
+
+	typedef TMap< FString, FString > FMetadata;
+	DECLARE_DELEGATE_OneParam( FOnClick, const FMetadata& /*Metadata*/ )
+
+public:
+
+	class FWidgetViewModel
+	{
+	public:
+
+		bool IsPressed() const { return bIsPressed; }
+		bool IsHovered() const { return bIsHovered; }
+
+		void SetIsPressed( bool Value ) { bIsPressed = Value; }
+		void SetIsHovered( bool Value ) { bIsHovered = Value; }
+
+	private:
+
+		bool bIsPressed;
+		bool bIsHovered;
+	};
+
+public:
+
+	static TSharedRef< FSlateHyperlinkRun > Create( const TSharedRef< const FString >& InText, const FHyperlinkStyle& InStyle, const FMetadata& Metadata, FOnClick NavigateDelegate );
+																																	 
+	static TSharedRef< FSlateHyperlinkRun > Create( const TSharedRef< const FString >& InText, const FHyperlinkStyle& InStyle, const FMetadata& Metadata, FOnClick NavigateDelegate, const FTextRange& InRange );
+
+public:
+
+	virtual ~FSlateHyperlinkRun() {}
+
+	virtual FTextRange GetTextRange() const OVERRIDE;
+
+	virtual int16 GetBaseLine( float Scale ) const OVERRIDE;
+
+	virtual int16 GetMaxHeight( float Scale ) const OVERRIDE;
+
+	virtual FVector2D Measure( int32 StartIndex, int32 EndIndex, float Scale ) const OVERRIDE;
+
+	virtual uint8 GetKerning( int32 CurrentIndex, float Scale ) const OVERRIDE;
+
+	virtual TSharedRef< ILayoutBlock > CreateBlock( int32 StartIndex, int32 EndIndex, FVector2D Size, const TSharedPtr< IRunHighlighter >& Highlighter ) OVERRIDE;
+
+	virtual int32 OnPaint( const FTextLayout::FLineView& Line, const TSharedRef< ILayoutBlock >& Block, const FTextBlockStyle& DefaultStyle, const FGeometry& AllottedGeometry, const FSlateRect& MyClippingRect, FSlateWindowElementList& OutDrawElements, int32 LayerId, const FWidgetStyle& InWidgetStyle, bool bParentEnabled ) const OVERRIDE;
+
+	virtual FChildren* GetChildren() OVERRIDE;
+
+	virtual void ArrangeChildren( const TSharedRef< ILayoutBlock >& Block, const FGeometry& AllottedGeometry, FArrangedChildren& ArrangedChildren ) const OVERRIDE;
+
+	virtual void BeginLayout() OVERRIDE { Children.Empty(); }
+	virtual void EndLayout() OVERRIDE {}
+
+private:
+
+	FSlateHyperlinkRun( const TSharedRef< const FString >& InText, const FHyperlinkStyle& InStyle, const FMetadata& InMetadata, FOnClick InNavigateDelegate );
+																										 
+	FSlateHyperlinkRun( const TSharedRef< const FString >& InText, const FHyperlinkStyle& InStyle, const FMetadata& InMetadata, FOnClick InNavigateDelegate, const FTextRange& InRange );
+
+	FSlateHyperlinkRun( const FSlateHyperlinkRun& Run );
+
+private:
+
+	void OnNavigate();
+
+private:
+
+	TSharedRef< const FString > Text;
+	FTextRange Range;
+	FHyperlinkStyle Style;
+	TMap<FString,FString> Metadata;
+	FOnClick NavigateDelegate;
+
+	TSharedRef< FWidgetViewModel > ViewModel;
+	TSlotlessChildren< SWidget > Children;
+};
\ No newline at end of file
diff --git a/Engine/Source/Runtime/Slate/Public/Framework/Text/SlateImageRun.h b/Engine/Source/Runtime/Slate/Public/Framework/Text/SlateImageRun.h
new file mode 100644
--- /dev/null
+++ b/Engine/Source/Runtime/Slate/Public/Framework/Text/SlateImageRun.h
@@ -0,0 +1,49 @@
+// Copyright 1998-2014 Epic Games, Inc. All Rights Reserved.
+#pragma once
+
+class SLATE_API FSlateImageRun : public ISlateRun, public TSharedFromThis< FSlateImageRun >
+{
+public:
+
+	static TSharedRef< FSlateImageRun > Create( const TSharedRef< const FString >& InText, const FSlateBrush* InImage, int16 InBaseline );
+	static TSharedRef< FSlateImageRun > Create( const TSharedRef< const FString >& InText, const FSlateBrush* InImage, int16 InBaseline, const FTextRange& InRange );
+
+public:
+
+	virtual ~FSlateImageRun() {}
+
+	virtual FTextRange GetTextRange() const OVERRIDE;
+	virtual int16 GetBaseLine( float Scale ) const OVERRIDE;
+	virtual int16 GetMaxHeight( float Scale ) const OVERRIDE;
+
+	virtual FVector2D Measure( int32 StartIndex, int32 EndIndex, float Scale ) const OVERRIDE;
+
+	virtual uint8 GetKerning( int32 CurrentIndex, float Scale ) const OVERRIDE;
+
+
+	virtual TSharedRef< ILayoutBlock > CreateBlock( int32 StartIndex, int32 EndIndex, FVector2D Size, const TSharedPtr< IRunHighlighter >& Highlighter ) OVERRIDE;
+
+	virtual int32 OnPaint( const FTextLayout::FLineView& Line, const TSharedRef< ILayoutBlock >& Block, const FTextBlockStyle& DefaultStyle, const FGeometry& AllottedGeometry, const FSlateRect& MyClippingRect, FSlateWindowElementList& OutDrawElements, int32 LayerId, const FWidgetStyle& InWidgetStyle, bool bParentEnabled ) const OVERRIDE;
+
+	virtual FChildren* GetChildren() OVERRIDE;
+
+	virtual void ArrangeChildren( const TSharedRef< ILayoutBlock >& Block, const FGeometry& AllottedGeometry, FArrangedChildren& ArrangedChildren ) const OVERRIDE;
+
+	virtual void BeginLayout() OVERRIDE {}
+	virtual void EndLayout() OVERRIDE {}
+
+private:
+
+	FSlateImageRun( const TSharedRef< const FString >& InText, const FSlateBrush* InImage, int16 InBaseline );
+
+	FSlateImageRun( const TSharedRef< const FString >& InText, const FSlateBrush* InImage, int16 InBaseline, const FTextRange& InRange );
+
+private:
+
+	TSharedRef< const FString > Text;
+	FTextRange Range;
+	const FSlateBrush* Image;
+	int16 Baseline;
+
+	static FNoChildren NoChildrenInstance;
+};
\ No newline at end of file
diff --git a/Engine/Source/Runtime/Slate/Public/Framework/Text/SlateSimpleRunHighlighter.h b/Engine/Source/Runtime/Slate/Public/Framework/Text/SlateSimpleRunHighlighter.h
new file mode 100644
--- /dev/null
+++ b/Engine/Source/Runtime/Slate/Public/Framework/Text/SlateSimpleRunHighlighter.h
@@ -0,0 +1,25 @@
+// Copyright 1998-2014 Epic Games, Inc. All Rights Reserved.
+#pragma once
+
+class SLATE_API FSlateSimpleRunHighlighter : public ISlateRunHighlighter
+{
+public:
+
+	static TSharedRef< FSlateSimpleRunHighlighter > Create();
+
+	virtual ~FSlateSimpleRunHighlighter() {}
+
+	virtual int32 OnPaint( const FTextLayout::FLineView& Line, const TSharedRef< ISlateRun >& Run, const TSharedRef< ILayoutBlock >& Block, const FTextBlockStyle& DefaultStyle, const FGeometry& AllottedGeometry, const FSlateRect& MyClippingRect, FSlateWindowElementList& OutDrawElements, int32 LayerId, const FWidgetStyle& InWidgetStyle, bool bParentEnabled ) const OVERRIDE;
+
+	virtual FChildren* GetChildren() OVERRIDE;
+
+	virtual void ArrangeChildren( const TSharedRef< ILayoutBlock >& Block, const FGeometry& AllottedGeometry, FArrangedChildren& ArrangedChildren ) const OVERRIDE;
+
+private:
+
+	FSlateSimpleRunHighlighter();
+
+private:
+
+	static FNoChildren NoChildrenInstance;
+};
\ No newline at end of file
diff --git a/Engine/Source/Runtime/Slate/Public/Framework/Text/SlateTextLayout.h b/Engine/Source/Runtime/Slate/Public/Framework/Text/SlateTextLayout.h
new file mode 100644
--- /dev/null
+++ b/Engine/Source/Runtime/Slate/Public/Framework/Text/SlateTextLayout.h
@@ -0,0 +1,37 @@
+// Copyright 1998-2014 Epic Games, Inc. All Rights Reserved.
+#pragma once
+
+class SLATE_API FSlateTextLayout : public FTextLayout
+{
+public:
+
+	static TSharedRef< FSlateTextLayout > Create();
+
+	FChildren* GetChildren();
+
+	void ArrangeChildren( const FGeometry& AllottedGeometry, FArrangedChildren& ArrangedChildren ) const;
+
+	int32 OnPaint( const FTextBlockStyle& DefaultTextStyle, const FGeometry& AllottedGeometry, const FSlateRect& MyClippingRect, FSlateWindowElementList& OutDrawElements, int32 LayerId, const FWidgetStyle& InWidgetStyle, bool bParentEnabled ) const;
+
+	//void AddLine( const TSharedRef< FString >& Text, const FTextBlockStyle& TextStyle )
+	//{
+	//	FLineModel LineModel( Text );
+	//	LineModel.Runs.Add( FLineModel::FRun( FSlateTextRun::Create( LineModel.Text, TextStyle ) ) );
+	//	LineModels.Add( LineModel );
+	//}
+
+	virtual void EndLayout() OVERRIDE;
+
+
+protected:
+
+	FSlateTextLayout();
+
+	void AggregateChildren();
+
+private:
+
+	TSlotlessChildren< SWidget > Children;
+
+	friend class FSlateTextLayoutFactory;
+};
\ No newline at end of file
diff --git a/Engine/Source/Runtime/Slate/Public/Framework/Text/SlateTextRun.h b/Engine/Source/Runtime/Slate/Public/Framework/Text/SlateTextRun.h
new file mode 100644
--- /dev/null
+++ b/Engine/Source/Runtime/Slate/Public/Framework/Text/SlateTextRun.h
@@ -0,0 +1,52 @@
+// Copyright 1998-2014 Epic Games, Inc. All Rights Reserved.
+#pragma once
+
+class SLATE_API FSlateTextRun : public ISlateRun, public TSharedFromThis< FSlateTextRun >
+{
+public:
+
+	static TSharedRef< FSlateTextRun > Create( const TSharedRef< const FString >& InText, const FTextBlockStyle& Style );
+
+	static TSharedRef< FSlateTextRun > Create( const TSharedRef< const FString >& InText, const FTextBlockStyle& Style, const FTextRange& InRange );
+
+public:
+
+	virtual ~FSlateTextRun() {}
+
+	virtual FTextRange GetTextRange() const OVERRIDE;
+	virtual int16 GetBaseLine( float Scale ) const OVERRIDE;
+	virtual int16 GetMaxHeight( float Scale ) const OVERRIDE;
+	virtual FVector2D Measure( int32 StartIndex, int32 EndIndex, float Scale ) const OVERRIDE;
+	virtual uint8 GetKerning( int32 CurrentIndex, float Scale ) const OVERRIDE;
+
+	virtual TSharedRef< ILayoutBlock > CreateBlock( int32 StartIndex, int32 EndIndex, FVector2D Size, const TSharedPtr< IRunHighlighter >& Highlighter ) OVERRIDE;
+
+	virtual int32 OnPaint( const FTextLayout::FLineView& Line, const TSharedRef< ILayoutBlock >& Block, const FTextBlockStyle& DefaultStyle, const FGeometry& AllottedGeometry, const FSlateRect& MyClippingRect, FSlateWindowElementList& OutDrawElements, int32 LayerId, const FWidgetStyle& InWidgetStyle, bool bParentEnabled ) const OVERRIDE;
+
+	virtual FChildren* GetChildren() OVERRIDE;
+
+	virtual void ArrangeChildren( const TSharedRef< ILayoutBlock >& Block, const FGeometry& AllottedGeometry, FArrangedChildren& ArrangedChildren ) const OVERRIDE;
+
+	virtual void BeginLayout() OVERRIDE {}
+	virtual void EndLayout() OVERRIDE {}
+
+private:
+
+	FSlateTextRun( const TSharedRef< const FString >& InText, const FTextBlockStyle& InStyle );
+
+	FSlateTextRun( const TSharedRef< const FString >& InText, const FTextBlockStyle& InStyle, const FTextRange& InRange );
+
+	FSlateTextRun( const FSlateTextRun& Run );
+
+private:
+
+	TSharedRef< const FString > Text;
+	FTextBlockStyle Style;
+	FTextRange Range;
+
+#if TEXT_LAYOUT_DEBUG
+	FString DebugSlice;
+#endif
+
+	static FNoChildren NoChildrenInstance;
+};
\ No newline at end of file
diff --git a/Engine/Source/Runtime/Slate/Public/Framework/Text/SlateWidgetRun.h b/Engine/Source/Runtime/Slate/Public/Framework/Text/SlateWidgetRun.h
new file mode 100644
--- /dev/null
+++ b/Engine/Source/Runtime/Slate/Public/Framework/Text/SlateWidgetRun.h
@@ -0,0 +1,69 @@
+// Copyright 1998-2014 Epic Games, Inc. All Rights Reserved.
+#pragma once
+
+class SLATE_API FSlateWidgetRun : public ISlateRun, public TSharedFromThis< FSlateWidgetRun >
+{
+public:
+
+	struct FWidgetRunInfo
+	{
+		TSharedRef< SWidget > Widget;
+		int16 Baseline;
+		TOptional< FVector2D > Size;
+
+		FWidgetRunInfo( const TSharedRef< SWidget >& InWidget, int16 InBaseline, TOptional< FVector2D > InSize = TOptional< FVector2D >() )
+			: Widget( InWidget )
+			, Baseline( InBaseline )
+			, Size( InSize )
+		{
+
+		}
+
+		FWidgetRunInfo( const FWidgetRunInfo& Other )
+			: Widget( Other.Widget )
+			, Baseline( Other.Baseline )
+			, Size( Other.Size )
+		{
+
+		}
+	};
+
+	static TSharedRef< FSlateWidgetRun > Create( const TSharedRef< const FString >& InText, const FWidgetRunInfo& InWidgetInfo );
+
+	static TSharedRef< FSlateWidgetRun > Create( const TSharedRef< const FString >& InText, const FWidgetRunInfo& InWidgetInfo, const FTextRange& InRange );
+
+	virtual ~FSlateWidgetRun() {}
+
+	virtual FTextRange GetTextRange() const OVERRIDE;
+	virtual int16 GetBaseLine( float Scale ) const OVERRIDE;
+	virtual int16 GetMaxHeight( float Scale ) const OVERRIDE;
+
+	virtual FVector2D Measure( int32 StartIndex, int32 EndIndex, float Scale ) const OVERRIDE;
+	virtual uint8 GetKerning( int32 CurrentIndex, float Scale ) const OVERRIDE;
+
+	virtual TSharedRef< ILayoutBlock > CreateBlock( int32 StartIndex, int32 EndIndex, FVector2D Size, const TSharedPtr< IRunHighlighter >& Highlighter ) OVERRIDE;
+
+	virtual int32 OnPaint( const FTextLayout::FLineView& Line, const TSharedRef< ILayoutBlock >& Block, const FTextBlockStyle& DefaultStyle, const FGeometry& AllottedGeometry, const FSlateRect& MyClippingRect, FSlateWindowElementList& OutDrawElements, int32 LayerId, const FWidgetStyle& InWidgetStyle, bool bParentEnabled ) const OVERRIDE;
+
+	virtual FChildren* GetChildren() OVERRIDE;
+
+	virtual void ArrangeChildren( const TSharedRef< ILayoutBlock >& Block, const FGeometry& AllottedGeometry, FArrangedChildren& ArrangedChildren ) const;
+
+	virtual void BeginLayout() OVERRIDE {}
+	virtual void EndLayout() OVERRIDE {}
+
+private:
+
+	FSlateWidgetRun( const TSharedRef< const FString >& InText, const FWidgetRunInfo& InWidgetInfo );
+
+	FSlateWidgetRun( const TSharedRef< const FString >& InText, const FWidgetRunInfo& InWidgetInfo, const FTextRange& InRange );
+
+	FSlateWidgetRun( const FSlateWidgetRun& Run );
+
+private:
+
+	TSharedRef< const FString > Text;
+	FTextRange Range;
+	FWidgetRunInfo Info;
+	TSlotlessChildren< SWidget > Children;
+};
\ No newline at end of file
diff --git a/Engine/Source/Runtime/Slate/Public/Framework/Text/TextDecorators.h b/Engine/Source/Runtime/Slate/Public/Framework/Text/TextDecorators.h
new file mode 100644
--- /dev/null
+++ b/Engine/Source/Runtime/Slate/Public/Framework/Text/TextDecorators.h
@@ -0,0 +1,77 @@
+// Copyright 1998-2014 Epic Games, Inc. All Rights Reserved.
+#pragma once
+
+#include "ITextDecorator.h"
+
+class SLATE_API FWidgetDecorator : public ITextDecorator
+{
+public:
+
+	DECLARE_DELEGATE_RetVal_TwoParams( FSlateWidgetRun::FWidgetRunInfo, FCreateWidget, const FTextRunInfo& /*RunInfo*/, const ISlateStyle* /*Style*/ )
+
+public:
+
+	static TSharedRef< FWidgetDecorator > Create( FString InRunName, const FCreateWidget& InCreateWidgetDelegate );
+	virtual ~FWidgetDecorator() {}
+
+public:
+
+	virtual bool Supports( const FTextRunParseResults& RunInfo, const FString& Text ) const OVERRIDE;
+
+	virtual TSharedRef< ISlateRun > Create( const FTextRunParseResults& RunInfo, const FString& OriginalText, const TSharedRef< FString >& InOutModelText, const ISlateStyle* Style ) OVERRIDE;
+
+private:
+
+	FWidgetDecorator( FString InRunName, const FCreateWidget& InCreateWidgetDelegate );
+
+private:
+
+	FString RunName;
+	FCreateWidget CreateWidgetDelegate;
+};
+
+class SLATE_API FImageDecorator : public ITextDecorator
+{
+public:
+
+	static TSharedRef< FImageDecorator > Create( FString InRunName, const ISlateStyle* const InOverrideStyle = NULL );
+	virtual ~FImageDecorator() {}
+
+public:
+
+	virtual bool Supports( const FTextRunParseResults& RunInfo, const FString& Text ) const OVERRIDE;
+
+	virtual TSharedRef< ISlateRun > Create( const FTextRunParseResults& RunInfo, const FString& OriginalText, const TSharedRef< FString >& InOutModelText, const ISlateStyle* Style ) OVERRIDE;
+
+private:
+
+	FImageDecorator( FString InRunName, const ISlateStyle* const InOverrideStyle );
+
+private:
+
+	FString RunName;
+	const ISlateStyle* OverrideStyle;
+};
+
+class SLATE_API FHyperlinkDecorator : public ITextDecorator
+{
+public:
+
+	static TSharedRef< FHyperlinkDecorator > Create( FString Id, const FSlateHyperlinkRun::FOnClick& NavigateDelegate );
+	virtual ~FHyperlinkDecorator() {}
+
+public:
+
+	virtual bool Supports( const FTextRunParseResults& RunInfo, const FString& Text ) const OVERRIDE;
+
+	virtual TSharedRef< ISlateRun > Create( const FTextRunParseResults& RunInfo, const FString& OriginalText, const TSharedRef< FString >& InOutModelText, const ISlateStyle* Style ) OVERRIDE;
+
+private:
+
+	FHyperlinkDecorator( FString InId, const FSlateHyperlinkRun::FOnClick& InNavigateDelegate );
+
+private:
+
+	FString Id;
+	FSlateHyperlinkRun::FOnClick NavigateDelegate;
+};
diff --git a/Engine/Source/Runtime/Slate/Public/Framework/Text/TextHighlight.h b/Engine/Source/Runtime/Slate/Public/Framework/Text/TextHighlight.h
new file mode 100644
--- /dev/null
+++ b/Engine/Source/Runtime/Slate/Public/Framework/Text/TextHighlight.h
@@ -0,0 +1,17 @@
+// Copyright 1998-2014 Epic Games, Inc. All Rights Reserved.
+#pragma once
+
+struct SLATE_API FTextHighlight
+{
+	FTextHighlight( int32 InLineIndex, const FTextRange& InRange, const TSharedRef< class IRunHighlighter >& InHighlighter )
+		: LineIndex( InLineIndex )
+		, Range( InRange )
+		, Highlighter( InHighlighter )
+	{
+
+	}
+
+	int32 LineIndex;
+	FTextRange Range;
+	TSharedRef< IRunHighlighter > Highlighter;
+};
\ No newline at end of file
diff --git a/Engine/Source/Runtime/Slate/Public/Framework/Text/TextLayout.h b/Engine/Source/Runtime/Slate/Public/Framework/Text/TextLayout.h
new file mode 100644
--- /dev/null
+++ b/Engine/Source/Runtime/Slate/Public/Framework/Text/TextLayout.h
@@ -0,0 +1,222 @@
+// Copyright 1998-2014 Epic Games, Inc. All Rights Reserved.
+#pragma once
+
+#include "LineBreakIterator.h"
+
+#define TEXT_LAYOUT_DEBUG 0
+
+namespace ETextJustify
+{
+	enum Type
+	{
+		Left,
+		Center,
+		Right
+	};
+}
+
+class SLATE_API FTextLayout
+{
+public:
+
+	struct FBlockDefinition
+	{
+		FTextRange Range;
+		FTextRange VisualRange;
+		TSharedPtr< IRunHighlighter > Highlighter;
+	};
+
+	struct FBreakCandidate
+	{
+		FTextRange Range; 
+		FTextRange RangeWithoutTrailingWhitespace; 
+		FVector2D Size;
+		FVector2D SizeWithoutTrailingWhitespace;
+
+		int16 MaxAboveBaseline;
+		int16 MaxBelowBaseline;
+
+		uint8 Kerning;
+
+#if TEXT_LAYOUT_DEBUG
+		FString DebugSlice;
+#endif
+	};
+
+	class FRunModel
+	{
+	public:
+
+		FRunModel( const TSharedRef< IRun >& InRun);
+
+	public:
+
+		TSharedRef< IRun > GetRun() const;;
+
+		void BeginLayout();
+		void EndLayout();
+
+		FTextRange GetTextRange() const;
+		int16 GetBaseLine( float Scale ) const;
+		int16 GetMaxHeight( float Scale ) const;
+
+		FVector2D Measure( int32 BeginIndex, int32 EndIndex, float Scale );
+
+		uint8 GetKerning( int32 CurrentIndex, float Scale );
+
+		static int BinarySearchForBeginIndex( const TArray< FTextRange >& Ranges, int32 BeginIndex );
+
+		static int BinarySearchForEndIndex( const TArray< FTextRange >& Ranges, int32 RangeBeginIndex, int32 EndIndex );
+
+		TSharedRef< ILayoutBlock > CreateBlock( const FBlockDefinition& BlockDefine, float Scale ) const;
+
+		void ClearCache();
+
+	private:
+
+		TSharedRef< IRun > Run;
+		TArray< FTextRange > MeasuredRanges;
+		TArray< FVector2D > MeasuredRangeSizes;
+	};
+
+	struct FLineModel
+	{
+	public:
+
+		FLineModel( const TSharedRef< FString >& InText );
+
+		TSharedRef< FString > Text;
+		TArray< FRunModel > Runs;
+		TArray< FBreakCandidate > BreakCandidates;
+		TArray< FTextHighlight > Highlights;
+	};
+
+	struct FLineView
+	{
+		TArray< TSharedRef< ILayoutBlock > > Blocks;
+		FVector2D Size;
+		FVector2D Offset;
+	};
+
+
+public:
+
+	virtual ~FTextLayout();
+
+	const TArray< FTextLayout::FLineView >& GetLineViews() const;
+	const TArray< FTextLayout::FLineModel >& GetLineModels() const;
+
+	FVector2D GetSize() const;
+	FVector2D GetDrawSize() const;
+
+	float GetWrappingWidth() const;
+	void SetWrappingWidth( float Value );
+
+	float GetLineHeightPercentage() const;
+	void SetLineHeightPercentage( float Value );
+
+	ETextJustify::Type GetJustification() const;
+	void SetJustification( ETextJustify::Type Value );
+
+	float GetScale() const;
+	void SetScale( float Value );
+
+	FMargin GetMargin() const;
+	void SetMargin( const FMargin& InMargin );
+
+	void ClearLines();
+
+	void AddLine( const TSharedRef< FString >& Text, const TArray< TSharedRef< IRun > >& Runs );
+
+	/**
+	* Clears all highlights
+	*/
+	void ClearHighlights();
+
+	/**
+	* Adds a single highlight to the existing set of highlights.
+	*/
+	void AddHighlight( const FTextHighlight& Highlight );
+
+	/**
+	* Replaces the current set of highlights with the provided highlights.
+	*/
+	void SetHighlights( const TArray< FTextHighlight >& Highlights );
+
+	/**
+	* Updates the TextLayout's if any changes have occurred since the last update.
+	*/
+	virtual void UpdateIfNeeded();
+
+protected:
+
+	FTextLayout();
+
+	/**
+	* Prepares the cache with the information needed to perform line wrapping if the information does not already exist.
+	*/
+	void ReadyForWrapping();
+
+	/**
+	* Clears the current layouts view information.
+	*/
+	void ClearView();
+
+	/**
+	* Notifies all Runs that we are beginning to generate a new layout.
+	*/
+	virtual void BeginLayout();
+
+	/**
+	* Notifies all Runs that the layout has finished generating.
+	*/
+	virtual void EndLayout();
+
+	/**
+	* Creates a new View for the layout, replacing the current one.
+	*/
+	void CreateView();
+
+
+private:
+
+	void FlowLayout();
+
+	void JustifyLayout();
+
+	void CreateLineViewBlocks( const FLineModel& Line, const int32 StopIndex, const int32 VisualStopIndex, int32& OutRunIndex, int32& OutHighlightIndex, int32& OutPreviousBlockEnd, TArray< TSharedRef< ILayoutBlock > >& OutSoftLine );
+
+	FBreakCandidate CreateBreakCandidate( int32& OutRunIndex, FLineModel& Line, int32 PreviousBreak, int32 CurrentBreak );
+
+protected:
+
+	/** The models for the lines of text. A LineModel represents a single string with no manual breaks. */
+	TArray< FLineModel > LineModels;
+
+	/** The views for the lines of text. A LineView represents a single visual line of text. Multiple LineViews can map to the same LineModel, if for example wrapping occurs. */
+	TArray< FLineView > LineViews;
+
+	/** Whether or not we currently have the information cached to perform line wrapping. */
+	bool HasWrapInformation;
+
+	/** Whether parameters on the layout have changed which requires the view be updated. */
+	bool HasChanged;
+
+	/** The scale to draw the text at */
+	float Scale;
+
+	/** The width that the text should be wrap at. If 0 or negative no wrapping occurs. */
+	float WrappingWidth;
+
+	/** The size of the margins to put about the text. This is an unscaled value. */
+	FMargin Margin;
+
+	/** How the text should be aligned with the margin. */
+	ETextJustify::Type Justification;
+
+	/** The percentage to modify a line height by. */
+	float LineHeightPercentage;
+
+	/** The final size of the text layout on screen. */
+	FVector2D DrawSize;
+};
\ No newline at end of file
diff --git a/Engine/Source/Runtime/Slate/Public/Framework/Text/TextLayoutFactory.h b/Engine/Source/Runtime/Slate/Public/Framework/Text/TextLayoutFactory.h
new file mode 100644
--- /dev/null
+++ b/Engine/Source/Runtime/Slate/Public/Framework/Text/TextLayoutFactory.h
@@ -0,0 +1,2 @@
+// Copyright 1998-2014 Epic Games, Inc. All Rights Reserved.
+
diff --git a/Engine/Source/Runtime/Slate/Public/Framework/Text/TextLine.h b/Engine/Source/Runtime/Slate/Public/Framework/Text/TextLine.h
new file mode 100644
--- /dev/null
+++ b/Engine/Source/Runtime/Slate/Public/Framework/Text/TextLine.h
@@ -0,0 +1,2 @@
+// Copyright 1998-2014 Epic Games, Inc. All Rights Reserved.
+
diff --git a/Engine/Source/Runtime/Slate/Public/Framework/Text/TextRange.h b/Engine/Source/Runtime/Slate/Public/Framework/Text/TextRange.h
new file mode 100644
--- /dev/null
+++ b/Engine/Source/Runtime/Slate/Public/Framework/Text/TextRange.h
@@ -0,0 +1,25 @@
+// Copyright 1998-2014 Epic Games, Inc. All Rights Reserved.
+#pragma once
+
+struct SLATE_API FTextRange
+{
+	FTextRange()
+		: BeginIndex(INDEX_NONE)
+		, EndIndex(INDEX_NONE)
+	{
+
+	}
+
+	FTextRange( int32 InBeginIndex, int32 InEndIndex )
+		: BeginIndex( InBeginIndex )
+		, EndIndex( InEndIndex )
+	{
+
+	}
+
+	int32 Len() const { return EndIndex - BeginIndex; }
+	bool IsEmpty() const { return (EndIndex - BeginIndex) <= 0; }
+
+	int32 BeginIndex;
+	int32 EndIndex;
+};
\ No newline at end of file
diff --git a/Engine/Source/Runtime/Slate/Public/Framework/Text/WidgetLayoutBlock.h b/Engine/Source/Runtime/Slate/Public/Framework/Text/WidgetLayoutBlock.h
new file mode 100644
--- /dev/null
+++ b/Engine/Source/Runtime/Slate/Public/Framework/Text/WidgetLayoutBlock.h
@@ -0,0 +1,66 @@
+// Copyright 1998-2014 Epic Games, Inc. All Rights Reserved.
+#pragma once
+
+class SLATE_API FWidgetLayoutBlock : public ILayoutBlock
+{
+public:
+
+	static TSharedRef< FWidgetLayoutBlock > Create( const TSharedRef< IRun >& InRun, const TSharedRef< SWidget >& InWidget, const FTextRange& InRange, FVector2D InSize, const TSharedPtr< IRunHighlighter >& InHighlighter )
+	{
+		return MakeShareable( new FWidgetLayoutBlock( InRun, InWidget, InRange, InSize, InHighlighter ) );
+	}
+
+	virtual ~FWidgetLayoutBlock() {}
+
+	virtual TSharedRef< IRun > GetRun() const OVERRIDE { return Run; }
+
+	virtual FTextRange GetTextRange() const OVERRIDE { return Range; }
+
+	virtual FVector2D GetSize() const OVERRIDE { return Size; }
+
+	virtual TSharedPtr< IRunHighlighter > GetHighlighter() const OVERRIDE { return Highlighter; }
+
+	TSharedRef< SWidget > GetWidget() const { return Widget; }
+
+	virtual void SetLocationOffset( const FVector2D& InLocationOffset ) OVERRIDE { LocationOffset = InLocationOffset; }
+	virtual FVector2D GetLocationOffset() const OVERRIDE { return LocationOffset; }
+
+private:
+
+	static TSharedRef< FWidgetLayoutBlock > Create( const FWidgetLayoutBlock& Block )
+	{
+		return MakeShareable( new FWidgetLayoutBlock( Block ) );
+	}
+
+	FWidgetLayoutBlock( const TSharedRef< IRun >& InRun, const TSharedRef< SWidget >& InWidget, const FTextRange& InRange, FVector2D InSize, const TSharedPtr< IRunHighlighter >& InHighlighter )
+		: Run( InRun )
+		, Widget( InWidget )
+		, Range( InRange )
+		, Size( InSize )
+		, LocationOffset( ForceInitToZero )
+		, Highlighter( InHighlighter )
+	{
+
+	}
+
+	FWidgetLayoutBlock( const FWidgetLayoutBlock& Block )
+		: Run( Block.Run )
+		, Widget( Block.Widget )
+		, Range( Block.Range )
+		, Size( Block.Size )
+		, LocationOffset( ForceInitToZero )
+		, Highlighter( Block.Highlighter )
+	{
+
+	}
+
+private:
+
+	TSharedRef< IRun > Run;
+	TSharedRef< SWidget > Widget;
+
+	FTextRange Range;
+	FVector2D Size;
+	FVector2D LocationOffset;
+	TSharedPtr< IRunHighlighter > Highlighter;
+};
\ No newline at end of file
diff --git a/Engine/Source/Runtime/Slate/Public/Slate.h b/Engine/Source/Runtime/Slate/Public/Slate.h
--- a/Engine/Source/Runtime/Slate/Public/Slate.h
+++ b/Engine/Source/Runtime/Slate/Public/Slate.h
@@ -93,6 +93,13 @@ DECLARE_CYCLE_STAT_EXTERN( TEXT("Slate Misc Time"), STAT_SlateMiscTime, STATGROU
 #include "SHeader.h"
 #include "SGridPanel.h"
 #include "SUniformGridPanel.h"
+#include "SMenuAnchor.h"
+#include "MultiBoxDefs.h"
+////
+#include "MultiBox.h"
+////
+#include "MultiBoxBuilder.h"
+#include "MultiBoxExtender.h"
 #include "SMultiLineEditableText.h"
 #include "SEditableText.h"
 #include "SEditableTextBox.h"
@@ -100,7 +107,7 @@ DECLARE_CYCLE_STAT_EXTERN( TEXT("Slate Misc Time"), STAT_SlateMiscTime, STATGROU
 #include "SButton.h"
 #include "SToolTip.h"
 #include "SScrollBar.h"
-#include "SMenuAnchor.h"
+#include "SScrollBorder.h"
 #include "SErrorText.h"
 #include "SErrorHint.h"
 #include "SPopUpErrorText.h"
@@ -125,8 +132,10 @@ DECLARE_CYCLE_STAT_EXTERN( TEXT("Slate Misc Time"), STAT_SlateMiscTime, STATGROU
 #include "SColorWheel.h"
 #include "SColorThemes.h"
 #include "SColorPicker.h"
-#include "MultiBoxBuilder.h"
-#include "MultiBoxExtender.h"
+///
+#include "SToolBarButtonBlock.h"
+#include "SToolBarComboButtonBlock.h"
+///
 #include "SComboButton.h"
 #include "SComboBox.h"
 #include "SHyperlink.h"
@@ -139,6 +148,7 @@ DECLARE_CYCLE_STAT_EXTERN( TEXT("Slate Misc Time"), STAT_SlateMiscTime, STATGROU
 #include "SNotificationList.h"
 #include "SWidgetSwitcher.h"
 #include "SSuggestionTextBox.h"
+#include "SBreadcrumbTrail.h"
 #include "SWizard.h"
 #include "STextComboBox.h"
 #include "SNumericEntryBox.h"
@@ -151,6 +161,9 @@ DECLARE_CYCLE_STAT_EXTERN( TEXT("Slate Misc Time"), STAT_SlateMiscTime, STATGROU
 #include "ScrollyZoomy.h"
 #include "SSafeZone.h"
 #include "MarqueeRect.h"
+#include "SRotatorInputBox.h"
+#include "SVectorInputBox.h"
+#include "SVirtualJoystick.h"
 
 // Docking Framework
 #include "WorkspaceItem.h"
diff --git a/Engine/Source/Runtime/Slate/Public/Widgets/Input/SEditableText.h b/Engine/Source/Runtime/Slate/Public/Widgets/Input/SEditableText.h
--- a/Engine/Source/Runtime/Slate/Public/Widgets/Input/SEditableText.h
+++ b/Engine/Source/Runtime/Slate/Public/Widgets/Input/SEditableText.h
@@ -2,10 +2,10 @@
 
 #pragma once
 
-#include "MultiBox/MultiBoxExtender.h"
 #include "ITextEditorWidget.h"
 #include "ITextInputMethodSystem.h"
 
+
 /**
  * Editable text widget
  */
diff --git a/Engine/Source/Runtime/Slate/Public/Widgets/SCanvas.h b/Engine/Source/Runtime/Slate/Public/Widgets/SCanvas.h
--- a/Engine/Source/Runtime/Slate/Public/Widgets/SCanvas.h
+++ b/Engine/Source/Runtime/Slate/Public/Widgets/SCanvas.h
@@ -1,12 +1,17 @@
 // Copyright 1998-2014 Epic Games, Inc. All Rights Reserved.
 
+/*=============================================================================
+	SCanvas.h: Declares the SCanvas class.
+=============================================================================*/
+
 #pragma once
 
 
 /**
  * Canvas is a layout widget that allows you to arbitrary position and size child widgets in a relative coordinate space
  */
-class SLATE_API SCanvas : public SPanel
+class SLATE_API SCanvas
+	: public SPanel
 {
 public:
 
@@ -96,12 +101,9 @@ public:
 			, SizeAttr( FVector2D( 1.0f, 1.0f ) )
 			, HAlignment(HAlign_Left)
 			, VAlignment(VAlign_Top)
-		{
-		}
+		{ }
 	};
 
-
-
 	SLATE_BEGIN_ARGS( SCanvas )
 		{
 			_Visibility = EVisibility::SelfHitTestInvisible;
@@ -123,6 +125,11 @@ public:
 		return *(new FSlot());
 	}
 
+	/**
+	 * Adds a content slot.
+	 *
+	 * @return The added slot.
+	 */
 	FSlot& AddSlot()
 	{
 		SCanvas::FSlot& NewSlot = SCanvas::Slot();
@@ -131,50 +138,30 @@ public:
 	}
 
 	/**
-	 * Panels arrange their children in a space described by the AllottedGeometry parameter. The results of the arrangement
-	 * should be returned by appending a FArrangedWidget pair for every child widget. See StackPanel for an example
+	 * Removes a particular content slot.
 	 *
-	 * @param AllottedGeometry    The geometry allotted for this widget by its parent.
-	 * @param ArrangedChildren    The array to which to add the WidgetGeometries that represent the arranged children.
+	 * @param SlotWidget The widget in the slot to remove.
 	 */
-	virtual void ArrangeChildren( const FGeometry& AllottedGeometry, FArrangedChildren& ArrangedChildren ) const OVERRIDE;
-	
+	int32 RemoveSlot( const TSharedRef<SWidget>& SlotWidget );
+
 	/**
-	 * The widget should respond by populating the OutDrawElements array with FDrawElements 
-	 * that represent it and any of its children.
-	 *
-	 * @param AllottedGeometry  The FGeometry that describes an area in which the widget should appear.
-	 * @param MyClippingRect    The clipping rectangle allocated for this widget and its children.
-	 * @param OutDrawElements   A list of FDrawElements to populate with the output.
-	 * @param LayerId           The Layer onto which this widget should be rendered.
-	 * @param InColorAndOpacity Color and Opacity to be applied to all the descendants of the widget being painted
-	 * @param bParentEnabled	True if the parent of this widget is enabled.
-	 *
-	 * @return The maximum layer ID attained by this widget or any of its children.
+	 * Removes all slots from the panel.
 	 */
+	void ClearChildren();
+
+public:
+
+	// Begin SWidget overrides
+
+	virtual void ArrangeChildren( const FGeometry& AllottedGeometry, FArrangedChildren& ArrangedChildren ) const OVERRIDE;
+	
 	virtual int32 OnPaint( const FGeometry& AllottedGeometry, const FSlateRect& MyClippingRect, FSlateWindowElementList& OutDrawElements, int32 LayerId, const FWidgetStyle& InWidgetStyle, bool bParentEnabled ) const OVERRIDE;
 
-	/**
-	 * A Panel's desired size in the space required to arrange of its children on the screen while respecting all of
-	 * the children's desired sizes and any layout-related options specified by the user. See StackPanel for an example.
-	 *
-	 * @return The desired size.
-	 */
 	virtual FVector2D ComputeDesiredSize() const OVERRIDE;
 
-	/** @return  The children of a panel in a slot-agnostic way. */
 	virtual FChildren* GetChildren() OVERRIDE;
 
-	/** 
-	 * Removes a slot from this box panel which contains the specified SWidget
-	 *
-	 * @param SlotWidget The widget to match when searching through the slots
-	 * @returns The index in the children array where the slot was removed and -1 if no slot was found matching the widget
-	 */
-	int32 RemoveSlot( const TSharedRef<SWidget>& SlotWidget );
-
-	/** Removes all slots from the panel */
-	void ClearChildren();
+	// End SWidget overrides
 
 protected:
 
diff --git a/Engine/Source/Runtime/Slate/Public/Widgets/SToolTip.h b/Engine/Source/Runtime/Slate/Public/Widgets/SToolTip.h
--- a/Engine/Source/Runtime/Slate/Public/Widgets/SToolTip.h
+++ b/Engine/Source/Runtime/Slate/Public/Widgets/SToolTip.h
@@ -1,5 +1,9 @@
 // Copyright 1998-2014 Epic Games, Inc. All Rights Reserved.
 
+/*=============================================================================
+	SToolTip.h: Declares the SToolTip class.
+=============================================================================*/
+
 #pragma once
 
 
diff --git a/Engine/Source/Runtime/Slate/Public/Widgets/SViewport.h b/Engine/Source/Runtime/Slate/Public/Widgets/SViewport.h
--- a/Engine/Source/Runtime/Slate/Public/Widgets/SViewport.h
+++ b/Engine/Source/Runtime/Slate/Public/Widgets/SViewport.h
@@ -1,5 +1,9 @@
 // Copyright 1998-2014 Epic Games, Inc. All Rights Reserved.
 
+/*=============================================================================
+	SViewport.h: Declares the SViewport class.
+=============================================================================*/
+
 #pragma once
 
 
@@ -43,14 +47,19 @@ public:
 		SLATE_ARGUMENT( bool, IgnoreTextureAlpha )
 	SLATE_END_ARGS()
 
+	/**
+	 * Default constructor.
+	 */
+	SViewport( );
+
 	/**
 	 * Construct the widget.
 	 *
 	 * @param InArgs  Declaration from which to construct the widget.
 	 */
 	void Construct(const FArguments& InArgs);
-	
-	SViewport();
+
+public:
 
 	/** SViewport wants keyboard focus */
 	virtual bool SupportsKeyboardFocus() const OVERRIDE { return true; }
@@ -116,7 +125,10 @@ public:
 		WidgetToFocusOnActivate.Reset();
 	}
 
+public:
+
 	// SWidget interface
+
 	virtual int32 OnPaint( const FGeometry& AllottedGeometry, const FSlateRect& MyClippingRect, FSlateWindowElementList& OutDrawElements, int32 LayerId, const FWidgetStyle& InWidgetStyle, bool bParentEnabled ) const OVERRIDE;
 	virtual void Tick( const FGeometry& AllottedGeometry, const double InCurrentTime, const float InDeltaTime ) OVERRIDE;
 	virtual FReply OnTouchStarted( const FGeometry& MyGeometry, const FPointerEvent& TouchEvent ) OVERRIDE;
@@ -140,27 +152,37 @@ public:
 	virtual FReply OnControllerButtonReleased( const FGeometry& MyGeometry, const FControllerEvent& ControllerEvent ) OVERRIDE;
 	virtual FReply OnControllerAnalogValueChanged( const FGeometry& MyGeometry, const FControllerEvent& ControllerEvent ) OVERRIDE;
 	virtual FReply OnMotionDetected( const FGeometry& MyGeometry, const FMotionEvent& MotionEvent ) OVERRIDE;
-	// End of SWidget interface
 
+	// End of SWidget interface
 
 private:
-	FSimpleSlot& Decl_GetContent() { return ChildSlot; }
+
+	FSimpleSlot& Decl_GetContent()
+	{
+		return ChildSlot;
+	}
 	
 private:
-	/** Interface to the rendering and I/O implementation of the viewport */
+
+	// Interface to the rendering and I/O implementation of the viewport.
 	TWeakPtr<ISlateViewport> ViewportInterface;
-	/** Whether or not to show the disabled effect when this viewport is disabled */
+
+	// Whether or not to show the disabled effect when this viewport is disabled.
 	TAttribute<bool> ShowDisabledEffect;
-	/** Whether or not this viewport renders directly to the window backbuffer */
+
+	// Whether or not this viewport renders directly to the window back-buffer.
 	bool bRenderDirectlyToWindow;
-	/** Whether or not to apply gamma correction on the render target supplied by the ISlateViewport */
+
+	// Whether or not to apply gamma correction on the render target supplied by the ISlateViewport.
 	bool bEnableGammaCorrection;
-	/** Whether or not to blend this viewport with the background */
+
+	// Whether or not to blend this viewport with the background.
 	bool bEnableBlending;
-	/** Whether or not to allow texture alpha to be used in blending calculations */
+
+	// Whether or not to allow texture alpha to be used in blending calculations.
 	bool bIgnoreTextureAlpha;
 
-	/** Widget to transfer keyboard focus to when this window becomes active, if any.  This is used to
-	    restore focus to a widget after a popup has been dismissed. */
+	// Widget to transfer keyboard focus to when this window becomes active, if any.
+	// This is used to restore focus to a widget after a popup has been dismissed.
 	TWeakPtr< SWidget > WidgetToFocusOnActivate;
 };
diff --git a/Engine/Source/Runtime/Slate/Public/Widgets/SWeakWidget.h b/Engine/Source/Runtime/Slate/Public/Widgets/SWeakWidget.h
--- a/Engine/Source/Runtime/Slate/Public/Widgets/SWeakWidget.h
+++ b/Engine/Source/Runtime/Slate/Public/Widgets/SWeakWidget.h
@@ -1,18 +1,25 @@
 // Copyright 1998-2014 Epic Games, Inc. All Rights Reserved.
 
+/*=============================================================================
+	SWeakWidget.h: Declares the SWeakWidget class.
+=============================================================================*/
+
 #pragma once
 
 
 /**
- * Weak widgets present a piece of content without owning it.
+ * Implements a widget that holds a weak pointer to one child widget.
+ *
+ * Weak widgets encapsulate a piece of content without owning it.
  * e.g. A tooltip is owned by the hovered widget but displayed
  *      by a floating window. That window cannot own the tooltip
  *      and must therefore use an SWeakWidget.
  */
-class SLATE_API SWeakWidget : public SWidget
+class SLATE_API SWeakWidget
+	: public SWidget
 {
-
 public:
+
 	SLATE_BEGIN_ARGS(SWeakWidget)
 		{
 			_Visibility = EVisibility::SelfHitTestInvisible;
@@ -20,8 +27,20 @@ public:
 		SLATE_ARGUMENT( TSharedPtr<SWidget>, PossiblyNullContent )
 	SLATE_END_ARGS()
 
+public:
+
 	void Construct(const FArguments& InArgs);
 
+	void SetContent(TWeakPtr<SWidget> InWidget);
+
+	bool ChildWidgetIsValid() const;
+
+	TWeakPtr<SWidget> GetChildWidget() const;
+
+public:
+
+	// Begin SWidget interface
+
 	virtual FVector2D ComputeDesiredSize() const OVERRIDE;
 
 	virtual void ArrangeChildren( const FGeometry& AllottedGeometry, FArrangedChildren& ArrangedChildren ) const OVERRIDE;
@@ -30,13 +49,9 @@ public:
 
 	virtual int32 OnPaint( const FGeometry& AllottedGeometry, const FSlateRect& MyClippingRect, FSlateWindowElementList& OutDrawElements, int32 LayerId, const FWidgetStyle& InWidgetStyle, bool bParentEnabled ) const OVERRIDE;
 
-	void SetContent(TWeakPtr<SWidget> InWidget);
-
-	bool ChildWidgetIsValid() const;
-
-	TWeakPtr<SWidget> GetChildWidget() const;
+	// End SWidget interface
 
 private:
-	TWeakChild<SWidget> WeakChild;
 
+	TWeakChild<SWidget> WeakChild;
 };
diff --git a/Engine/Source/Runtime/Slate/Public/Widgets/Text/ITextEditorWidget.h b/Engine/Source/Runtime/Slate/Public/Widgets/Text/ITextEditorWidget.h
new file mode 100644
--- /dev/null
+++ b/Engine/Source/Runtime/Slate/Public/Widgets/Text/ITextEditorWidget.h
@@ -0,0 +1,173 @@
+// Copyright 1998-2014 Epic Games, Inc. All Rights Reserved.
+
+#pragma once
+
+namespace ECursorMoveMethod
+{
+	enum Type
+	{
+		/** Move one character at a time (e.g. arrow left, right). */
+		CharacterHorizontal,
+		/** Move vertically (e.g. arrow up, down) */
+		CharacterVertical,
+		/** Move one word at a time (e.g. ctrl is held down). */
+		Word,
+	};
+}
+
+namespace ECursorAction
+{
+	enum Type
+	{
+		/** Just relocate the cursor. */
+		MoveCursor,
+		/** Move the cursor and select any text that it passes over. */
+		SelectText
+	};
+}
+
+namespace ETextLocation
+{
+	enum Type
+	{
+		/** Jump to the beginning of text. (e.g. Ctrl+Home) */
+		BeginningOfDocument,
+		/** Jump to the end of text. (e.g. Ctrl+End) */
+		EndOfDocument,
+		/** Jump to the beginning of the line or beginning of text within a line (e.g. Home) */
+		BeginningOfLine,
+		/** Jump to the end of line (e.g. End) */
+		EndOfLine
+	};
+}
+
+class SLATE_API ITextEditorWidget
+{
+	public:
+
+	/**
+	 * Call this before the text is changed interactively by the user (handles Undo, etc.)
+	 */
+	virtual void StartChangingText() = 0;
+
+	/**
+	 * Called when text is finished being changed interactively by the user
+	 */
+	virtual void FinishChangingText() = 0;
+	
+	/**
+	 * Can the user edit the text?
+	 */
+	virtual bool GetIsReadOnly() const = 0;
+
+	/**
+	 * Deletes the character to the left of the caret
+	 */
+	virtual void BackspaceChar() = 0;
+
+	/**
+	 * Deletes the character to the right of the caret
+	 */
+	virtual void DeleteChar() = 0;
+	
+	/**
+	 * The user pressed this character key. Should we actually type it?
+	 * NOT RELIABLE VALIDATION.
+	 */
+	virtual bool CanTypeCharacter(const TCHAR CharInQuestion) const = 0;
+
+	/**
+	 * Adds a character at the caret position
+	 *
+	 * @param  InChar  Character to add
+	 */
+	virtual void TypeChar( const int32 Character ) = 0;
+	
+	/**
+	 * Move the cursors, left, right, up down, one character or word at a time.
+	 * Optionally select text while moving.
+	 * 
+	 * @param  Method       Move horizontally, vertically, or whole words horizontally.
+	 * @param  Direction    Move forward/backward by number of units specified; usuatlly +/-1.
+	 * @param  Action       Just move or select text.
+	 */
+	virtual FReply MoveCursor( ECursorMoveMethod::Type Method, int8 Direction, ECursorAction::Type Action ) = 0;
+	
+	/**
+	 * Jump the cursor to a special location: e.g. end of line, beginning of document, etc..
+	 */
+	virtual void JumpTo(ETextLocation::Type JumpLocation, ECursorAction::Type Action) = 0;
+
+	/**
+	 * Clears text selection
+	 */
+	virtual void ClearSelection() = 0;
+	
+	/**
+	 * Selects all text
+	 */
+	virtual void SelectAllText() = 0;
+
+	/**
+	 * Executed when the user presses ESC
+	 */
+	virtual FReply OnEscape() = 0;
+
+	/**
+	 * Executed when the user presses ENTER
+	 */
+	virtual void OnEnter() = 0;
+
+	/**
+	 * Returns true if 'Cut' is currently a valid action
+	 *
+	 * @return	True to allow this action, otherwise false
+	 */
+	virtual bool CanExecuteCut() const = 0;
+	
+	/**
+	 * Cuts selected text to the clipboard
+	 */
+	virtual void CutSelectedTextToClipboard() = 0;
+
+	/**
+	 * Returns true if 'Copy' is currently a valid action
+	 *
+	 * @return	True to allow this action, otherwise false
+	 */
+	virtual bool CanExecuteCopy() const = 0;
+
+	/**
+	 * Copies selected text to the clipboard
+	 */
+	virtual void CopySelectedTextToClipboard() = 0;
+
+	/**
+	 * Returns true if 'Paste' is currently a valid action
+	 *
+	 * @return	True to allow this action, otherwise false
+	 */
+	virtual bool CanExecutePaste() const = 0;
+
+	/**
+	 * Pastes text from the clipboard
+	 */
+	virtual void PasteTextFromClipboard() = 0;
+	
+	/**
+	 * Returns true if 'Undo' is currently a valid action
+	 *
+	 * @return	True to allow this action, otherwise false
+	 */
+	virtual bool CanExecuteUndo() const = 0;
+
+	/**
+	 * Rolls back state to the most recent undo history's state
+	 */
+	virtual void Undo() = 0;
+
+	/**
+	 * Restores cached undo state that was previously rolled back
+	 */
+	virtual void Redo() = 0;
+};
diff --git a/Engine/Source/Runtime/Slate/Public/Widgets/Text/SInlineEditableTextBlock.h b/Engine/Source/Runtime/Slate/Public/Widgets/Text/SInlineEditableTextBlock.h
new file mode 100644
--- /dev/null
+++ b/Engine/Source/Runtime/Slate/Public/Widgets/Text/SInlineEditableTextBlock.h
@@ -0,0 +1,148 @@
+// Copyright 1998-2014 Epic Games, Inc. All Rights Reserved.
+
+#pragma once
+
+
+DECLARE_DELEGATE_RetVal_TwoParams(bool, FOnVerifyTextChanged, const FText&, FText&)
+DECLARE_DELEGATE_OneParam(FOnBeginTextEdit, const FText&)
+
+/**
+ * Slate's InlineEditableTextBlock's are double selectable to go from a STextBlock to become SEditableTextBox.
+ */
+class SLATE_API SInlineEditableTextBlock: public SCompoundWidget
+{
+	SLATE_BEGIN_ARGS( SInlineEditableTextBlock )
+		: _Text()
+		, _Style( &FCoreStyle::Get().GetWidgetStyle<FInlineEditableTextBlockStyle>("InlineEditableTextBlockStyle") )
+		, _Font()
+		, _ColorAndOpacity()
+		, _ShadowOffset()
+		, _ShadowColorAndOpacity()
+		, _HighlightText()
+		, _WrapTextAt(0.0f)
+		, _IsReadOnly(false)
+	{}
+
+		/** The text displayed in this text block */
+		SLATE_ATTRIBUTE( FText, Text )
+
+		/** Pointer to a style of the inline editable text block, which dictates the font, color, and shadow options. */
+		SLATE_STYLE_ARGUMENT( FInlineEditableTextBlockStyle, Style )
+
+		/** Sets the font used to draw the text (overrides style) */
+		SLATE_ATTRIBUTE( FSlateFontInfo, Font )
+
+		/** Text color and opacity (overrides style) */
+		SLATE_ATTRIBUTE( FSlateColor, ColorAndOpacity )
+
+		/** Drop shadow offset in pixels (overrides style) */
+		SLATE_ATTRIBUTE( FVector2D, ShadowOffset )
+
+		/** Shadow color and opacity (overrides style) */
+		SLATE_ATTRIBUTE( FLinearColor, ShadowColorAndOpacity )
+
+		/** Highlight this text in the text block */
+		SLATE_ATTRIBUTE( FText, HighlightText )
+
+		/** Whether text wraps onto a new line when it's length exceeds this width; if this value is zero or negative, no wrapping occurs. */
+		SLATE_ATTRIBUTE( float, WrapTextAt )
+
+		/** True if the editable text block is read-only. It will not be able to enter edit mode if read-only */
+		SLATE_ATTRIBUTE( bool, IsReadOnly )
+
+		/** Callback when the text starts to be edited */
+		SLATE_EVENT( FOnBeginTextEdit, OnBeginTextEdit )
+
+		/** Callback when the text is committed. */
+		SLATE_EVENT( FOnTextCommitted, OnTextCommitted )
+
+		/** Callback to check if the widget is selected, should only be hooked up if parent widget is handling selection or focus. */
+		SLATE_EVENT( FIsSelected, IsSelected )
+
+		/** Called whenever the text is changed interactively by the user */
+		SLATE_EVENT( FOnVerifyTextChanged, OnVerifyTextChanged )
+	SLATE_END_ARGS()
+
+	~SInlineEditableTextBlock();
+
+	/**
+	 * Construct this widget.  Called by the SNew() Slate macro.
+	 *
+	 * @param	InArgs				Declaration used by the SNew() macro to construct this widget
+	 * @param	InViewModel			The UI logic not specific to slate
+	 */
+	void Construct( const FArguments& InArgs );
+
+	/**
+	 * See SWidget::SupportsKeyboardFocus().
+	 *
+	 * @return  This widget can gain focus if IsSelected is not bound.
+	 */
+	virtual bool SupportsKeyboardFocus() const OVERRIDE;
+
+	virtual FReply OnMouseButtonDown( const FGeometry& MyGeometry, const FPointerEvent& MouseEvent ) OVERRIDE;
+
+	virtual FReply OnDragOver( const FGeometry& MyGeometry, const FDragDropEvent& DragDropEvent ) OVERRIDE;
+
+	virtual FReply OnMouseButtonDoubleClick( const FGeometry& InMyGeometry, const FPointerEvent& InMouseEvent ) OVERRIDE;
+
+	virtual FReply OnKeyDown( const FGeometry& MyGeometry, const FKeyboardEvent& InKeyboardEvent ) OVERRIDE;
+
+	virtual void Tick( const FGeometry& AllottedGeometry, const double InCurrentTime, const float InDeltaTime ) OVERRIDE;
+	 
+	/** Switches the widget to editing mode */
+	void EnterEditingMode();
+
+	/** Switches the widget to label mode. */
+	void ExitEditingMode();
+
+	/** Checks if the widget is in edit mode */
+	bool IsInEditMode() const;
+
+	void SetText( const TAttribute< FText >& InText );
+	void SetText( const FString& InText );
+
+protected:
+	/** Callback for the text box's OnTextChanged event */
+	void OnTextChanged(const FText& InText);
+
+	/** Callback when the text box is committed, switches back to label mode. */ 
+	void OnTextBoxCommitted(const FText& InText, ETextCommit::Type InCommitType);
+
+	/** Cancels the edit mode and switches back to label mode */
+	void CancelEditMode();
+
+protected:
+	/** The widget used when in label mode */
+	TSharedPtr< STextBlock > TextBlock;
+
+	/** The widget used when in editing mode */ 
+	TSharedPtr< SEditableTextBox > TextBox;
+
+	/** Delegate to execute when the text starts to be edited */
+	FOnBeginTextEdit OnBeginTextEditDelegate;
+
+	/** Delegate to execute when editing mode text is committed. */
+	FOnTextCommitted OnTextCommittedDelegate;
+
+	/** Delegate to execute to check the status of if the widget is selected or not
+		Only needs to be hooked up if an external widget is managing selection, such
+		as a SListView or STreeView. */
+	FIsSelected IsSelected;
+
+	/** When selection of widget is managed by another widget, this delays the "double select" from 
+		occurring immediately, offering a chance for double clicking to take action. */
+	float DoubleSelectDelay;
+
+	/** Main horizontal box, used to dynamically add and remove the editable slot. */
+	TSharedPtr< SHorizontalBox > HorizontalBox;
+
+	/** Callback to verify text when changed. Will return an error message to denote problems. */
+	FOnVerifyTextChanged OnVerifyTextChanged;
+
+	/** Attribute for the text to use for the widget */
+	TAttribute<FText> Text;
+
+	/** Attribute to look up if the widget is read-only */
+	TAttribute< bool > bIsReadOnly;
+};
diff --git a/Engine/Source/Runtime/Slate/Public/Widgets/Text/SMultiLineEditableText.h b/Engine/Source/Runtime/Slate/Public/Widgets/Text/SMultiLineEditableText.h
new file mode 100644
--- /dev/null
+++ b/Engine/Source/Runtime/Slate/Public/Widgets/Text/SMultiLineEditableText.h
@@ -0,0 +1,232 @@
+// Copyright 1998-2014 Epic Games, Inc. All Rights Reserved.
+
+#pragma once
+
+#include "ITextEditorWidget.h"
+
+/** Information about the substring as it would appear as a result of soft-wrapping a line of text. */
+struct FWrappedStringSlice
+{
+	/** 
+	 * Create a substring info that sources InSourceString, begins at InFirstCharIndex and ends at InLastCharIndex (exclusive).
+	 * The substring measures InSize.X by InSize.Y.
+	 */
+	FWrappedStringSlice( FString& InSourceString, int32 InFirstCharIndex, int32 InLastCharIndex, FVector2D InSize )
+	: SourceString( InSourceString )
+	, FirstCharIndex( InFirstCharIndex )
+	, LastCharIndex( InLastCharIndex )
+	, Size( InSize )
+	{
+	}
+
+	/** The substring sources this string. */
+	FString& SourceString;
+	/** Index of the first character in the substring */
+	int32 FirstCharIndex;
+	/** Index of the last character (exclusive) */
+	int32 LastCharIndex;
+	/** The measured size of this substring given the current font and point size. */
+	FVector2D Size;
+};
+
+/** An editable text widget that supports multiple lines and soft word-wrapping. */
+class SLATE_API SMultiLineEditableText : public SLeafWidget, public ITextEditorWidget
+{
+public:
+	SLATE_BEGIN_ARGS( SMultiLineEditableText )
+	: _Text()
+	, _ScrollBar()
+	, _PreferredSize(FVector2D(100,100))
+	{}
+		/** The initial text that will appear in the widget. */
+		SLATE_TEXT_ARGUMENT( Text )
+
+		/** Scroll bar that is used to visualize and manipulate the scrolling of the text */
+		SLATE_ARGUMENT( TSharedPtr<class SScrollBar>, ScrollBar )
+
+		/** How big should this text field be for optimum operation. The value is irrelevant when the editable text fills all available space. */
+		SLATE_ARGUMENT( FVector2D, PreferredSize)
+
+	SLATE_END_ARGS()
+
+	SMultiLineEditableText();
+
+	void Construct( const FArguments& InArgs );
+
+private:
+	
+	/** Location within the text model. */
+	struct FTextLocation
+	{
+	public:
+		FTextLocation( const int32 InLineIndex = 0, const int32 InOffset = 0 )
+		: LineIndex( InLineIndex )
+		, Offset( InOffset )
+		{
+
+		}
+
+		bool operator==( const FTextLocation& Other ) const
+		{
+			return
+				LineIndex == Other.LineIndex &&
+				Offset == Other.Offset;
+		}
+
+		bool operator<( const FTextLocation& Other ) const
+		{
+			return this->LineIndex < Other.LineIndex || (this->LineIndex == Other.LineIndex && this->Offset < Other.Offset);
+		}
+
+		int32 GetLineIndex() const { return LineIndex; }
+		int32 GetOffset() const { return Offset; }
+
+		private:
+		int32 LineIndex;
+		int32 Offset;
+	};
+
+	/** Range between two locations. Guarantees that the First location is before the Last. */
+	struct FTextRange
+	{
+		/**
+		 * Construct a Range from two locations: A and B.
+		 * Guarantees that First will be the earlier of the two.
+		 */
+		FTextRange( const FTextLocation& A, const FTextLocation& B )
+		: First( (A < B) ? A : B )
+		, Last( (A < B) ? B : A )
+		{
+		}
+
+		const FTextLocation First;
+		const FTextLocation Last;
+	};
+
+private:
+	// BEGIN ITextEditorWidget interface
+	virtual void StartChangingText() OVERRIDE;
+	virtual void FinishChangingText() OVERRIDE;
+	virtual bool GetIsReadOnly() const OVERRIDE;
+	virtual void BackspaceChar() OVERRIDE;
+	virtual void DeleteChar() OVERRIDE;
+	virtual bool CanTypeCharacter(const TCHAR CharInQuestion) const OVERRIDE;
+	virtual void TypeChar( const int32 Character ) OVERRIDE;
+	virtual FReply MoveCursor( ECursorMoveMethod::Type Method, int8 Direction, ECursorAction::Type Action ) OVERRIDE;
+	virtual void JumpTo(ETextLocation::Type JumpLocation, ECursorAction::Type Action) OVERRIDE;
+	virtual void ClearSelection() OVERRIDE;
+	virtual void SelectAllText() OVERRIDE;
+	virtual FReply OnEscape() OVERRIDE;
+	virtual void OnEnter() OVERRIDE;
+	virtual bool CanExecuteCut() const OVERRIDE;
+	virtual void CutSelectedTextToClipboard() OVERRIDE;
+	virtual bool CanExecuteCopy() const OVERRIDE;
+	virtual void CopySelectedTextToClipboard() OVERRIDE;
+	virtual bool CanExecutePaste() const OVERRIDE;
+	virtual void PasteTextFromClipboard() OVERRIDE;
+	virtual bool CanExecuteUndo() const OVERRIDE;
+	virtual void Undo() OVERRIDE;
+	virtual void Redo() OVERRIDE;
+	// END ITextEditorWidget interface
+
+private:
+	// BEGIN SWidget interface
+	virtual void Tick( const FGeometry& AllottedGeometry, const double InCurrentTime, const float InDeltaTime ) OVERRIDE;
+	virtual int32 OnPaint( const FGeometry& AllottedGeometry, const FSlateRect& MyClippingRect, FSlateWindowElementList& OutDrawElements, int32 LayerId, const FWidgetStyle& InWidgetStyle, bool bParentEnabled ) const OVERRIDE;
+	virtual FVector2D ComputeDesiredSize() const OVERRIDE;
+	virtual bool SupportsKeyboardFocus() const OVERRIDE;
+	virtual FReply OnKeyChar( const FGeometry& MyGeometry,const FCharacterEvent& InCharacterEvent ) OVERRIDE;
+	virtual FReply OnKeyDown( const FGeometry& MyGeometry, const FKeyboardEvent& InKeyboardEvent ) OVERRIDE;
+	virtual FReply OnKeyUp( const FGeometry& MyGeometry, const FKeyboardEvent& InKeyboardEvent ) OVERRIDE;
+	// END SWidget interface
+
+private:
+	
+	/** Re-compute the wrapping information about the current text. */
+	void ReWrap(const FGeometry& AllottedGeometry);
+
+	/** Invoked when the user manipulates the scrollbar*/
+	void OnScrollbarScrolled(float OffsetFraction);
+
+	/** Remember where the cursor was when we started selecting. */
+	void BeginSelecting( const FTextLocation& Endpoint );
+
+	/** Given the current geometry, make the scrollbar reflect the current offset and fraction of the visible text. */
+	void UpdateScrollbar( const FGeometry &AllottedGeometry );
+
+	/**
+	 * Given a location and a Direction to offset, return a new location.
+	 *
+	 * @param Location    Cursor location from which to offset
+	 * @param Direction   Positive means right, negative means left.
+	 */
+	FTextLocation TranslatedLocation( const FTextLocation& CurrentLocation, int8 Direction ) const;
+
+	/**
+	 * Given a location and a Direction to offset, return a new location.
+	 *
+	 * @param Location    Cursor location from which to offset
+	 * @param Direction   Positive means down, negative means up.
+	 */
+	FTextLocation TranslateLocationVertical( const FTextLocation& Location, int8 Direction ) const;
+
+	/** Find the closest word boundary */
+	FTextLocation ScanForWordBoundary( const FTextLocation& Location, int8 Direction ) const; 
+
+	/** Get the character at Location */
+	TCHAR GetCharacterAt( const FTextLocation& Location ) const;
+
+	/** Are we at the beginning of all the text. */
+	bool IsAtStartOfDoc( const FTextLocation& Location ) const;
+	
+	/** Are we at the end of all the text. */
+	bool IsAtEndOfDoc( const FTextLocation& Location ) const;
+
+	/** Is this location the beginning of a line */
+	bool IsAtStartOfLine( const FTextLocation& Location ) const;
+
+	/** Is this location the end of a line. */
+	bool IsAtEndOfLine( const FTextLocation& Location ) const;
+
+	/** Are we currently at the beginning of a word */
+	bool IsAtWordStart( const FTextLocation& Location ) const;
+
+	/** Given a location in unwrapped text, compute the location in wrapped line. */
+	FTextLocation AbsoluteToWrapped( const FTextLocation& Location ) const;
+
+	/** Given a range in unwrapped text, compute the range in wrapped text. */
+	FTextRange AbsoluteToWrapped( const FTextRange& Range ) const;
+
+	/** If we have a pointer to this scroll bar, use it to visualize and manipulate the scroll offset. */
+	TWeakPtr<class SScrollBar> OptionalScrollBarPtr;
+
+	/** We maintain our own version of the text as an array of strings. This is done for performance reasons. */
+	TArray< TSharedRef<FString> > TextLines;
+	
+	/** The information about how the text is currently wrapped. */
+	TArray< FWrappedStringSlice > WrappedText;
+
+	/** Remember last frame's geometry for comparison and  */
+	FGeometry CachedLastFrameGeometry;
+
+	/** That start of the selection when there is a selection. The end is implicitly wherever the cursor happens to be. */
+	TOptional<FTextLocation> SelectionStart;
+
+	/** Current cursor position; there is always one. */
+	FTextLocation CursorPosition;
+
+	/** The font being used to render the text. */
+	FSlateFontInfo FontInfo;
+
+	/** Floating point offset in number of lines from the top of the document. */
+	float NumLinesScrollOffset;
+
+	/** The user probably wants the cursor where they last explicitly positioned it horizontally. */
+	int32 PreferredCursorOffsetInLine;
+
+	/** Last time the user did anything with the cursor.*/
+	double LastCursorInteractionTime;
+
+	/** How big should this text field be for optimum operation. The value is irrelevant when the editable text fills all available space. */
+	FVector2D PreferredSize;
+};
