Slate: Reorganized private Slate files; proper pre-compiled header for all files.

diff --git a/Engine/Source/Runtime/Slate/Public/Widgets/Text/SRichTextBlock.h b/Engine/Source/Runtime/Slate/Public/Widgets/Text/SRichTextBlock.h
new file mode 100644
--- /dev/null
+++ b/Engine/Source/Runtime/Slate/Public/Widgets/Text/SRichTextBlock.h
@@ -0,0 +1,163 @@
+// Copyright 1998-2014 Epic Games, Inc. All Rights Reserved.
+#pragma once
+
+/**
+ * A rich static text widget. 
+ * Through the use of markup and text decorators, text with different styles, embedded image and widgets can be achieved.
+ */
+class SLATE_API SRichTextBlock : public SWidget
+{
+public:
+
+	SLATE_BEGIN_ARGS( SRichTextBlock )
+		: _Text(  )
+		, _WrapTextAt( 0.0f )
+		, _AutoWrapText(false)
+		, _DecoratorStyleSet( &FCoreStyle::Get() )
+		, _TextStyle( &FCoreStyle::Get().GetWidgetStyle<FTextBlockStyle>( "NormalText" ) )
+		, _Margin( FMargin() )
+		, _LineHeightPercentage( 1.0f )
+		, _Justification( ETextJustify::Left )
+		, _Decorators()
+		, _Parser()
+	{}
+		/** The text displayed in this text block */
+		SLATE_ARGUMENT( FText, Text )
+
+		/** Whether text wraps onto a new line when it's length exceeds this width; if this value is zero or negative, no wrapping occurs. */
+		SLATE_ATTRIBUTE( float, WrapTextAt )
+
+		/** Whether to wrap text automatically based on the widget's computed horizontal space.  IMPORTANT: Using automatic wrapping can result
+		    in visual artifacts, as the the wrapped size will computed be at least one frame late!  Consider using WrapTextAt instead.  The initial 
+			desired size will not be clamped.  This works best in cases where the text block's size is not affecting other widget's layout. */
+		SLATE_ATTRIBUTE( bool, AutoWrapText )
+
+		/** The style set used for looking up styles used by decorators*/
+		SLATE_ARGUMENT( const ISlateStyle*, DecoratorStyleSet )
+
+		/** The style of the text block, which dictates the default font, color, and shadow options. */
+		SLATE_STYLE_ARGUMENT( FTextBlockStyle, TextStyle )
+
+		/** The amount of blank space left around the edges of text area. */
+		SLATE_ATTRIBUTE( FMargin, Margin )
+
+		/** The amount to scale each lines height by. */
+		SLATE_ATTRIBUTE( float, LineHeightPercentage )
+
+		/** How the text should be aligned with the margin. */
+		SLATE_ATTRIBUTE( ETextJustify::Type, Justification )
+
+		/** Any decorators that should be used while parsing the text. */
+		SLATE_ARGUMENT( TArray< TSharedRef< class ITextDecorator > >, Decorators )
+
+		/** The parser used to resolve any markup used in the provided string. */
+		SLATE_ARGUMENT( TSharedPtr< class IRichTextMarkupParser >, Parser )
+
+		/** Additional decorators can be append to the widget inline. Inline decorators get precedence over decorators not specified inline. */
+		TArray< TSharedRef< ITextDecorator > > InlineDecorators;
+		SRichTextBlock::FArguments& operator + (const TSharedRef< ITextDecorator >& DecoratorToAdd)
+		{
+			InlineDecorators.Add( DecoratorToAdd );
+			return *this;
+		}
+
+	SLATE_END_ARGS()
+
+	static TSharedRef< ITextDecorator > Decorator( const TSharedRef< ITextDecorator >& Decorator )
+	{
+		return Decorator;
+	}
+
+	static TSharedRef< ITextDecorator > WidgetDecorator( const FString RunName, const FWidgetDecorator::FCreateWidget& FactoryDelegate )
+	{
+		return FWidgetDecorator::Create( RunName, FactoryDelegate );
+	}
+
+	template< class UserClass >
+	static TSharedRef< ITextDecorator >  WidgetDecorator( const FString RunName, UserClass* InUserObjectPtr, typename FWidgetDecorator::FCreateWidget::TSPMethodDelegate_Const< UserClass >::FMethodPtr InFunc )
+	{
+		return FWidgetDecorator::Create( RunName, FWidgetDecorator::FCreateWidget::CreateSP( InUserObjectPtr, InFunc ) );
+	}
+
+	static TSharedRef< ITextDecorator > ImageDecorator( FString RunName = TEXT("img"), const ISlateStyle* const OverrideStyle = NULL )
+	{
+		return FImageDecorator::Create( RunName, OverrideStyle );
+	}
+
+	static TSharedRef< ITextDecorator > HyperlinkDecorator( const FString Id, const FSlateHyperlinkRun::FOnClick& NavigateDelegate )
+	{
+		return FHyperlinkDecorator::Create( Id, NavigateDelegate );
+	}
+
+	template< class UserClass >
+	static TSharedRef< ITextDecorator > HyperlinkDecorator( const FString Id, UserClass* InUserObjectPtr, typename FSlateHyperlinkRun::FOnClick::TSPMethodDelegate< UserClass >::FMethodPtr NavigateFunc )
+	{
+		return FHyperlinkDecorator::Create( Id, FSlateHyperlinkRun::FOnClick::CreateSP( InUserObjectPtr, NavigateFunc ) );
+	}
+
+	/************************************************************************/
+	/* SWidget                                                              */
+	/************************************************************************/
+	void Construct( const FArguments& InArgs );
+
+	virtual int32 OnPaint( const FGeometry& AllottedGeometry, const FSlateRect& MyClippingRect, FSlateWindowElementList& OutDrawElements, int32 LayerId, const FWidgetStyle& InWidgetStyle, bool bParentEnabled ) const;
+	virtual void Tick( const FGeometry& AllottedGeometry, const double InCurrentTime, const float InDeltaTime ) OVERRIDE;
+
+	virtual void CacheDesiredSize() OVERRIDE;
+	virtual FVector2D ComputeDesiredSize() const OVERRIDE;
+
+	virtual FChildren* GetChildren() OVERRIDE;
+
+	virtual void ArrangeChildren( const FGeometry& AllottedGeometry, FArrangedChildren& ArrangedChildren ) const OVERRIDE;
+
+	/**
+	 * Gets the text assigned to this text block
+	 */
+	const FText& GetText() const
+	{
+		return Text;
+	}
+
+	/**
+	 * Sets the text for this text block
+	 */
+	void SetText( const FText& InText );
+
+	void SetHighlightText( const FText& InHighlightText );
+
+
+private:
+	
+	TSharedPtr< ITextDecorator > TryGetDecorator( const TArray< TSharedRef< ITextDecorator > >& Decorators, const FString& Line, const FTextRunParseResults& TextRun ) const;
+
+private:
+
+	/** The text displayed in this text block */
+	FText Text;
+
+	TSharedPtr< FSlateTextLayout > TextLayout;
+
+	FText HighlightText;
+	TArray< FTextHighlight > Highlights;
+	TSharedPtr< ISlateRunHighlighter> TextHighlighter;
+
+	const ISlateStyle* TagStyleSet;
+	FTextBlockStyle TextStyle;
+	
+	/** Whether text wraps onto a new line when it's length exceeds this width; if this value is zero or negative, no wrapping occurs. */
+	TAttribute<float> WrapTextAt;
+	
+	/** True if we're wrapping text automatically based on the computed horizontal space for this widget */
+	TAttribute<bool> AutoWrapText;
+
+	/** The last known width of the control from the previous OnPaint, used to recalculate wrapping. */
+	mutable float CachedAutoWrapTextWidth;
+
+	TAttribute< FMargin > Margin;
+	TAttribute< ETextJustify::Type > Justification; 
+	TAttribute< float > LineHeightPercentage;
+
+	TArray< TSharedRef< class ITextDecorator > > Decorators;
+
+	TSharedPtr< class IRichTextMarkupParser > Parser;
+};
diff --git a/Engine/Source/Runtime/Slate/Slate.Build.cs b/Engine/Source/Runtime/Slate/Slate.Build.cs
--- a/Engine/Source/Runtime/Slate/Slate.Build.cs
+++ b/Engine/Source/Runtime/Slate/Slate.Build.cs
@@ -27,19 +27,30 @@ public class Slate : ModuleRules
 		PrivateIncludePaths.AddRange(
 			new string[] {
 				"Runtime/Slate/Private",
-				"Runtime/Slate/Private/Docking",
-				"Runtime/Slate/Private/DockingFramework",
 				"Runtime/Slate/Private/Framework",
-				"Runtime/Slate/Private/Framework/Animation",
 				"Runtime/Slate/Private/Framework/Application",
 				"Runtime/Slate/Private/Framework/Commands",
+				"Runtime/Slate/Private/Framework/Docking",
+				"Runtime/Slate/Private/Framework/Layout",
+				"Runtime/Slate/Private/Framework/MultiBox",
+				"Runtime/Slate/Private/Framework/Notifications",
 				"Runtime/Slate/Private/Framework/Styling",
+				"Runtime/Slate/Private/Framework/Testing",
+				"Runtime/Slate/Private/Framework/Text",
+				"Runtime/Slate/Private/Framework/Text/IOS",
+				"Runtime/Slate/Private/Framework/Text/Tests",
 				"Runtime/Slate/Private/Framework/Widgets",
-				"Runtime/Slate/Private/IOS",
-				"Runtime/Slate/Private/MultiBox",
-				"Runtime/Slate/Private/TableView",
-				"Runtime/Slate/Private/Testing",
-				"Runtime/Slate/Private/Text",
+				"Runtime/Slate/Private/Widgets/Colors",
+				"Runtime/Slate/Private/Widgets/Docking",
+				"Runtime/Slate/Private/Widgets/Images",
+				"Runtime/Slate/Private/Widgets/Input",
+				"Runtime/Slate/Private/Widgets/Layout",
+				"Runtime/Slate/Private/Widgets/Navigation",
+				"Runtime/Slate/Private/Widgets/Notifications",
+				"Runtime/Slate/Private/Widgets/Testing",
+				"Runtime/Slate/Private/Widgets/Text",
+				"Runtime/Slate/Private/Widgets/Tutorials",
+				"Runtime/Slate/Private/Widgets/Views",
 			}
 		);
 
