fixed non-unity builds

diff --git a/Engine/Source/Runtime/Slate/Private/Framework/MultiBox/SEditableTextBlock.cpp b/Engine/Source/Runtime/Slate/Private/Framework/MultiBox/SEditableTextBlock.cpp
--- a/Engine/Source/Runtime/Slate/Private/Framework/MultiBox/SEditableTextBlock.cpp
+++ b/Engine/Source/Runtime/Slate/Private/Framework/MultiBox/SEditableTextBlock.cpp
@@ -1,7 +1,11 @@
 // Copyright 1998-2014 Epic Games, Inc. All Rights Reserved.
 
+/*=============================================================================
+	SEditableTextBlock.h: Implements the SEditableTextBlock class.
+=============================================================================*/
+
 #include "SlatePrivatePCH.h"
-#include "MultiBox.h"
+#include "SEditableTextBlock.h"
 
 
 FEditableTextBlock::FEditableTextBlock( const FText& InLabel, const FText& InToolTip, const FSlateIcon& InIcon, const TAttribute< FText >& InTextAttribute, bool bInReadOnly, const FOnTextCommitted& InOnTextCommitted, const FOnTextChanged& InOnTextChanged )
@@ -13,25 +17,17 @@ FEditableTextBlock::FEditableTextBlock( const FText& InLabel, const FText& InToo
 	  OnTextCommitted( InOnTextCommitted ),
 	  OnTextChanged( InOnTextChanged ),
 	  bReadOnly( bInReadOnly )
-{
-}
+{ }
 
 
-TSharedRef< class IMultiBlockBaseWidget > FEditableTextBlock::ConstructWidget() const
+TSharedRef<class IMultiBlockBaseWidget> FEditableTextBlock::ConstructWidget() const
 {
 	return SNew( SEditableTextBlock )
-			.Cursor(EMouseCursor::Default);
+		.Cursor(EMouseCursor::Default);
 }
 
 
-
-void SEditableTextBlock::Construct( const FArguments& InArgs )
-{
-}
-
-
-
-void SEditableTextBlock::BuildMultiBlockWidget(const ISlateStyle* StyleSet, const FName& StyleName)
+void SEditableTextBlock::BuildMultiBlockWidget( const ISlateStyle* StyleSet, const FName& StyleName )
 {
 	TSharedPtr< const FMultiBox > MultiBox = OwnerMultiBoxWidget.Pin()->GetMultiBox();
 	TSharedPtr< const FEditableTextBlock > EditableTextBlock = StaticCastSharedRef< const FEditableTextBlock >( MultiBlock.ToSharedRef() );
@@ -57,8 +53,7 @@ void SEditableTextBlock::BuildMultiBlockWidget(const ISlateStyle* StyleSet, cons
 		const FSlateBrush* IconBrush = ActualIcon.GetIcon();
 		if( IconBrush->GetResourceName() != NAME_None )
 		{
-			IconWidget =
-				SNew( SImage )
+			IconWidget = SNew( SImage )
 				.Image( IconBrush );
 		}
 	}
@@ -128,7 +123,6 @@ void SEditableTextBlock::BuildMultiBlockWidget(const ISlateStyle* StyleSet, cons
 }
 
 
-
 bool SEditableTextBlock::IsEnabled() const
 {
 	TSharedPtr< const FUICommandList > ActionList = MultiBlock->GetActionList();
@@ -148,4 +142,3 @@ bool SEditableTextBlock::IsEnabled() const
 
 	return bEnabled;
 }
-
diff --git a/Engine/Source/Runtime/Slate/Private/Framework/MultiBox/SEditableTextBlock.h b/Engine/Source/Runtime/Slate/Private/Framework/MultiBox/SEditableTextBlock.h
--- a/Engine/Source/Runtime/Slate/Private/Framework/MultiBox/SEditableTextBlock.h
+++ b/Engine/Source/Runtime/Slate/Private/Framework/MultiBox/SEditableTextBlock.h
@@ -1,5 +1,9 @@
 // Copyright 1998-2014 Epic Games, Inc. All Rights Reserved.
 
+/*=============================================================================
+	SEditableTextBlock.cpp: Declares the SEditableTextBlock class.
+=============================================================================*/
+
 #pragma once
 
 /**
@@ -24,7 +28,6 @@ public:
 	 */
 	FEditableTextBlock( const FText& InLabel, const FText& InToolTip, const FSlateIcon& InIcon, const TAttribute< FText >& InTextAttribute, bool bInReadOnly, const FOnTextCommitted& InOnTextCommitted, const FOnTextChanged& InOnTextChanged );
 
-
 private:
 
 	/**
@@ -32,8 +35,7 @@ private:
 	 *
 	 * @return  MultiBlock widget object
 	 */
-	virtual TSharedRef< class IMultiBlockBaseWidget > ConstructWidget() const;
-
+	virtual TSharedRef<class IMultiBlockBaseWidget> ConstructWidget() const;
 
 private:
 
@@ -64,8 +66,6 @@ private:
 };
 
 
-
-
 /**
  * Editable text block widget
  */
@@ -78,7 +78,6 @@ public:
 	SLATE_BEGIN_ARGS( SEditableTextBlock ){}
 	SLATE_END_ARGS()
 
-
 	/**
 	 * Builds this MultiBlock widget up from the MultiBlock associated with it
 	 */
@@ -89,7 +88,7 @@ public:
 	 *
 	 * @param	InArgs	The declaration data for this widget
 	 */
-	void Construct( const FArguments& InArgs );
+	void Construct( const FArguments& InArgs ) { }
 
 protected:
 
@@ -100,7 +99,6 @@ protected:
 	 */
 	bool IsEnabled() const;
 
-
 private:
 
 };
