Slate: fixed minor tab wells missing (TTP# 333740)

diff --git a/Engine/Source/Runtime/Slate/Private/Framework/Docking/SDockingTabStack.cpp b/Engine/Source/Runtime/Slate/Private/Framework/Docking/SDockingTabStack.cpp
--- a/Engine/Source/Runtime/Slate/Private/Framework/Docking/SDockingTabStack.cpp
+++ b/Engine/Source/Runtime/Slate/Private/Framework/Docking/SDockingTabStack.cpp
@@ -117,94 +117,100 @@ void SDockingTabStack::Construct( const FArguments& InArgs, const TSharedRef<FTa
 	ChildSlot
 	[
 		SNew(SVerticalBox)
-		.Visibility( EVisibility::SelfHitTestInvisible )
-		+SVerticalBox::Slot()
-		. AutoHeight()
-		[
-			// TAB WELL AREA
-			SNew(SBorder)
-			.Visibility( this, &SDockingTabStack::GetTabWellVisibility )
-			.DesiredSizeScale( this, &SDockingTabStack::GetTabWellScale )
-			.BorderImage( FCoreStyle::Get().GetBrush("NoBorder") )
-			.VAlign(VAlign_Bottom)
-			.OnMouseButtonDown( this, &SDockingTabStack::TabWellRightClicked )
-			.Padding(0)
-			[
-				SNew(SVerticalBox)
-				.Visibility( EVisibility::SelfHitTestInvisible )
-				+SVerticalBox::Slot()
-				.AutoHeight()
-				.Expose(TitleBarSlot)
-				[
-					SNew(SSpacer)
-				]
-				+SVerticalBox::Slot()
-				.AutoHeight()
-				[
-					SNew(SImage)
-					.Image( this, &SDockingTabStack::GetTabWellBrush )
-				]
-			]
-		]
-		+SVerticalBox::Slot()
-		. FillHeight(1)
-		[
-			SAssignNew(OverlayManagement.ContentAreaOverlay, SOverlay)
-			+ SOverlay::Slot()
-			[
-				SAssignNew(ContentSlot, SBorder)
-				. BorderImage( this, &SDockingTabStack::GetContentAreaBrush )
-				. Padding( this, &SDockingTabStack::GetContentPadding )
-				[
-					// Selected CONTENT AREA
-					SNew(STextBlock)
-					. Text( LOCTEXT("EmptyTabMessage", "Empty Tab!") )
-				]
-			]
-			+SOverlay::Slot()
-			.Padding(0)
-			.HAlign(HAlign_Left)
-			.VAlign(VAlign_Top)
+			.Visibility( EVisibility::SelfHitTestInvisible )
+
+		+ SVerticalBox::Slot()
+			.AutoHeight()
 			[
-				SNew(SButton)
-				.ButtonStyle( UnhideTabWellButtonStyle )
-				.OnClicked( this, &SDockingTabStack::UnhideTabWell )
-				.ContentPadding(0)
-				.Visibility( this, &SDockingTabStack::GetUnhideButtonVisibility )
-				.DesiredSizeScale( this, &SDockingTabStack::GetUnhideTabWellButtonScale )
-				.ButtonColorAndOpacity( this, &SDockingTabStack::GetUnhideTabWellButtonOpacity )
-				[
-					// Button should be big enough to show its own image
-					SNew(SSpacer)
-					.Size( UnhideTabWellButtonStyle->Normal.ImageSize )
-				]
+				// tab well area
+				SNew(SBorder)
+					.Visibility(this, &SDockingTabStack::GetTabWellVisibility)
+					.DesiredSizeScale(this, &SDockingTabStack::GetTabWellScale)
+					.BorderImage(FCoreStyle::Get().GetBrush("NoBorder"))
+					.VAlign(VAlign_Bottom)
+					.OnMouseButtonDown(this, &SDockingTabStack::TabWellRightClicked)
+					.Padding(0.0f)
+					[
+						SNew(SVerticalBox)
+							.Visibility(EVisibility::SelfHitTestInvisible)
+
+						+ SVerticalBox::Slot()
+							.AutoHeight()
+							.Expose(TitleBarSlot)
+
+						+ SVerticalBox::Slot()
+							.AutoHeight()
+							[
+								SNew(SImage)
+									.Image(this, &SDockingTabStack::GetTabWellBrush)
+							]
+					]
 			]
-			
-			#if DEBUG_TAB_MANAGEMENT
-			+SOverlay::Slot()
-			.HAlign(HAlign_Left)
-			.VAlign(VAlign_Top)
+
+		+ SVerticalBox::Slot()
+			.FillHeight(1.0f)
 			[
-				SNew(SBorder)
-				.BorderImage( FCoreStyle::Get().GetBrush( "Docking.Border" ) )
-				.BorderBackgroundColor(FLinearColor(1,0.5,0,0.75f))
-				.Visibility(EVisibility::HitTestInvisible)
-				[
-					SNew(STextBlock)
-					.Text(this, &SDockingTabStack::ShowPersistentTabs)
-					.ShadowOffset(FVector2D::UnitVector)
-				]
+				// tab content area
+				SAssignNew(OverlayManagement.ContentAreaOverlay, SOverlay)
+
+				+ SOverlay::Slot()
+					[
+						// content goes here
+						SAssignNew(ContentSlot, SBorder)
+							.BorderImage(this, &SDockingTabStack::GetContentAreaBrush)
+							.Padding(this, &SDockingTabStack::GetContentPadding)
+							[
+								SNew(STextBlock)
+									.Text(LOCTEXT("EmptyTabMessage", "Empty Tab!"))
+							]
+					]
+
+				+ SOverlay::Slot()
+					.Padding(0.0f)
+					.HAlign(HAlign_Left)
+					.VAlign(VAlign_Top)
+					[
+						// unhide tab well button (yellow triangle)
+						SNew(SButton)
+							.ButtonStyle(UnhideTabWellButtonStyle)
+							.OnClicked(this, &SDockingTabStack::UnhideTabWell)
+							.ContentPadding(0.0f)
+							.Visibility(this, &SDockingTabStack::GetUnhideButtonVisibility)
+							.DesiredSizeScale(this, &SDockingTabStack::GetUnhideTabWellButtonScale)
+							.ButtonColorAndOpacity(this, &SDockingTabStack::GetUnhideTabWellButtonOpacity)
+							[
+								// button should be big enough to show its own image
+								SNew(SSpacer)
+									.Size(UnhideTabWellButtonStyle->Normal.ImageSize)
+							]
+					]
+
+#if DEBUG_TAB_MANAGEMENT
+				+ SOverlay::Slot()
+					.HAlign(HAlign_Left)
+					.VAlign(VAlign_Top)
+					[
+						SNew(SBorder)
+							.BorderImage(FCoreStyle::Get().GetBrush("Docking.Border"))
+							.BorderBackgroundColor(FLinearColor(1.0f, 0.5f, 0.0f, 0.75f))
+							.Visibility(EVisibility::HitTestInvisible)
+							[
+								SNew(STextBlock)
+									.Text(this, &SDockingTabStack::ShowPersistentTabs)
+									.ShadowOffset(FVector2D::UnitVector)
+							]
+					]
+#endif
 			]
-			#endif
-		]
 	];
 
 	if (bIsDocumentArea)
 	{
-		this->SetNodeContent( SDocumentAreaWidget::MakeDocumentAreaWidget(), SNullWidget::NullWidget, SNullWidget::NullWidget );
+		this->SetNodeContent(SDocumentAreaWidget::MakeDocumentAreaWidget(), SNullWidget::NullWidget, SNullWidget::NullWidget);
 	}
 }
 
+
 void SDockingTabStack::OnLastTabRemoved()
 {
 	if (!bIsDocumentArea)
@@ -701,23 +707,29 @@ TSharedRef< SDockingTabStack > SDockingTabStack::CreateNewTabStackBySplitting( c
 
 void SDockingTabStack::SetParentNode( TSharedRef<class SDockingSplitter> InParent )
 {
-	SDockingNode::SetParentNode( InParent );
+	SDockingNode::SetParentNode(InParent);
 
 	// OK, if this docking area has a parent window, we'll assume the window was created with no title bar, and we'll
 	// place the title bar widgets into our content instead!
 	const TSharedPtr<SDockingArea>& DockArea = GetDockArea();
-	if( DockArea.IsValid() && DockArea->GetParentWindow().IsValid() )
+
+	if (DockArea.IsValid() && DockArea->GetParentWindow().IsValid())
 	{
 		// @todo mainframe: Really we only want these to show up for tab stacks that are along the top of the window,
 		//                  and only the first one!  Currently, all SDockingAreas with a parent window set will get
 		//                  title area widgets added!
 		const TSharedRef<SWindow>& ParentWindow = DockArea->GetParentWindow().ToSharedRef();
+
 		TSharedPtr<IWindowTitleBar> TitleBar;
 		TSharedRef<SWidget> TitleBarWidget = FSlateApplication::Get().MakeWindowTitleBar(ParentWindow, TitleBarContent, TitleBar);
 
 		(*TitleBarSlot)[TitleBarWidget];
 		ParentWindow->SetTitleBar(TitleBar);
 	}
+	else
+	{
+		(*TitleBarSlot)[TitleBarContent.ToSharedRef()];
+	}
 }
 
 
