#UE4 - Remove double SBorder.h inclue from Slate.h

diff --git a/Engine/Source/Runtime/Slate/Public/Slate.h b/Engine/Source/Runtime/Slate/Public/Slate.h
--- a/Engine/Source/Runtime/Slate/Public/Slate.h
+++ b/Engine/Source/Runtime/Slate/Public/Slate.h
@@ -78,7 +78,6 @@ DECLARE_CYCLE_STAT_EXTERN( TEXT("Slate Misc Time"), STAT_SlateMiscTime, STATGROU
 #include "SCompoundWidget.h"
 #include "SFxWidget.h"
 #include "SBorder.h"
-#include "SBorder.h"
 #include "SSeparator.h"
 #include "SSpacer.h"
 #include "SWrapBox.h"
