Slate - Adding additional setters and getters for a few attributes, SBorder, SCompoundWidget.

diff --git a/Engine/Source/Runtime/Slate/Private/Widgets/Layout/SBorder.cpp b/Engine/Source/Runtime/Slate/Private/Widgets/Layout/SBorder.cpp
--- a/Engine/Source/Runtime/Slate/Private/Widgets/Layout/SBorder.cpp
+++ b/Engine/Source/Runtime/Slate/Private/Widgets/Layout/SBorder.cpp
@@ -65,16 +65,6 @@ void SBorder::ClearContent()
 	ChildSlot = FSimpleSlot();
 }
 
-void SBorder::SetBorderBackgroundColor( const TAttribute<FSlateColor>& InColorAndOpacity )
-{
-	BorderBackgroundColor = InColorAndOpacity;
-}
-
-void SBorder::SetDesiredSizeScale( const TAttribute<FVector2D>& InDesiredSizeScale )
-{
-	DesiredSizeScale = InDesiredSizeScale;
-}
-
 
 /**
  * The widget should respond by populating the OutDrawElements array with FDrawElements 
diff --git a/Engine/Source/Runtime/Slate/Public/Widgets/Layout/SBorder.h b/Engine/Source/Runtime/Slate/Public/Widgets/Layout/SBorder.h
--- a/Engine/Source/Runtime/Slate/Public/Widgets/Layout/SBorder.h
+++ b/Engine/Source/Runtime/Slate/Public/Widgets/Layout/SBorder.h
@@ -86,11 +86,29 @@ public:
 	/** Clears out the content for the border */
 	void ClearContent();
 
+	/** Gets the color and opacity of the background image of this border. */
+	FSlateColor GetBorderBackgroundColor()
+	{
+		return BorderBackgroundColor.Get();
+	}
+
 	/** Sets the color and opacity of the background image of this border. */
-	void SetBorderBackgroundColor( const TAttribute<FSlateColor>& InColorAndOpacity );
+	void SetBorderBackgroundColor( const TAttribute<FSlateColor>& InColorAndOpacity )
+	{
+		BorderBackgroundColor = InColorAndOpacity;
+	}
+
+	/** Set the desired size scale multiplier */
+	FVector2D GetDesiredSizeScale()
+	{
+		return DesiredSizeScale.Get();
+	}
 
 	/** Set the desired size scale multiplier */
-	void SetDesiredSizeScale( const TAttribute<FVector2D>& InDesiredSizeScale );
+	void SetDesiredSizeScale( const TAttribute<FVector2D>& InDesiredSizeScale )
+	{
+		DesiredSizeScale = InDesiredSizeScale;
+	}
 
 public:
 	// SWidget interface
diff --git a/Engine/Source/Runtime/SlateCore/Public/Widgets/SCompoundWidget.h b/Engine/Source/Runtime/SlateCore/Public/Widgets/SCompoundWidget.h
--- a/Engine/Source/Runtime/SlateCore/Public/Widgets/SCompoundWidget.h
+++ b/Engine/Source/Runtime/SlateCore/Public/Widgets/SCompoundWidget.h
@@ -38,6 +38,14 @@ public:
 		ContentScale = InContentScale;
 	}
 
+	/**
+	 * Gets the widget's color.
+	 */
+	FLinearColor GetColorAndOpacity() const
+	{
+		return ColorAndOpacity.Get();
+	}
+
 	/**
 	 * Sets the widget's color.
 	 *
