Slate - Adding get/set for Content padding on SButton.

diff --git a/Engine/Source/Runtime/Slate/Public/Widgets/Input/SButton.h b/Engine/Source/Runtime/Slate/Public/Widgets/Input/SButton.h
--- a/Engine/Source/Runtime/Slate/Public/Widgets/Input/SButton.h
+++ b/Engine/Source/Runtime/Slate/Public/Widgets/Input/SButton.h
@@ -114,6 +114,18 @@ public:
 	 */
 	void Construct( const FArguments& InArgs );
 
+	/** Gets the content padding of the button */
+	FMargin GetContentPadding() const
+	{
+		return ContentPadding.Get();
+	}
+
+	/** Sets the content padding of the button */
+	void SetContentPadding(const TAttribute<FMargin>& InContentPadding)
+	{
+		ContentPadding = InContentPadding;
+	}
+
 protected:
 	/** @return combines the user-specified margin and the button's internal margin. */
 	FMargin GetCombinedPadding() const;
