Slate/UMG - Exposing SetValue for the SSlider.  Exposing GetCursor for SWidget.  Adding a new UStruct EditableTextWidgetStyle.

diff --git a/Engine/Source/Runtime/Slate/Private/Framework/Styling/EditableTextWidgetStyle.cpp b/Engine/Source/Runtime/Slate/Private/Framework/Styling/EditableTextWidgetStyle.cpp
new file mode 100644
--- /dev/null
+++ b/Engine/Source/Runtime/Slate/Private/Framework/Styling/EditableTextWidgetStyle.cpp
@@ -0,0 +1,10 @@
+// Copyright 1998-2014 Epic Games, Inc. All Rights Reserved.
+
+#include "SlatePrivatePCH.h"
+
+
+UEditableTextWidgetStyle::UEditableTextWidgetStyle( const class FPostConstructInitializeProperties& PCIP )
+	: Super(PCIP)
+{
+	
+}
\ No newline at end of file
diff --git a/Engine/Source/Runtime/Slate/Private/SlatePrivatePCH.h b/Engine/Source/Runtime/Slate/Private/SlatePrivatePCH.h
--- a/Engine/Source/Runtime/Slate/Private/SlatePrivatePCH.h
+++ b/Engine/Source/Runtime/Slate/Private/SlatePrivatePCH.h
@@ -27,3 +27,4 @@
 #include "ComboBoxWidgetStyle.h"
 #include "ComboButtonWidgetStyle.h"
 #include "TextBlockWidgetStyle.h"
+#include "EditableTextWidgetStyle.h"
diff --git a/Engine/Source/Runtime/Slate/Private/Widgets/Input/SSlider.cpp b/Engine/Source/Runtime/Slate/Private/Widgets/Input/SSlider.cpp
--- a/Engine/Source/Runtime/Slate/Private/Widgets/Input/SSlider.cpp
+++ b/Engine/Source/Runtime/Slate/Private/Widgets/Input/SSlider.cpp
@@ -189,3 +189,8 @@ float SSlider::PositionToValue( const FGeometry& MyGeometry, const FVector2D& Ab
 	
 	return FMath::Clamp(RelativeValue, 0.0f, 1.0f);
 }
+
+void SSlider::SetValue(const TAttribute<float>& InValueAttribute)
+{
+	ValueAttribute = InValueAttribute;
+}
diff --git a/Engine/Source/Runtime/Slate/Public/Framework/Styling/EditableTextWidgetStyle.h b/Engine/Source/Runtime/Slate/Public/Framework/Styling/EditableTextWidgetStyle.h
new file mode 100644
--- /dev/null
+++ b/Engine/Source/Runtime/Slate/Public/Framework/Styling/EditableTextWidgetStyle.h
@@ -0,0 +1,24 @@
+// Copyright 1998-2014 Epic Games, Inc. All Rights Reserved.
+
+#pragma once
+
+#include "SlateWidgetStyleContainerBase.h"
+#include "EditableTextWidgetStyle.generated.h"
+
+/**
+ */
+UCLASS(hidecategories=Object, MinimalAPI)
+class UEditableTextWidgetStyle : public USlateWidgetStyleContainerBase
+{
+	GENERATED_UCLASS_BODY()
+
+public:
+	/** The actual data describing the button's appearance. */
+	UPROPERTY(Category=Appearance, EditAnywhere, meta=(ShowOnlyInnerProperties))
+	FEditableTextStyle EditableTextStyle;
+
+	virtual const struct FSlateWidgetStyle* const GetStyle() const OVERRIDE
+	{
+		return static_cast< const struct FSlateWidgetStyle* >( &EditableTextStyle );
+	}
+};
diff --git a/Engine/Source/Runtime/Slate/Public/Widgets/Input/SSlider.h b/Engine/Source/Runtime/Slate/Public/Widgets/Input/SSlider.h
--- a/Engine/Source/Runtime/Slate/Public/Widgets/Input/SSlider.h
+++ b/Engine/Source/Runtime/Slate/Public/Widgets/Input/SSlider.h
@@ -74,6 +74,13 @@ public:
 		return ValueAttribute.Get();
 	}
 
+	/**
+	* Gets the slider's value.
+	*
+	* @return Slider value.
+	*/
+	void SetValue(const TAttribute<float>& InValueAttribute);
+
 public:
 
 	virtual int32 OnPaint( const FGeometry& AllottedGeometry, const FSlateRect& MyClippingRect, FSlateWindowElementList& OutDrawElements, int32 LayerId, const FWidgetStyle& InWidgetStyle, bool bParentEnabled ) const OVERRIDE;
diff --git a/Engine/Source/Runtime/SlateCore/Private/Widgets/SWidget.cpp b/Engine/Source/Runtime/SlateCore/Private/Widgets/SWidget.cpp
--- a/Engine/Source/Runtime/SlateCore/Private/Widgets/SWidget.cpp
+++ b/Engine/Source/Runtime/SlateCore/Private/Widgets/SWidget.cpp
@@ -459,6 +459,12 @@ void SWidget::EnableToolTipForceField( const bool bEnableForceField )
 }
 
 
+TOptional<EMouseCursor::Type> SWidget::GetCursor() const
+{
+	return Cursor.Get();
+}
+
+
 void SWidget::SetCursor( const TAttribute< TOptional<EMouseCursor::Type> >& InCursor )
 {
 	Cursor = InCursor;
diff --git a/Engine/Source/Runtime/SlateCore/Public/Widgets/SWidget.h b/Engine/Source/Runtime/SlateCore/Public/Widgets/SWidget.h
--- a/Engine/Source/Runtime/SlateCore/Public/Widgets/SWidget.h
+++ b/Engine/Source/Runtime/SlateCore/Public/Widgets/SWidget.h
@@ -580,6 +580,11 @@ public:
 	 */
 	void SetToolTip( const TSharedPtr<IToolTip>& InToolTip );
 
+	/**
+	 * Get the cursor that should appear when this widget is hovered
+	 */
+	TOptional<EMouseCursor::Type> GetCursor() const;
+
 	/**
 	 * Set the cursor that should appear when this widget is hovered
 	 */
