giving SProgressBar support to fill up from the left to right, right to left, and from center to the outer edges.

diff --git a/Engine/Source/Runtime/Slate/Private/Widgets/Notifications/SProgressBar.cpp b/Engine/Source/Runtime/Slate/Private/Widgets/Notifications/SProgressBar.cpp
--- a/Engine/Source/Runtime/Slate/Private/Widgets/Notifications/SProgressBar.cpp
+++ b/Engine/Source/Runtime/Slate/Private/Widgets/Notifications/SProgressBar.cpp
@@ -9,6 +9,8 @@ void SProgressBar::Construct( const FArguments& InArgs )
 
 	Percent = InArgs._Percent;
 
+	BarFillType = InArgs._BarFillType;
+
 	BackgroundImage = &InArgs._Style->BackgroundImage;
 	FillImage = &InArgs._Style->FillImage;
 	MarqueeImage = &InArgs._Style->MarqueeImage;
@@ -49,20 +51,61 @@ int32 SProgressBar::OnPaint( const FGeometry& AllottedGeometry, const FSlateRect
 	
 	if( ProgressFraction.IsSet() )
 	{
-		const float ClampedFraction = FMath::Clamp( ProgressFraction.GetValue(), 0.0f, 1.0f  );
-
-		// Draw Fill
-		FSlateDrawElement::MakeBox(
-			OutDrawElements,
-			RetLayerId++,
-			AllottedGeometry.ToPaintGeometry(
-				FVector2D::ZeroVector,
-				FVector2D( AllottedGeometry.Size.X * ( ClampedFraction ) , AllottedGeometry.Size.Y )),
-			FillImage,
-			ForegroundClippingRect,
-			DrawEffects,
-			FillColorAndOpacitySRGB
-			);
+		switch (BarFillType)
+		{
+			case EProgressBarFillType::RightToLeft:
+			{
+				const float ClampedFraction = FMath::Clamp( ProgressFraction.GetValue(), 0.0f, 1.0f  );
+				// Draw Fill
+				FSlateDrawElement::MakeBox(
+					OutDrawElements,
+					RetLayerId++,
+					AllottedGeometry.ToPaintGeometry(
+						FVector2D( AllottedGeometry.Size.X - (AllottedGeometry.Size.X * ( ClampedFraction )) , 0.0f),
+						FVector2D( AllottedGeometry.Size.X * ( ClampedFraction ) , AllottedGeometry.Size.Y )),
+					FillImage,
+					ForegroundClippingRect,
+					DrawEffects,
+					FillColorAndOpacitySRGB
+					);
+				break;
+			}
+			case EProgressBarFillType::FillFromCenter:
+			{
+				const float ClampedFraction = FMath::Clamp( ProgressFraction.GetValue(), 0.0f, 1.0f  );
+				// Draw Fill
+				FSlateDrawElement::MakeBox(
+					OutDrawElements,
+					RetLayerId++,
+					AllottedGeometry.ToPaintGeometry(
+						FVector2D( (AllottedGeometry.Size.X * 0.5f) - ((AllottedGeometry.Size.X * ( ClampedFraction ))*0.5), 0.0f),
+						FVector2D( AllottedGeometry.Size.X * ( ClampedFraction ) , AllottedGeometry.Size.Y )),
+					FillImage,
+					ForegroundClippingRect,
+					DrawEffects,
+					FillColorAndOpacitySRGB
+					);
+				break;
+			}
+			case EProgressBarFillType::LeftToRight:
+			default:
+			{
+				const float ClampedFraction = FMath::Clamp( ProgressFraction.GetValue(), 0.0f, 1.0f  );
+				// Draw Fill
+				FSlateDrawElement::MakeBox(
+					OutDrawElements,
+					RetLayerId++,
+					AllottedGeometry.ToPaintGeometry(
+						FVector2D::ZeroVector,
+						FVector2D( AllottedGeometry.Size.X * ( ClampedFraction ) , AllottedGeometry.Size.Y )),
+					FillImage,
+					ForegroundClippingRect,
+					DrawEffects,
+					FillColorAndOpacitySRGB
+					);
+				break;
+			}
+		}
 		
 	}
 	else
diff --git a/Engine/Source/Runtime/Slate/Public/Widgets/Notifications/SProgressBar.h b/Engine/Source/Runtime/Slate/Public/Widgets/Notifications/SProgressBar.h
--- a/Engine/Source/Runtime/Slate/Public/Widgets/Notifications/SProgressBar.h
+++ b/Engine/Source/Runtime/Slate/Public/Widgets/Notifications/SProgressBar.h
@@ -2,6 +2,25 @@
 
 #pragma once
 
+#include "SProgressBar.generated.h"
+
+/**
+ * SProgressBar Fill Type 
+ */
+UENUM()
+namespace EProgressBarFillType
+{
+	enum Type
+	{
+		// will fill up from the left side to the right
+		LeftToRight,
+		// will fill up from the right side to the left side
+		RightToLeft,
+		// will fill up from the center to the outer edges
+		FillFromCenter,
+	};
+}
+
 /** A progress bar widget.*/
 class SLATE_API SProgressBar : public SLeafWidget
 {
@@ -9,6 +28,7 @@ class SLATE_API SProgressBar : public SLeafWidget
 public:
 	SLATE_BEGIN_ARGS(SProgressBar)
 		: _Style( &FCoreStyle::Get().GetWidgetStyle<FProgressBarStyle>("ProgressBar") )
+		, _BarFillType(EProgressBarFillType::LeftToRight)
 		, _Percent( TOptional<float>() )
 		, _FillColorAndOpacity( FLinearColor::White )
 		, _BorderPadding( FVector2D(1,0) )
@@ -17,6 +37,9 @@ public:
 		/** Style used for the progress bar */
 		SLATE_STYLE_ARGUMENT( FProgressBarStyle, Style )
 
+		/** Defines if this progress bar fills Left to right or right to left*/
+		SLATE_ARGUMENT( EProgressBarFillType::Type, BarFillType )
+
 		/** Used to determine the fill position of the progress bar ranging 0..1 */
 		SLATE_ATTRIBUTE( TOptional<float>, Percent )
 
@@ -44,6 +67,8 @@ private:
 	/** The text displayed over the progress bar */
 	TAttribute< TOptional<float> > Percent;
 
+	EProgressBarFillType::Type BarFillType;
+
 	/** Background image to use for the progress bar */
 	const FSlateBrush* BackgroundImage;
 
