Fixed user defined keybindings not loading

TTP# 333749 - CRITICAL: Regression: Editor: Keyboard shortcuts reset on editor relaunch

CharToEscapeSeqMap had recently had a conversion from \\ -> \\\\ added to it, but ReplaceCharWithEscapedChar and ReplaceEscapedCharWithChar both span the array in the same order, so the unescape sequence wouldn't catch all the places where \\ had been escaped when processing carriage returns. ReplaceEscapedCharWithChar now spins the array backwards to ensure it's doing the inverse operation.

I also had the gestures written out using the compact JSON format, since the new-lines aren't used in the ini file.

diff --git a/Engine/Source/Runtime/Core/Private/Containers/String.cpp b/Engine/Source/Runtime/Core/Private/Containers/String.cpp
--- a/Engine/Source/Runtime/Core/Private/Containers/String.cpp
+++ b/Engine/Source/Runtime/Core/Private/Containers/String.cpp
@@ -1001,7 +1001,8 @@ FString FString::ReplaceEscapedCharWithChar( const TArray<TCHAR>* Chars/*=NULL*/
 	if ( Len() > 0 && (Chars == NULL || Chars->Num() > 0) )
 	{
 		FString Result(*this);
-		for ( int32 ChIdx = 0; ChIdx < MAX_SUPPORTED_ESCAPE_CHARS; ChIdx++ )
+		// Spin CharToEscapeSeqMap backwards to ensure we're doing the inverse of ReplaceCharWithEscapedChar
+		for ( int32 ChIdx = MAX_SUPPORTED_ESCAPE_CHARS - 1; ChIdx >= 0; ChIdx-- )
 		{
 			if ( Chars == NULL || Chars->Contains(*(CharToEscapeSeqMap[ChIdx][0])) )
 			{
diff --git a/Engine/Source/Runtime/Slate/Private/Framework/Commands/InputBindingManager.cpp b/Engine/Source/Runtime/Slate/Private/Framework/Commands/InputBindingManager.cpp
--- a/Engine/Source/Runtime/Slate/Private/Framework/Commands/InputBindingManager.cpp
+++ b/Engine/Source/Runtime/Slate/Private/Framework/Commands/InputBindingManager.cpp
@@ -14,8 +14,8 @@ class FUserDefinedGestures
 {
 public:
 	void LoadGestures();
-	void SaveGestures();
-	bool GetUserDefinedGesture( const FName BindingContext, const FName CommandName, FInputGesture& OutUserDefinedGesture );
+	void SaveGestures() const;
+	bool GetUserDefinedGesture( const FName BindingContext, const FName CommandName, FInputGesture& OutUserDefinedGesture ) const;
 	void SetUserDefinedGesture( const FUICommandInfo& CommandInfo );
 	/** Remove all user defined gestures */
 	void RemoveAll();
@@ -36,12 +36,11 @@ void FUserDefinedGestures::LoadGestures()
 		if (!Gestures.IsValid())
 		{
 			// Gestures have not been loaded from the ini file, try reading them from the txt file now
-			FArchive* Ar = IFileManager::Get().CreateFileReader( *( FPaths::GameSavedDir() / TEXT("Preferences/EditorKeyBindings.txt") ) );
-			if( Ar )
+			TSharedPtr<FArchive> Ar = MakeShareable( IFileManager::Get().CreateFileReader( *( FPaths::GameSavedDir() / TEXT("Preferences/EditorKeyBindings.txt") ) ) );
+			if( Ar.IsValid() )
 			{
-				TSharedRef< TJsonReader<ANSICHAR> > TextReader = TJsonReaderFactory<ANSICHAR>::Create( Ar );
+				TSharedRef< TJsonReader<ANSICHAR> > TextReader = TJsonReaderFactory<ANSICHAR>::Create( Ar.Get() );
 				bResult = FJsonSerializer::Deserialize( TextReader, Gestures );
-				delete Ar;
 			}
 		}
 		
@@ -52,19 +51,19 @@ void FUserDefinedGestures::LoadGestures()
 	}
 }
 
-void FUserDefinedGestures::SaveGestures()
+void FUserDefinedGestures::SaveGestures() const
 {
 	if( Gestures.IsValid() ) 
 	{
 		FString Content;
-		TSharedRef< TJsonWriter<> > Writer = TJsonWriterFactory<>::Create( &Content );
+		auto Writer = TJsonWriterFactory< TCHAR, TCondensedJsonPrintPolicy<TCHAR> >::Create( &Content );
 		FJsonSerializer::Serialize( Gestures.ToSharedRef(), Writer );
 
 		GConfig->SetString(TEXT("UserDefinedGestures"), TEXT("Content"), *FRemoteConfig::ReplaceIniCharWithSpecialChar(Content).ReplaceCharWithEscapedChar(), GEditorKeyBindingsIni);
 	}
 }
 
-bool FUserDefinedGestures::GetUserDefinedGesture( const FName BindingContext, const FName CommandName, FInputGesture& OutUserDefinedGesture )
+bool FUserDefinedGestures::GetUserDefinedGesture( const FName BindingContext, const FName CommandName, FInputGesture& OutUserDefinedGesture ) const
 {
 	bool bResult = false;
 
