User gesture data is now saved as an ini-style array of JSON data

This allows it to be merged when the config files are loaded (it was previously just a string of JSON data, which meant whichever file was loaded last would be the winner).

diff --git a/Engine/Source/Runtime/Slate/Private/Framework/Commands/InputBindingManager.cpp b/Engine/Source/Runtime/Slate/Private/Framework/Commands/InputBindingManager.cpp
--- a/Engine/Source/Runtime/Slate/Private/Framework/Commands/InputBindingManager.cpp
+++ b/Engine/Source/Runtime/Slate/Private/Framework/Commands/InputBindingManager.cpp
@@ -10,6 +10,38 @@
 /* FUserDefinedGestures helper class
  *****************************************************************************/
 
+/** An identifier for a user defined gesture */
+struct FUserDefinedGestureKey
+{
+	FUserDefinedGestureKey()
+	{
+	}
+
+	FUserDefinedGestureKey(const FName InBindingContext, const FName InCommandName)
+		: BindingContext(InBindingContext)
+		, CommandName(InCommandName)
+	{
+	}
+
+	bool operator==(const FUserDefinedGestureKey& Other) const
+	{
+		return BindingContext == Other.BindingContext && CommandName == Other.CommandName;
+	}
+
+	bool operator!=(const FUserDefinedGestureKey& Other) const
+	{
+		return !(*this == Other);
+	}
+
+	FName BindingContext;
+	FName CommandName;
+};
+
+uint32 GetTypeHash( const FUserDefinedGestureKey& Key )
+{
+	return GetTypeHash(Key.BindingContext) ^ GetTypeHash(Key.CommandName);
+}
+
 class FUserDefinedGestures
 {
 public:
@@ -20,46 +52,136 @@ public:
 	/** Remove all user defined gestures */
 	void RemoveAll();
 private:
-	TSharedPtr<FJsonObject> Gestures;
+	/* Mapping from a gesture key to the user defined gesture */
+	typedef TMap<FUserDefinedGestureKey, FInputGesture> FGesturesMap;
+	TSharedPtr<FGesturesMap> Gestures;
 };
 
 void FUserDefinedGestures::LoadGestures()
 {
 	if( !Gestures.IsValid() )
 	{
-		FString Content;
-		GConfig->GetString(TEXT("UserDefinedGestures"), TEXT("Content"), Content, GEditorKeyBindingsIni);
-
-		TSharedRef< TJsonReader<> > Reader = TJsonReaderFactory<>::Create( FRemoteConfig::ReplaceIniSpecialCharWithChar(Content).ReplaceEscapedCharWithChar() );
-		bool bResult = FJsonSerializer::Deserialize( Reader, Gestures );
+		Gestures = MakeShareable( new FGesturesMap );
 
-		if (!Gestures.IsValid())
+		// First, try and load the gestures from their new location in the ini file
+		// Failing that, try and load them from the older txt file
+		TArray<FString> GestureJsonArray;
+		if( GConfig->GetArray(TEXT("UserDefinedGestures"), TEXT("UserDefinedGestures"), GestureJsonArray, GEditorKeyBindingsIni) )
 		{
-			// Gestures have not been loaded from the ini file, try reading them from the txt file now
-			TSharedPtr<FArchive> Ar = MakeShareable( IFileManager::Get().CreateFileReader( *( FPaths::GameSavedDir() / TEXT("Preferences/EditorKeyBindings.txt") ) ) );
-			if( Ar.IsValid() )
+			// This loads an array of JSON strings representing the FUserDefinedGestureKey and FInputGesture in a single JSON object
+			for(const FString& GestureJson : GestureJsonArray)
 			{
-				TSharedRef< TJsonReader<ANSICHAR> > TextReader = TJsonReaderFactory<ANSICHAR>::Create( Ar.Get() );
-				bResult = FJsonSerializer::Deserialize( TextReader, Gestures );
+				const FString UnescapedContent = FRemoteConfig::ReplaceIniSpecialCharWithChar(GestureJson).ReplaceEscapedCharWithChar();
+
+				TSharedPtr<FJsonObject> GestureInfoObj;
+				auto JsonReader = TJsonReaderFactory<>::Create( UnescapedContent );
+				if( FJsonSerializer::Deserialize( JsonReader, GestureInfoObj ) )
+				{
+					const TSharedPtr<FJsonValue> BindingContextObj = GestureInfoObj->Values.FindRef( TEXT("BindingContext") );
+					const TSharedPtr<FJsonValue> CommandNameObj = GestureInfoObj->Values.FindRef( TEXT("CommandName") );
+					const TSharedPtr<FJsonValue> CtrlObj = GestureInfoObj->Values.FindRef( TEXT("Control") );
+					const TSharedPtr<FJsonValue> AltObj = GestureInfoObj->Values.FindRef( TEXT("Alt") );
+					const TSharedPtr<FJsonValue> ShiftObj = GestureInfoObj->Values.FindRef( TEXT("Shift") );
+					const TSharedPtr<FJsonValue> KeyObj = GestureInfoObj->Values.FindRef( TEXT("Key") );
+
+					const FName BindingContext = *BindingContextObj->AsString();
+					const FName CommandName = *CommandNameObj->AsString();
+
+					const FUserDefinedGestureKey GestureKey(BindingContext, CommandName);
+					FInputGesture& UserDefinedGesture = Gestures->FindOrAdd(GestureKey);
+
+					UserDefinedGesture.bCtrl = CtrlObj->AsBool();
+					UserDefinedGesture.bAlt = AltObj->AsBool();
+					UserDefinedGesture.bShift = ShiftObj->AsBool();
+					UserDefinedGesture.Key = *KeyObj->AsString();
+				}
 			}
 		}
-		
-		if( !Gestures.IsValid() )
+		else
 		{
-			Gestures = MakeShareable( new FJsonObject );
+			TSharedPtr<FJsonObject> GesturesObj;
+
+			// This loads a JSON object containing BindingContexts, containing objects of CommandNames, containing the FInputGesture information
+			FString Content;
+			if( GConfig->GetString(TEXT("UserDefinedGestures"), TEXT("Content"), Content, GEditorKeyBindingsIni) )
+			{
+				const FString UnescapedContent = FRemoteConfig::ReplaceIniSpecialCharWithChar(Content).ReplaceEscapedCharWithChar();
+
+				auto JsonReader = TJsonReaderFactory<>::Create( UnescapedContent );
+				FJsonSerializer::Deserialize( JsonReader, GesturesObj );
+			}
+
+			if (!GesturesObj.IsValid())
+			{
+				// Gestures have not been loaded from the ini file, try reading them from the txt file now
+				TSharedPtr<FArchive> Ar = MakeShareable( IFileManager::Get().CreateFileReader( *( FPaths::GameSavedDir() / TEXT("Preferences/EditorKeyBindings.txt") ) ) );
+				if( Ar.IsValid() )
+				{
+					auto TextReader = TJsonReaderFactory<ANSICHAR>::Create( Ar.Get() );
+					FJsonSerializer::Deserialize( TextReader, GesturesObj );
+				}
+			}
+
+			if (GesturesObj.IsValid())
+			{
+				for(const auto& BindingContextInfo : GesturesObj->Values)
+				{
+					const FName BindingContext = *BindingContextInfo.Key;
+					TSharedPtr<FJsonObject> BindingContextObj = BindingContextInfo.Value->AsObject();
+					for(const auto& CommandInfo : BindingContextObj->Values)
+					{
+						const FName CommandName = *CommandInfo.Key;
+						TSharedPtr<FJsonObject> CommandObj = CommandInfo.Value->AsObject();
+
+						const TSharedPtr<FJsonValue> CtrlObj = CommandObj->Values.FindRef( TEXT("Control") );
+						const TSharedPtr<FJsonValue> AltObj = CommandObj->Values.FindRef( TEXT("Alt") );
+						const TSharedPtr<FJsonValue> ShiftObj = CommandObj->Values.FindRef( TEXT("Shift") );
+						const TSharedPtr<FJsonValue> KeyObj = CommandObj->Values.FindRef( TEXT("Key") );
+
+						const FUserDefinedGestureKey GestureKey(BindingContext, CommandName);
+						FInputGesture& UserDefinedGesture = Gestures->FindOrAdd(GestureKey);
+
+						UserDefinedGesture.bCtrl = CtrlObj->AsBool();
+						UserDefinedGesture.bAlt = AltObj->AsBool();
+						UserDefinedGesture.bShift = ShiftObj->AsBool();
+						UserDefinedGesture.Key = *KeyObj->AsString();
+					}
+				}
+			}
 		}
 	}
 }
 
 void FUserDefinedGestures::SaveGestures() const
 {
-	if( Gestures.IsValid() ) 
+	if( Gestures.IsValid() )
 	{
-		FString Content;
-		auto Writer = TJsonWriterFactory< TCHAR, TCondensedJsonPrintPolicy<TCHAR> >::Create( &Content );
-		FJsonSerializer::Serialize( Gestures.ToSharedRef(), Writer );
+		FString GestureRawJsonContent;
+		TArray<FString> GestureJsonArray;
+		for(const auto& GestureInfo : *Gestures)
+		{
+			TSharedPtr<FJsonValueObject> GestureInfoValueObj = MakeShareable( new FJsonValueObject( MakeShareable( new FJsonObject ) ) );
+			TSharedPtr<FJsonObject> GestureInfoObj = GestureInfoValueObj->AsObject();
+
+			// Set the gesture values for the command
+			GestureInfoObj->Values.Add( TEXT("BindingContext"), MakeShareable( new FJsonValueString( GestureInfo.Key.BindingContext.ToString() ) ) );
+			GestureInfoObj->Values.Add( TEXT("CommandName"), MakeShareable( new FJsonValueString( GestureInfo.Key.CommandName.ToString() ) ) );
+			GestureInfoObj->Values.Add( TEXT("Control"), MakeShareable( new FJsonValueBoolean( GestureInfo.Value.bCtrl ) ) );
+			GestureInfoObj->Values.Add( TEXT("Alt"), MakeShareable( new FJsonValueBoolean( GestureInfo.Value.bAlt ) ) );
+			GestureInfoObj->Values.Add( TEXT("Shift"), MakeShareable( new FJsonValueBoolean( GestureInfo.Value.bShift ) ) );
+			GestureInfoObj->Values.Add( TEXT("Key"), MakeShareable( new FJsonValueString( GestureInfo.Value.Key.ToString() ) ) );
+
+			auto JsonWriter = TJsonWriterFactory< TCHAR, TCondensedJsonPrintPolicy<TCHAR> >::Create( &GestureRawJsonContent );
+			FJsonSerializer::Serialize( GestureInfoObj.ToSharedRef(), JsonWriter );
+
+			const FString EscapedContent = FRemoteConfig::ReplaceIniCharWithSpecialChar(GestureRawJsonContent).ReplaceCharWithEscapedChar();
+			GestureJsonArray.Add(EscapedContent);
+		}
+
+		GConfig->SetArray(TEXT("UserDefinedGestures"), TEXT("UserDefinedGestures"), GestureJsonArray, GEditorKeyBindingsIni);
 
-		GConfig->SetString(TEXT("UserDefinedGestures"), TEXT("Content"), *FRemoteConfig::ReplaceIniCharWithSpecialChar(Content).ReplaceCharWithEscapedChar(), GEditorKeyBindingsIni);
+		// Clean up the old Content key (if it still exists)
+		GConfig->RemoveKey(TEXT("UserDefinedGestures"), TEXT("Content"), GEditorKeyBindingsIni);
 	}
 }
 
@@ -69,26 +191,12 @@ bool FUserDefinedGestures::GetUserDefinedGesture( const FName BindingContext, co
 
 	if( Gestures.IsValid() )
 	{
-		const TSharedPtr<FJsonValue> ContextObj = Gestures->Values.FindRef( BindingContext.ToString() );
-		if( ContextObj.IsValid() )
+		const FUserDefinedGestureKey GestureKey(BindingContext, CommandName);
+		const FInputGesture* const UserDefinedGesturePtr = Gestures->Find(GestureKey);
+		if( UserDefinedGesturePtr )
 		{
-			const TSharedPtr<FJsonValue> CommandObj = ContextObj->AsObject()->Values.FindRef( CommandName.ToString() );
-			if( CommandObj.IsValid() )
-			{
-				TSharedPtr<FJsonObject> GestureObj = CommandObj->AsObject();
-
-				const TSharedPtr<FJsonValue> KeyObj = GestureObj->Values.FindRef( TEXT("Key") );
-				const TSharedPtr<FJsonValue> CtrlObj = GestureObj->Values.FindRef( TEXT("Control") );
-				const TSharedPtr<FJsonValue> AltObj = GestureObj->Values.FindRef( TEXT("Alt") );
-				const TSharedPtr<FJsonValue> ShiftObj = GestureObj->Values.FindRef( TEXT("Shift") );
-
-				OutUserDefinedGesture.Key = *KeyObj->AsString();
-				OutUserDefinedGesture.bCtrl = CtrlObj->AsBool();
-				OutUserDefinedGesture.bAlt = AltObj->AsBool();
-				OutUserDefinedGesture.bShift = ShiftObj->AsBool();
-
-				bResult = true;
-			}
+			OutUserDefinedGesture = *UserDefinedGesturePtr;
+			bResult = true;
 		}
 	}
 
@@ -97,45 +205,25 @@ bool FUserDefinedGestures::GetUserDefinedGesture( const FName BindingContext, co
 
 void FUserDefinedGestures::SetUserDefinedGesture( const FUICommandInfo& CommandInfo )
 {
-	// Find or create the command context json value
-	const FName BindingContext = CommandInfo.GetBindingContext();
-	const FName CommandName = CommandInfo.GetCommandName();
-
-	TSharedPtr<FJsonValue> ContextObj = Gestures->Values.FindRef( BindingContext.ToString() );
-	if( !ContextObj.IsValid() )
-	{
-		ContextObj = Gestures->Values.Add( BindingContext.ToString(), MakeShareable( new FJsonValueObject( MakeShareable( new FJsonObject ) ) ) );
-	}
-
-	const TSharedPtr<const FInputGesture> InputGesture = CommandInfo.GetActiveGesture();
-
-	FInputGesture GestureToSave;
-
-	// Save an empty invalid gesture if one was not set
-	// This is an indication that the user doesn't want this bound and not to use the default gesture
-	if( InputGesture.IsValid() )
-	{
-		GestureToSave = *InputGesture;
-	}
-
+	if( Gestures.IsValid() )
 	{
-		TSharedPtr<FJsonValueObject> GestureValueObj = MakeShareable( new FJsonValueObject( MakeShareable( new FJsonObject ) ) );
-		TSharedPtr<FJsonObject> GestureObj = GestureValueObj->AsObject();
+		const FName BindingContext = CommandInfo.GetBindingContext();
+		const FName CommandName = CommandInfo.GetCommandName();
 
-		// Set the gesture values for the command
-		GestureObj->Values.Add( TEXT("Key"), MakeShareable( new FJsonValueString( GestureToSave.Key.ToString() ) ) );
-		GestureObj->Values.Add( TEXT("Control"),  MakeShareable( new FJsonValueBoolean( GestureToSave.bCtrl ) ) );
-		GestureObj->Values.Add( TEXT("Alt"),  MakeShareable( new FJsonValueBoolean( GestureToSave.bAlt ) ) );
-		GestureObj->Values.Add( TEXT("Shift"),  MakeShareable( new FJsonValueBoolean( GestureToSave.bShift ) ) );
+		// Find or create the command context
+		const FUserDefinedGestureKey GestureKey(BindingContext, CommandName);
+		FInputGesture& UserDefinedGesture = Gestures->FindOrAdd(GestureKey);
 
-		ContextObj->AsObject()->Values.Add( CommandName.ToString(), GestureValueObj );
+		// Save an empty invalid gesture if one was not set
+		// This is an indication that the user doesn't want this bound and not to use the default gesture
+		const TSharedPtr<const FInputGesture> InputGesture = CommandInfo.GetActiveGesture();
+		UserDefinedGesture = (InputGesture.IsValid()) ? *InputGesture : FInputGesture();
 	}
-
 }
 
 void FUserDefinedGestures::RemoveAll()
 {
-	Gestures = MakeShareable(new FJsonObject);
+	Gestures = MakeShareable( new FGesturesMap );
 }
 
 
