Slate - Fixing a crash when a user clicks and drags on a dropdown combobox menu that causes the captor to be destroyed, but the mouse captor code wasn't checking for a valid path.

diff --git a/Engine/Source/Runtime/Slate/Private/Framework/Application/SlateApplication.cpp b/Engine/Source/Runtime/Slate/Private/Framework/Application/SlateApplication.cpp
--- a/Engine/Source/Runtime/Slate/Private/Framework/Application/SlateApplication.cpp
+++ b/Engine/Source/Runtime/Slate/Private/Framework/Application/SlateApplication.cpp
@@ -2079,12 +2079,15 @@ void FSlateApplication::QueryCursor()
 			if (MouseCaptor.IsValid())
 			{
 				FWidgetPath MouseCaptorPath = MouseCaptor.ToWidgetPath();
-				TSharedRef< SWindow > CaptureWindow = MouseCaptorPath.GetWindow();
-			
-				// Never query the mouse captor path if it is outside an active modal window
-				if (!ActiveModalWindow.IsValid() || (CaptureWindow == ActiveModalWindow || CaptureWindow->IsDescendantOf(ActiveModalWindow)))
+				if ( MouseCaptorPath.IsValid() )
 				{
-					WidgetsToQueryForCursor = MouseCaptorPath;
+					TSharedRef< SWindow > CaptureWindow = MouseCaptorPath.GetWindow();
+
+					// Never query the mouse captor path if it is outside an active modal window
+					if ( !ActiveModalWindow.IsValid() || ( CaptureWindow == ActiveModalWindow || CaptureWindow->IsDescendantOf(ActiveModalWindow) ) )
+					{
+						WidgetsToQueryForCursor = MouseCaptorPath;
+					}
 				}
 			}
 			else
