Modified the code that decides whether a window title bar should be highlighted to depend on IsForegroundWindow(). This solves an issue with title bar getting deactivated when opening menus.

diff --git a/Engine/Source/Runtime/Slate/Private/Framework/Application/SWindowTitleBar.h b/Engine/Source/Runtime/Slate/Private/Framework/Application/SWindowTitleBar.h
--- a/Engine/Source/Runtime/Slate/Private/Framework/Application/SWindowTitleBar.h
+++ b/Engine/Source/Runtime/Slate/Private/Framework/Application/SWindowTitleBar.h
@@ -540,7 +540,8 @@ private:
 	const FSlateBrush* GetWindowTitlebackgroundImage( ) const
 	{
 		TSharedPtr<SWindow> OwnerWindow = OwnerWindowPtr.Pin();
-		const bool bIsActive = FSlateApplication::Get().GetActiveTopLevelWindow() == OwnerWindow || OwnerWindow->HasActiveChildren();
+		TSharedPtr<FGenericWindow> NativeWindow = OwnerWindow->GetNativeWindow();
+		const bool bIsActive = NativeWindow.IsValid() && NativeWindow->IsForegroundWindow();
 
 		return bIsActive ? &Style->ActiveTitleBrush : &Style->InactiveTitleBrush;
 	}
