Fixing Linux CIS (type resolution problems).

diff --git a/Engine/Source/Runtime/Slate/Public/Widgets/Views/STableRow.h b/Engine/Source/Runtime/Slate/Public/Widgets/Views/STableRow.h
--- a/Engine/Source/Runtime/Slate/Public/Widgets/Views/STableRow.h
+++ b/Engine/Source/Runtime/Slate/Public/Widgets/Views/STableRow.h
@@ -82,7 +82,7 @@ public:
 	 *
 	 * @param	InArgs	The declaration data for this widget
 	 */
-	void Construct(const FArguments& InArgs, const TSharedRef<STableViewBase>& InOwnerTableView)
+	 void Construct(const STableRow<ItemType>::FArguments& InArgs, const TSharedRef<STableViewBase>& InOwnerTableView)
 	{
 		ConstructInternal(InArgs, InOwnerTableView);
 
@@ -703,7 +703,7 @@ public:
 	typedef typename STableRow<ItemType>::FArguments FTableRowArgs;
 
 protected:
-	void Construct( const FArguments& InArgs, const TSharedRef<STableViewBase>& OwnerTableView )
+	void Construct(const FTableRowArgs& InArgs, const TSharedRef<STableViewBase>& OwnerTableView)
 	{
 		STableRow<ItemType>::Construct(
 			FTableRowArgs()
