A fix for a fix (now for Windows CIS). Both clang and cl should be happy.

diff --git a/Engine/Source/Runtime/Slate/Public/Widgets/Views/STableRow.h b/Engine/Source/Runtime/Slate/Public/Widgets/Views/STableRow.h
--- a/Engine/Source/Runtime/Slate/Public/Widgets/Views/STableRow.h
+++ b/Engine/Source/Runtime/Slate/Public/Widgets/Views/STableRow.h
@@ -82,7 +82,7 @@ public:
 	 *
 	 * @param	InArgs	The declaration data for this widget
 	 */
-	 void Construct(const STableRow<ItemType>::FArguments& InArgs, const TSharedRef<STableViewBase>& InOwnerTableView)
+	 void Construct(const typename STableRow<ItemType>::FArguments& InArgs, const TSharedRef<STableViewBase>& InOwnerTableView)
 	{
 		ConstructInternal(InArgs, InOwnerTableView);
 
