Fixing crash on Mac, when re-rooting a blueprint's component tree.
(TTP #334167)

diff --git a/Engine/Source/Runtime/Slate/Public/Widgets/Views/STreeView.h b/Engine/Source/Runtime/Slate/Public/Widgets/Views/STreeView.h
--- a/Engine/Source/Runtime/Slate/Public/Widgets/Views/STreeView.h
+++ b/Engine/Source/Runtime/Slate/Public/Widgets/Views/STreeView.h
@@ -312,17 +312,24 @@ private:
 
 	virtual bool Private_DoesItemHaveChildren( int32 ItemIndexInList ) const OVERRIDE
 	{
-		return DenseItemInfos[ItemIndexInList].bHasChildren;
+		bool bHasChildren = false;
+		if (ItemIndexInList < DenseItemInfos.Num())
+		{
+			bHasChildren = DenseItemInfos[ItemIndexInList].bHasChildren;
+		}
+		return bHasChildren;
 	}
 
 	virtual int32 Private_GetNestingDepth( int32 ItemIndexInList ) const OVERRIDE
 	{
-		return DenseItemInfos[ItemIndexInList].NestingLevel;
+		int32 NextingLevel = 0;
+		if (ItemIndexInList < DenseItemInfos.Num())
+		{
+			NextingLevel = DenseItemInfos[ItemIndexInList].NestingLevel;
+		}
+		return NextingLevel;
 	}
 
-
-
-
 public:
 		/**
 		 * See SWidget::Tick()
@@ -352,8 +359,6 @@ public:
 						// Rebuild the linearized view of the tree data.
 						LinearizedItems.Empty();
 						PopulateLinearizedItems( *TreeItemsSource, LinearizedItems, TempDenseItemInfos, 0, TempSelectedItemsMap, TempSparseItemInfo, true );
-						SparseItemInfos = TempSparseItemInfo;
-						DenseItemInfos = TempDenseItemInfos;
 
 						if(	this->SelectedItems.Num() != TempSelectedItemsMap.Num() || 
 							this->SelectedItems.Difference(TempSelectedItemsMap).Num() > 0 || 
@@ -375,6 +380,13 @@ public:
 
 							this->Private_SignalSelectionChanged(ESelectInfo::Direct);
 						}
+						
+						// these should come after Private_SignalSelectionChanged(); because through a
+						// series of calls, Private_SignalSelectionChanged() could end up in a child
+						// that indexes into either of these arrays (the index wouldn't be updated yet,
+						// and could be invalid)
+						SparseItemInfos = TempSparseItemInfo;
+						DenseItemInfos  = TempDenseItemInfos;
 					}
 				}
 			}
