SLATE - Adding a method to dynamically set the OnGetMenuContent of the SComboButton.

diff --git a/Engine/Source/Runtime/Slate/Private/Widgets/Input/SComboButton.cpp b/Engine/Source/Runtime/Slate/Private/Widgets/Input/SComboButton.cpp
--- a/Engine/Source/Runtime/Slate/Private/Widgets/Input/SComboButton.cpp
+++ b/Engine/Source/Runtime/Slate/Private/Widgets/Input/SComboButton.cpp
@@ -136,3 +136,8 @@ void SComboButton::SetMenuContent( const TSharedRef< SWidget >& InContent )
 			]
 		];
 }
+
+void SComboButton::SetOnGetMenuContent(FOnGetContent InOnGetMenuContent)
+{
+	OnGetMenuContent = InOnGetMenuContent;
+}
diff --git a/Engine/Source/Runtime/Slate/Public/Widgets/Input/SComboButton.h b/Engine/Source/Runtime/Slate/Public/Widgets/Input/SComboButton.h
--- a/Engine/Source/Runtime/Slate/Public/Widgets/Input/SComboButton.h
+++ b/Engine/Source/Runtime/Slate/Public/Widgets/Input/SComboButton.h
@@ -33,8 +33,13 @@ public:
 		SLATE_STYLE_ARGUMENT( FButtonStyle, ButtonStyle )
 		
 		SLATE_NAMED_SLOT( FArguments, ButtonContent )
+		
+		/** Optional static menu content.  If the menu content needs to be dynamically built, use the OnGetMenuContent event */
 		SLATE_NAMED_SLOT( FArguments, MenuContent )
+		
+		/** Sets an event handler to generate a widget dynamically when the menu is needed. */
 		SLATE_EVENT( FOnGetContent, OnGetMenuContent )
+		
 		SLATE_EVENT( FOnComboBoxOpened, OnComboBoxOpened )
 		SLATE_ARGUMENT( bool, IsFocusable )
 		SLATE_ARGUMENT( bool, HasDownArrow )
@@ -62,6 +67,9 @@ public:
 	 */
 	void SetMenuContent( const TSharedRef< SWidget >& InContent );
 
+	/** See the OnGetMenuContent event */
+	void SetOnGetMenuContent( FOnGetContent InOnGetMenuContent );
+
 	/**
 	 * Construct the widget from a declaration
 	 *
