Temporary fix for Linux window movement.

- In some cases SDockingTabStack is consuming all mouse events and reporting them as handled, which prevents any events from ever being routed to SWindow.
- This platform-specific hack makes events routed to SWindow, bypassing an underlying problem.

diff --git a/Engine/Source/Runtime/Slate/Private/Framework/Docking/SDockingTabStack.cpp b/Engine/Source/Runtime/Slate/Private/Framework/Docking/SDockingTabStack.cpp
--- a/Engine/Source/Runtime/Slate/Private/Framework/Docking/SDockingTabStack.cpp
+++ b/Engine/Source/Runtime/Slate/Private/Framework/Docking/SDockingTabStack.cpp
@@ -398,7 +398,14 @@ FReply SDockingTabStack::OnMouseButtonDown( const FGeometry& MyGeometry, const F
 	if ( ForegroundTab.IsValid() && !ForegroundTab->IsActive() )
 	{
 		FGlobalTabmanager::Get()->SetActiveTab( ForegroundTab );
+#if PLATFORM_LINUX
+		// Don't stop further event handling in case the user wants to move this window.
+		// Returning FReply::Handled() here will prevent SWindow from seeing the event.
+		// FIXME: In some cases the foreground tab is never the active tab and this handler will consume every mouse down event.
+		return FReply::Unhandled();
+#else
 		return FReply::Handled();
+#endif
 	}
 	else
 	{
