Made the "Restore open assets" message an editor notification rather than a modal dialog

TTP# 334332 - EDITOR: Open 'last opened assets' dialog should *NOT BE MODAL ON EDITOR STARTUP*

Adjusted the layout of the Slate notifications to make them better suited for longer messages that also contained buttons and other interactive widgets.  I ran these UI layout changes (visually) past NickD.

Added the ability for the notifications to contain a check box.

Improved the SlateViewer test suite for notifications so you could test a wider variation of settings without using the editor.

diff --git a/Engine/Source/Editor/UnrealEd/Private/Toolkits/AssetEditorManager.cpp b/Engine/Source/Editor/UnrealEd/Private/Toolkits/AssetEditorManager.cpp
--- a/Engine/Source/Editor/UnrealEd/Private/Toolkits/AssetEditorManager.cpp
+++ b/Engine/Source/Editor/UnrealEd/Private/Toolkits/AssetEditorManager.cpp
@@ -359,46 +359,136 @@ void FAssetEditorManager::RestorePreviouslyOpenAssets()
 	bool bCleanShutdown =  false;
 	GConfig->GetBool(TEXT("AssetEditorManager"), TEXT("CleanShutdown"), bCleanShutdown, GEditorUserSettingsIni);
 
+	SaveOpenAssetEditors(false);
+
 	if(OpenAssets.Num() > 0)
 	{
-		bool bRestore = false;
-		if(!bCleanShutdown)
+		if(bCleanShutdown)
 		{
-			if(FMessageDialog::Open(EAppMsgType::YesNo, LOCTEXT("RestoreOpenAssetsAfterCrash", "The Editor did not shut down cleanly, would you like to attempt to restore previously open assets now?")) == EAppReturnType::Yes)
+			// Do we have permission to automatically re-open the assets, or should we ask?
+			const bool bAutoRestore = GetDefault<UEditorLoadingSavingSettings>()->bRestoreOpenAssetTabsOnRestart;
+
+			if(bAutoRestore)
 			{
-				bRestore = true;
+				// Pretend that we showed the notification and that the user clicked "Restore Now"
+				OnConfirmRestorePreviouslyOpenAssets(OpenAssets);
 			}
-		}
-		else
-		{
-			bRestore = GetDefault<UEditorLoadingSavingSettings>()->bRestoreOpenAssetTabsOnRestart;
-
-			if(!bRestore)
+			else
 			{
-				// assets to restore, but we haven't been explicitly told not to restore them, so check now
-				FSuppressableWarningDialog::FSetupInfo Info(LOCTEXT("RestoreOpenAssetsAfterClose_Message", "Assets were open when the Editor was last closed, would you like to restore them now?"), LOCTEXT("RestoreOpenAssetsAfterClose_Title", "Reopen assets"), TEXT("SuppressPreviouslyOpenAssetsDialog"));
-				Info.ConfirmText = LOCTEXT("RestoreOpenAssetsAfterClose_Confirm", "Restore Open Assets");
-				Info.CancelText = LOCTEXT("RestoreOpenAssetsAfterClose_Cancel", "Don't Open Assets");
-				FSuppressableWarningDialog RestoreAssetsDialog( Info );
+				// Has this notification previously been suppressed by the user?
+				bool bSuppressNotification = false;
+				GConfig->GetBool(TEXT("AssetEditorManager"), TEXT("SuppressRestorePreviouslyOpenAssetsNotification"), bSuppressNotification, GEditorUserSettingsIni);
 
-				if(RestoreAssetsDialog.ShowModal() == FSuppressableWarningDialog::EResult::Confirm)
+				if(!bSuppressNotification)
 				{
-					GetMutableDefault<UEditorLoadingSavingSettings>()->bRestoreOpenAssetTabsOnRestart = true;
-					bRestore = true;
+					// Ask the user; this doesn't block so will reopen the assets later
+					SpawnRestorePreviouslyOpenAssetsNotification(bCleanShutdown, OpenAssets);
 				}
 			}
 		}
+		else
+		{
+			// If we crashed, we always ask regardless of what the user previously said
+			SpawnRestorePreviouslyOpenAssetsNotification(bCleanShutdown, OpenAssets);
+ 		}
+	}
+}
+
+void FAssetEditorManager::SpawnRestorePreviouslyOpenAssetsNotification(const bool bCleanShutdown, const TArray<FString>& AssetsToOpen)
+{
+	/** Utility functions for the notification which don't rely on the state from FAssetEditorManager */
+	struct Local
+	{
+		static ESlateCheckBoxState::Type GetDontAskAgainCheckBoxState()
+		{
+			bool bSuppressNotification = false;
+			GConfig->GetBool(TEXT("AssetEditorManager"), TEXT("SuppressRestorePreviouslyOpenAssetsNotification"), bSuppressNotification, GEditorUserSettingsIni);
+			return bSuppressNotification ? ESlateCheckBoxState::Checked : ESlateCheckBoxState::Unchecked;
+		}
 
-		if(bRestore)
+		static void OnDontAskAgainCheckBoxStateChanged(ESlateCheckBoxState::Type NewState)
 		{
-			for( const FString& AssetName : OpenAssets )
-			{
-				OpenEditorForAsset(AssetName);
-			}
+			const bool bSuppressNotification = (NewState == ESlateCheckBoxState::Checked);
+			GConfig->SetBool(TEXT("AssetEditorManager"), TEXT("SuppressRestorePreviouslyOpenAssetsNotification"), bSuppressNotification, GEditorUserSettingsIni);
 		}
+	};
+
+	FNotificationInfo Info(bCleanShutdown 
+		? LOCTEXT("RestoreOpenAssetsAfterClose_Message", "Assets were open when the Editor was last closed, would you like to restore them now?")
+		: LOCTEXT("RestoreOpenAssetsAfterCrash", "The Editor did not shut down cleanly, would you like to attempt to restore previously open assets now?")
+		);
+
+	// Add the buttons
+	Info.ButtonDetails.Add(FNotificationButtonInfo(
+		LOCTEXT("RestoreOpenAssetsAfterClose_Confirm", "Restore Now"), 
+		FText(), 
+		FSimpleDelegate::CreateRaw(this, &FAssetEditorManager::OnConfirmRestorePreviouslyOpenAssets, AssetsToOpen), 
+		SNotificationItem::CS_None
+		));
+	Info.ButtonDetails.Add(FNotificationButtonInfo(
+		LOCTEXT("RestoreOpenAssetsAfterClose_Cancel", "Don't Restore"), 
+		FText(), 
+		FSimpleDelegate::CreateRaw(this, &FAssetEditorManager::OnCancelRestorePreviouslyOpenAssets), 
+		SNotificationItem::CS_None
+		));
+
+	// We will be keeping track of this ourselves
+	Info.bFireAndForget = false;
+
+	// We want the auto-save to be subtle
+	Info.bUseLargeFont = false;
+	Info.bUseThrobber = false;
+	Info.bUseSuccessFailIcons = false;
+
+	// Only let the user suppress the non-crash version
+	if(bCleanShutdown)
+	{
+		Info.CheckBoxState = TAttribute<ESlateCheckBoxState::Type>::Create(&Local::GetDontAskAgainCheckBoxState);
+		Info.CheckBoxStateChanged = FOnCheckStateChanged::CreateStatic(&Local::OnDontAskAgainCheckBoxStateChanged);
+		Info.CheckBoxText = NSLOCTEXT("ModalDialogs", "DefaultCheckBoxMessage", "Don't show this again");
 	}
 
-	SaveOpenAssetEditors(false);
+	// Close any existing notification
+	if(RestorePreviouslyOpenAssetsNotification.IsValid())
+	{
+		RestorePreviouslyOpenAssetsNotification->ExpireAndFadeout();
+		RestorePreviouslyOpenAssetsNotification.Reset();
+	}
+
+	RestorePreviouslyOpenAssetsNotification = FSlateNotificationManager::Get().AddNotification(Info);
+}
+
+void FAssetEditorManager::OnConfirmRestorePreviouslyOpenAssets(TArray<FString> AssetsToOpen)
+{
+	// Close any existing notification
+	if(RestorePreviouslyOpenAssetsNotification.IsValid())
+	{
+		RestorePreviouslyOpenAssetsNotification->ExpireAndFadeout();
+		RestorePreviouslyOpenAssetsNotification.Reset();
+
+		// If the user suppressed the notification for future sessions, make sure this is reflected in their settings
+		// Note: We do that inside this if statement so that we only do it if we were showing a UI they could interact with
+		bool bSuppressNotification = false;
+		GConfig->GetBool(TEXT("AssetEditorManager"), TEXT("SuppressRestorePreviouslyOpenAssetsNotification"), bSuppressNotification, GEditorUserSettingsIni);
+		UEditorLoadingSavingSettings& Settings = *GetMutableDefault<UEditorLoadingSavingSettings>();
+		Settings.bRestoreOpenAssetTabsOnRestart = bSuppressNotification;
+		Settings.PostEditChange();
+	}
+
+	for( const FString& AssetName : AssetsToOpen )
+	{
+		OpenEditorForAsset(AssetName);
+	}
+}
+
+void FAssetEditorManager::OnCancelRestorePreviouslyOpenAssets()
+{
+	// Close any existing notification
+	if(RestorePreviouslyOpenAssetsNotification.IsValid())
+	{
+		RestorePreviouslyOpenAssetsNotification->ExpireAndFadeout();
+		RestorePreviouslyOpenAssetsNotification.Reset();
+	}
 }
 
 void FAssetEditorManager::SaveOpenAssetEditors(bool bOnShutdown)
@@ -407,7 +497,7 @@ void FAssetEditorManager::SaveOpenAssetEditors(bool bOnShutdown)
 	{
 		TArray<FString> OpenAssets;
 
-		// Dont save a list of assets to restore if we are running under a debugger
+		// Don't save a list of assets to restore if we are running under a debugger
 		if(!FPlatformMisc::IsDebuggerPresent())
 		{
 			for (auto EditorPair : OpenedEditors)
diff --git a/Engine/Source/Editor/UnrealEd/Public/Toolkits/AssetEditorManager.h b/Engine/Source/Editor/UnrealEd/Public/Toolkits/AssetEditorManager.h
--- a/Engine/Source/Editor/UnrealEd/Public/Toolkits/AssetEditorManager.h
+++ b/Engine/Source/Editor/UnrealEd/Public/Toolkits/AssetEditorManager.h
@@ -97,6 +97,15 @@ private:
 	// Handles ticks from the ticker.
 	bool HandleTicker( float DeltaTime );
 
+	// Spawn a notification asking the user if they want to restore their previously open assets
+	void SpawnRestorePreviouslyOpenAssetsNotification(const bool bCleanShutdown, const TArray<FString>& AssetsToOpen);
+
+	// Handler for when the "Restore Now" button is clicked on the RestorePreviouslyOpenAssets notification
+	void OnConfirmRestorePreviouslyOpenAssets(TArray<FString> AssetsToOpen);
+
+	// Handler for when the "Don't Restore" button is clicked on the RestorePreviouslyOpenAssets notification
+	void OnCancelRestorePreviouslyOpenAssets();
+
 	// Saves a list of open asset editors so they can be restored on editor restart
 	void SaveOpenAssetEditors(bool bOnShutdown);
 
@@ -148,6 +157,9 @@ private:
 
 	/** Flag whether we are currently shutting down */
 	bool bSavingOnShutdown;
+
+	/** A pointer to the notification used by RestorePreviouslyOpenAssets */
+	TSharedPtr<SNotificationItem> RestorePreviouslyOpenAssetsNotification;
 };
 
 
diff --git a/Engine/Source/Runtime/Slate/Private/Widgets/Notifications/SNotificationList.cpp b/Engine/Source/Runtime/Slate/Private/Widgets/Notifications/SNotificationList.cpp
--- a/Engine/Source/Runtime/Slate/Private/Widgets/Notifications/SNotificationList.cpp
+++ b/Engine/Source/Runtime/Slate/Private/Widgets/Notifications/SNotificationList.cpp
@@ -188,14 +188,14 @@ protected:
 		return CompletionState == CS_Pending ? EVisibility::Visible : EVisibility::Collapsed;
 	}
 
-	EVisibility GetSuccessImageVisibility() const
+	EVisibility GetSuccessFailImageVisibility() const
 	{
-		return CompletionState == CS_Success ? EVisibility::Visible : EVisibility::Collapsed;
+		return (CompletionState == CS_Success || CompletionState == CS_Fail) ? EVisibility::Visible : EVisibility::Collapsed;
 	}
 
-	EVisibility GetFailImageVisibility() const
+	const FSlateBrush* GetSuccessFailImage() const
 	{
-		return CompletionState == CS_Fail ? EVisibility::Visible : EVisibility::Collapsed;
+		return CompletionState == CS_Success ? FCoreStyle::Get().GetBrush("NotificationList.SuccessImage") : FCoreStyle::Get().GetBrush("NotificationList.FailImage");
 	}
 
 public:
@@ -275,7 +275,13 @@ public:
 		/** When true the larger bolder font will be used to display the message */
 		SLATE_ATTRIBUTE( bool, bUseLargeFont)
 		/** When set this forces the width of the box, used to stop resizeing on text change */
-		SLATE_ATTRIBUTE( FOptionalSize, WidthOverride )
+		SLATE_ARGUMENT( FOptionalSize, WidthOverride )
+		/** When set this will display a check box on the notification; handles getting the current check box state */
+		SLATE_ATTRIBUTE( ESlateCheckBoxState::Type, CheckBoxState )
+		/** When set this will display a check box on the notification; handles setting the new check box state */
+		SLATE_EVENT( FOnCheckStateChanged, CheckBoxStateChanged );
+		/** Text to display for the check box message */
+		SLATE_ATTRIBUTE( FText, CheckBoxText );
 		/** When set this will display as a hyperlink on the right side of the notification. */
 		SLATE_EVENT( FSimpleDelegate, Hyperlink)
 		/** Text to display for the hyperlink (if Hyperlink is valid) */
@@ -320,6 +326,7 @@ public:
 	 */
 	TSharedRef<SHorizontalBox> ConstructInternals( const FArguments& InArgs ) 
 	{
+		CheckBoxStateChanged = InArgs._CheckBoxStateChanged;
 		Hyperlink = InArgs._Hyperlink;
 		HyperlinkText = InArgs._HyperlinkText;
 
@@ -336,103 +343,148 @@ public:
 			.Image( InArgs._Image )
 		];
 
-		// Set font to default case.
-		FSlateFontInfo font = FCoreStyle::Get().GetFontStyle(TEXT("NotificationList.FontBold"));
-
-		// only override if requested
-		if (!InArgs._bUseLargeFont.Get())
 		{
-			font = FCoreStyle::Get().GetFontStyle(TEXT("NotificationList.FontLight"));
-		}
+			const FSlateFontInfo Font = InArgs._bUseLargeFont.Get() 
+				? FCoreStyle::Get().GetFontStyle(TEXT("NotificationList.FontBold")) 
+				: FCoreStyle::Get().GetFontStyle(TEXT("NotificationList.FontLight"));
 
-		// Build Text box
-		HorizontalBox->AddSlot()
-		.AutoWidth()
-		.Padding(10.f, 0.f, 25.f, 0.f)
-		.VAlign(VAlign_Center)
-		.HAlign(HAlign_Left)
-		[
-			SNew(SBox)
-			.WidthOverride(InArgs._WidthOverride)
-			[
-				SAssignNew(MyTextBlock, STextBlock)
-				.Text(Text)
-				.Font(font)
-			]
-		];
+			// Container for the text and optional interactive widgets (buttons, check box, and hyperlink)
+			TSharedRef<SVerticalBox> TextAndInteractiveWidgetsBox = SNew(SVerticalBox);
 
-		if (InArgs._bUseThrobber.Get())
-		{
-			// Build pending throbber
 			HorizontalBox->AddSlot()
 			.AutoWidth()
-			.Padding(0.f, 0.f, 25.f, 0.f)
-			.HAlign(HAlign_Center)
+			.Padding(10.f, 0.f, 15.f, 0.f)
 			.VAlign(VAlign_Center)
+			.HAlign(HAlign_Left)
 			[
-				SNew(SThrobber)
-				.Visibility( this, &SNotificationItemImpl::GetThrobberVisibility )
+				TextAndInteractiveWidgetsBox
 			];
-		}
 
-		if (InArgs._bUseSuccessFailIcons.Get())
-		{
-			// Build success image
-			HorizontalBox->AddSlot()
-			.AutoWidth()
-			.Padding(0.f, 0.f, 25.f, 0.f)
-			.HAlign(HAlign_Center)
-			.VAlign(VAlign_Center)
+			// Build Text box
+			TextAndInteractiveWidgetsBox->AddSlot()
+			.AutoHeight()
 			[
-				SNew(SImage)
-				.Image( FCoreStyle::Get().GetBrush("NotificationList.SuccessImage") )
-				.Visibility( this, &SNotificationItemImpl::GetSuccessImageVisibility )
+				SNew(SBox)
+				.WidthOverride(InArgs._WidthOverride)
+				[
+					SAssignNew(MyTextBlock, STextBlock)
+					.Text(Text)
+					.Font(Font)
+					.AutoWrapText(InArgs._WidthOverride.IsSet()) // only auto-wrap the text if we've been given a size constraint; otherwise, fill the notification area
+				]
 			];
 
-			// Build fail image
-			HorizontalBox->AddSlot()
+			TSharedRef<SHorizontalBox> InteractiveWidgetsBox = SNew(SHorizontalBox);
+			TextAndInteractiveWidgetsBox->AddSlot()
+			.AutoHeight()
+			[
+				InteractiveWidgetsBox
+			];
+
+			// Adds any buttons that were passed in.
+			{
+				TSharedRef<SHorizontalBox> ButtonsBox = SNew(SHorizontalBox);
+				for (int32 idx = 0; idx < InArgs._ButtonDetails.Get().Num(); idx++)
+				{
+					FNotificationButtonInfo Button = InArgs._ButtonDetails.Get()[idx];
+
+					ButtonsBox->AddSlot()
+					.AutoWidth()
+					.HAlign(HAlign_Left)
+					.VAlign(VAlign_Center)
+					.Padding(0.0f, 0.0f, 4.0f, 0.0f)
+					[
+						SNew(SButton)
+						.Text(Button.Text)
+						.ToolTipText(Button.ToolTip)
+						.OnClicked(this, &SNotificationItemImpl::OnButtonClicked, Button.Callback) 
+						.Visibility( this, &SNotificationItemImpl::GetButtonVisibility, Button.VisibilityOnNone, Button.VisibilityOnPending, Button.VisibilityOnSuccess, Button.VisibilityOnFail )
+					];
+				}
+				InteractiveWidgetsBox->AddSlot()
+				.AutoWidth()
+				.Padding(0.0f, 2.0f, 0.0f, 0.0f)
+				.VAlign(VAlign_Center)
+				.HAlign(HAlign_Left)
+				[
+					ButtonsBox
+				];
+			}
+
+			// Adds a check box, but only visible when bound
+			InteractiveWidgetsBox->AddSlot()
 			.AutoWidth()
-			.Padding(0.f, 0.f, 25.f, 0.f )
-			.HAlign(HAlign_Center)
 			.VAlign(VAlign_Center)
+			.HAlign(HAlign_Left)
 			[
-				SNew(SImage)
-				.Image( FCoreStyle::Get().GetBrush("NotificationList.FailImage") )
-				.Visibility( this, &SNotificationItemImpl::GetFailImageVisibility )
+				SNew(SBox)
+				.Padding(FMargin(0.0f, 2.0f, 4.0f, 0.0f))
+				.VAlign(VAlign_Center)
+				.HAlign(HAlign_Left)
+				.Visibility(this, &SNotificationItemImpl::GetCheckBoxVisibility)
+				[
+					SNew(SCheckBox)
+					.IsChecked(InArgs._CheckBoxState)
+					.OnCheckStateChanged(CheckBoxStateChanged)
+					[
+						SNew(STextBlock)
+						.Text(InArgs._CheckBoxText)
+					]
+				]
+			];
+
+			// Adds a hyperlink, but only visible when bound
+			InteractiveWidgetsBox->AddSlot()
+			.VAlign(VAlign_Bottom)
+			.HAlign(HAlign_Right)
+			[
+				SNew(SBox)
+				.Padding(FMargin(0.0f, 2.0f, 0.0f, 2.0f))
+				.VAlign(VAlign_Center)
+				.HAlign(HAlign_Left)
+				.Visibility(this, &SNotificationItemImpl::GetHyperlinkVisibility)
+				[
+					SNew(SHyperlink)
+					.Text(this, &SNotificationItemImpl::GetHyperlinkText)
+					.OnNavigate(this, &SNotificationItemImpl::OnHyperlinkClicked)
+				]
 			];
 		}
 
-		// Adds any buttons that were passed in.
-		for (int32 idx = 0; idx < InArgs._ButtonDetails.Get().Num(); idx++)
+		if (InArgs._bUseThrobber.Get())
 		{
-			FNotificationButtonInfo button = InArgs._ButtonDetails.Get()[idx];
-
+			// Build pending throbber
 			HorizontalBox->AddSlot()
 			.AutoWidth()
-			.HAlign(HAlign_Right)
-			.VAlign(VAlign_Center)
-			.Padding(2)
 			[
-				SNew(SButton)
-				.Text(button.Text)
-				.ToolTipText(button.ToolTip)
-				.OnClicked(this, &SNotificationItemImpl::OnButtonClicked, button.Callback) 
-				.Visibility( this, &SNotificationItemImpl::GetButtonVisibility, button.VisibilityOnNone, button.VisibilityOnPending, button.VisibilityOnSuccess, button.VisibilityOnFail )
+				SNew(SBox)
+				.Padding(FMargin(5.f, 0.f, 10.f, 0.f))
+				.HAlign(HAlign_Center)
+				.VAlign(VAlign_Center)
+				.Visibility( this, &SNotificationItemImpl::GetThrobberVisibility )
+				[
+					SNew(SThrobber)
+				]
 			];
 		}
 
-		// Adds a hyperlink, but only visible when bound
-		HorizontalBox->AddSlot()
+		if (InArgs._bUseSuccessFailIcons.Get())
+		{
+			// Build success/fail image
+			HorizontalBox->AddSlot()
 			.AutoWidth()
-			.HAlign(HAlign_Right)
-			.VAlign(VAlign_Center)
-			.Padding(2)
 			[
-				SNew(SHyperlink)
-					.Text(this, &SNotificationItemImpl::GetHyperlinkText)
-					.OnNavigate(this, &SNotificationItemImpl::OnHyperlinkClicked)
-					.Visibility(this, &SNotificationItemImpl::GetHyperlinkVisibility)
-			];	
+				SNew(SBox)
+				.Padding(FMargin(8.f, 0.f, 10.f, 0.f))
+				.HAlign(HAlign_Center)
+				.VAlign(VAlign_Center)
+				.Visibility( this, &SNotificationItemImpl::GetSuccessFailImageVisibility )
+				[
+					SNew(SImage)
+					.Image( this, &SNotificationItemImpl::GetSuccessFailImage )
+				]
+			];
+		}
 
 		return HorizontalBox;
 	}
@@ -466,6 +518,12 @@ protected:
 		}
 	}
 
+	/* Used to determine whether the check box is visible */
+	EVisibility GetCheckBoxVisibility() const
+	{
+		return CheckBoxStateChanged.IsBound() ? EVisibility::Visible : EVisibility::Collapsed;
+	}
+
 	/* Used to determine whether the hyperlink is visible */
 	EVisibility GetHyperlinkVisibility() const
 	{
@@ -495,6 +553,9 @@ private:
 
 protected:
 
+	/** When set this will display a check box on the notification; handles setting the new check box state */
+	FOnCheckStateChanged CheckBoxStateChanged;
+
 	/** When set this will display as a hyperlink on the right side of the notification. */
 	FSimpleDelegate Hyperlink;
 
@@ -506,15 +567,8 @@ protected:
 //// SNotificationList
 TSharedRef<SNotificationItem> SNotificationList::AddNotification(const FNotificationInfo& Info)
 {
-	EVisibility visibility = EVisibility::HitTestInvisible;
 	static const FSlateBrush* CachedImage = FCoreStyle::Get().GetBrush("NotificationList.DefaultMessage");
 
-	// Notifications with buttons need to be hit test enabled.
-	if (Info.ButtonDetails.Num() > 0)
-	{
-		visibility = EVisibility::Collapsed;
-	}
-
 	// Create notification.
 	TSharedRef<SNotificationExtendable> NewItem =
 		SNew(SNotificationItemImpl)
@@ -528,7 +582,9 @@ TSharedRef<SNotificationItem> SNotificationList::AddNotification(const FNotifica
 		.bUseSuccessFailIcons(Info.bUseSuccessFailIcons)
 		.bUseLargeFont(Info.bUseLargeFont)
 		.WidthOverride(Info.WidthOverride)
-		.Visibility(visibility)
+		.CheckBoxState(Info.CheckBoxState)
+		.CheckBoxStateChanged(Info.CheckBoxStateChanged)
+		.CheckBoxText(Info.CheckBoxText)
 		.Hyperlink(Info.Hyperlink)
 		.HyperlinkText(Info.HyperlinkText);
 		
diff --git a/Engine/Source/Runtime/Slate/Private/Widgets/Testing/STestSuite.cpp b/Engine/Source/Runtime/Slate/Private/Widgets/Testing/STestSuite.cpp
--- a/Engine/Source/Runtime/Slate/Private/Widgets/Testing/STestSuite.cpp
+++ b/Engine/Source/Runtime/Slate/Private/Widgets/Testing/STestSuite.cpp
@@ -2700,6 +2700,13 @@ public:
 
 	void Construct(const FArguments& InArgs)
 	{
+		bUseLargeFont = true;
+		bAddDummyButtons = false;
+		bAddDummyCheckBox = false;
+		bAddDummyHyperlink = false;
+
+		DummyCheckBoxState = ESlateCheckBoxState::Unchecked;
+
 		this->ChildSlot
 		[
 			SNew(SOverlay)
@@ -2728,21 +2735,69 @@ public:
 					[
 						SNew(SButton)
 						.OnClicked(this, &SNotificationListTest::SpawnPendingNotification)
-						.Text( LOCTEXT("NotificationListTest-SpawnPendingNotificationLabel", "Spawn Pending Notificaiton") )
+						.Text( LOCTEXT("NotificationListTest-SpawnPendingNotificationLabel", "Spawn Pending Notification") )
 					]
 
 					+SHorizontalBox::Slot().AutoWidth()
 					[
 						SNew(SButton)
 						.OnClicked(this, &SNotificationListTest::FinishPendingNotification_Success)
-						.Text( LOCTEXT("NotificationListTest-FinishPendingNotificationSuccessLabel", "Finish Pending Notificaiton - Success") )
+						.Text( LOCTEXT("NotificationListTest-FinishPendingNotificationSuccessLabel", "Finish Pending Notification - Success") )
 					]
 
 					+SHorizontalBox::Slot().AutoWidth()
 					[
 						SNew(SButton)
 						.OnClicked(this, &SNotificationListTest::FinishPendingNotification_Fail)
-						.Text( LOCTEXT("NotificationListTest-FinishPendingNotificationFailLabel", "Finish Pending Notificaiton - Fail") )
+						.Text( LOCTEXT("NotificationListTest-FinishPendingNotificationFailLabel", "Finish Pending Notification - Fail") )
+					]
+				]
+
+				+SVerticalBox::Slot()
+				.AutoHeight()
+				[
+					SNew(SCheckBox)
+					.IsChecked(this, &SNotificationListTest::IsUseLargeFontChecked)
+					.OnCheckStateChanged(this, &SNotificationListTest::OnUseLargeFontCheckStateChanged)
+					[
+						SNew(STextBlock)
+						.Text( LOCTEXT("NotificationListTest-UseLargeFont", "Use Large Font" ) )
+					]
+				]
+
+				+SVerticalBox::Slot()
+				.AutoHeight()
+				[
+					SNew(SCheckBox)
+					.IsChecked(this, &SNotificationListTest::IsAddDummyButtonsChecked)
+					.OnCheckStateChanged(this, &SNotificationListTest::OnAddDummyButtonsCheckStateChanged)
+					[
+						SNew(STextBlock)
+						.Text( LOCTEXT("NotificationListTest-AddDummyButtons", "Add Dummy Buttons" ) )
+					]
+				]
+
+				+SVerticalBox::Slot()
+				.AutoHeight()
+				[
+					SNew(SCheckBox)
+					.IsChecked(this, &SNotificationListTest::IsAddDummyCheckBoxChecked)
+					.OnCheckStateChanged(this, &SNotificationListTest::OnAddDummyCheckBoxCheckStateChanged)
+					[
+						SNew(STextBlock)
+						.Text( LOCTEXT("NotificationListTest-AddDummyCheckBox", "Add Dummy Check Box" ) )
+					]
+				]
+
+				+SVerticalBox::Slot()
+				.AutoHeight()
+				[
+					SNew(SCheckBox)
+					.IsChecked(this, &SNotificationListTest::IsAddDummyHyperlinkChecked)
+					.OnCheckStateChanged(this, &SNotificationListTest::OnAddDummyHyperlinkCheckStateChanged)
+					[
+						SNew(STextBlock)
+						.Text( LOCTEXT("NotificationListTest-AddDummyHyperlink", "Add Dummy Hyperlink" ) )
 					]
 				]
 			]
@@ -2753,21 +2808,117 @@ public:
 			.Padding(15)
 			[
 				SAssignNew(NotificationListPtr, SNotificationList)
-				.Visibility(EVisibility::HitTestInvisible)
+				.Visibility(EVisibility::SelfHitTestInvisible)
 			]
 		];
 	}
 
 protected:
+	ESlateCheckBoxState::Type IsUseLargeFontChecked() const
+	{
+		return bUseLargeFont ? ESlateCheckBoxState::Checked : ESlateCheckBoxState::Unchecked;
+	}
+
+	void OnUseLargeFontCheckStateChanged(ESlateCheckBoxState::Type NewState)
+	{
+		bUseLargeFont = (NewState == ESlateCheckBoxState::Checked);
+	}
+
+	ESlateCheckBoxState::Type IsAddDummyButtonsChecked() const
+	{
+		return bAddDummyButtons ? ESlateCheckBoxState::Checked : ESlateCheckBoxState::Unchecked;
+	}
+
+	void OnAddDummyButtonsCheckStateChanged(ESlateCheckBoxState::Type NewState)
+	{
+		bAddDummyButtons = (NewState == ESlateCheckBoxState::Checked);
+	}
+
+	ESlateCheckBoxState::Type IsAddDummyCheckBoxChecked() const
+	{
+		return bAddDummyCheckBox ? ESlateCheckBoxState::Checked : ESlateCheckBoxState::Unchecked;
+	}
+
+	void OnAddDummyCheckBoxCheckStateChanged(ESlateCheckBoxState::Type NewState)
+	{
+		bAddDummyCheckBox = (NewState == ESlateCheckBoxState::Checked);
+	}
+
+	ESlateCheckBoxState::Type IsAddDummyHyperlinkChecked() const
+	{
+		return bAddDummyHyperlink ? ESlateCheckBoxState::Checked : ESlateCheckBoxState::Unchecked;
+	}
+
+	void OnAddDummyHyperlinkCheckStateChanged(ESlateCheckBoxState::Type NewState)
+	{
+		bAddDummyHyperlink = (NewState == ESlateCheckBoxState::Checked);
+	}
+
+	void SetNotificationInfoFlags(FNotificationInfo& Info)
+	{
+		// Dummy callback for buttons and hyperlinks
+		struct Local
+		{
+			static void Dummy()
+			{
+			}
+		};
+		const FSimpleDelegate DummyDelegate = FSimpleDelegate::CreateStatic(&Local::Dummy);
+
+		Info.bUseLargeFont = bUseLargeFont;
+
+		if(bAddDummyButtons)
+		{
+			Info.ButtonDetails.Add(FNotificationButtonInfo(LOCTEXT("NotificationListTest-DummyButton1Text", "Dummy Button 1"), FText(), DummyDelegate, SNotificationItem::CS_None));
+			Info.ButtonDetails.Add(FNotificationButtonInfo(LOCTEXT("NotificationListTest-DummyButton2Text", "Dummy Button 2"), FText(), DummyDelegate, SNotificationItem::CS_None));
+
+			Info.ButtonDetails.Add(FNotificationButtonInfo(LOCTEXT("NotificationListTest-DummyPendingButton1Text", "Dummy Pending Button 1"), FText(), DummyDelegate, SNotificationItem::CS_Pending));
+			Info.ButtonDetails.Add(FNotificationButtonInfo(LOCTEXT("NotificationListTest-DummyPendingButton2Text", "Dummy Pending Button 2"), FText(), DummyDelegate, SNotificationItem::CS_Pending));
+
+			Info.ButtonDetails.Add(FNotificationButtonInfo(LOCTEXT("NotificationListTest-DummySuccessButton1Text", "Dummy Success Button 1"), FText(), DummyDelegate, SNotificationItem::CS_Success));
+			Info.ButtonDetails.Add(FNotificationButtonInfo(LOCTEXT("NotificationListTest-DummySuccessButton2Text", "Dummy Success Button 2"), FText(), DummyDelegate, SNotificationItem::CS_Success));
+
+			Info.ButtonDetails.Add(FNotificationButtonInfo(LOCTEXT("NotificationListTest-DummyFailButton1Text", "Dummy Fail Button 1"), FText(), DummyDelegate, SNotificationItem::CS_Fail));
+			Info.ButtonDetails.Add(FNotificationButtonInfo(LOCTEXT("NotificationListTest-DummyFailButton2Text", "Dummy Fail Button 2"), FText(), DummyDelegate, SNotificationItem::CS_Fail));
+		}
+
+		if(bAddDummyCheckBox)
+		{
+			Info.CheckBoxState = TAttribute<ESlateCheckBoxState::Type>(this, &SNotificationListTest::GetDummyCheckBoxState);
+			Info.CheckBoxStateChanged = FOnCheckStateChanged::CreateSP(this, &SNotificationListTest::OnDummyCheckBoxStateChanged);
+			Info.CheckBoxText = LOCTEXT("NotificationListTest-DummyCheckBoxText", "Dummy Check Box");
+		}
+
+		if(bAddDummyHyperlink)
+		{
+			Info.Hyperlink = DummyDelegate;
+			Info.HyperlinkText = LOCTEXT("NotificationListTest-DummyHyperlinkText", "Dummy Hyperlink");
+		}
+	}
+
+	ESlateCheckBoxState::Type GetDummyCheckBoxState() const
+	{
+		return DummyCheckBoxState;
+	}
+
+	void OnDummyCheckBoxStateChanged(ESlateCheckBoxState::Type NewState)
+	{
+		DummyCheckBoxState = NewState;
+	}
+
 	FReply SpawnNotification1()
 	{
-		NotificationListPtr->AddNotification( FNotificationInfo( LOCTEXT("TestNotification01", "A Notification" )) );
+		FNotificationInfo Info( LOCTEXT("TestNotification01", "A Notification" ));
+		SetNotificationInfoFlags(Info);
+		NotificationListPtr->AddNotification(Info);
 		return FReply::Handled();
 	}
 
 	FReply SpawnNotification2()
 	{
-		NotificationListPtr->AddNotification( FNotificationInfo( LOCTEXT("TestNotification02", "Another Notification" )) );
+		FNotificationInfo Info( LOCTEXT("TestNotification02", "Another Notification" ));
+		SetNotificationInfoFlags(Info);
+		NotificationListPtr->AddNotification(Info);
 		return FReply::Handled();
 	}
 
@@ -2778,10 +2929,11 @@ protected:
 			PendingProgressPtr.Pin()->ExpireAndFadeout();
 		}
 
-		FNotificationInfo info(LOCTEXT("TestNotificationInProgress", "Operation In Progress" ));
-		info.bFireAndForget = false;
+		FNotificationInfo Info(LOCTEXT("TestNotificationInProgress", "Operation In Progress" ));
+		SetNotificationInfoFlags(Info);
+		Info.bFireAndForget = false;
 		
-		PendingProgressPtr = NotificationListPtr->AddNotification(info);
+		PendingProgressPtr = NotificationListPtr->AddNotification(Info);
 
 		PendingProgressPtr.Pin()->SetCompletionState(SNotificationItem::CS_Pending);
 		return FReply::Handled();
@@ -2819,10 +2971,24 @@ protected:
 		return FReply::Handled();
 	}
 
+	/** Should we use the larger bolder font to display the message? */
+	bool bUseLargeFont;
+
+	/** Should we add dummy buttons to the notifications? */
+	bool bAddDummyButtons;
+
+	/** Should we add a dummy check box to the notifications? */
+	bool bAddDummyCheckBox;
+
+	/** Should we add a dummy hyperlink to the notifications? */
+	bool bAddDummyHyperlink;
+
 	/** The list of active system messages */
 	TSharedPtr<SNotificationList> NotificationListPtr;
 	/** The pending progress message */
 	TWeakPtr<SNotificationItem> PendingProgressPtr;
+	/** If showing a dummy check box, this maintains its state in the UI */
+	ESlateCheckBoxState::Type DummyCheckBoxState;
 };
 
 class SGridPanelTest : public SCompoundWidget
@@ -3301,8 +3467,21 @@ TSharedRef<SDockTab> SpawnTestSuite2( const FSpawnTabArgs& Args )
 			->AddTab("GridPanelTest", ETabState::OpenedTab)
 			->AddTab("DPIScalingTest", ETabState::OpenedTab)
 		)
-	);
-
+	)
+#if PLATFORM_SUPPORTS_MULTIPLE_NATIVE_WINDOWS
+	->AddArea
+	(
+		// This area will get a 400x600 window at 10,10
+		FTabManager::NewArea(400,600)
+		->SetWindow( FVector2D(10,10), false )
+		->Split
+		(
+			// The area contains a single tab with the widget reflector.
+			FTabManager::NewStack()->AddTab( "WidgetReflector", ETabState::OpenedTab )
+		)
+	)
+#endif
+	;
 
 	TSharedRef<SDockTab> TestSuite2Tab = 
 		SNew(SDockTab)
diff --git a/Engine/Source/Runtime/Slate/Public/Widgets/Notifications/SNotificationList.h b/Engine/Source/Runtime/Slate/Public/Widgets/Notifications/SNotificationList.h
--- a/Engine/Source/Runtime/Slate/Public/Widgets/Notifications/SNotificationList.h
+++ b/Engine/Source/Runtime/Slate/Public/Widgets/Notifications/SNotificationList.h
@@ -45,10 +45,10 @@ struct FNotificationButtonInfo
 		: Text( InText )
 		, ToolTip( InToolTip )
 		, Callback( InCallback )
-		, VisibilityOnNone( VisibleInState == SNotificationItem::CS_None ? EVisibility::Visible : EVisibility::Hidden )
-		, VisibilityOnPending( VisibleInState == SNotificationItem::CS_Pending ? EVisibility::Visible : EVisibility::Hidden )
-		, VisibilityOnSuccess( VisibleInState == SNotificationItem::CS_Success ? EVisibility::Visible : EVisibility::Hidden )
-		, VisibilityOnFail( VisibleInState == SNotificationItem::CS_Fail ? EVisibility::Visible : EVisibility::Hidden )
+		, VisibilityOnNone( VisibleInState == SNotificationItem::CS_None ? EVisibility::Visible : EVisibility::Collapsed )
+		, VisibilityOnPending( VisibleInState == SNotificationItem::CS_Pending ? EVisibility::Visible : EVisibility::Collapsed )
+		, VisibilityOnSuccess( VisibleInState == SNotificationItem::CS_Success ? EVisibility::Visible : EVisibility::Collapsed )
+		, VisibilityOnFail( VisibleInState == SNotificationItem::CS_Fail ? EVisibility::Visible : EVisibility::Collapsed )
 	{}
 
 	/** Message on the button */
@@ -90,6 +90,9 @@ struct FNotificationInfo
 		bUseLargeFont(true),
 		WidthOverride(),
 		bFireAndForget(true),
+		CheckBoxState(ESlateCheckBoxState::Unchecked),
+		CheckBoxStateChanged(),
+		CheckBoxText(),
 		Hyperlink(),
 		HyperlinkText( NSLOCTEXT("EditorNotification", "DefaultHyperlinkText", "Show Log") ),
 		bAllowThrottleWhenFrameRateIsLow(true)
@@ -129,6 +132,15 @@ struct FNotificationInfo
 	/** When true the notification will automatically time out after the expire duration. */ 
 	bool bFireAndForget;
 
+	/** When set this will display a check box on the notification; handles getting the current check box state */
+	TAttribute< ESlateCheckBoxState::Type > CheckBoxState;
+
+	/** When set this will display a check box on the notification; handles setting the new check box state */
+	FOnCheckStateChanged CheckBoxStateChanged;
+
+	/** Text to display for the check box message */
+	TAttribute< FText > CheckBoxText;
+
 	/** When set this will display as a hyperlink on the right side of the notification. */
 	FSimpleDelegate Hyperlink;
 
