SLATE - We no longer ensure() that the onclick event returns an FReply that is Handled().

diff --git a/Engine/Source/Runtime/Slate/Private/Widgets/Input/SButton.cpp b/Engine/Source/Runtime/Slate/Private/Widgets/Input/SButton.cpp
--- a/Engine/Source/Runtime/Slate/Private/Widgets/Input/SButton.cpp
+++ b/Engine/Source/Runtime/Slate/Private/Widgets/Input/SButton.cpp
@@ -182,7 +182,6 @@ FReply SButton::OnMouseButtonUp( const FGeometry& MyGeometry, const FPointerEven
 				if( ( ClickMethod == EButtonClickMethod::MouseUp || HasMouseCapture() ) && OnClicked.IsBound() == true )
 				{
 					Reply = OnClicked.Execute();
-					ensure(Reply.IsEventHandled() == true);
 				}
 			}
 		}
