Fixed menu popping regression in Slate

diff --git a/Engine/Source/Runtime/Slate/Private/Widgets/Layout/SScrollBox.cpp b/Engine/Source/Runtime/Slate/Private/Widgets/Layout/SScrollBox.cpp
--- a/Engine/Source/Runtime/Slate/Private/Widgets/Layout/SScrollBox.cpp
+++ b/Engine/Source/Runtime/Slate/Private/Widgets/Layout/SScrollBox.cpp
@@ -148,6 +148,8 @@ void SScrollBox::Construct( const FArguments& InArgs )
 			.OnUserScrolled( this, &SScrollBox::ScrollBar_OnUserScrolled )
 		]
 	];
+
+	ScrollBar->SetState( 0.0f, 1.0f );
 }
 
 
@@ -208,7 +210,9 @@ void SScrollBox::Tick( const FGeometry& AllottedGeometry, const double InCurrent
 	const FGeometry ScrollPanelGeometry = FindChildGeometry( AllottedGeometry, ScrollPanel.ToSharedRef() );
 	
 	const float ContentHeight = ScrollPanel->GetDesiredSize().Y;
-	const float ViewFraction = ScrollPanelGeometry.Size.Y / ContentHeight;
+
+	// If this scroll box has no size, do not compute a view fraction because it will be wrong and causes pop in when the size is available
+	const float ViewFraction = AllottedGeometry.Size.Y > 0 ? ScrollPanelGeometry.Size.Y / ContentHeight : 1;
 	const float ViewOffset = FMath::Clamp<float>( DesiredScrollOffset/ContentHeight, 0.0, 1.0 - ViewFraction );
 	
 	// Update the scrollbar with the clamped version of the offset
diff --git a/Engine/Source/Runtime/Slate/Public/Widgets/Layout/SScrollBox.h b/Engine/Source/Runtime/Slate/Public/Widgets/Layout/SScrollBox.h
--- a/Engine/Source/Runtime/Slate/Public/Widgets/Layout/SScrollBox.h
+++ b/Engine/Source/Runtime/Slate/Public/Widgets/Layout/SScrollBox.h
@@ -94,9 +94,6 @@ private:
 	/** Does the user need a hint that they can scroll down? */
 	FSlateColor GetBottomShadowOpacity() const;
 
-	/** Does the user need to see the scroll bar? */
-	EVisibility GetScrollBarVisibility() const;
-
 	TSharedPtr<class SScrollPanel> ScrollPanel;
 	TSharedPtr<SScrollBar> ScrollBar;
 	
