Slate: Added ability to specify center content alignment of window title bars

diff --git a/Engine/Source/Runtime/Slate/Private/Framework/Application/SWindowTitleBar.h b/Engine/Source/Runtime/Slate/Private/Framework/Application/SWindowTitleBar.h
--- a/Engine/Source/Runtime/Slate/Private/Framework/Application/SWindowTitleBar.h
+++ b/Engine/Source/Runtime/Slate/Private/Framework/Application/SWindowTitleBar.h
@@ -73,7 +73,15 @@ public:
 	
 public:
 
-	void Construct( const FArguments& InArgs, const TSharedRef<SWindow>& InWindow, const TSharedPtr<SWidget>& InCenterContent )
+	/**
+	 * Creates and initializes a new window title bar widget.
+	 *
+	 * @param InArgs The construction arguments.
+	 * @param InWindow The window to create the title bar for.
+	 * @param InCenterContent The content for the title bar's center area.
+	 * @param CenterContentAlignment The horizontal alignment of the center content.
+	 */
+	void Construct( const FArguments& InArgs, const TSharedRef<SWindow>& InWindow, const TSharedPtr<SWidget>& InCenterContent, EHorizontalAlignment InCenterContentAlignment )
 	{
 		OwnerWindowPtr = InWindow;
 		Style = InArgs._Style;
@@ -98,7 +106,7 @@ public:
 
 					+ SOverlay::Slot()
 						[
-							MakeTitleBarContent(InCenterContent)
+							MakeTitleBarContent(InCenterContent, InCenterContentAlignment)
 						]
 				]
 		];
@@ -262,9 +270,12 @@ protected:
 	/**
 	 * Creates the title bar's content.
 	 *
+	 * @param CenterContent The content for the title bar's center area.
+	 * @param CenterContentAlignment The horizontal alignment of the center content.
+	 *
 	 * @return The content widget.
 	 */
-	TSharedRef<SWidget> MakeTitleBarContent( TSharedPtr<SWidget> CenterContent )
+	TSharedRef<SWidget> MakeTitleBarContent( TSharedPtr<SWidget> CenterContent, EHorizontalAlignment CenterContentAlignment )
 	{
 		TSharedPtr<SWidget> LeftContent;
 		TSharedPtr<SWidget> RightContent;
@@ -294,6 +305,13 @@ protected:
 
 		FVector2D LeftSize = LeftContent->GetDesiredSize();
 		FVector2D RightSize = RightContent->GetDesiredSize();
+
+		if (CenterContentAlignment == HAlign_Center)
+		{
+			LeftSize = FMath::Max(LeftSize, RightSize);
+			RightSize = LeftSize;
+		}
+
 		float SpacerHeight = FMath::Max(LeftSize.Y, RightSize.Y);
 
 		// create title bar
@@ -318,6 +336,7 @@ protected:
 							]
 
 						+ SHorizontalBox::Slot()
+							.HAlign(CenterContentAlignment)
 							.VAlign(VAlign_Top)
 							.FillWidth(1.0f)
 							[
diff --git a/Engine/Source/Runtime/Slate/Private/Framework/Application/SlateApplication.cpp b/Engine/Source/Runtime/Slate/Private/Framework/Application/SlateApplication.cpp
--- a/Engine/Source/Runtime/Slate/Private/Framework/Application/SlateApplication.cpp
+++ b/Engine/Source/Runtime/Slate/Private/Framework/Application/SlateApplication.cpp
@@ -3012,9 +3012,9 @@ TSharedRef<SWidget> FSlateApplication::MakeImage( const TAttribute<const FSlateB
 }
 
 
-TSharedRef<SWidget> FSlateApplication::MakeWindowTitleBar( const TSharedRef<SWindow>& Window, const TSharedPtr<SWidget>& CenterContent, TSharedPtr<IWindowTitleBar>& OutTitleBar ) const
+TSharedRef<SWidget> FSlateApplication::MakeWindowTitleBar( const TSharedRef<SWindow>& Window, const TSharedPtr<SWidget>& CenterContent, EHorizontalAlignment CenterContentAlignment, TSharedPtr<IWindowTitleBar>& OutTitleBar ) const
 {
-	TSharedRef<SWindowTitleBar> TitleBar = SNew(SWindowTitleBar, Window, CenterContent)
+	TSharedRef<SWindowTitleBar> TitleBar = SNew(SWindowTitleBar, Window, CenterContent, CenterContentAlignment)
 		.Visibility(EVisibility::SelfHitTestInvisible);
 
 	OutTitleBar = TitleBar;
diff --git a/Engine/Source/Runtime/Slate/Private/Framework/Docking/SDockingTabStack.cpp b/Engine/Source/Runtime/Slate/Private/Framework/Docking/SDockingTabStack.cpp
--- a/Engine/Source/Runtime/Slate/Private/Framework/Docking/SDockingTabStack.cpp
+++ b/Engine/Source/Runtime/Slate/Private/Framework/Docking/SDockingTabStack.cpp
@@ -728,7 +728,7 @@ void SDockingTabStack::SetParentNode( TSharedRef<class SDockingSplitter> InParen
 		const TSharedRef<SWindow>& ParentWindow = DockArea->GetParentWindow().ToSharedRef();
 
 		TSharedPtr<IWindowTitleBar> TitleBar;
-		TSharedRef<SWidget> TitleBarWidget = FSlateApplication::Get().MakeWindowTitleBar(ParentWindow, TitleBarContent, TitleBar);
+		TSharedRef<SWidget> TitleBarWidget = FSlateApplication::Get().MakeWindowTitleBar(ParentWindow, TitleBarContent, HAlign_Fill, TitleBar);
 
 		(*TitleBarSlot)[TitleBarWidget];
 		ParentWindow->SetTitleBar(TitleBar);
diff --git a/Engine/Source/Runtime/Slate/Public/Framework/Application/SlateApplication.h b/Engine/Source/Runtime/Slate/Public/Framework/Application/SlateApplication.h
--- a/Engine/Source/Runtime/Slate/Public/Framework/Application/SlateApplication.h
+++ b/Engine/Source/Runtime/Slate/Public/Framework/Application/SlateApplication.h
@@ -1011,7 +1011,7 @@ public:
 
 	virtual TSharedRef<SWidget> MakeImage( const TAttribute<const FSlateBrush*>& Image, const TAttribute<FSlateColor>& Color, const TAttribute<EVisibility>& Visibility ) const OVERRIDE;
 
-	virtual TSharedRef<SWidget> MakeWindowTitleBar( const TSharedRef<SWindow>& Window, const TSharedPtr<SWidget>& CenterContent, TSharedPtr<IWindowTitleBar>& OutTitleBar ) const OVERRIDE;
+	virtual TSharedRef<SWidget> MakeWindowTitleBar( const TSharedRef<SWindow>& Window, const TSharedPtr<SWidget>& CenterContent, EHorizontalAlignment CenterContentAlignment, TSharedPtr<IWindowTitleBar>& OutTitleBar ) const OVERRIDE;
 
 	virtual TSharedRef<IToolTip> MakeToolTip( const TAttribute<FString>& ToolTipString ) OVERRIDE;
 
diff --git a/Engine/Source/Runtime/SlateCore/Private/Widgets/SWindow.cpp b/Engine/Source/Runtime/SlateCore/Private/Widgets/SWindow.cpp
--- a/Engine/Source/Runtime/SlateCore/Private/Widgets/SWindow.cpp
+++ b/Engine/Source/Runtime/SlateCore/Private/Widgets/SWindow.cpp
@@ -357,7 +357,7 @@ void SWindow::ConstructWindowInternals( const bool bCreateTitleBar )
 		MainWindowArea->AddSlot()
 			.AutoHeight()
 			[
-				FSlateApplicationBase::Get().MakeWindowTitleBar(SharedThis(this), nullptr, TitleBar)
+				FSlateApplicationBase::Get().MakeWindowTitleBar(SharedThis(this), nullptr, HAlign_Center, TitleBar)
 			];
 	}
 	else
diff --git a/Engine/Source/Runtime/SlateCore/Public/Application/SlateApplicationBase.h b/Engine/Source/Runtime/SlateCore/Public/Application/SlateApplicationBase.h
--- a/Engine/Source/Runtime/SlateCore/Public/Application/SlateApplicationBase.h
+++ b/Engine/Source/Runtime/SlateCore/Public/Application/SlateApplicationBase.h
@@ -207,11 +207,12 @@ public:
 	 *
 	 * @param Window The window to create the title bar for.
 	 * @param CenterContent Optional content for the title bar's center (will override window title).
+	 * @param CenterContentAlignment The horizontal alignment of the center content.
 	 * @param OutTitleBar Will hold a pointer to the title bar's interface.
 	 *
 	 * @return The new title bar widget.
 	 */
-	virtual TSharedRef<SWidget> MakeWindowTitleBar( const TSharedRef<SWindow>& Window, const TSharedPtr<SWidget>& CenterContent, TSharedPtr<IWindowTitleBar>& OutTitleBar ) const = 0;
+	virtual TSharedRef<SWidget> MakeWindowTitleBar( const TSharedRef<SWindow>& Window, const TSharedPtr<SWidget>& CenterContent, EHorizontalAlignment CenterContentAlignment, TSharedPtr<IWindowTitleBar>& OutTitleBar ) const = 0;
 
 	/**
 	 * Destroying windows has implications on some OSs (e.g. destroying Win32 HWNDs can cause events to be lost).
