Color picker now displays an error when renaming a theme the same as another

#ttp 331840 - EDITOR: In the color picker, renaming a color theme with a duplicate name should yield an error notification

#branch UE4

#add added area to color picker to display an errortext at the bottom (initially tried to change rename box to SInlineEditableTextBlock, however it wouldn't display the popup error correctly)

#change When the name is modified it checks to see if these is a conflict and displays an error if there is.

#change 'Accept' button can now only be pressed when there's no error present.

diff --git a/Engine/Source/Runtime/Slate/Private/Widgets/Colors/SColorThemes.cpp b/Engine/Source/Runtime/Slate/Private/Widgets/Colors/SColorThemes.cpp
--- a/Engine/Source/Runtime/Slate/Private/Widgets/Colors/SColorThemes.cpp
+++ b/Engine/Source/Runtime/Slate/Private/Widgets/Colors/SColorThemes.cpp
@@ -736,11 +736,13 @@ void SColorThemesViewer::Construct(const FArguments& InArgs)
 		+SHorizontalBox::Slot() .FillWidth(1) .Padding(3)
 		[
 			SAssignNew(RenameTextBox, SEditableTextBox) .Font(SmallLayoutFont)
+				.OnTextChanged(this, &SColorThemesViewer::ChangeThemeName)
 		]
 		+SHorizontalBox::Slot().AutoWidth() .Padding(3)
 		[
 			SNew(SButton) .Text(LOCTEXT("AcceptRenameButton", "Accept")) .HAlign(HAlign_Right)
-				.OnClicked(this, &SColorThemesViewer::ChangeThemeName)
+				.OnClicked(this, &SColorThemesViewer::AcceptThemeName)
+				.IsEnabled(this, &SColorThemesViewer::CanAcceptThemeName)
 		]
 		+SHorizontalBox::Slot().AutoWidth() .Padding(3)
 		[
@@ -792,6 +794,11 @@ void SColorThemesViewer::Construct(const FArguments& InArgs)
 		[
 			SAssignNew(Menu, SBorder)
 		]
+		+ SVerticalBox::Slot().AutoHeight()
+		[
+			SAssignNew(ErrorText, SErrorText)
+				.Visibility(this, &SColorThemesViewer::OnGetErrorTextVisibility)
+		]
 	];
 
 	MenuToStandardNoReturn();
@@ -809,6 +816,7 @@ END_SLATE_FUNCTION_BUILD_OPTIMIZATION
 
 void SColorThemesViewer::MenuToStandardNoReturn()
 {
+	ErrorText->SetError(FString());
 	Menu->SetContent(MenuStandard.ToSharedRef());
 }
 
@@ -871,7 +879,7 @@ TSharedPtr<FColorTheme> SColorThemesViewer::IsColorTheme(const FString& ThemeNam
 	// Find the desired theme
 	for (int32 ThemeIndex = 0; ThemeIndex < ColorThemes.Num(); ++ThemeIndex)
 	{
-		TSharedPtr<FColorTheme>& ColorTheme = ColorThemes[ThemeIndex];
+		const TSharedPtr<FColorTheme>& ColorTheme = ColorThemes[ThemeIndex];
 		if ( ColorTheme->Name == ThemeName )
 		{
 			return ColorTheme;
@@ -924,7 +932,7 @@ TSharedPtr<FColorTheme> SColorThemesViewer::GetDefaultColorTheme(bool bCreateNew
 	return GetColorTheme( Name.ToString() );
 }
 
-FReply SColorThemesViewer::ChangeThemeName()
+FReply SColorThemesViewer::AcceptThemeName()
 {
 	// Update the theme name if it differs, ensuring it is still unique
 	const FString Name = RenameTextBox->GetText().ToString();
@@ -936,6 +944,34 @@ FReply SColorThemesViewer::ChangeThemeName()
 	return FReply::Handled();
 }
 
+bool SColorThemesViewer::CanAcceptThemeName() const
+{
+	return !ErrorText->HasError();
+}
+
+void SColorThemesViewer::ChangeThemeName(const FText& InText)
+{
+	ErrorText->SetError(FString());
+
+	const FString ThemeName = InText.ToString();
+	for (int32 ThemeIndex = 0; ThemeIndex < ColorThemes.Num(); ++ThemeIndex)
+	{
+		const TSharedPtr<FColorTheme>& ColorTheme = ColorThemes[ThemeIndex];
+		if (ColorTheme != GetCurrentColorTheme() && ColorTheme->Name == ThemeName)
+		{
+			FFormatNamedArguments Args;
+			Args.Add(TEXT("Name"), InText);
+			ErrorText->SetError(FText::Format(NSLOCTEXT("ColorThemesViewer", "VerifyTextDup", "A theme already exists with the name '{Name}'."), Args).ToString());
+			return;
+		}
+	}
+}
+
+EVisibility SColorThemesViewer::OnGetErrorTextVisibility() const
+{
+	return !CanAcceptThemeName() ? EVisibility::Visible : EVisibility::Collapsed;
+}
+
 FReply SColorThemesViewer::NewColorTheme()
 {
 	// Create a new, defaultly named theme and update the display
diff --git a/Engine/Source/Runtime/Slate/Public/Widgets/Colors/SColorThemes.h b/Engine/Source/Runtime/Slate/Public/Widgets/Colors/SColorThemes.h
--- a/Engine/Source/Runtime/Slate/Public/Widgets/Colors/SColorThemes.h
+++ b/Engine/Source/Runtime/Slate/Public/Widgets/Colors/SColorThemes.h
@@ -513,7 +513,10 @@ private:
 	FReply NewColorTheme();
 	FReply DuplicateColorTheme();
 	FReply DeleteColorTheme();
-	FReply ChangeThemeName();
+	FReply AcceptThemeName();
+	bool CanAcceptThemeName() const;
+	void ChangeThemeName(const FText& InText);
+	EVisibility OnGetErrorTextVisibility() const;
 
 	/** Sets the current color theme to the existing theme */
 	void SetCurrentColorTheme(TSharedPtr<FColorTheme> NewTheme);
@@ -552,7 +555,9 @@ private:
 	TSharedPtr<SWidget> MenuRename;
 	TSharedPtr<SWidget> MenuConfirmDelete;
 	TSharedPtr<SWidget> MenuTrashColor;
-	
+	/** If any error occurs, it needs to be displayed here as SPopupErrorText won't work with nested combo buttons */
+	TSharedPtr<SErrorText> ErrorText;
+
 	/** The text box for renaming themes */
 	TSharedPtr<SEditableTextBox> RenameTextBox;
 
