fix for visual scale being applied to the joystick input scale
TTP334561
#ue4

diff --git a/Engine/Source/Runtime/Slate/Private/Widgets/Input/SVirtualJoystick.cpp b/Engine/Source/Runtime/Slate/Private/Widgets/Input/SVirtualJoystick.cpp
--- a/Engine/Source/Runtime/Slate/Private/Widgets/Input/SVirtualJoystick.cpp
+++ b/Engine/Source/Runtime/Slate/Private/Widgets/Input/SVirtualJoystick.cpp
@@ -345,7 +345,7 @@ void SVirtualJoystick::Tick( const FGeometry& AllottedGeometry, const double InC
 			Control.VisualSize = FVector2D(ResolveRelativePosition(Control.VisualSize.X, AllottedGeometry.Size.X, ScaleFactor), ResolveRelativePosition(Control.VisualSize.Y, AllottedGeometry.Size.Y, ScaleFactor));
 			Control.InteractionSize = FVector2D(ResolveRelativePosition(Control.InteractionSize.X, AllottedGeometry.Size.X, ScaleFactor), ResolveRelativePosition(Control.InteractionSize.Y, AllottedGeometry.Size.Y, ScaleFactor));
 			Control.ThumbSize = FVector2D(ResolveRelativePosition(Control.ThumbSize.X, AllottedGeometry.Size.X, ScaleFactor), ResolveRelativePosition(Control.ThumbSize.Y, AllottedGeometry.Size.Y, ScaleFactor));
-			Control.InputScale = Control.InputScale * ScaleFactor;
+//			Control.InputScale = Control.InputScale * ScaleFactor;
 			Control.bHasBeenPositioned = true;
 		}
 
