Editor: refactored FLayoutService to be INI file agnostic; moved Editor layout from EditorUserSettings.ini to EditorLayout.ini

I also added code to retain backwards compatibility. Existing settings from EditorUserSettings.ini will be migrated to EditorLayout.ini

diff --git a/Engine/Source/Editor/Kismet/Private/ApplicationMode.cpp b/Engine/Source/Editor/Kismet/Private/ApplicationMode.cpp
--- a/Engine/Source/Editor/Kismet/Private/ApplicationMode.cpp
+++ b/Engine/Source/Editor/Kismet/Private/ApplicationMode.cpp
@@ -11,7 +11,7 @@ void FApplicationMode::DeactivateMode(TSharedPtr<FTabManager> InTabManager)
 {
 	// Save the layout to INI
 	check(InTabManager.IsValid());
-	FLayoutSaveRestore::SaveTheLayout(InTabManager->PersistLayout());
+	FLayoutSaveRestore::SaveToConfig(GEditorLayoutIni, InTabManager->PersistLayout());
 		
 	// Unregister the tabs
 	/*
@@ -40,5 +40,5 @@ TSharedRef<FTabManager::FLayout> FApplicationMode::ActivateMode(TSharedPtr<FTabM
 
 	// Try loading the layout from INI
 	check(TabLayout.IsValid());
-	return FLayoutSaveRestore::LoadUserConfigVersionOf(TabLayout.ToSharedRef());
+	return FLayoutSaveRestore::LoadFromConfig(GEditorLayoutIni, TabLayout.ToSharedRef());
 }
diff --git a/Engine/Source/Editor/LevelEditor/Private/LevelEditor.cpp b/Engine/Source/Editor/LevelEditor/Private/LevelEditor.cpp
--- a/Engine/Source/Editor/LevelEditor/Private/LevelEditor.cpp
+++ b/Engine/Source/Editor/LevelEditor/Private/LevelEditor.cpp
@@ -366,7 +366,7 @@ void FLevelEditorModule::SetLevelEditorTabManager( const TSharedPtr<SDockTab>& O
 	if (OwnerTab.IsValid())
 	{
 		LevelEditorTabManager = FGlobalTabmanager::Get()->NewTabManager(OwnerTab.ToSharedRef());
-		LevelEditorTabManager->SetOnPersistLayout( FTabManager::FOnPersistLayout::CreateStatic( FLayoutSaveRestore::SaveTheLayout ) );
+		LevelEditorTabManager->SetOnPersistLayout(FTabManager::FOnPersistLayout::CreateRaw(this, &FLevelEditorModule::HandleTabManagerPersistLayout));
 	}
 }
 
diff --git a/Engine/Source/Editor/LevelEditor/Private/SLevelEditor.cpp b/Engine/Source/Editor/LevelEditor/Private/SLevelEditor.cpp
--- a/Engine/Source/Editor/LevelEditor/Private/SLevelEditor.cpp
+++ b/Engine/Source/Editor/LevelEditor/Private/SLevelEditor.cpp
@@ -1067,8 +1067,7 @@ TSharedRef<SWidget> SLevelEditor::RestoreContentArea( const TSharedRef<SDockTab>
 	// or there wont be any tab spawners for the modes.
 	RefreshEditorModeCommands();
 
-	const TSharedRef<FTabManager::FLayout> Layout = FLayoutSaveRestore::LoadUserConfigVersionOf
-	(	
+	const TSharedRef<FTabManager::FLayout> Layout = FLayoutSaveRestore::LoadFromConfig(GEditorLayoutIni,
 		FTabManager::NewLayout( "LevelEditor_Layout_v1.1" )
 		->AddArea
 		(
diff --git a/Engine/Source/Editor/LevelEditor/Public/LevelEditor.h b/Engine/Source/Editor/LevelEditor/Public/LevelEditor.h
--- a/Engine/Source/Editor/LevelEditor/Public/LevelEditor.h
+++ b/Engine/Source/Editor/LevelEditor/Public/LevelEditor.h
@@ -228,6 +228,12 @@ private:
 	 */
 	void BindGlobalLevelEditorCommands();
 
+	// Callback for persisting the Level Editor's layout.
+	void HandleTabManagerPersistLayout( const TSharedRef<FTabManager::FLayout>& LayoutToSave )
+	{
+		FLayoutSaveRestore::SaveToConfig(GEditorLayoutIni, LayoutToSave);
+	}
+
 private:
 	TSharedPtr<FExtensibilityManager> MenuExtensibilityManager;
 	TSharedPtr<FExtensibilityManager> ToolBarExtensibilityManager;
diff --git a/Engine/Source/Editor/MainFrame/Private/Frame/MainFrameActions.cpp b/Engine/Source/Editor/MainFrame/Private/Frame/MainFrameActions.cpp
--- a/Engine/Source/Editor/MainFrame/Private/Frame/MainFrameActions.cpp
+++ b/Engine/Source/Editor/MainFrame/Private/Frame/MainFrameActions.cpp
@@ -699,29 +699,29 @@ void FMainFrameActionCallbacks::ResetLayout()
 	// make a backup
 	GEditor->SaveEditorUserSettings();
 
-	FString BackupUserSettingsIni = FString::Printf(TEXT("%s_Backup.ini"), *FPaths::GetBaseFilename(GEditorUserSettingsIni, false));
+	FString BackupEditorLayoutIni = FString::Printf(TEXT("%s_Backup.ini"), *FPaths::GetBaseFilename(GEditorLayoutIni, false));
 
-	if( COPY_Fail == IFileManager::Get().Copy(*BackupUserSettingsIni, *GEditorUserSettingsIni) )
+	if( COPY_Fail == IFileManager::Get().Copy(*BackupEditorLayoutIni, *GEditorLayoutIni) )
 	{
 		FMessageLog EditorErrors("EditorErrors");
-		if(!FPaths::FileExists(GEditorUserSettingsIni))
+		if(!FPaths::FileExists(GEditorLayoutIni))
 		{
 			FFormatNamedArguments Arguments;
-			Arguments.Add(TEXT("FileName"), FText::FromString(GEditorUserSettingsIni));
+			Arguments.Add(TEXT("FileName"), FText::FromString(GEditorLayoutIni));
 			EditorErrors.Warning(FText::Format(LOCTEXT("UnsuccessfulBackup_NoExist_Notification", "Unsuccessful backup! {FileName} does not exist!"), Arguments));
 		}
-		else if(IFileManager::Get().IsReadOnly(*BackupUserSettingsIni))
+		else if(IFileManager::Get().IsReadOnly(*BackupEditorLayoutIni))
 		{
 			FFormatNamedArguments Arguments;
-			Arguments.Add(TEXT("FileName"), FText::FromString(FPaths::ConvertRelativePathToFull(BackupUserSettingsIni)));
+			Arguments.Add(TEXT("FileName"), FText::FromString(FPaths::ConvertRelativePathToFull(BackupEditorLayoutIni)));
 			EditorErrors.Warning(FText::Format(LOCTEXT("UnsuccessfulBackup_ReadOnly_Notification", "Unsuccessful backup! {FileName} is read-only!"), Arguments));
 		}
 		else
 		{
 			// We don't specifically know why it failed, this is a fallback.
 			FFormatNamedArguments Arguments;
-			Arguments.Add(TEXT("SourceFileName"), FText::FromString(GEditorUserSettingsIni));
-			Arguments.Add(TEXT("BackupFileName"), FText::FromString(FPaths::ConvertRelativePathToFull(BackupUserSettingsIni)));
+			Arguments.Add(TEXT("SourceFileName"), FText::FromString(GEditorLayoutIni));
+			Arguments.Add(TEXT("BackupFileName"), FText::FromString(FPaths::ConvertRelativePathToFull(BackupEditorLayoutIni)));
 			EditorErrors.Warning(FText::Format(LOCTEXT("UnsuccessfulBackup_Fallback_Notification", "Unsuccessful backup of {SourceFileName} to {BackupFileName}"), Arguments));
 		}
 		EditorErrors.Notify(LOCTEXT("BackupUnsuccessful_Title", "Backup Unsuccessful!"));
@@ -732,7 +732,7 @@ void FMainFrameActionCallbacks::ResetLayout()
 		ErrorNotification.bFireAndForget = true;
 		ErrorNotification.ExpireDuration = 3.0f;
 		ErrorNotification.bUseThrobber = true;
-		ErrorNotification.Hyperlink = FSimpleDelegate::CreateStatic(&FMainFrameActionCallbacks::OpenBackupDirectory, BackupUserSettingsIni);
+		ErrorNotification.Hyperlink = FSimpleDelegate::CreateStatic(&FMainFrameActionCallbacks::OpenBackupDirectory, BackupEditorLayoutIni);
 		ErrorNotification.HyperlinkText = LOCTEXT("SuccessfulBackup_Notification_Hyperlink", "Open Directory");
 		ErrorNotification.Text = LOCTEXT("SuccessfulBackup_Notification", "Backup Successful!");
 		ErrorNotification.Image = FEditorStyle::GetBrush(TEXT("NotificationList.SuccessImage"));
@@ -749,7 +749,7 @@ void FMainFrameActionCallbacks::SaveLayout()
 	FGlobalTabmanager::Get()->SaveAllVisualState();
 
 	// Write the saved state's config to disk
-	GConfig->Flush( false, GEditorUserSettingsIni );
+	GConfig->Flush( false, GEditorLayoutIni );
 }
 
 void FMainFrameActionCallbacks::ToggleFullscreen_Execute()
diff --git a/Engine/Source/Editor/MainFrame/Private/Frame/MainFrameHandler.h b/Engine/Source/Editor/MainFrame/Private/Frame/MainFrameHandler.h
--- a/Engine/Source/Editor/MainFrame/Private/Frame/MainFrameHandler.h
+++ b/Engine/Source/Editor/MainFrame/Private/Frame/MainFrameHandler.h
@@ -55,7 +55,7 @@ public:
 		}
 		else
 		{
-			GConfig->EmptySection(TEXT("EditorLayouts"), *GEditorUserSettingsIni);
+			GConfig->EmptySection(TEXT("EditorLayouts"), *GEditorLayoutIni);
 		}
 
 
@@ -242,7 +242,7 @@ public:
 		TSharedRef<FGlobalTabmanager> GlobalTabManager = FGlobalTabmanager::Get();
 		
 		// Persistent layouts should get stored using the specified method.
-		GlobalTabManager->SetOnPersistLayout( FTabManager::FOnPersistLayout::CreateStatic( &FLayoutSaveRestore::SaveTheLayout ) );
+		GlobalTabManager->SetOnPersistLayout(FTabManager::FOnPersistLayout::CreateRaw(this, &FMainFrameHandler::HandleTabManagerPersistLayout));
 		
 		const bool bIncludeGameName = true;
 		GlobalTabManager->SetApplicationTitle( StaticGetApplicationTitle( bIncludeGameName ) );
@@ -323,6 +323,13 @@ public:
 		}
 	}
 
+private:
+
+	// Callback for persisting the Level Editor's layout.
+	void HandleTabManagerPersistLayout( const TSharedRef<FTabManager::FLayout>& LayoutToSave )
+	{
+		FLayoutSaveRestore::SaveToConfig(GEditorLayoutIni, LayoutToSave);
+	}
 
 private:
 
diff --git a/Engine/Source/Editor/MainFrame/Private/MainFrameModule.cpp b/Engine/Source/Editor/MainFrame/Private/MainFrameModule.cpp
--- a/Engine/Source/Editor/MainFrame/Private/MainFrameModule.cpp
+++ b/Engine/Source/Editor/MainFrame/Private/MainFrameModule.cpp
@@ -115,7 +115,7 @@ void FMainFrameModule::CreateDefaultMainFrame( const bool bStartImmersivePIE )
 				DisplayMetrics.PrimaryDisplayWorkAreaRect.Bottom - DisplayMetrics.PrimaryDisplayWorkAreaRect.Top );
 			const FVector2D WindowSize = CenterScale * DisplaySize;
 
-			TSharedRef<FTabManager::FLayout> LoadedLayout = FLayoutSaveRestore::LoadUserConfigVersionOf(
+			TSharedRef<FTabManager::FLayout> LoadedLayout = FLayoutSaveRestore::LoadFromConfig(GEditorLayoutIni,
 				// We persist the positioning of the level editor and the content browser.
 				// The asset editors currently do not get saved.
 				FTabManager::NewLayout( "UnrealEd_Layout_v1.1" )
@@ -563,6 +563,9 @@ void FMainFrameModule::StartupModule( )
 	CompileFailSound->AddToRoot();
 
 	ModuleCompileStartTime = 0.0f;
+
+	// migrate old layout settings
+	FLayoutSaveRestore::MigrateConfig(GEditorUserSettingsIni, GEditorLayoutIni);
 }
 
 
diff --git a/Engine/Source/Editor/UnrealEd/Private/Toolkits/AssetEditorToolkit.cpp b/Engine/Source/Editor/UnrealEd/Private/Toolkits/AssetEditorToolkit.cpp
--- a/Engine/Source/Editor/UnrealEd/Private/Toolkits/AssetEditorToolkit.cpp
+++ b/Engine/Source/Editor/UnrealEd/Private/Toolkits/AssetEditorToolkit.cpp
@@ -92,7 +92,7 @@ void FAssetEditorToolkit::InitAssetEditor( const EToolkitMode::Type Mode, const
 		NewMajorTab->SetRightContent(UserFeedbackWidget);
 
 		const TSharedRef<FTabManager> NewTabManager = FGlobalTabmanager::Get()->NewTabManager( NewMajorTab.ToSharedRef() );		
-		NewTabManager->SetOnPersistLayout( FTabManager::FOnPersistLayout::CreateStatic( &FLayoutSaveRestore::SaveTheLayout ) );
+		NewTabManager->SetOnPersistLayout(FTabManager::FOnPersistLayout::CreateRaw(this, &FAssetEditorToolkit::HandleTabManagerPersistLayout));
 		this->TabManager = NewTabManager;
 
 		NewMajorTab->SetContent
@@ -113,7 +113,7 @@ void FAssetEditorToolkit::InitAssetEditor( const EToolkitMode::Type Mode, const
 	
 	if (ToolkitMode == EToolkitMode::Standalone)
 	{
-		TSharedRef<FTabManager::FLayout> LayoutToUse = FLayoutSaveRestore::LoadUserConfigVersionOf(StandaloneDefaultLayout);
+		TSharedRef<FTabManager::FLayout> LayoutToUse = FLayoutSaveRestore::LoadFromConfig(GEditorLayoutIni, StandaloneDefaultLayout);
 
 		// Actually create the widget content
 		NewStandaloneHost->SetupInitialContent( LayoutToUse, NewMajorTab, bCreateDefaultStandaloneMenu );
@@ -704,10 +704,10 @@ void FAssetEditorToolkit::RestoreFromLayout(const TSharedRef<FTabManager::FLayou
 	if (HostWidget.Get() != NULL)
 	{
 		// Save the old layout
-		FLayoutSaveRestore::SaveTheLayout(TabManager->PersistLayout());
+		FLayoutSaveRestore::SaveToConfig(GEditorIni, TabManager->PersistLayout());
 
 		// Load the potentially previously saved new layout
-		TSharedRef<FTabManager::FLayout> UserConfiguredNewLayout = FLayoutSaveRestore::LoadUserConfigVersionOf(NewLayout);
+		TSharedRef<FTabManager::FLayout> UserConfiguredNewLayout = FLayoutSaveRestore::LoadFromConfig(GEditorLayoutIni, NewLayout);
 
 		// Apply the new layout
 		HostWidget->RestoreFromLayout(UserConfiguredNewLayout);
diff --git a/Engine/Source/Editor/UnrealEd/Public/Toolkits/AssetEditorToolkit.h b/Engine/Source/Editor/UnrealEd/Public/Toolkits/AssetEditorToolkit.h
--- a/Engine/Source/Editor/UnrealEd/Public/Toolkits/AssetEditorToolkit.h
+++ b/Engine/Source/Editor/UnrealEd/Public/Toolkits/AssetEditorToolkit.h
@@ -191,6 +191,12 @@ private:
 	/** Spawns the toolbar tab */
 	TSharedRef<SDockTab> SpawnTab_Toolbar(const FSpawnTabArgs& Args);
 
+	// Callback for persisting the Asset Editor's layout.
+	void HandleTabManagerPersistLayout( const TSharedRef<FTabManager::FLayout>& LayoutToSave )
+	{
+		FLayoutSaveRestore::SaveToConfig(GEditorLayoutIni, LayoutToSave);
+	}
+
 protected:
 
 	/** For standalone asset editing tool-kits that were switched from world-centric mode on the fly, this stores
diff --git a/Engine/Source/Programs/UnrealFrontend/Private/UnrealFrontendMain.cpp b/Engine/Source/Programs/UnrealFrontend/Private/UnrealFrontendMain.cpp
--- a/Engine/Source/Programs/UnrealFrontend/Private/UnrealFrontendMain.cpp
+++ b/Engine/Source/Programs/UnrealFrontend/Private/UnrealFrontendMain.cpp
@@ -342,7 +342,7 @@ void RunStatsConvertCommand()
 
 void RunUI()
 {
-	FString UnrealFrontendUserSettingsIni = FPaths::GetPath(GEngineIni) + "/UnrealFrontend.ini";
+	FString UnrealFrontendLayoutIni = FPaths::GetPath(GEngineIni) + "/Layout.ini";
 
 	FSlateApplication::InitializeAsStandaloneApplication(GetStandardStandaloneRenderer());
 
@@ -402,7 +402,7 @@ void RunUI()
 				)
 		);
 
-	TSharedRef<FTabManager::FLayout> UserConfiguredNewLayout = FLayoutSaveRestore::LoadFromConfig(NewLayout, UnrealFrontendUserSettingsIni);
+	TSharedRef<FTabManager::FLayout> UserConfiguredNewLayout = FLayoutSaveRestore::LoadFromConfig(UnrealFrontendLayoutIni, NewLayout);
 	FGlobalTabmanager::Get()->RestoreFrom(UserConfiguredNewLayout, TSharedPtr<SWindow>());
 
 	// enter main loop
@@ -440,8 +440,8 @@ void RunUI()
 	}
 
 	// save application layout
-	FLayoutSaveRestore::SaveToConfig(UserConfiguredNewLayout, UnrealFrontendUserSettingsIni);
-	GConfig->Flush(false, UnrealFrontendUserSettingsIni);
+	FLayoutSaveRestore::SaveToConfig(UnrealFrontendLayoutIni, UserConfiguredNewLayout);
+	GConfig->Flush(false, UnrealFrontendLayoutIni);
 
 	// shut down application
 	FSlateApplication::Shutdown();
diff --git a/Engine/Source/Runtime/Core/Private/Misc/ConfigCacheIni.cpp b/Engine/Source/Runtime/Core/Private/Misc/ConfigCacheIni.cpp
--- a/Engine/Source/Runtime/Core/Private/Misc/ConfigCacheIni.cpp
+++ b/Engine/Source/Runtime/Core/Private/Misc/ConfigCacheIni.cpp
@@ -2613,6 +2613,7 @@ void FConfigCacheIni::InitializeConfigSystem()
 #if WITH_EDITOR
 	// load some editor specific .ini files
 	FConfigCacheIni::LoadGlobalIniFile(GEditorIni, TEXT("Editor"));
+	FConfigCacheIni::LoadGlobalIniFile(GEditorLayoutIni, TEXT("EditorLayout"));
 	FConfigCacheIni::LoadGlobalIniFile(GEditorUserSettingsIni, TEXT("EditorUserSettings"));
 	FConfigCacheIni::LoadGlobalIniFile(GEditorKeyBindingsIni, TEXT("EditorKeyBindings"));
 
diff --git a/Engine/Source/Runtime/Core/Private/Misc/Core.cpp b/Engine/Source/Runtime/Core/Private/Misc/Core.cpp
--- a/Engine/Source/Runtime/Core/Private/Misc/Core.cpp
+++ b/Engine/Source/Runtime/Core/Private/Misc/Core.cpp
@@ -122,6 +122,7 @@ uint64					GMakeCacheIDIndex				= 0;						/* Cache ID */
 FString				GEngineIni;													/* Engine ini filename */
 FString				GEditorIni;													/* Editor ini filename */
 FString				GEditorKeyBindingsIni;										/* Editor Key Bindings ini file */
+FString				GEditorLayoutIni;											/* Editor UI Layout ini filename */
 FString				GEditorUserSettingsIni;										/* Editor User Settings ini filename */
 FString				GEditorGameAgnosticIni;										/* Editor Settings (shared between games) ini filename */
 FString				GCompatIni;
diff --git a/Engine/Source/Runtime/Core/Public/Core.h b/Engine/Source/Runtime/Core/Public/Core.h
--- a/Engine/Source/Runtime/Core/Public/Core.h
+++ b/Engine/Source/Runtime/Core/Public/Core.h
@@ -238,6 +238,7 @@ extern CORE_API uint64 GMakeCacheIDIndex;
 extern CORE_API FString GEditorKeyBindingsIni;
 extern CORE_API FString GEngineIni;
 extern CORE_API FString GEditorIni;
+extern CORE_API FString GEditorLayoutIni;
 extern CORE_API FString GEditorUserSettingsIni;
 extern CORE_API FString GEditorGameAgnosticIni;
 extern CORE_API FString GCompatIni;
diff --git a/Engine/Source/Runtime/Slate/Private/Framework/Docking/LayoutService.cpp b/Engine/Source/Runtime/Slate/Private/Framework/Docking/LayoutService.cpp
--- a/Engine/Source/Runtime/Slate/Private/Framework/Docking/LayoutService.cpp
+++ b/Engine/Source/Runtime/Slate/Private/Framework/Docking/LayoutService.cpp
@@ -3,6 +3,9 @@
 #include "SlatePrivatePCH.h"
 
 
+const TCHAR* EditorLayoutsSectionName = TEXT("EditorLayouts");
+
+
 const FString& FLayoutSaveRestore::GetAdditionalLayoutConfigIni()
 {
 	static const FString IniSectionAdditionalConfig = TEXT("SlateAdditionalLayoutConfig");
@@ -10,28 +13,19 @@ const FString& FLayoutSaveRestore::GetAdditionalLayoutConfigIni()
 }
 
 
-void FLayoutSaveRestore::SaveTheLayout( const TSharedRef<FTabManager::FLayout>& LayoutToSave )
-{
-	SaveToConfig(LayoutToSave, GEditorUserSettingsIni);
-}
-
-void FLayoutSaveRestore::SaveToConfig( const TSharedRef<FTabManager::FLayout>& LayoutToSave, FString ConfigFileName )
+void FLayoutSaveRestore::SaveToConfig( FString ConfigFileName, const TSharedRef<FTabManager::FLayout>& LayoutToSave )
 {
 	const FString LayoutAsString = FLayoutSaveRestore::PrepareLayoutStringForIni( LayoutToSave->ToString() );
 
-	GConfig->SetString( TEXT("EditorLayouts"), *LayoutToSave->GetLayoutName().ToString(), *LayoutAsString, ConfigFileName );
+	GConfig->SetString(EditorLayoutsSectionName, *LayoutToSave->GetLayoutName().ToString(), *LayoutAsString, ConfigFileName );
 }
 
-TSharedRef<FTabManager::FLayout> FLayoutSaveRestore::LoadUserConfigVersionOf( const TSharedRef<FTabManager::FLayout>& DefaultLayout )
-{
-	return LoadFromConfig(DefaultLayout,GEditorUserSettingsIni);
-}
 
-TSharedRef<FTabManager::FLayout> FLayoutSaveRestore::LoadFromConfig( const TSharedRef<FTabManager::FLayout>& DefaultLayout, FString ConfigFileName )
+TSharedRef<FTabManager::FLayout> FLayoutSaveRestore::LoadFromConfig( const FString& ConfigFileName, const TSharedRef<FTabManager::FLayout>& DefaultLayout )
 {
 	const FName LayoutName = DefaultLayout->GetLayoutName();
 	FString UserLayoutString;
-	if ( GConfig->GetString( TEXT("EditorLayouts"), *LayoutName.ToString(), UserLayoutString, ConfigFileName ) && !UserLayoutString.IsEmpty() )
+	if ( GConfig->GetString(EditorLayoutsSectionName, *LayoutName.ToString(), UserLayoutString, ConfigFileName ) && !UserLayoutString.IsEmpty() )
 	{
 		TSharedPtr<FTabManager::FLayout> UserLayout = FTabManager::FLayout::NewFromString( FLayoutSaveRestore::GetLayoutStringFromIni( UserLayoutString ));
 		if ( UserLayout.IsValid() && UserLayout->GetPrimaryArea().IsValid() )
@@ -43,12 +37,47 @@ TSharedRef<FTabManager::FLayout> FLayoutSaveRestore::LoadFromConfig( const TShar
 	return DefaultLayout;
 }
 
+
+void FLayoutSaveRestore::MigrateConfig( const FString& OldConfigFileName, const FString& NewConfigFileName )
+{
+	TArray<FString> OldSectionStrings;
+
+	// check whether any layout configuration needs to be migrated
+	if (!GConfig->GetSection(EditorLayoutsSectionName, OldSectionStrings, OldConfigFileName) || (OldSectionStrings.Num() == 0))
+	{
+		return;
+	}
+
+	TArray<FString> NewSectionStrings;
+
+	// migrate old configuration if a new layout configuration does not yet exist
+	if (!GConfig->GetSection(EditorLayoutsSectionName, NewSectionStrings, NewConfigFileName) || (NewSectionStrings.Num() == 0))
+	{
+		FString Key, Value;
+
+		for (auto SectionString : OldSectionStrings)
+		{
+			if (SectionString.Split(TEXT("="), &Key, &Value))
+			{
+				GConfig->SetString(EditorLayoutsSectionName, *Key, *Value, NewConfigFileName);
+			}
+		}
+	}
+
+	// remove old configuration
+	GConfig->EmptySection(EditorLayoutsSectionName, OldConfigFileName);
+	GConfig->Flush(false, GEditorUserSettingsIni);
+	GConfig->Flush(false, GEditorLayoutIni);
+}
+
+
 FString FLayoutSaveRestore::PrepareLayoutStringForIni(const FString& LayoutString)
 {
 	// Have to store braces as parentheses due to braces causing ini issues
 	return LayoutString.Replace(TEXT("{"), TEXT("(")).Replace(TEXT("}"), TEXT(")")).Replace(LINE_TERMINATOR, TEXT("\\") LINE_TERMINATOR);
 }
 
+
 FString FLayoutSaveRestore::GetLayoutStringFromIni(const FString& LayoutString)
 {
 	// Revert parenthesis to braces, from ini readable to Json readable
diff --git a/Engine/Source/Runtime/Slate/Public/Framework/Docking/LayoutService.h b/Engine/Source/Runtime/Slate/Public/Framework/Docking/LayoutService.h
--- a/Engine/Source/Runtime/Slate/Public/Framework/Docking/LayoutService.h
+++ b/Engine/Source/Runtime/Slate/Public/Framework/Docking/LayoutService.h
@@ -8,35 +8,29 @@ struct SLATE_API FLayoutSaveRestore
 	/** Gets the ini section label for the additional layout configs */
 	static const FString& GetAdditionalLayoutConfigIni();
 
-	/** Write the LayoutToSave out into a a config file uses GEditorUserSettingsIni
-	 *
-	 * @param LayoutToSave the layout to save.
-	 */
-	static void SaveTheLayout( const TSharedRef<FTabManager::FLayout>& LayoutToSave );
-
 	/** Write the layout out into a named config file.
 	 *
-	 * @param LayoutToSave the layout to save.
 	 * @param ConfigFileName file to be save to
+	 * @param LayoutToSave the layout to save.
 	 */
-	static void SaveToConfig( const TSharedRef<FTabManager::FLayout>& LayoutToSave, FString ConfigFileName );
+	static void SaveToConfig( FString ConfigFileName, const TSharedRef<FTabManager::FLayout>& LayoutToSave );
 
-	/** Given a named DefaultLayout, return any saved versions of it. Otherwise return the default layout uses GEditorUserSettingsIni
+	/** Given a named DefualtLayout, return any saved version of it from the given ini file, otherwise return the default, also default to open tabs based on bool
 	 *
+	 * @param ConfigFileName file to be used to load an existing layout
 	 * @param DefaultLayout the layout to be used if the file does not exist.
 	 *
 	 * @return Loaded layout or the default
 	 */
-	static TSharedRef<FTabManager::FLayout> LoadUserConfigVersionOf( const TSharedRef<FTabManager::FLayout>& DefaultLayout );
+	static TSharedRef<FTabManager::FLayout> LoadFromConfig( const FString& ConfigFileName, const TSharedRef<FTabManager::FLayout>& DefaultLayout );
 
-	/** Given a named DefualtLayout, return any saved version of it from the given ini file, otherwise return the default, also default to open tabs based on bool
-	 *
-	 * @param DefaultLayout the layout to be used if the file does not exist.
-	 * @param ConfigFileName file to be used to load an existing layout
+	/**
+	 * Migrates the layout configuration from one config file to another.
 	 *
-	 * @return Loaded layout or the default
+	 * @param OldConfigFileName The name of the old configuration file.
+	 * @param NewConfigFileName The name of the new configuration file.
 	 */
-	static TSharedRef<FTabManager::FLayout> LoadFromConfig( const TSharedRef<FTabManager::FLayout>& DefaultLayout, FString ConfigFileName);
+	static void MigrateConfig( const FString& OldConfigFileName, const FString& NewConfigFileName );
 
 private:
 	/**
@@ -50,4 +44,4 @@ private:
 	 * The opposite of PrepareLayoutStringForIni.
 	 */
 	static FString GetLayoutStringFromIni(const FString& LayoutString);
-};
\ No newline at end of file
+};
