Slate: Moved window management related code to SlateCore

diff --git a/Engine/Plugins/Slate/SlateRemote/Source/SlateRemote/Private/Server/SlateRemoteServer.cpp b/Engine/Plugins/Slate/SlateRemote/Source/SlateRemote/Private/Server/SlateRemoteServer.cpp
--- a/Engine/Plugins/Slate/SlateRemote/Source/SlateRemote/Private/Server/SlateRemoteServer.cpp
+++ b/Engine/Plugins/Slate/SlateRemote/Source/SlateRemote/Private/Server/SlateRemoteServer.cpp
@@ -145,7 +145,7 @@ void FSlateRemoteServer::ProcessTouchMessage( const FSlateRemoteServerMessage& M
 						
 				if (WidgetPath.Widgets.Num() == 0 || WidgetPath.Widgets.Last().Widget != GameViewport)
 				{
-					SlateApplication.FindPathToWidget(SlateApplication.GetInteractiveTopLevelWindows(), GameViewport.ToSharedRef(), WidgetPath);
+					SlateApplication.FindPathToWidget(GameViewport.ToSharedRef(), WidgetPath);
 					GameViewportWidgetPath = WidgetPath;
 				}
 
diff --git a/Engine/Source/Runtime/Slate/Private/Framework/Application/SlateApplication.cpp b/Engine/Source/Runtime/Slate/Private/Framework/Application/SlateApplication.cpp
--- a/Engine/Source/Runtime/Slate/Private/Framework/Application/SlateApplication.cpp
+++ b/Engine/Source/Runtime/Slate/Private/Framework/Application/SlateApplication.cpp
@@ -11,7 +11,6 @@
 DECLARE_CYCLE_STAT( TEXT("Message Tick Time"), STAT_SlateMessageTick, STATGROUP_Slate );
 DECLARE_CYCLE_STAT( TEXT("Update Tooltip Time"), STAT_SlateUpdateTooltip, STATGROUP_Slate );
 DECLARE_CYCLE_STAT( TEXT("Tick Window And Children Time"), STAT_SlateTickWindowAndChildren, STATGROUP_Slate );
-DECLARE_CYCLE_STAT( TEXT("FindPathToWidget"), STAT_FindPathToWidget, STATGROUP_Slate );
 
 
 // Slate Event Logging is enabled to allow crash log dumping
@@ -296,103 +295,6 @@ static bool LocateWidgetsUnderCursor_Helper( FArrangedWidget& Candidate, FVector
 }
 
 
-void FSlateApplication::ArrangeWindowToFront( TArray< TSharedRef<SWindow> >& Windows, const TSharedRef<SWindow>& WindowToBringToFront )
-{
-	Windows.Remove( WindowToBringToFront );
-
-	if ( Windows.Num() == 0 || WindowToBringToFront->IsTopmostWindow() )
-	{
-		Windows.Add( WindowToBringToFront );
-	}
-	else
-	{
-		bool PerformedInsert = false;
-		for (int Index = Windows.Num() - 1; Index >= 0 ; Index--)
-		{
-			if ( !Windows[Index]->IsTopmostWindow() )
-			{
-				Windows.Insert( WindowToBringToFront, Index + 1 );
-				PerformedInsert = true;
-				break;
-			}
-		}
-
-		if ( !PerformedInsert )
-		{
-			Windows.Insert( WindowToBringToFront, 0 );
-		}
-	}
-}
-
-/**
- * Make BringMeToFront first among its peers.
- * i.e. make it the last window in its parent's list of child windows.
- *
- * @return The top-most window whose children were re-arranged
- */
-static TSharedRef<SWindow> BringToFrontInParent( const TSharedRef<SWindow>& BringMeToFront )
-{
-	const TSharedPtr<SWindow> ParentWindow = BringMeToFront->GetParentWindow();
-
-	if ( ParentWindow.IsValid() )
-	{
-		FSlateApplication::ArrangeWindowToFront( ParentWindow->GetChildWindows(), BringMeToFront );
-
-		return BringToFrontInParent( ParentWindow.ToSharedRef() );
-	}
-		
-	// This window has no parent, so it must be the top-most window.
-	return BringMeToFront;
-}
-
-/**
- * Put 'BringMeToFront' at the font of the list of 'WindowsToReorder'. The top-most (front-most) window is last in Z-order and therefore
- * is added to the end of the list.
- *
- * @param WindowsToReorder  An ordered list of windows.
- * @param BingMeToFront     Window to be brough to front.
- */
-static void BringWindowToFront( TArray< TSharedRef<SWindow> >& WindowsToReorder, const TSharedRef<SWindow>& BringMeToFront )
-{
-	const TSharedRef<SWindow> TopLeveWindowToReorder = BringToFrontInParent( BringMeToFront );
-
-	FSlateApplication::ArrangeWindowToFront( WindowsToReorder, TopLeveWindowToReorder );
-}
-
-/** @return true if the list of WindowsToSearch or any of their children contains the WindowToFind */
-static bool ContainsWindow( const TArray< TSharedRef<SWindow> >& WindowsToSearch, const TSharedRef<SWindow>& WindowToFind  )
-{
-	if ( WindowsToSearch.Contains(WindowToFind) )
-	{
-		return true;
-	}
-	else
-	{
-		for (int32 WindowIndex=0; WindowIndex < WindowsToSearch.Num(); ++WindowIndex)
-		{
-			if ( ContainsWindow( WindowsToSearch[WindowIndex]->GetChildWindows(), WindowToFind ) )
-			{
-				return true;
-			}
-		}
-	}
-
-	return false;
-}
-
-/** Remove the WindowToRemove (and its children) from TheList of windows or from their children. */
-static void RemoveWindowFromList( TArray< TSharedRef<SWindow> >& TheList, const TSharedRef<SWindow>& WindowToRemove )
-{
-	int32 NumRemoved = TheList.Remove( WindowToRemove );
-	if (NumRemoved == 0)
-	{
-		for ( int32 ChildIndex=0; ChildIndex < TheList.Num(); ++ChildIndex )
-		{
-			RemoveWindowFromList( TheList[ChildIndex]->GetChildWindows(), WindowToRemove ) ;
-		}	
-	}
-}
-
 void FSlateApplication::Create()
 {
 	EKeys::Initialize();
@@ -1038,7 +940,7 @@ TSharedRef<SWindow> FSlateApplication::AddWindow( TSharedRef<SWindow> InSlateWin
 	// activation message may be sent by the OS as soon as the window is shown (in the Init function), and if we
 	// don't add the Slate window to our window list, we wouldn't be able to route that message to the window.
 
-	ArrangeWindowToFront( SlateWindows, InSlateWindow );
+	FSlateWindowHelper::ArrangeWindowToFront(SlateWindows, InSlateWindow);
 	TSharedRef<FGenericWindow> NewWindow = MakeWindow( InSlateWindow, bShowImmediately );
 
 	if( bShowImmediately )
@@ -1216,7 +1118,7 @@ void FSlateApplication::SetModalWindowStackEndedDelegate(FModalWindowStackEnded
 TSharedRef<SWindow> FSlateApplication::AddWindowAsNativeChild( TSharedRef<SWindow> InSlateWindow, TSharedRef<SWindow> InParentWindow, const bool bShowImmediately )
 {
 	// Parent window must already have been added
-	checkSlow( ContainsWindow( SlateWindows, InParentWindow ) );
+	checkSlow(FSlateWindowHelper::ContainsWindow(SlateWindows, InParentWindow));
 
 	// Add the Slate window to the Slate application's top-level window array.  Note that neither the Slate window
 	// or the native window are ready to be used yet, however we need to make sure they're in the Slate window
@@ -1399,7 +1301,7 @@ void FSlateApplication::RegisterGameViewport( TSharedRef<SViewport> InViewport )
 
 	FWidgetPath PathToViewport;
 	// If we cannot find the window it could have been destroyed.
-	if( FindPathToWidget( SlateWindows, InViewport, PathToViewport, EVisibility::All ) )
+	if (FSlateWindowHelper::FindPathToWidget(SlateWindows, InViewport, PathToViewport, EVisibility::All))
 	{
 		FReply Reply = FReply::Handled().SetKeyboardFocus( InViewport, EKeyboardFocusCause::SetDirectly );
 	
@@ -1442,7 +1344,7 @@ void FSlateApplication::SetJoystickCaptorToGameViewport()
 	if ( CurrentGameViewportWidget.IsValid() )
 	{
 		FWidgetPath PathToWidget;
-		FindPathToWidget( SlateWindows, CurrentGameViewportWidget.ToSharedRef(), /*OUT*/ PathToWidget );
+		FSlateWindowHelper::FindPathToWidget(SlateWindows, CurrentGameViewportWidget.ToSharedRef(), /*OUT*/ PathToWidget);
 
 		FReply Temp = FReply::Handled().CaptureJoystick(CurrentGameViewportWidget.ToSharedRef());
 
@@ -1455,7 +1357,7 @@ void FSlateApplication::SetKeyboardFocus( const TSharedPtr< SWidget >& OptionalW
 	if (OptionalWidgetToFocus.IsValid())
 	{
 		FWidgetPath PathToWidget;
-		FindPathToWidget( SlateWindows, OptionalWidgetToFocus.ToSharedRef(), /*OUT*/ PathToWidget );
+		FSlateWindowHelper::FindPathToWidget(SlateWindows, OptionalWidgetToFocus.ToSharedRef(), /*OUT*/ PathToWidget);
 
 		FReply Reply = FReply::Handled();
 		Reply.SetKeyboardFocus( OptionalWidgetToFocus.ToSharedRef(), EKeyboardFocusCause::SetDirectly );
@@ -1785,51 +1687,15 @@ void FSlateApplication::SetExitRequestedHandler( const FSimpleDelegate& OnExitRe
 }
 
 
-bool FSlateApplication::FindPathToWidget( const TArray< TSharedRef<SWindow> > WindowsToSearch,  TSharedRef< const SWidget > InWidget, FWidgetPath& OutWidgetPath, EVisibility VisibilityFilter )
-{
-	SCOPE_CYCLE_COUNTER( STAT_FindPathToWidget);
-
-	// Iterate over our top level windows
-	bool bFoundWidget = false;
-	for( int32 WindowIndex = 0; !bFoundWidget && WindowIndex < WindowsToSearch.Num(); ++WindowIndex )
-	{
-		// Make a widget path that contains just the top-level window
-		TSharedRef< SWindow > CurWindow = WindowsToSearch[ WindowIndex ];
-		
-		FArrangedChildren JustWindow(VisibilityFilter);
-		{
-			JustWindow.AddWidget( FArrangedWidget(CurWindow, CurWindow->GetWindowGeometryInScreen()) );
-		}
-		
-		FWidgetPath PathToWidget( CurWindow, JustWindow );
-		
-		// Attempt to extend it to the desired child widget; essentially a full-window search for 'InWidget
-		if ( CurWindow == InWidget || PathToWidget.ExtendPathTo( FWidgetMatcher(InWidget), VisibilityFilter ) )
-		{
-			OutWidgetPath = PathToWidget;
-			bFoundWidget = true;
-		}
-
-		if ( !bFoundWidget )
-		{
-			// Search this window's children
-			bFoundWidget = FindPathToWidget(CurWindow->GetChildWindows(), InWidget, OutWidgetPath, VisibilityFilter);
-		}
-	}
-
-	return bFoundWidget;
-}
-
-
 bool FSlateApplication::GeneratePathToWidgetUnchecked( TSharedRef< const SWidget > InWidget, FWidgetPath& OutWidgetPath, EVisibility VisibilityFilter ) const
 {
-	return FindPathToWidget( SlateWindows, InWidget, OutWidgetPath, VisibilityFilter );
+	return FSlateWindowHelper::FindPathToWidget(SlateWindows, InWidget, OutWidgetPath, VisibilityFilter);
 }
 
 
 void FSlateApplication::GeneratePathToWidgetChecked( TSharedRef< const SWidget > InWidget, FWidgetPath& OutWidgetPath, EVisibility VisibilityFilter ) const
 {
-	const bool bWasFound = FindPathToWidget( SlateWindows, InWidget, OutWidgetPath, VisibilityFilter );
+	const bool bWasFound = FSlateWindowHelper::FindPathToWidget(SlateWindows, InWidget, OutWidgetPath, VisibilityFilter);
 	check( bWasFound );
 }
 
@@ -1844,7 +1710,7 @@ TSharedPtr<SWindow> FSlateApplication::FindWidgetWindow( TSharedRef< const SWidg
 TSharedPtr<SWindow> FSlateApplication::FindWidgetWindow( TSharedRef< const SWidget > InWidget, FWidgetPath& OutWidgetPath ) const
 {
 	// If the user wants a widget path back populate it instead
-	const bool bWasFound = FindPathToWidget( SlateWindows, InWidget, OutWidgetPath, EVisibility::All );
+	const bool bWasFound = FSlateWindowHelper::FindPathToWidget(SlateWindows, InWidget, OutWidgetPath, EVisibility::All);
 	if( bWasFound )
 	{
 		return OutWidgetPath.TopLevelWindow;
@@ -2810,28 +2676,7 @@ FSlateRect FSlateApplication::GetPreferredWorkArea() const
 		const FVector2D CursorPos = GetCursorPos();
 		const FSlateRect WorkArea = GetWorkArea( FSlateRect( CursorPos.X, CursorPos.Y, CursorPos.X + 1.0f, CursorPos.Y + 1.0f ) );
 
-		struct Local
-		{
-			static bool CheckWorkAreaForWindows_Helper( const TArray< TSharedRef<SWindow> >& WindowsToSearch, const FSlateRect& InWorkAreaRect )
-			{
-				for( TArray< TSharedRef<SWindow> >::TConstIterator CurrentWindowIt( WindowsToSearch ); CurrentWindowIt; ++CurrentWindowIt )
-				{
-					const TSharedRef<SWindow>& CurrentWindow = *CurrentWindowIt;
-					const FVector2D Position = CurrentWindow->GetPositionInScreen();
-					const FVector2D Size = CurrentWindow->GetSizeInScreen();
-					const FSlateRect WindowRect( Position.X, Position.Y, Size.X, Size.Y );
-
-					if( FSlateRect::DoRectanglesIntersect( InWorkAreaRect, WindowRect ) || CheckWorkAreaForWindows_Helper( CurrentWindow->GetChildWindows(), InWorkAreaRect ) )
-					{
-						return true;
-					}
-				}
-
-				return false;
-			}
-		};
-
-		if( Local::CheckWorkAreaForWindows_Helper( SlateWindows, WorkArea ) )
+		if (FSlateWindowHelper::CheckWorkAreaForWindows(SlateWindows, WorkArea))
 		{
 			return WorkArea;
 		}
@@ -2937,30 +2782,6 @@ bool FSlateApplication::TakeScreenshot(TSharedRef<SWidget>& Widget, const FIntRe
 	return true;
 }
 
-static TSharedPtr<SWindow> FindWindowByPlatformWindow( const TArray< TSharedRef<SWindow> >& WindowsToSearch, const TSharedRef< FGenericWindow >& PlatformWindow )
-{
-	for (int32 WindowIndex=0; WindowIndex < WindowsToSearch.Num(); ++WindowIndex)
-	{
-		TSharedRef<SWindow> SomeWindow = WindowsToSearch[WindowIndex];
-		TSharedRef<FGenericWindow> SomeNativeWindow = StaticCastSharedRef<FGenericWindow>( SomeWindow->GetNativeWindow().ToSharedRef() );
-		if ( SomeNativeWindow == PlatformWindow )
-		{
-			return SomeWindow;
-		}
-		else
-		{
-			// Search child windows
-			TSharedPtr<SWindow> FoundChildWindow = FindWindowByPlatformWindow( SomeWindow->GetChildWindows(), PlatformWindow );
-			if (FoundChildWindow.IsValid())
-			{
-				return FoundChildWindow;
-			}
-		}
-	}
-
-	return TSharedPtr<SWindow>( NULL );
-}
-
 
 /* FSlateApplicationBase interface
  *****************************************************************************/
@@ -3051,7 +2872,7 @@ bool FSlateApplication::ShouldProcessUserInputMessages( const TSharedPtr< FGener
 	TSharedPtr< SWindow > Window;
 	if ( PlatformWindow.IsValid() )
 	{
-		Window = FindWindowByPlatformWindow( SlateWindows, PlatformWindow.ToSharedRef() );
+		Window = FSlateWindowHelper::FindWindowByPlatformWindow( SlateWindows, PlatformWindow.ToSharedRef() );
 	}
 
 	if ( ActiveModalWindows.Num() == 0 || 
@@ -4212,7 +4033,7 @@ void FSlateApplication::OnMotionDetectedMessage( FMotionEvent& MotionEvent )
 
 bool FSlateApplication::OnSizeChanged( const TSharedRef< FGenericWindow >& PlatformWindow, const int32 Width, const int32 Height, bool bWasMinimized )
 {
-	TSharedPtr< SWindow > Window = FindWindowByPlatformWindow( SlateWindows, PlatformWindow );
+	TSharedPtr< SWindow > Window = FSlateWindowHelper::FindWindowByPlatformWindow( SlateWindows, PlatformWindow );
 
 	if ( Window.IsValid() )
 	{
@@ -4241,7 +4062,7 @@ bool FSlateApplication::OnSizeChanged( const TSharedRef< FGenericWindow >& Platf
 
 void FSlateApplication::OnOSPaint( const TSharedRef< FGenericWindow >& PlatformWindow )
 {
-	TSharedPtr< SWindow > Window = FindWindowByPlatformWindow( SlateWindows, PlatformWindow );
+	TSharedPtr< SWindow > Window = FSlateWindowHelper::FindWindowByPlatformWindow( SlateWindows, PlatformWindow );
 	PrivateDrawWindows( Window );
 	Renderer->FlushCommands();
 }
@@ -4277,7 +4098,7 @@ void FSlateApplication::FinishedReshapingWindow( const TSharedRef< FGenericWindo
 
 void FSlateApplication::OnMovedWindow( const TSharedRef< FGenericWindow >& PlatformWindow, const int32 X, const int32 Y )
 {
-	TSharedPtr< SWindow > Window = FindWindowByPlatformWindow( SlateWindows, PlatformWindow );
+	TSharedPtr< SWindow > Window = FSlateWindowHelper::FindWindowByPlatformWindow( SlateWindows, PlatformWindow );
 
 	if ( Window.IsValid() )
 	{
@@ -4309,7 +4130,7 @@ FWindowActivateEvent::EActivationType TranslationWindowActivationMessage( const
 
 bool FSlateApplication::OnWindowActivationChanged( const TSharedRef< FGenericWindow >& PlatformWindow, const EWindowActivation::Type ActivationType )
 {
-	TSharedPtr< SWindow > Window = FindWindowByPlatformWindow( SlateWindows, PlatformWindow );
+	TSharedPtr< SWindow > Window = FSlateWindowHelper::FindWindowByPlatformWindow( SlateWindows, PlatformWindow );
 
 	if ( !Window.IsValid() )
 	{
@@ -4333,7 +4154,7 @@ bool FSlateApplication::ProcessWindowActivatedMessage( const FWindowActivateEven
 		{
 			// Window being ACTIVATED
 
-			BringWindowToFront( SlateWindows, ActivateEvent.GetAffectedWindow() );
+			FSlateWindowHelper::BringWindowToFront(SlateWindows, ActivateEvent.GetAffectedWindow());
 
 			{
 				// Switch worlds widgets in the current path
@@ -4450,7 +4271,7 @@ void FSlateApplication::ProcessApplicationActivationMessage( bool InAppActivated
 
 EWindowZone::Type FSlateApplication::GetWindowZoneForPoint( const TSharedRef< FGenericWindow >& PlatformWindow, const int32 X, const int32 Y )
 {
-	TSharedPtr< SWindow > Window = FindWindowByPlatformWindow( SlateWindows, PlatformWindow );
+	TSharedPtr< SWindow > Window = FSlateWindowHelper::FindWindowByPlatformWindow( SlateWindows, PlatformWindow );
 
 	if ( Window.IsValid() )
 	{
@@ -4475,7 +4296,7 @@ void FSlateApplication::PrivateDestroyWindow( const TSharedRef<SWindow>& Destroy
 	DestroyedWindow->DestroyWindowImmediately();
 
 	// Remove the window and all its children from the Slate window list
-	RemoveWindowFromList(SlateWindows, DestroyedWindow);
+	FSlateWindowHelper::RemoveWindowFromList(SlateWindows, DestroyedWindow);
 
 	// Shutdown the application if there are no more windows
 	{
@@ -4499,7 +4320,7 @@ void FSlateApplication::PrivateDestroyWindow( const TSharedRef<SWindow>& Destroy
 
 void FSlateApplication::OnWindowClose( const TSharedRef< FGenericWindow >& PlatformWindow )
 {
-	TSharedPtr< SWindow > Window = FindWindowByPlatformWindow( SlateWindows, PlatformWindow );
+	TSharedPtr< SWindow > Window = FSlateWindowHelper::FindWindowByPlatformWindow( SlateWindows, PlatformWindow );
 
 	if ( Window.IsValid() )
 	{
@@ -4510,7 +4331,7 @@ void FSlateApplication::OnWindowClose( const TSharedRef< FGenericWindow >& Platf
 EDropEffect::Type FSlateApplication::OnDragEnterText( const TSharedRef< FGenericWindow >& Window, const FString& Text )
 {
 	const TSharedPtr< FExternalDragOperation > DragDropOperation = FExternalDragOperation::NewText( Text );
-	const TSharedPtr< SWindow > EffectingWindow = FindWindowByPlatformWindow( SlateWindows, Window );
+	const TSharedPtr< SWindow > EffectingWindow = FSlateWindowHelper::FindWindowByPlatformWindow( SlateWindows, Window );
 
 	EDropEffect::Type Result = EDropEffect::None;
 	if ( DragDropOperation.IsValid() && EffectingWindow.IsValid() )
@@ -4524,7 +4345,7 @@ EDropEffect::Type FSlateApplication::OnDragEnterText( const TSharedRef< FGeneric
 EDropEffect::Type FSlateApplication::OnDragEnterFiles( const TSharedRef< FGenericWindow >& Window, const TArray< FString >& Files )
 {
 	const TSharedPtr< FExternalDragOperation > DragDropOperation = FExternalDragOperation::NewFiles( Files );
-	const TSharedPtr< SWindow > EffectingWindow = FindWindowByPlatformWindow( SlateWindows, Window );
+	const TSharedPtr< SWindow > EffectingWindow = FSlateWindowHelper::FindWindowByPlatformWindow( SlateWindows, Window );
 
 	EDropEffect::Type Result = EDropEffect::None;
 	if ( DragDropOperation.IsValid() && EffectingWindow.IsValid() )
diff --git a/Engine/Source/Runtime/Slate/Public/Framework/Application/SlateApplication.h b/Engine/Source/Runtime/Slate/Public/Framework/Application/SlateApplication.h
--- a/Engine/Source/Runtime/Slate/Public/Framework/Application/SlateApplication.h
+++ b/Engine/Source/Runtime/Slate/Public/Framework/Application/SlateApplication.h
@@ -216,11 +216,6 @@ public:
 	 */
 	void SetCursorPos( const FVector2D& MouseCoordinate );
 
-	/**
-	 * Reorders an array of windows so the specified window is "brought to the front"
-	 */
-	static void ArrangeWindowToFront( TArray< TSharedRef<SWindow> >& Windows, const TSharedRef<SWindow>& WindowToBringToFront );
-
 	/**
 	 * Polls game devices for input
 	 */
@@ -423,19 +418,6 @@ public:
 	 */
 	void SetExitRequestedHandler( const FSimpleDelegate& OnExitRequestedHandler );
 		
-	/**
-	 * @todo slate: Remove this method or make it private.
-	 * Searches for the specified widget and generates a full path to it.  Note that this is
-	 * a relatively slow operation!
-	 * 
-	 * @param  InWidget       Widget to generate a path to
-	 * @param  OutWidgetPath  The generated widget path
-	 * @param  VisibilityFilter	Widgets must have this type of visibility to be included the path
-	 *
-	 * @return	True if the widget path was found
-	 */
-	static bool FindPathToWidget( const TArray< TSharedRef<SWindow> > WindowsToSearch, TSharedRef< const SWidget > InWidget, FWidgetPath& OutWidgetPath, EVisibility VisibilityFilter = EVisibility::Visible );
-
 	/**
 	 * @todo slate: Remove this method or make it private.
 	 * Searches for the specified widget and generates a full path to it.  Note that this is
@@ -961,12 +943,12 @@ public:
 
 	virtual void ArrangeWindowToFrontVirtual( TArray<TSharedRef<SWindow>>& Windows, const TSharedRef<SWindow>& WindowToBringToFront ) OVERRIDE
 	{
-		FSlateApplication::ArrangeWindowToFront(Windows, WindowToBringToFront);
+		FSlateWindowHelper::ArrangeWindowToFront(Windows, WindowToBringToFront);
 	}
 
-	virtual bool FindPathToWidgetVirtual( const TArray<TSharedRef<SWindow>> WindowsToSearch, TSharedRef<const SWidget> InWidget, FWidgetPath& OutWidgetPath, EVisibility VisibilityFilter = EVisibility::Visible ) OVERRIDE
+	virtual bool FindPathToWidget( TSharedRef<const SWidget> InWidget, FWidgetPath& OutWidgetPath, EVisibility VisibilityFilter = EVisibility::Visible ) OVERRIDE
 	{
-		return FSlateApplication::FindPathToWidget(WindowsToSearch, InWidget, OutWidgetPath, VisibilityFilter);
+		return FSlateWindowHelper::FindPathToWidget(GetInteractiveTopLevelWindows(), InWidget, OutWidgetPath, VisibilityFilter);
 	}
 
 	virtual const double GetCurrentTime( ) const OVERRIDE
diff --git a/Engine/Source/Runtime/SlateCore/Private/Application/SlateWindowHelper.cpp b/Engine/Source/Runtime/SlateCore/Private/Application/SlateWindowHelper.cpp
new file mode 100644
--- /dev/null
+++ b/Engine/Source/Runtime/SlateCore/Private/Application/SlateWindowHelper.cpp
@@ -0,0 +1,183 @@
+// Copyright 1998-2014 Epic Games, Inc. All Rights Reserved.
+
+/*=============================================================================
+	SlateWindowHelper.cpp: Implements the FSlateWindowHelper class.
+=============================================================================*/
+
+#include "SlateCorePrivatePCH.h"
+
+
+DECLARE_CYCLE_STAT( TEXT("FindPathToWidget"), STAT_FindPathToWidget, STATGROUP_Slate );
+
+
+/* FSlateWindowHelper static functions
+ *****************************************************************************/
+
+void FSlateWindowHelper::ArrangeWindowToFront( TArray< TSharedRef<SWindow> >& Windows, const TSharedRef<SWindow>& WindowToBringToFront )
+{
+	Windows.Remove(WindowToBringToFront);
+
+	if ((Windows.Num() == 0) || WindowToBringToFront->IsTopmostWindow())
+	{
+		Windows.Add(WindowToBringToFront);
+	}
+	else
+	{
+		bool PerformedInsert = false;
+
+		for (int WindowIndex = Windows.Num() - 1; WindowIndex >= 0; --WindowIndex)
+		{
+			if (!Windows[WindowIndex]->IsTopmostWindow())
+			{
+				Windows.Insert(WindowToBringToFront, WindowIndex + 1);
+				PerformedInsert = true;
+
+				break;
+			}
+		}
+
+		if (!PerformedInsert)
+		{
+			Windows.Insert(WindowToBringToFront, 0);
+		}
+	}
+}
+
+
+void FSlateWindowHelper::BringWindowToFront( TArray<TSharedRef<SWindow>>& Windows, const TSharedRef<SWindow>& BringMeToFront )
+{
+	const TSharedRef<SWindow> TopLeveWindowToReorder = BringToFrontInParent(BringMeToFront);
+	FSlateWindowHelper::ArrangeWindowToFront(Windows, TopLeveWindowToReorder);
+}
+
+
+bool FSlateWindowHelper::CheckWorkAreaForWindows( const TArray<TSharedRef<SWindow>>& WindowsToSearch, const FSlateRect& WorkAreaRect )
+{
+	for (TArray<TSharedRef<SWindow>>::TConstIterator CurrentWindowIt(WindowsToSearch); CurrentWindowIt; ++CurrentWindowIt)
+	{
+		const TSharedRef<SWindow>& CurrentWindow = *CurrentWindowIt;
+		const FVector2D Position = CurrentWindow->GetPositionInScreen();
+		const FVector2D Size = CurrentWindow->GetSizeInScreen();
+		const FSlateRect WindowRect(Position.X, Position.Y, Size.X, Size.Y);
+
+		if (FSlateRect::DoRectanglesIntersect(WorkAreaRect, WindowRect) || CheckWorkAreaForWindows(CurrentWindow->GetChildWindows(), WorkAreaRect))
+		{
+			return true;
+		}
+	}
+
+	return false;
+}
+
+
+bool FSlateWindowHelper::ContainsWindow( const TArray<TSharedRef<SWindow>>& WindowsToSearch, const TSharedRef<SWindow>& WindowToFind )
+{
+	if (WindowsToSearch.Contains(WindowToFind))
+	{
+		return true;
+	}
+
+	for (int32 WindowIndex=0; WindowIndex < WindowsToSearch.Num(); ++WindowIndex)
+	{
+		if (ContainsWindow(WindowsToSearch[WindowIndex]->GetChildWindows(), WindowToFind))
+		{
+			return true;
+		}
+	}
+
+	return false;
+}
+
+
+bool FSlateWindowHelper::FindPathToWidget( const TArray<TSharedRef<SWindow>> WindowsToSearch,  TSharedRef<const SWidget> InWidget, FWidgetPath& OutWidgetPath, EVisibility VisibilityFilter )
+{
+	SCOPE_CYCLE_COUNTER(STAT_FindPathToWidget);
+
+	// Iterate over our top level windows
+	bool bFoundWidget = false;
+
+	for (int32 WindowIndex = 0; !bFoundWidget && WindowIndex < WindowsToSearch.Num(); ++WindowIndex)
+	{
+		// Make a widget path that contains just the top-level window
+		TSharedRef<SWindow> CurWindow = WindowsToSearch[WindowIndex];
+		
+		FArrangedChildren JustWindow(VisibilityFilter);
+		{
+			JustWindow.AddWidget(FArrangedWidget(CurWindow, CurWindow->GetWindowGeometryInScreen()));
+		}
+		
+		FWidgetPath PathToWidget(CurWindow, JustWindow);
+		
+		// Attempt to extend it to the desired child widget; essentially a full-window search for 'InWidget
+		if ((CurWindow == InWidget) || PathToWidget.ExtendPathTo(FWidgetMatcher(InWidget), VisibilityFilter))
+		{
+			OutWidgetPath = PathToWidget;
+			bFoundWidget = true;
+		}
+
+		if (!bFoundWidget)
+		{
+			// Search this window's children
+			bFoundWidget = FindPathToWidget(CurWindow->GetChildWindows(), InWidget, OutWidgetPath, VisibilityFilter);
+		}
+	}
+
+	return bFoundWidget;
+}
+
+
+TSharedPtr<SWindow> FSlateWindowHelper::FindWindowByPlatformWindow( const TArray< TSharedRef<SWindow> >& WindowsToSearch, const TSharedRef< FGenericWindow >& PlatformWindow )
+{
+	for (int32 WindowIndex = 0; WindowIndex < WindowsToSearch.Num(); ++WindowIndex)
+	{
+		TSharedRef<SWindow> SomeWindow = WindowsToSearch[WindowIndex];
+		TSharedRef<FGenericWindow> SomeNativeWindow = StaticCastSharedRef<FGenericWindow>(SomeWindow->GetNativeWindow().ToSharedRef());
+		
+		if (SomeNativeWindow == PlatformWindow)
+		{
+			return SomeWindow;
+		}
+
+		// Search child windows
+		TSharedPtr<SWindow> FoundChildWindow = FindWindowByPlatformWindow(SomeWindow->GetChildWindows(), PlatformWindow);
+
+		if (FoundChildWindow.IsValid())
+		{
+			return FoundChildWindow;
+		}
+	}
+
+	return TSharedPtr<SWindow>(nullptr);
+}
+
+
+void FSlateWindowHelper::RemoveWindowFromList( TArray<TSharedRef<SWindow>>& Windows, const TSharedRef<SWindow>& WindowToRemove )
+{
+	int32 NumRemoved = Windows.Remove(WindowToRemove);
+
+	if (NumRemoved == 0)
+	{
+		for (int32 ChildIndex=0; ChildIndex < Windows.Num(); ++ChildIndex)
+		{
+			RemoveWindowFromList(Windows[ChildIndex]->GetChildWindows(), WindowToRemove) ;
+		}	
+	}
+}
+
+
+/* FSlateWindowHelper implementation
+ *****************************************************************************/
+
+TSharedRef<SWindow> FSlateWindowHelper::BringToFrontInParent( const TSharedRef<SWindow>& WindowToBringToFront )
+{
+	const TSharedPtr<SWindow> ParentWindow = WindowToBringToFront->GetParentWindow();
+
+	if (!ParentWindow.IsValid())
+	{
+		return WindowToBringToFront;
+	}
+
+	FSlateWindowHelper::ArrangeWindowToFront(ParentWindow->GetChildWindows(), WindowToBringToFront);
+
+	return BringToFrontInParent(ParentWindow.ToSharedRef());
+}
diff --git a/Engine/Source/Runtime/SlateCore/Private/Widgets/SWindow.cpp b/Engine/Source/Runtime/SlateCore/Private/Widgets/SWindow.cpp
--- a/Engine/Source/Runtime/SlateCore/Private/Widgets/SWindow.cpp
+++ b/Engine/Source/Runtime/SlateCore/Private/Widgets/SWindow.cpp
@@ -1252,7 +1252,7 @@ FReply SWindow::OnKeyboardFocusReceived( const FGeometry& MyGeometry, const FKey
 		}
 
 		FWidgetPath WidgetToFocusPath;
-		if( FSlateApplicationBase::Get().FindPathToWidgetVirtual( JustThisWindow, PinnedWidgetToFocusOnActivate.ToSharedRef(), WidgetToFocusPath ) )
+		if( FSlateWindowHelper::FindPathToWidget( JustThisWindow, PinnedWidgetToFocusOnActivate.ToSharedRef(), WidgetToFocusPath ) )
 		{
 			FSlateApplicationBase::Get().SetKeyboardFocus( WidgetToFocusPath, EKeyboardFocusCause::SetDirectly );
 		}
diff --git a/Engine/Source/Runtime/SlateCore/Public/Application/SlateApplicationBase.h b/Engine/Source/Runtime/SlateCore/Public/Application/SlateApplicationBase.h
--- a/Engine/Source/Runtime/SlateCore/Public/Application/SlateApplicationBase.h
+++ b/Engine/Source/Runtime/SlateCore/Public/Application/SlateApplicationBase.h
@@ -58,9 +58,9 @@ public:
 	virtual void ArrangeWindowToFrontVirtual( TArray<TSharedRef<SWindow>>& Windows, const TSharedRef<SWindow>& WindowToBringToFront ) = 0;
 
 	/**
-	 * @todo slate: Remove this method or make it private.
-	 * Searches for the specified widget and generates a full path to it.  Note that this is
-	 * a relatively slow operation!
+	 * Searches for the specified widget and generates a full path to it.
+	 *
+	 * Note: this is a relatively slow operation!
 	 * 
 	 * @param  InWidget       Widget to generate a path to
 	 * @param  OutWidgetPath  The generated widget path
@@ -68,7 +68,7 @@ public:
 	 *
 	 * @return	True if the widget path was found
 	 */
-	virtual bool FindPathToWidgetVirtual( const TArray<TSharedRef<SWindow>> WindowsToSearch, TSharedRef<const SWidget> InWidget, FWidgetPath& OutWidgetPath, EVisibility VisibilityFilter = EVisibility::Visible ) = 0;
+	virtual bool FindPathToWidget( TSharedRef<const SWidget> InWidget, FWidgetPath& OutWidgetPath, EVisibility VisibilityFilter = EVisibility::Visible ) = 0;
 
 	/**
 	 * Gets the active top-level window.
diff --git a/Engine/Source/Runtime/SlateCore/Public/Application/SlateWindowHelper.h b/Engine/Source/Runtime/SlateCore/Public/Application/SlateWindowHelper.h
new file mode 100644
--- /dev/null
+++ b/Engine/Source/Runtime/SlateCore/Public/Application/SlateWindowHelper.h
@@ -0,0 +1,95 @@
+// Copyright 1998-2014 Epic Games, Inc. All Rights Reserved.
+
+/*=============================================================================
+	SlateWindowHelper.h: Declares the FSlateWindowHelper class.
+=============================================================================*/
+
+#pragma once
+
+
+/**
+ * Implements a manager for Slate windows.
+ */
+class SLATECORE_API FSlateWindowHelper
+{
+public:
+
+	/**
+	 * Reorders the given collection of windows so the specified window is brought to the front.
+	 *
+	 * @param Windows The collection of windows to reorder.
+	 * @param WindowToBringToFront The window to bring to the front.
+	 */
+	static void ArrangeWindowToFront( TArray<TSharedRef<SWindow>>& Windows, const TSharedRef<SWindow>& WindowToBringToFront );
+
+	/**
+	 * Put 'BringMeToFront' at the font of the list of 'WindowsToReorder'.
+	 *
+	 * The top-most (front-most) window is last in Z-order and therefore is added to the end of the list.
+	 *
+	 * @param Windows An ordered list of windows.
+	 * @param BingMeToFront The window to bring to front.
+	 */
+	static void BringWindowToFront( TArray<TSharedRef<SWindow>>& Windows, const TSharedRef<SWindow>& BringMeToFront );
+
+	/**
+	 * Checks whether any of the given windows overlap the specified work area.
+	 *
+	 * @param WindowsToSearch The collection of windows to check.
+	 * @param WorkAreaRect The work area to check.
+	 *
+	 * @return true if at least one window overlaps with the work area, false otherwise.
+	 */
+	static bool CheckWorkAreaForWindows( const TArray< TSharedRef<SWindow> >& WindowsToSearch, const FSlateRect& WorkAreaRect );
+
+	/**
+	 * Checks whether the given collection of windows contains the specified window.
+	 *
+	 * @param WindowsToSearch The collection of windows to search.
+	 * @param WindowToFind The window to find.
+	 *
+	 * @return true if the collection contains the window, false otherwise.
+	 */
+	static bool ContainsWindow( const TArray<TSharedRef<SWindow>>& WindowsToSearch, const TSharedRef<SWindow>& WindowToFind );
+
+	/**
+	 * Searches for the specified widget and generates a full path to it.
+	 *
+	 * Note: this is a relatively slow operation!
+	 * 
+	 * @param InWidget Widget to generate a path to.
+	 * @param OutWidgetPath The generated widget path.
+	 * @param VisibilityFilter Widgets must have this type of visibility to be included the path.
+	 *
+	 * @return True if the widget path was found.
+	 */
+	static bool FindPathToWidget( const TArray<TSharedRef<SWindow>> WindowsToSearch, TSharedRef<const SWidget> InWidget, FWidgetPath& OutWidgetPath, EVisibility VisibilityFilter = EVisibility::Visible );
+
+	/**
+	 * Searches the given collection of windows to find the Slate window that corresponds to the specified platform window.
+	 *
+	 * @param WindowsToSearch The collection of windows to search.
+	 * @param PlatformWindow The platform window to find the Slate window for.
+	 *
+	 * @return The corresponding Slate window, or nullptr if it was not found.
+	 */
+	static TSharedPtr<SWindow> FindWindowByPlatformWindow( const TArray<TSharedRef<SWindow>>& WindowsToSearch, const TSharedRef<FGenericWindow>& PlatformWindow );
+
+	/**
+	 * Removes the specified window from the given collection of windows.
+	 *
+	 * @param Windows The collection of windows.
+	 * @param WindowToRemove The window to remove.
+	 */
+	static void RemoveWindowFromList( TArray<TSharedRef<SWindow>>& Windows, const TSharedRef<SWindow>& WindowToRemove );
+
+protected:
+
+	/**
+	 * Make BringMeToFront first among its peers.
+	 * i.e. make it the last window in its parent's list of child windows.
+	 *
+	 * @return The top-most window whose children were re-arranged
+	 */
+	static TSharedRef<SWindow> BringToFrontInParent( const TSharedRef<SWindow>& WindowToBringToFront );
+};
diff --git a/Engine/Source/Runtime/SlateCore/Public/SlateCore.h b/Engine/Source/Runtime/SlateCore/Public/SlateCore.h
--- a/Engine/Source/Runtime/SlateCore/Public/SlateCore.h
+++ b/Engine/Source/Runtime/SlateCore/Public/SlateCore.h
@@ -136,6 +136,7 @@ DECLARE_MEMORY_STAT_EXTERN(TEXT("Index Buffer Memory"), STAT_SlateIndexBufferMem
 
 // Application
 #include "IToolTip.h"
+#include "SlateWindowHelper.h"
 #include "SlateApplicationBase.h"
 #include "ThrottleManager.h"
 
