Slate: consistent names for event processing methods; removed dead code

diff --git a/Engine/Plugins/Slate/SlateRemote/Source/SlateRemote/Private/Server/SlateRemoteServer.cpp b/Engine/Plugins/Slate/SlateRemote/Source/SlateRemote/Private/Server/SlateRemoteServer.cpp
--- a/Engine/Plugins/Slate/SlateRemote/Source/SlateRemote/Private/Server/SlateRemoteServer.cpp
+++ b/Engine/Plugins/Slate/SlateRemote/Source/SlateRemote/Private/Server/SlateRemoteServer.cpp
@@ -82,7 +82,7 @@ void FSlateRemoteServer::ProcessMotionMessage( const FSlateRemoteServerMessage&
 	FVector Accel = FVector(Message.Data[9], Message.Data[10], Message.Data[11]);
 
 	FMotionEvent Event(0, Attitude, RotationRate, Gravity, Accel);
-	FSlateApplication::Get().OnMotionDetectedMessage(Event);
+	FSlateApplication::Get().ProcessMotionDetectedEvent(Event);
 }
 
 
@@ -116,7 +116,7 @@ void FSlateRemoteServer::ProcessTiltMessage( const FSlateRemoteServerMessage& Me
 	LastRoll = Roll;
 
 	FMotionEvent Event(0, Attitude, RotationRate, Gravity, Accel);
-	FSlateApplication::Get().OnMotionDetectedMessage(Event);
+	FSlateApplication::Get().ProcessMotionDetectedEvent(Event);
 }
 
 
@@ -173,15 +173,15 @@ void FSlateRemoteServer::ProcessTouchMessage( const FSlateRemoteServerMessage& M
 			// send input to handler
 			if (Message.DataType == DT_TouchBegan)
 			{
-				SlateApplication.OnTouchStartedMessage(nullptr, Event);
+				SlateApplication.ProcessTouchStartedEvent(nullptr, Event);
 			}
 			else if (Message.DataType == DT_TouchEnded)
 			{
-				SlateApplication.OnTouchEndedMessage(Event);
+				SlateApplication.ProcessTouchEndedEvent(Event);
 			}
 			else
 			{
-				SlateApplication.OnTouchMovedMessage(Event);
+				SlateApplication.ProcessTouchMovedEvent(Event);
 			}
 		}
 	}
diff --git a/Engine/Source/Runtime/Slate/Private/Framework/Application/SlateApplication.cpp b/Engine/Source/Runtime/Slate/Private/Framework/Application/SlateApplication.cpp
--- a/Engine/Source/Runtime/Slate/Private/Framework/Application/SlateApplication.cpp
+++ b/Engine/Source/Runtime/Slate/Private/Framework/Application/SlateApplication.cpp
@@ -1586,36 +1586,6 @@ void FSlateApplication::ClearKeyboardFocus( const EKeyboardFocusCause::Type InCa
 }
 
 
-void FSlateApplication::OnSetCursorMessage()
-{
-	QueryCursor();
-}
-
-void FSlateApplication::OnDragLeaveMessage()
-{
-	DragDropContent.Reset();
-}
-
-bool FSlateApplication::OnDropMessage( TSharedRef<SWindow> WindowDroppedInto, FDragDropEvent& DragDropEvent )
-{
-	FWidgetPath WidgetsUnderCursor = LocateWindowUnderMouse( DragDropEvent.GetScreenSpacePosition(), GetInteractiveTopLevelWindows() );
-	DragDropEvent.SetEventPath( WidgetsUnderCursor );
-
-	// Switch worlds for widgets in the current path
-	FScopedSwitchWorldHack SwitchWorld( WidgetsUnderCursor );
-
-	FReply Reply = FReply::Unhandled();
-	for ( int32 WidgetIndex = WidgetsUnderCursor.Widgets.Num()-1; !Reply.IsEventHandled() && WidgetIndex >=0; --WidgetIndex )
-	{
-		const FArrangedWidget& CurWidgetGeometry = WidgetsUnderCursor.Widgets(WidgetIndex);
-		Reply = CurWidgetGeometry.Widget->OnDrop( CurWidgetGeometry.Geometry, DragDropEvent ).SetHandler(CurWidgetGeometry.Widget);
-	}
-
-	LOG_EVENT( EEventLog::DropMessage, Reply )
-
-	return Reply.IsEventHandled();
-}
-
 void FSlateApplication::OnShutdown()
 {
 	// Clean up our tooltip window
@@ -2463,7 +2433,7 @@ void FSlateApplication::SynthesizeMouseMove()
 		PlatformApplication->GetModifierKeys()
 	);
 
-	ProcessMouseMoveMessage(MouseEvent);
+	ProcessMouseMoveEvent(MouseEvent);
 }
 
 void FSlateApplication::OnLogSlateEvent(EEventLog::Type Event, const FString& AdditionalContent)
@@ -2887,10 +2857,10 @@ bool FSlateApplication::ShouldProcessUserInputMessages( const TSharedPtr< FGener
 bool FSlateApplication::OnKeyChar( const TCHAR Character, const bool IsRepeat )
 {
 	FCharacterEvent CharacterEvent( Character, PlatformApplication->GetModifierKeys(), IsRepeat );
-	return ProcessKeyCharMessage( CharacterEvent );
+	return ProcessKeyCharEvent( CharacterEvent );
 }
 
-bool FSlateApplication::ProcessKeyCharMessage( FCharacterEvent& InCharacterEvent )
+bool FSlateApplication::ProcessKeyCharEvent( FCharacterEvent& InCharacterEvent )
 {
 	FReply Reply = FReply::Unhandled();
 
@@ -2933,10 +2903,10 @@ bool FSlateApplication::OnKeyDown( const int32 KeyCode, const uint32 CharacterCo
 	FKey const Key = FInputKeyManager::Get().GetKeyFromCodes( KeyCode, CharacterCode );
 	FKeyboardEvent KeyboardEvent( Key, PlatformApplication->GetModifierKeys(), IsRepeat, CharacterCode );
 
-	return ProcessKeyDownMessage( KeyboardEvent );
+	return ProcessKeyDownEvent( KeyboardEvent );
 }
 
-bool FSlateApplication::ProcessKeyDownMessage( FKeyboardEvent& InKeyboardEvent )
+bool FSlateApplication::ProcessKeyDownEvent( FKeyboardEvent& InKeyboardEvent )
 {
 	FReply Reply = FReply::Unhandled();
 
@@ -3035,10 +3005,10 @@ bool FSlateApplication::OnKeyUp( const int32 KeyCode, const uint32 CharacterCode
 	FKey const Key = FInputKeyManager::Get().GetKeyFromCodes( KeyCode, CharacterCode );
 	FKeyboardEvent KeyboardEvent( Key, PlatformApplication->GetModifierKeys(), IsRepeat, CharacterCode );
 
-	return ProcessKeyUpMessage( KeyboardEvent );
+	return ProcessKeyUpEvent( KeyboardEvent );
 }
 
-bool FSlateApplication::ProcessKeyUpMessage( FKeyboardEvent& InKeyboardEvent )
+bool FSlateApplication::ProcessKeyUpEvent( FKeyboardEvent& InKeyboardEvent )
 {
 	FReply Reply = FReply::Unhandled();
 
@@ -3137,10 +3107,10 @@ bool FSlateApplication::OnMouseDown( const TSharedPtr< FGenericWindow >& Platfor
 		PlatformApplication->GetModifierKeys()
 		);
 
-	return ProcessMouseButtonDownMessage( PlatformWindow, MouseEvent );
+	return ProcessMouseButtonDownEvent( PlatformWindow, MouseEvent );
 }
 
-bool FSlateApplication::ProcessMouseButtonDownMessage( const TSharedPtr< FGenericWindow >& PlatformWindow, FPointerEvent& MouseEvent )
+bool FSlateApplication::ProcessMouseButtonDownEvent( const TSharedPtr< FGenericWindow >& PlatformWindow, FPointerEvent& MouseEvent )
 {
 	LastUserInteractionTimeForThrottling = this->GetCurrentTime();
 
@@ -3288,10 +3258,10 @@ bool FSlateApplication::OnMouseDoubleClick( const TSharedPtr< FGenericWindow >&
 		PlatformApplication->GetModifierKeys()
 		);
 
-	return ProcessMouseButtonDoubleClickMessage( PlatformWindow, MouseEvent );
+	return ProcessMouseButtonDoubleClickEvent( PlatformWindow, MouseEvent );
 }
 
-bool FSlateApplication::ProcessMouseButtonDoubleClickMessage( const TSharedPtr< FGenericWindow >& PlatformWindow, FPointerEvent& InMouseEvent )
+bool FSlateApplication::ProcessMouseButtonDoubleClickEvent( const TSharedPtr< FGenericWindow >& PlatformWindow, FPointerEvent& InMouseEvent )
 {
 	LastUserInteractionTimeForThrottling = this->GetCurrentTime();
 
@@ -3343,10 +3313,10 @@ bool FSlateApplication::OnMouseUp( const EMouseButtons::Type Button )
 		PlatformApplication->GetModifierKeys()
 		);
 
-	return ProcessMouseButtonUpMessage( MouseEvent );
+	return ProcessMouseButtonUpEvent( MouseEvent );
 }
 
-bool FSlateApplication::ProcessMouseButtonUpMessage( FPointerEvent& MouseEvent )
+bool FSlateApplication::ProcessMouseButtonUpEvent( FPointerEvent& MouseEvent )
 {
 	LastUserInteractionTimeForThrottling = this->GetCurrentTime();
 	PressedMouseButtons.Remove( MouseEvent.GetEffectingButton() );
@@ -3452,10 +3422,10 @@ bool FSlateApplication::OnMouseWheel( const float Delta )
 		PlatformApplication->GetModifierKeys()
 		);
 
-	return ProcessMouseWheelOrGestureMessage( MouseWheelEvent, NULL );
+	return ProcessMouseWheelOrGestureEvent( MouseWheelEvent, NULL );
 }
 
-bool FSlateApplication::ProcessMouseWheelOrGestureMessage( FPointerEvent& InWheelEvent, const FPointerEvent* InGestureEvent )
+bool FSlateApplication::ProcessMouseWheelOrGestureEvent( FPointerEvent& InWheelEvent, const FPointerEvent* InGestureEvent )
 {
 	const bool bShouldProcessEvent = InWheelEvent.GetWheelDelta() != 0 ||
 		(InGestureEvent!=NULL && InGestureEvent->GetGestureDelta()!=FVector2D::ZeroVector);
@@ -3530,7 +3500,7 @@ bool FSlateApplication::OnMouseMove()
 			PlatformApplication->GetModifierKeys()
 			);
 
-		Result = ProcessMouseMoveMessage( MouseEvent );
+		Result = ProcessMouseMoveEvent( MouseEvent );
 	}
 
 	return Result;
@@ -3548,13 +3518,13 @@ bool FSlateApplication::OnRawMouseMove( const int32 X, const int32 Y )
 			PlatformApplication->GetModifierKeys()
 		);
 
-		ProcessMouseMoveMessage(MouseEvent);
+		ProcessMouseMoveEvent(MouseEvent);
 	}
 	
 	return true;
 }
 
-bool FSlateApplication::ProcessMouseMoveMessage( FPointerEvent& MouseEvent )
+bool FSlateApplication::ProcessMouseMoveEvent( FPointerEvent& MouseEvent )
 {
 	if ( !MouseEvent.GetCursorDelta().IsZero() )
 	{
@@ -3862,7 +3832,7 @@ bool FSlateApplication::OnControllerAnalog( FKey Button, int32 ControllerId, flo
 			false
 		);
 
-		OnControllerAnalogValueChangedMessage( ControllerEvent );
+		ProcessControllerAnalogValueChangedEvent( ControllerEvent );
 	}
 
 	return true;
@@ -3873,7 +3843,7 @@ bool FSlateApplication::OnControllerAnalog( EControllerButtons::Type Button, int
 	return OnControllerAnalog( TranslateControllerButtonToKey( Button ), ControllerId, AnalogValue );
 }
 
-void FSlateApplication::OnControllerAnalogValueChangedMessage( FControllerEvent& ControllerEvent )
+void FSlateApplication::ProcessControllerAnalogValueChangedEvent( FControllerEvent& ControllerEvent )
 {
 	CALL_WIDGET_FUNCTION(ControllerEvent, OnControllerAnalogValueChanged);
 }
@@ -3889,7 +3859,7 @@ bool FSlateApplication::OnControllerButtonPressed( FKey Button, int32 Controller
 			IsRepeat
 		);
 
-		OnControllerButtonPressedMessage( ControllerEvent );
+		ProcessControllerButtonPressedEvent( ControllerEvent );
 	}
 
 	return true;
@@ -3900,7 +3870,7 @@ bool FSlateApplication::OnControllerButtonPressed( EControllerButtons::Type Butt
 	return OnControllerButtonPressed( TranslateControllerButtonToKey( Button ), ControllerId, IsRepeat );
 }
 
-void FSlateApplication::OnControllerButtonPressedMessage( FControllerEvent& ControllerEvent )
+void FSlateApplication::ProcessControllerButtonPressedEvent( FControllerEvent& ControllerEvent )
 {
 	CALL_WIDGET_FUNCTION(ControllerEvent, OnControllerButtonPressed);
 }
@@ -3916,7 +3886,7 @@ bool FSlateApplication::OnControllerButtonReleased( FKey Button, int32 Controlle
 			IsRepeat
 		);
 
-		OnControllerButtonReleasedMessage( ControllerEvent );
+		ProcessControllerButtonReleasedEvent( ControllerEvent );
 	}
 
 	return false;
@@ -3927,7 +3897,7 @@ bool FSlateApplication::OnControllerButtonReleased( EControllerButtons::Type But
 	return OnControllerButtonReleased( TranslateControllerButtonToKey( Button ), ControllerId, IsRepeat );
 }
 
-void FSlateApplication::OnControllerButtonReleasedMessage( FControllerEvent& ControllerEvent )
+void FSlateApplication::ProcessControllerButtonReleasedEvent( FControllerEvent& ControllerEvent )
 {
 	CALL_WIDGET_FUNCTION(ControllerEvent, OnControllerButtonReleased);
 
@@ -3955,7 +3925,7 @@ bool FSlateApplication::OnTouchGesture( EGestureEvent::Type GestureType, const F
 		PlatformApplication->GetModifierKeys()
 	);
 	
-	return ProcessMouseWheelOrGestureMessage( MouseWheelEvent, &GestureEvent );
+	return ProcessMouseWheelOrGestureEvent( MouseWheelEvent, &GestureEvent );
 }
 
 bool FSlateApplication::OnTouchStarted( const TSharedPtr< FGenericWindow >& PlatformWindow, const FVector2D& Location, int32 TouchIndex, int32 ControllerId )
@@ -3966,14 +3936,14 @@ bool FSlateApplication::OnTouchStarted( const TSharedPtr< FGenericWindow >& Plat
 		Location,
 		Location,
 		true);
-	OnTouchStartedMessage( PlatformWindow, PointerEvent );
+	ProcessTouchStartedEvent( PlatformWindow, PointerEvent );
 
 	return true;
 }
 
-void FSlateApplication::OnTouchStartedMessage( const TSharedPtr< FGenericWindow >& PlatformWindow, FPointerEvent& PointerEvent )
+void FSlateApplication::ProcessTouchStartedEvent( const TSharedPtr< FGenericWindow >& PlatformWindow, FPointerEvent& PointerEvent )
 {
-	ProcessMouseButtonDownMessage( PlatformWindow, PointerEvent );
+	ProcessMouseButtonDownEvent( PlatformWindow, PointerEvent );
 }
 
 bool FSlateApplication::OnTouchMoved( const FVector2D& Location, int32 TouchIndex, int32 ControllerId )
@@ -3984,14 +3954,14 @@ bool FSlateApplication::OnTouchMoved( const FVector2D& Location, int32 TouchInde
 		Location,
 		Location,
 		true);
-	OnTouchMovedMessage(PointerEvent);
+	ProcessTouchMovedEvent(PointerEvent);
 
 	return true;
 }
 
-void FSlateApplication::OnTouchMovedMessage( FPointerEvent& PointerEvent )
+void FSlateApplication::ProcessTouchMovedEvent( FPointerEvent& PointerEvent )
 {
-	ProcessMouseMoveMessage(PointerEvent);
+	ProcessMouseMoveEvent(PointerEvent);
 }
 
 bool FSlateApplication::OnTouchEnded( const FVector2D& Location, int32 TouchIndex, int32 ControllerId )
@@ -4002,14 +3972,14 @@ bool FSlateApplication::OnTouchEnded( const FVector2D& Location, int32 TouchInde
 		Location,
 		Location,
 		true);
-	OnTouchEndedMessage(PointerEvent);
+	ProcessTouchEndedEvent(PointerEvent);
 
 	return true;
 }
 
-void FSlateApplication::OnTouchEndedMessage( FPointerEvent& PointerEvent )
+void FSlateApplication::ProcessTouchEndedEvent( FPointerEvent& PointerEvent )
 {
-	ProcessMouseButtonUpMessage(PointerEvent);
+	ProcessMouseButtonUpEvent(PointerEvent);
 }
 
 bool FSlateApplication::OnMotionDetected(const FVector& Tilt, const FVector& RotationRate, const FVector& Gravity, const FVector& Acceleration, int32 ControllerId)
@@ -4021,12 +3991,12 @@ bool FSlateApplication::OnMotionDetected(const FVector& Tilt, const FVector& Rot
 		Gravity,
 		Acceleration
 		);
-	OnMotionDetectedMessage(MotionEvent);
+	ProcessMotionDetectedEvent(MotionEvent);
 
 	return true;
 }
 
-void FSlateApplication::OnMotionDetectedMessage( FMotionEvent& MotionEvent )
+void FSlateApplication::ProcessMotionDetectedEvent( FMotionEvent& MotionEvent )
 {
 	CALL_WIDGET_FUNCTION(MotionEvent, OnMotionDetected);
 }
@@ -4140,10 +4110,10 @@ bool FSlateApplication::OnWindowActivationChanged( const TSharedRef< FGenericWin
 	FWindowActivateEvent::EActivationType TranslatedActivationType = TranslationWindowActivationMessage( ActivationType );
 	FWindowActivateEvent WindowActivateEvent( TranslatedActivationType, Window.ToSharedRef() );
 
-	return ProcessWindowActivatedMessage( WindowActivateEvent );
+	return ProcessWindowActivatedEvent( WindowActivateEvent );
 }
 
-bool FSlateApplication::ProcessWindowActivatedMessage( const FWindowActivateEvent& ActivateEvent )
+bool FSlateApplication::ProcessWindowActivatedEvent( const FWindowActivateEvent& ActivateEvent )
 {
 	TSharedPtr<SWindow> ActiveModalWindow = GetActiveModalWindow();
 
@@ -4233,11 +4203,11 @@ bool FSlateApplication::ProcessWindowActivatedMessage( const FWindowActivateEven
 
 bool FSlateApplication::OnApplicationActivationChanged( const bool IsActive )
 {
-	ProcessApplicationActivationMessage( IsActive );
+	ProcessApplicationActivationEvent( IsActive );
 	return true;
 }
 
-void FSlateApplication::ProcessApplicationActivationMessage( bool InAppActivated )
+void FSlateApplication::ProcessApplicationActivationEvent( bool InAppActivated )
 {
 	const bool UserSwitchedAway = bAppIsActive && !InAppActivated;
 
@@ -4379,11 +4349,11 @@ EDropEffect::Type FSlateApplication::OnDragEnter( const TSharedRef< SWindow >& W
 
 	LastCursorPosition = CurrentCursorPosition;
 
-	ProcessDragEnterMessage( Window, DragDropEvent );
+	ProcessDragEnterEvent( Window, DragDropEvent );
 	return EDropEffect::None;
 }
 
-bool FSlateApplication::ProcessDragEnterMessage( TSharedRef<SWindow> WindowEntered, FDragDropEvent& DragDropEvent )
+bool FSlateApplication::ProcessDragEnterEvent( TSharedRef<SWindow> WindowEntered, FDragDropEvent& DragDropEvent )
 {
 	FWidgetPath WidgetsUnderCursor = LocateWindowUnderMouse( DragDropEvent.GetScreenSpacePosition(), GetInteractiveTopLevelWindows() );
 	DragDropEvent.SetEventPath( WidgetsUnderCursor );
@@ -4418,7 +4388,7 @@ EDropEffect::Type FSlateApplication::OnDragOver( const TSharedPtr< FGenericWindo
 				PlatformApplication->GetModifierKeys()
 			);
 
-			MouseMoveHandled = ProcessMouseMoveMessage( MouseEvent );
+			MouseMoveHandled = ProcessMouseMoveEvent( MouseEvent );
 			CursorMovementDelta = MouseEvent.GetCursorDelta();
 		}
 
@@ -4460,7 +4430,7 @@ EDropEffect::Type FSlateApplication::OnDragDrop( const TSharedPtr< FGenericWindo
 
 		// User dropped into a Slate window. Slate is already in drag and drop mode.
 		// It knows what to do based on a mouse up.
-		if ( ProcessMouseButtonUpMessage( MouseEvent ) )
+		if ( ProcessMouseButtonUpEvent( MouseEvent ) )
 {
 			Result = EDropEffect::Copy;
 		}
diff --git a/Engine/Source/Runtime/Slate/Public/Framework/Application/SlateApplication.h b/Engine/Source/Runtime/Slate/Public/Framework/Application/SlateApplication.h
--- a/Engine/Source/Runtime/Slate/Public/Framework/Application/SlateApplication.h
+++ b/Engine/Source/Runtime/Slate/Public/Framework/Application/SlateApplication.h
@@ -654,7 +654,7 @@ protected:
 	void PrivateDrawWindows( TSharedPtr<SWindow> DrawOnlyThisWindow = NULL );
 
 	/**
-	 * Prepass step before drawing windows to compute geometry size and reshape autosized windows
+	 * Pre-pass step before drawing windows to compute geometry size and reshape autosized windows
 	 */
 	void DrawPrepass( TSharedPtr<SWindow> DrawOnlyThisWindow );
 
@@ -672,7 +672,7 @@ protected:
 	/** Engages or disengages application throttling based on user behavior */
 	void ThrottleApplicationBasedOnMouseMovement();
 
-protected:
+public:
 
 	/**
 	 * Called by the native application in response to a mouse move. Routs the event to Slate Widgets.
@@ -681,7 +681,7 @@ protected:
 	 *
 	 * @return  Was this event handled by the Slate application?
 	 */
-	bool ProcessMouseMoveMessage( FPointerEvent& MouseEvent );
+	bool ProcessMouseMoveEvent( FPointerEvent& MouseEvent );
 
 	/**
 	 * Called by the native application in response to a mouse button press. Routs the event to Slate Widgets.
@@ -690,7 +690,7 @@ protected:
 	 *
 	 * @return  Was this event handled by the Slate application?
 	 */
-	bool ProcessMouseButtonDownMessage( const TSharedPtr< FGenericWindow >& PlatformWindow, FPointerEvent& MouseEvent );
+	bool ProcessMouseButtonDownEvent( const TSharedPtr< FGenericWindow >& PlatformWindow, FPointerEvent& MouseEvent );
 
 	/**
 	 * Called by the native application in response to a mouse button release. Routs the event to Slate Widgets.
@@ -699,8 +699,7 @@ protected:
 	 *
 	 * @return  Was this event handled by the Slate application?
 	 */
-	bool ProcessMouseButtonUpMessage( FPointerEvent& MouseEvent );
-
+	bool ProcessMouseButtonUpEvent( FPointerEvent& MouseEvent );
 
 	/**
 	 * Called by the native application in response to a mouse release. Routs the event to Slate Widgets.
@@ -709,7 +708,7 @@ protected:
 	 *
 	 * @return  Was this event handled by the Slate application?
 	 */
-	bool ProcessMouseButtonDoubleClickMessage( const TSharedPtr< FGenericWindow >& PlatformWindow, FPointerEvent& InMouseEvent );
+	bool ProcessMouseButtonDoubleClickEvent( const TSharedPtr< FGenericWindow >& PlatformWindow, FPointerEvent& InMouseEvent );
 	
 	/**
 	 * Called by the native application in response to a mouse wheel spin or a touch gesture. Routs the event to Slate Widgets.
@@ -719,7 +718,7 @@ protected:
 	 *
 	 * @return  Was this event handled by the Slate application?
 	 */
-	bool ProcessMouseWheelOrGestureMessage( FPointerEvent& InWheelEvent, const FPointerEvent* InGestureEvent );
+	bool ProcessMouseWheelOrGestureEvent( FPointerEvent& InWheelEvent, const FPointerEvent* InGestureEvent );
 
 	/**
 	 * Called when a character is entered
@@ -728,7 +727,7 @@ protected:
 	 *
 	 * @return  Was this event handled by the Slate application?
 	 */
-	bool ProcessKeyCharMessage( FCharacterEvent& InCharacterEvent );
+	bool ProcessKeyCharEvent( FCharacterEvent& InCharacterEvent );
 
 	/**
 	 * Called when a key is pressed
@@ -737,7 +736,7 @@ protected:
 	 *
 	 * @return  Was this event handled by the Slate application?
 	 */
-	bool ProcessKeyDownMessage( FKeyboardEvent& InKeyboardEvent );
+	bool ProcessKeyDownEvent( FKeyboardEvent& InKeyboardEvent );
 
 	/**
 	 * Called when a key is released
@@ -746,15 +745,8 @@ protected:
 	 *
 	 * @return  Was this event handled by the Slate application?
 	 */
-	bool ProcessKeyUpMessage( FKeyboardEvent& InKeyboardEvent );
+	bool ProcessKeyUpEvent( FKeyboardEvent& InKeyboardEvent );
 	
-public:
-
-	/**
-	 * Called by the native application in response to a set cursor event.
-	 */
-	void OnSetCursorMessage();
-
 	/**
 	 * Called when a drag from an external (non-slate) source enters a window
 	 *
@@ -763,78 +755,57 @@ public:
 	 *
 	 * @return true if the drag enter was handled and can be processed by some widget in this window; false otherwise
 	 */
-	bool ProcessDragEnterMessage( TSharedRef<SWindow> WindowEntered, FDragDropEvent& DragDropEvent );
-
-	/**
-	 * Called when a drag from an external (non-slate) source leaves a window
-	 *
-	 * @param WindowLeft     The window that was left by the drag and drop
-	 * @param DragDropEvent  Describes the mouse state (position, pressed buttons, etc) and associated payload
-	 *
-	 * @return true if the drag enter was handled and can be processed by some widget in this window; false otherwise
-	 */
-	void OnDragLeaveMessage();
-
-	/**
-	 * Called when a drag from an external (non-slate) is dropped onto a window
-	 *
-	 * @param WindowDroppedInto  The window that was dropped into
-	 * @param DragDropEvent  Describes the mouse state (position, pressed buttons, etc) and associated payload
-	 *
-	 * @return true if the drag enter was handled and can be processed by some widget in this window; false otherwise
-	 */
-	bool OnDropMessage( TSharedRef<SWindow> WindowDroppedInto, FDragDropEvent& DragDropEvent );
+	bool ProcessDragEnterEvent( TSharedRef<SWindow> WindowEntered, FDragDropEvent& DragDropEvent );
 
 	/**
 	 * Called when a controller button is found pressed when polling game device state
 	 * 
 	 * @param ControllerEvent	The controller event generated
 	 */
-	void OnControllerButtonPressedMessage( FControllerEvent& ControllerEvent );
+	void ProcessControllerButtonPressedEvent( FControllerEvent& ControllerEvent );
 
 	/**
 	 * Called when a controller button is found released when polling game device state
 	 * 
 	 * @param ControllerEvent	The controller event generated
 	 */
-	void OnControllerButtonReleasedMessage( FControllerEvent& ControllerEvent );
+	void ProcessControllerButtonReleasedEvent( FControllerEvent& ControllerEvent );
 
 	/**
 	 * Called when a controller analog value has changed while polling game device state
 	 * 
 	 * @param ControllerEvent	The controller event generated
 	 */
-	void OnControllerAnalogValueChangedMessage( FControllerEvent& ControllerEvent );
+	void ProcessControllerAnalogValueChangedEvent( FControllerEvent& ControllerEvent );
 
 	/**
 	 * Called when a touchpad touch is started (finger down) when polling game device state
 	 * 
 	 * @param ControllerEvent	The touch event generated
 	 */
-	void OnTouchStartedMessage( const TSharedPtr< FGenericWindow >& PlatformWindow, FPointerEvent& InTouchEvent );
+	void ProcessTouchStartedEvent( const TSharedPtr< FGenericWindow >& PlatformWindow, FPointerEvent& InTouchEvent );
 
 	/**
 	 * Called when a touchpad touch is moved  (finger moved) when polling game device state
 	 * 
 	 * @param ControllerEvent	The touch event generated
 	 */
-	void OnTouchMovedMessage( FPointerEvent& InTouchEvent );
+	void ProcessTouchMovedEvent( FPointerEvent& InTouchEvent );
 
 	/**
 	 * Called when a touchpad touch is ended (finger lifted) when polling game device state
 	 * 
 	 * @param ControllerEvent	The touch event generated
 	 */
-	void OnTouchEndedMessage( FPointerEvent& InTouchEvent );
+	void ProcessTouchEndedEvent( FPointerEvent& InTouchEvent );
 
 	/**
 	 * Called when motion is detected (controller or device) when polling game device state
 	 * 
 	 * @param MotionEvent		The motion event generated
 	 */
-	void OnMotionDetectedMessage( FMotionEvent& InMotionEvent );
+	void ProcessMotionDetectedEvent( FMotionEvent& InMotionEvent );
 
-public:
 	/**
 	 * Called by the native application in response to an activation or deactivation. 
 	 *
@@ -842,12 +813,16 @@ public:
 	 *
 	 * @return  Was this event handled by the Slate application?
 	 */
-	bool ProcessWindowActivatedMessage( const FWindowActivateEvent& ActivateEvent );
+	bool ProcessWindowActivatedEvent( const FWindowActivateEvent& ActivateEvent );
 	
-	/** Called when the application is activated (i.e. one of its windows becomes active) or deactivated. */
-	void ProcessApplicationActivationMessage( bool InAppActivated );
-
+	/**
+	 * Called when the application is activated (i.e. one of its windows becomes active) or deactivated.
+	 *
+	 * @param InAppActivated Whether the application was activated.
+	 */
+	void ProcessApplicationActivationEvent( bool InAppActivated );
 
+public:
 
 	/** Called when the slate application is being shut down. */
 	void OnShutdown();
