#slate List, Trees, Tiles now respond to touch events.

diff --git a/Engine/Source/Runtime/Slate/Private/Framework/Application/SlateApplication.cpp b/Engine/Source/Runtime/Slate/Private/Framework/Application/SlateApplication.cpp
--- a/Engine/Source/Runtime/Slate/Private/Framework/Application/SlateApplication.cpp
+++ b/Engine/Source/Runtime/Slate/Private/Framework/Application/SlateApplication.cpp
@@ -3952,7 +3952,7 @@ bool FSlateApplication::OnTouchMoved( const FVector2D& Location, int32 TouchInde
 		ControllerId,
 		TouchIndex,
 		Location,
-		Location,
+		LastCursorPosition,
 		true);
 	ProcessTouchMovedEvent(PointerEvent);
 
diff --git a/Engine/Source/Runtime/Slate/Private/Widgets/Views/STableViewBase.cpp b/Engine/Source/Runtime/Slate/Private/Widgets/Views/STableViewBase.cpp
--- a/Engine/Source/Runtime/Slate/Private/Widgets/Views/STableViewBase.cpp
+++ b/Engine/Source/Runtime/Slate/Private/Widgets/Views/STableViewBase.cpp
@@ -317,7 +317,8 @@ FReply STableViewBase::OnMouseMove( const FGeometry& MyGeometry, const FPointerE
 
 			FReply Reply = FReply::Handled();
 
-			// Capture the mouse if we need to
+			// The mouse moved enough that we're now dragging the view. Capture the mouse
+			// so the user does not have to stay within the bounds of the list while dragging.
 			if( FSlateApplication::Get().GetMouseCaptor().Get() != this )
 			{
 				Reply.CaptureMouse( AsShared() ).UseHighPrecisionMouseMovement( AsShared() );
@@ -396,7 +397,7 @@ FReply STableViewBase::OnKeyDown( const FGeometry& MyGeometry, const FKeyboardEv
 
 FCursorReply STableViewBase::OnCursorQuery( const FGeometry& MyGeometry, const FPointerEvent& CursorEvent ) const
 {
-	if ( IsRightClickScrolling() )
+	if ( IsRightClickScrolling() && CursorEvent.IsMouseButtonDown(EKeys::RightMouseButton) )
 	{
 		// We hide the native cursor as we'll be drawing the software EMouseCursor::GrabHandClosed cursor
 		return FCursorReply::Cursor( EMouseCursor::None );
@@ -407,6 +408,49 @@ FCursorReply STableViewBase::OnCursorQuery( const FGeometry& MyGeometry, const F
 	}
 }
 
+FReply STableViewBase::OnTouchStarted( const FGeometry& MyGeometry, const FPointerEvent& InTouchEvent )
+{
+	// Clear any inertia 
+	this->InertialScrollManager.ClearScrollVelocity();
+	// We have started a new interaction; track how far the user has moved since they put their finger down.
+	AmountScrolledWhileRightMouseDown = 0;
+
+	return FReply::Unhandled();
+}
+
+FReply STableViewBase::OnTouchMoved( const FGeometry& MyGeometry, const FPointerEvent& InTouchEvent )
+{
+	const float ScrollByAmount = InTouchEvent.GetCursorDelta().Y;
+	AmountScrolledWhileRightMouseDown += FMath::Abs( ScrollByAmount );
+
+	this->InertialScrollManager.AddScrollSample( ScrollByAmount, FPlatformTime::Seconds());
+	const float AmountScrolled = this->ScrollBy( MyGeometry, ScrollByAmount );
+
+	if (AmountScrolledWhileRightMouseDown > SlatePanTriggerDistance)
+	{
+		// The user has moved the list some amount; they are probably
+		// trying to scroll. From now on, the list assumes the user is scrolling
+		// until they lift their finger.
+		return FReply::Handled().CaptureMouse( AsShared() );
+	}
+
+	return FReply::Handled();
+}
+
+FReply STableViewBase::OnTouchEnded( const FGeometry& MyGeometry, const FPointerEvent& InTouchEvent )
+{
+	AmountScrolledWhileRightMouseDown = 0;
+
+	if (HasMouseCapture())
+	{
+		return FReply::Handled().ReleaseMouseCapture();
+	}
+	else
+	{
+		return FReply::Handled();
+	}
+}
+
 
 /** @return The number of Widgets we currently have generated. */
 int32 STableViewBase::GetNumGeneratedChildren() const
diff --git a/Engine/Source/Runtime/Slate/Public/Widgets/Views/STableViewBase.h b/Engine/Source/Runtime/Slate/Public/Widgets/Views/STableViewBase.h
--- a/Engine/Source/Runtime/Slate/Public/Widgets/Views/STableViewBase.h
+++ b/Engine/Source/Runtime/Slate/Public/Widgets/Views/STableViewBase.h
@@ -36,6 +36,9 @@ public:
 	virtual FReply OnMouseWheel( const FGeometry& MyGeometry, const FPointerEvent& MouseEvent ) OVERRIDE;
 	virtual FReply OnKeyDown( const FGeometry& MyGeometry, const FKeyboardEvent& InKeyboardEvent ) OVERRIDE;
 	virtual FCursorReply OnCursorQuery( const FGeometry& MyGeometry, const FPointerEvent& CursorEvent ) const OVERRIDE;
+	virtual FReply OnTouchStarted( const FGeometry& MyGeometry, const FPointerEvent& InTouchEvent ) OVERRIDE;
+	virtual FReply OnTouchMoved( const FGeometry& MyGeometry, const FPointerEvent& InTouchEvent ) OVERRIDE;
+	virtual FReply OnTouchEnded( const FGeometry& MyGeometry, const FPointerEvent& InTouchEvent ) OVERRIDE;
 	// End of SWidget interface
 
 	/** @return The number of Widgets we currently have generated. */
diff --git a/Engine/Source/Runtime/SlateCore/Public/Widgets/SWidget.h b/Engine/Source/Runtime/SlateCore/Public/Widgets/SWidget.h
--- a/Engine/Source/Runtime/SlateCore/Public/Widgets/SWidget.h
+++ b/Engine/Source/Runtime/SlateCore/Public/Widgets/SWidget.h
@@ -408,6 +408,7 @@ public:
 
 	/**
 	 * Called when motion is detected (controller or device)
+	 * e.g. Someone tilts or shakes their controller.
 	 * 
 	 * @param MotionEvent	The motion event generated
 	 */
