Improve detection for SListView item types being a pointer type.
* Remove the static assert from SListView because the failure was being triggered earlier in TListTypeTraits<>
* Moved the check to TListTypeTraits.

diff --git a/Engine/Source/Runtime/Slate/Public/Framework/Views/TableViewTypeTraits.h b/Engine/Source/Runtime/Slate/Public/Framework/Views/TableViewTypeTraits.h
--- a/Engine/Source/Runtime/Slate/Public/Framework/Views/TableViewTypeTraits.h
+++ b/Engine/Source/Runtime/Slate/Public/Framework/Views/TableViewTypeTraits.h
@@ -22,6 +22,7 @@ template<typename T> struct TIsValidListItem< TWeakObjectPtr<T> >						{ enum {
  */
 template <typename T> struct TListTypeTraits
 {
+	checkAtCompileTime(TIsValidListItem<T>::Value, ItemType_must_be_a_pointer_or_TSharedPtr);
 };
 
 
diff --git a/Engine/Source/Runtime/Slate/Public/Widgets/Views/SListView.h b/Engine/Source/Runtime/Slate/Public/Widgets/Views/SListView.h
--- a/Engine/Source/Runtime/Slate/Public/Widgets/Views/SListView.h
+++ b/Engine/Source/Runtime/Slate/Public/Widgets/Views/SListView.h
@@ -31,8 +31,6 @@
 template <typename ItemType>
 class SListView : public STableViewBase, TListTypeTraits<ItemType>::SerializerType, public ITypedTableView< ItemType >
 {
-	checkAtCompileTime( TIsValidListItem<ItemType>::Value, ItemType_must_be_a_pointer_or_TSharedPtr );
-
 public:
 	typedef typename TListTypeTraits< ItemType >::NullableType NullableItemType;
 
