UMG - Exposing slate list related enums to UMG.

diff --git a/Engine/Source/Runtime/Slate/Public/Framework/Views/ITypedTableView.h b/Engine/Source/Runtime/Slate/Public/Framework/Views/ITypedTableView.h
--- a/Engine/Source/Runtime/Slate/Public/Framework/Views/ITypedTableView.h
+++ b/Engine/Source/Runtime/Slate/Public/Framework/Views/ITypedTableView.h
@@ -2,6 +2,9 @@
 
 #pragma once
 
+#include "ITypedTableView.generated.h"
+
+UENUM()
 namespace ESelectionMode
 {
 	enum Type
@@ -21,6 +24,7 @@ namespace ESelectionMode
 }
 
 /** Lists can exist as regular lists or as backing widgets for a tree */
+UENUM()
 namespace ETableViewMode
 {
 	enum Type
