SComboBox fixes
* SComboRow uses a SLATE_DEFAULT_SLOT() instead of SLATE_ARGUMENT called RowContent.
* If OnGenerateWidget is not bound, a default STextBlock is generated reminding the user to bind that delegate instead of failing to compile.
#BUN SComboRow RowContent is now replaced by a default slot as it matches Slate convention better.

diff --git a/Engine/Plugins/Developer/PerforceSourceControl/Source/PerforceSourceControl/Private/SPerforceSourceControlSettings.cpp b/Engine/Plugins/Developer/PerforceSourceControl/Source/PerforceSourceControl/Private/SPerforceSourceControlSettings.cpp
--- a/Engine/Plugins/Developer/PerforceSourceControl/Source/PerforceSourceControl/Private/SPerforceSourceControlSettings.cpp
+++ b/Engine/Plugins/Developer/PerforceSourceControl/Source/PerforceSourceControl/Private/SPerforceSourceControlSettings.cpp
@@ -421,7 +421,7 @@ TSharedRef<ITableRow> SPerforceSourceControlSettings::OnGenerateWorkspaceRow(TSh
 {
 	return
 		SNew(SComboRow< TSharedRef<FString> >, OwnerTable)
-		.RowContent( 
+		[
 			SNew(SHorizontalBox)
 			+SHorizontalBox::Slot()
 			.FillWidth(1)
@@ -431,7 +431,7 @@ TSharedRef<ITableRow> SPerforceSourceControlSettings::OnGenerateWorkspaceRow(TSh
 				.Text(*InItem)
 				.Font(FEditorStyle::GetFontStyle(TEXT("PropertyWindow.NormalFont")))
 			]
-		);
+		];
 }
 
 void SPerforceSourceControlSettings::OnWorkspaceSelected(TSharedPtr<FString> InItem, ESelectInfo::Type InSelectInfo)
diff --git a/Engine/Source/Editor/UnrealEd/Private/SComponentClassCombo.cpp b/Engine/Source/Editor/UnrealEd/Private/SComponentClassCombo.cpp
--- a/Engine/Source/Editor/UnrealEd/Private/SComponentClassCombo.cpp
+++ b/Engine/Source/Editor/UnrealEd/Private/SComponentClassCombo.cpp
@@ -260,8 +260,7 @@ TSharedRef<ITableRow> SComponentClassCombo::GenerateAddComponentRow( FComponentC
 		return
 			SNew( SComboRow< TSharedPtr<FString> >, OwnerTable )
 			.ToolTip( IDocumentation::Get()->CreateToolTip(Entry->GetComponentClass()->GetToolTipText(), NULL, DocLink, DocExcerptName) )
-			.RowContent
-			(
+			[
 				SNew(SHorizontalBox)
 				+SHorizontalBox::Slot()
 				.AutoWidth()
@@ -291,7 +290,7 @@ TSharedRef<ITableRow> SComponentClassCombo::GenerateAddComponentRow( FComponentC
 					.HighlightText(this, &SComponentClassCombo::GetCurrentSearchString)
 					.Text( FriendlyComponentName )
 				]
-			);
+			];
 	}
 }
 
diff --git a/Engine/Source/Runtime/Slate/Public/Widgets/Input/SComboBox.h b/Engine/Source/Runtime/Slate/Public/Widgets/Input/SComboBox.h
--- a/Engine/Source/Runtime/Slate/Public/Widgets/Input/SComboBox.h
+++ b/Engine/Source/Runtime/Slate/Public/Widgets/Input/SComboBox.h
@@ -12,10 +12,10 @@ class SComboRow : public STableRow< OptionType >
 public:
 
 	SLATE_BEGIN_ARGS( SComboRow )
-		: _RowContent()
+		: _Content()
 		{}
 
-		SLATE_ARGUMENT( TSharedPtr<SWidget>, RowContent )
+		SLATE_DEFAULT_SLOT( FArguments, Content )
 	SLATE_END_ARGS()
 
 public:
@@ -29,7 +29,7 @@ public:
 			typename STableRow<OptionType>::FArguments()
 			.Content()
 			[
-				InArgs._RowContent->AsShared()
+				InArgs._Content.Widget
 			]
 			, InOwnerTable
 		);
@@ -284,9 +284,21 @@ private:
 	/** Generate a row for the InItem in the combo box's list (passed in as OwnerTable). Do this by calling the user-specified OnGenerateWidget */
 	TSharedRef<ITableRow> GenerateMenuItemRow( OptionType InItem, const TSharedRef<STableViewBase>& OwnerTable)
 	{
-		return
-		SNew( SComboRow<OptionType>, OwnerTable)
-		.RowContent(OnGenerateWidget.Execute(InItem) );
+		if (OnGenerateWidget.IsBound())
+		{
+			return SNew(SComboRow<OptionType>, OwnerTable)
+				[
+					OnGenerateWidget.Execute(InItem)
+				];
+		}
+		else
+		{
+			return SNew(SComboRow<OptionType>, OwnerTable)
+				[
+					SNew(STextBlock).Text(NSLOCTEXT("SlateCore", "ComboBoxMissingOnGenerateWidgetMethod", "Please provide a .OnGenerateWidget() handler."))
+				];
+
+		}
 	}
 
 	/** Invoked when the selection in the list changes */
