Slate: properly checking weak pointer to owner window; added missing documentation

diff --git a/Engine/Source/Runtime/Slate/Private/Framework/Application/SWindowTitleBar.h b/Engine/Source/Runtime/Slate/Private/Framework/Application/SWindowTitleBar.h
--- a/Engine/Source/Runtime/Slate/Private/Framework/Application/SWindowTitleBar.h
+++ b/Engine/Source/Runtime/Slate/Private/Framework/Application/SWindowTitleBar.h
@@ -152,6 +152,11 @@ protected:
 	{
 		TSharedPtr<SWindow> OwnerWindow = OwnerWindowPtr.Pin();
 
+		if (!OwnerWindow.IsValid())
+		{
+			return;
+		}
+
 		MinimizeButton = SNew(SButton)
 				.IsFocusable(false)
 				.IsEnabled(OwnerWindow->HasMinimizeBox())
@@ -386,7 +391,11 @@ private:
 	FReply CloseButton_OnClicked()
 	{
 		TSharedPtr<SWindow> OwnerWindow = OwnerWindowPtr.Pin();
-		OwnerWindow->RequestDestroyWindow();
+
+		if (OwnerWindow.IsValid())
+		{
+			OwnerWindow->RequestDestroyWindow();
+		}
 
 		return FReply::Handled();
 	}
@@ -395,6 +404,12 @@ private:
 	const FSlateBrush* GetCloseImage() const
 	{
 		TSharedPtr<SWindow> OwnerWindow = OwnerWindowPtr.Pin();
+
+		if (!OwnerWindow.IsValid())
+		{
+			return nullptr;
+		}
+
 		TSharedPtr<FGenericWindow> NativeWindow = OwnerWindow->GetNativeWindow();
 
 		if (CloseButton->IsPressed())
@@ -424,17 +439,21 @@ private:
 	FReply MaximizeRestoreButton_OnClicked( )
 	{
 		TSharedPtr<SWindow> OwnerWindow = OwnerWindowPtr.Pin();
-		TSharedPtr<FGenericWindow> NativeWindow = OwnerWindow->GetNativeWindow();
 
-		if (NativeWindow.IsValid())
+		if (OwnerWindow.IsValid())
 		{
-			if (NativeWindow->IsMaximized())
-			{
-				NativeWindow->Restore();
-			}
-			else
+			TSharedPtr<FGenericWindow> NativeWindow = OwnerWindow->GetNativeWindow();
+
+			if (NativeWindow.IsValid())
 			{
-				NativeWindow->Maximize();
+				if (NativeWindow->IsMaximized())
+				{
+					NativeWindow->Restore();
+				}
+				else
+				{
+					NativeWindow->Maximize();
+				}
 			}
 		}
 
@@ -445,6 +464,12 @@ private:
 	const FSlateBrush* GetMaximizeRestoreImage( ) const
 	{
 		TSharedPtr<SWindow> OwnerWindow = OwnerWindowPtr.Pin();
+
+		if (!OwnerWindow.IsValid())
+		{
+			return nullptr;
+		}
+
 		TSharedPtr<FGenericWindow> NativeWindow = OwnerWindow->GetNativeWindow();
 
 		if (NativeWindow.IsValid() && NativeWindow->IsMaximized())
@@ -511,11 +536,15 @@ private:
 	FReply MinimizeButton_OnClicked( )
 	{
 		TSharedPtr<SWindow> OwnerWindow = OwnerWindowPtr.Pin();
-		TSharedPtr<FGenericWindow> NativeWindow = OwnerWindow->GetNativeWindow();
 
-		if (NativeWindow.IsValid())
+		if (OwnerWindow.IsValid())
 		{
-			NativeWindow->Minimize();
+			TSharedPtr<FGenericWindow> NativeWindow = OwnerWindow->GetNativeWindow();
+
+			if (NativeWindow.IsValid())
+			{
+				NativeWindow->Minimize();
+			}
 		}
 
 		return FReply::Handled();
@@ -525,6 +554,12 @@ private:
 	const FSlateBrush* GetMinimizeImage( ) const
 	{
 		TSharedPtr<SWindow> OwnerWindow = OwnerWindowPtr.Pin();
+
+		if (!OwnerWindow.IsValid())
+		{
+			return nullptr;
+		}
+
 		TSharedPtr<FGenericWindow> NativeWindow = OwnerWindow->GetNativeWindow();
 
 		if (!OwnerWindow->HasMinimizeButton())
@@ -559,6 +594,12 @@ private:
 	const FSlateBrush* GetWindowTitlebackgroundImage( ) const
 	{
 		TSharedPtr<SWindow> OwnerWindow = OwnerWindowPtr.Pin();
+
+		if (!OwnerWindow.IsValid())
+		{
+			return nullptr;
+		}
+
 		TSharedPtr<FGenericWindow> NativeWindow = OwnerWindow->GetNativeWindow();
 		const bool bIsActive = NativeWindow.IsValid() && NativeWindow->IsForegroundWindow();
 
@@ -594,22 +635,34 @@ private:
 	{
 		TSharedPtr<SWindow> OwnerWindow = OwnerWindowPtr.Pin();
 
+		if (!OwnerWindow.IsValid())
+		{
+			return FText::GetEmpty();
+		}
+
 		return OwnerWindow->GetTitle();
 	}
 
 private:
 
+	// Holds a weak pointer to the owner window.
 	TWeakPtr<SWindow> OwnerWindowPtr;
 
+	// Holds the window style to use (for buttons, text, etc.).
 	const FWindowStyle* Style;
 
+	// Holds the content widget of the title area.
 	TSharedPtr<SWidget> TitleArea;
 
+	// Holds the curve sequence for the window flash animation.
 	FCurveSequence TitleFlashSequence;
 
+	// Holds the minimize button.
 	TSharedPtr<SButton> MinimizeButton;
 
+	// Holds the maximize/restore button.
 	TSharedPtr<SButton> MaximizeRestoreButton;
 
+	// Holds the close button.
 	TSharedPtr<SButton> CloseButton;
 };
