Slate - Adding a way to get and set the percent value of a progress bar.

diff --git a/Engine/Source/Runtime/Slate/Private/Widgets/Notifications/SProgressBar.cpp b/Engine/Source/Runtime/Slate/Private/Widgets/Notifications/SProgressBar.cpp
--- a/Engine/Source/Runtime/Slate/Private/Widgets/Notifications/SProgressBar.cpp
+++ b/Engine/Source/Runtime/Slate/Private/Widgets/Notifications/SProgressBar.cpp
@@ -21,6 +21,15 @@ void SProgressBar::Construct( const FArguments& InArgs )
 	CurveSequence.Play();
 }
 
+TOptional<float> SProgressBar::GetPercent() const
+{
+	return Percent.Get();
+}
+
+void SProgressBar::SetPercent(TAttribute< TOptional<float> > InPercent)
+{
+	Percent = InPercent;
+}
 
 int32 SProgressBar::OnPaint( const FGeometry& AllottedGeometry, const FSlateRect& MyClippingRect, FSlateWindowElementList& OutDrawElements, int32 LayerId, const FWidgetStyle& InWidgetStyle, bool bParentEnabled ) const
 {
diff --git a/Engine/Source/Runtime/Slate/Public/Widgets/Notifications/SProgressBar.h b/Engine/Source/Runtime/Slate/Public/Widgets/Notifications/SProgressBar.h
--- a/Engine/Source/Runtime/Slate/Public/Widgets/Notifications/SProgressBar.h
+++ b/Engine/Source/Runtime/Slate/Public/Widgets/Notifications/SProgressBar.h
@@ -62,6 +62,12 @@ public:
 
 	FVector2D ComputeDesiredSize() const;
 
+	/** See Slate attribute Percent */
+	TOptional<float> GetPercent() const;
+
+	/** See Slate attribute Percent */
+	void SetPercent(TAttribute< TOptional<float> > InPercent);
+
 private:
 
 	/** The text displayed over the progress bar */
