Added overscroll support for ListView, TileView, and TreeView on desktops (Right-Mouse-Button scrolling) and Touch-Devices.

diff --git a/Engine/Source/Runtime/Slate/Private/Widgets/Views/STableViewBase.cpp b/Engine/Source/Runtime/Slate/Private/Widgets/Views/STableViewBase.cpp
--- a/Engine/Source/Runtime/Slate/Private/Widgets/Views/STableViewBase.cpp
+++ b/Engine/Source/Runtime/Slate/Private/Widgets/Views/STableViewBase.cpp
@@ -3,6 +3,11 @@
 #include "SlatePrivatePCH.h"
 #include "SListPanel.h"
 
+namespace ListConstants
+{
+	static const float OvershootMax = 50.0f;
+	static const float OvershootBounceRate = 350.0f;
+}
 
 /** Create the child widgets that comprise the list */
 void STableViewBase::ConstructChildren( const TAttribute<float>& InItemWidth, const TAttribute<float>& InItemHeight, const TSharedPtr<SHeaderRow>& InHeaderRow, const TSharedPtr<SScrollBar>& InScrollBar )
@@ -160,11 +165,20 @@ void STableViewBase::Tick( const FGeometry& AllottedGeometry, const double InCur
 		if ( !IsRightClickScrolling() )
 		{
 			this->InertialScrollManager.UpdateScrollVelocity(InDeltaTime);
-
 			const float ScrollVelocity = this->InertialScrollManager.GetScrollVelocity();
+
 			if ( ScrollVelocity != 0.f )
 			{
-				this->ScrollBy(AllottedGeometry, ScrollVelocity * InDeltaTime);
+				this->ScrollBy(AllottedGeometry, ScrollVelocity * InDeltaTime, EAllowOverscroll::Yes);
+			}
+
+			Overscroll.UpdateOverscroll( InDeltaTime );
+
+			if (Overscroll.GetOverscroll() != 0.0f)
+			{
+				InertialScrollManager.ClearScrollVelocity();	
+				
+				this->RequestListRefresh();
 			}
 		}
 
@@ -192,8 +206,11 @@ void STableViewBase::Tick( const FGeometry& AllottedGeometry, const double InCur
 				ScrollOffset = ReGenerateResults.NewScrollOffset;
 			}
 			
+			
 			ScrollOffset = FMath::Max(0.0, ScrollOffset);
-			ItemsPanel->SmoothScrollOffset( FMath::Fractional(ScrollOffset / GetNumItemsWide()) );
+			const float OverscrollAmount = Overscroll.GetOverscroll();
+			ItemsPanel->SmoothScrollOffset( FMath::Fractional(ScrollOffset / GetNumItemsWide()) +  OverscrollAmount );
+			
 
 			UpdateSelectionSet();
 
@@ -313,7 +330,7 @@ FReply STableViewBase::OnMouseMove( const FGeometry& MyGeometry, const FPointerE
 		if( IsRightClickScrolling() )
 		{
 			this->InertialScrollManager.AddScrollSample(-ScrollByAmount, FPlatformTime::Seconds());
-			const float AmountScrolled = this->ScrollBy( MyGeometry, -ScrollByAmount );
+			const float AmountScrolled = this->ScrollBy( MyGeometry, -ScrollByAmount, EAllowOverscroll::Yes );
 
 			FReply Reply = FReply::Handled();
 
@@ -365,7 +382,7 @@ FReply STableViewBase::OnMouseWheel( const FGeometry& MyGeometry, const FPointer
 		// Make sure scroll velocity is cleared so it doesn't fight with the mouse wheel input
 		this->InertialScrollManager.ClearScrollVelocity();
 
-		const float AmountScrolledInItems = this->ScrollBy( MyGeometry, -MouseEvent.GetWheelDelta()*WheelScrollAmount );
+		const float AmountScrolledInItems = this->ScrollBy( MyGeometry, -MouseEvent.GetWheelDelta()*WheelScrollAmount, EAllowOverscroll::No );
 		if (FMath::Abs(AmountScrolledInItems) > 0.0f)
 		{
 			return FReply::Handled();
@@ -424,7 +441,7 @@ FReply STableViewBase::OnTouchMoved( const FGeometry& MyGeometry, const FPointer
 	AmountScrolledWhileRightMouseDown += FMath::Abs( ScrollByAmount );
 
 	this->InertialScrollManager.AddScrollSample( ScrollByAmount, FPlatformTime::Seconds());
-	const float AmountScrolled = this->ScrollBy( MyGeometry, ScrollByAmount );
+	const float AmountScrolled = this->ScrollBy( MyGeometry, ScrollByAmount, EAllowOverscroll::Yes );
 
 	if (AmountScrolledWhileRightMouseDown > SlatePanTriggerDistance)
 	{
@@ -518,21 +535,26 @@ STableViewBase::STableViewBase( ETableViewMode::Type InTableViewMode )
 	, SelectionMode( ESelectionMode::Multi )
 	, SoftwareCursorPosition( ForceInitToZero )
 	, bShowSoftwareCursor( false )
-	, bItemsNeedRefresh( true )
+	, Overscroll( ListConstants::OvershootMax )
+	, bItemsNeedRefresh( true )	
 {
 }
 
-float STableViewBase::ScrollBy( const FGeometry& MyGeometry, float ScrollByAmountInSlateUnits )
+float STableViewBase::ScrollBy( const FGeometry& MyGeometry, float ScrollByAmountInSlateUnits, EAllowOverscroll AllowOverscroll )
 {
 	const int32 NumItemsBeingObserved = GetNumItemsBeingObserved();
 	const float FractionalScrollOffsetInItems = (ScrollOffset + GetScrollRateInItems() * ScrollByAmountInSlateUnits) / NumItemsBeingObserved;
-	const double ClampedScrollOffsetInItems = FMath::Clamp<double>( FractionalScrollOffsetInItems, 0.0, 1.0 ) * NumItemsBeingObserved;
+	const double ClampedScrollOffsetInItems = FMath::Clamp<double>( FractionalScrollOffsetInItems*NumItemsBeingObserved, -10.0f, NumItemsBeingObserved+10.0f ) * NumItemsBeingObserved;
+	if (AllowOverscroll == EAllowOverscroll::Yes)
+	{
+		Overscroll.ScrollBy( ClampedScrollOffsetInItems - ScrollByAmountInSlateUnits );
+	}
 	return ScrollTo( ClampedScrollOffsetInItems );
 }
 
 float STableViewBase::ScrollTo( float InScrollOffset )
 {
-	const float NewScrollOffset = FMath::Clamp( InScrollOffset, 0.0f, (float)GetNumItemsBeingObserved() );
+	const float NewScrollOffset = FMath::Clamp( InScrollOffset, -10.0f, GetNumItemsBeingObserved()+10.0f );
 	float AmountScrolled = FMath::Abs( ScrollOffset - NewScrollOffset );
 	ScrollOffset = NewScrollOffset;
 
@@ -645,3 +667,52 @@ TSharedRef<class SWidget> STableViewBase::GetScrollWidget()
 {
 	return SharedThis(this);
 }
+
+
+STableViewBase::FOverscroll::FOverscroll( const float InMaxOverscroll )
+: OverscrollAmount( 0.0f )
+, MaxOverscroll( InMaxOverscroll )
+{
+}
+
+float STableViewBase::FOverscroll::ScrollBy( float Delta )
+{
+	const float ValueBeforeDeltaApplied = OverscrollAmount;
+	const float OverscrollMagnitude = OverscrollAmount*OverscrollAmount + 1.0f;
+	const float EasedDelta = Delta / OverscrollMagnitude;
+	OverscrollAmount = FMath::Clamp(OverscrollAmount + EasedDelta, -ListConstants::OvershootMax, ListConstants::OvershootMax);
+
+	return ValueBeforeDeltaApplied - OverscrollAmount;
+}
+
+float STableViewBase::FOverscroll::GetOverscroll() const
+{
+	return OverscrollAmount;
+}
+
+void STableViewBase::FOverscroll::UpdateOverscroll( float InDeltaTime )
+{
+	const float OverscrollMagnitude = OverscrollAmount*OverscrollAmount;
+	const float PullForce = OverscrollMagnitude + 1.0f;
+	const float EasedDelta = ListConstants::OvershootBounceRate * InDeltaTime * PullForce;
+
+	if ( OverscrollAmount > 0 )
+	{
+		OverscrollAmount = FMath::Max( 0.0f, OverscrollAmount - EasedDelta * InDeltaTime );
+	}
+	else
+	{
+		OverscrollAmount = FMath::Min( 0.0f, OverscrollAmount +  EasedDelta * InDeltaTime );
+	}
+}
+
+bool STableViewBase::FOverscroll::ShouldApplyOverscroll( const bool bIsAtStartOfList, const bool bIsAtEndOfList, const float ScrollDelta ) const
+{
+	const bool bShouldApplyOverscroll =
+		// We can scroll past the edge of the list only if we are at the edge
+		(bIsAtStartOfList && ScrollDelta < 0) || (bIsAtEndOfList && ScrollDelta > 0) ||
+		// ... or if we are already past the edge and are scrolling in the opposite direction.
+		(OverscrollAmount > 0 && ScrollDelta < 0) || (OverscrollAmount < 0 && ScrollDelta > 0);
+
+	return bShouldApplyOverscroll;
+}
diff --git a/Engine/Source/Runtime/Slate/Public/Widgets/Views/SListView.h b/Engine/Source/Runtime/Slate/Public/Widgets/Views/SListView.h
--- a/Engine/Source/Runtime/Slate/Public/Widgets/Views/SListView.h
+++ b/Engine/Source/Runtime/Slate/Public/Widgets/Views/SListView.h
@@ -1116,108 +1116,133 @@ protected:
 		}
 	}
 
-	virtual float ScrollBy( const FGeometry& MyGeometry, float ScrollByAmountInSlateUnits ) OVERRIDE
+	virtual float ScrollBy( const FGeometry& MyGeometry, float ScrollByAmountInSlateUnits, EAllowOverscroll AllowOverscroll ) OVERRIDE
 	{
 		float AbsScrollByAmount = FMath::Abs( ScrollByAmountInSlateUnits );
 		int32 StartingItemIndex = (int32)ScrollOffset;
 		double NewScrollOffset = ScrollOffset;
 
-		const TArray<ItemType>* SourceItems = ItemsSource;
-		if ( SourceItems != NULL && SourceItems->Num() > 0 )
+		const bool bWholeListVisible = ScrollOffset == 0 && bWasAtEndOfList;
+		if (bWholeListVisible)
+		{
+			return 0;
+		}
+		else if ( AllowOverscroll == EAllowOverscroll::Yes && Overscroll.ShouldApplyOverscroll( ScrollOffset == 0, bWasAtEndOfList, ScrollByAmountInSlateUnits ) )
+		{
+			const float UnclampedScrollDelta = FMath::Sign(ScrollByAmountInSlateUnits) * AbsScrollByAmount / 10;				
+			const float ActuallyScrolledBy = Overscroll.ScrollBy( UnclampedScrollDelta );
+			if (ActuallyScrolledBy != 0.0f)
+			{
+				this->RequestListRefresh();
+			}
+			return ActuallyScrolledBy;
+		}
+		else
 		{
-			int ItemIndex = StartingItemIndex;
-			while( AbsScrollByAmount != 0 && ItemIndex < SourceItems->Num() && ItemIndex >= 0 )
+			// We know how far we want to scroll in SlateUnits, but we store scroll offset in "number of widgets".
+			// Challenge: each widget can be a different height.
+			// Strategy:
+			//           Scroll "one widget's height" at a time until we've scrolled as far as the user asked us to.
+			//           Generate widgets on demand so we can figure out how tall they are.
+
+			const TArray<ItemType>* SourceItems = ItemsSource;
+			if ( SourceItems != NULL && SourceItems->Num() > 0 )
 			{
-				const ItemType CurItem = (*SourceItems)[ ItemIndex ];
-				TSharedPtr<ITableRow> RowWidget = WidgetGenerator.GetWidgetForItem( CurItem );
-				if ( !RowWidget.IsValid() )
+				int ItemIndex = StartingItemIndex;
+				while( AbsScrollByAmount != 0 && ItemIndex < SourceItems->Num() && ItemIndex >= 0 )
 				{
-					// We couldn't find an existing widgets, meaning that this data item was not visible before.
-					// Make a new widget for it.
-					RowWidget = this->GenerateNewWidget( CurItem );
-
-					// It is useful to know the item's index that the widget was generated from.
-					// Helps with even/odd coloring
-					RowWidget->SetIndexInList(ItemIndex);
+					const ItemType CurItem = (*SourceItems)[ ItemIndex ];
+					TSharedPtr<ITableRow> RowWidget = WidgetGenerator.GetWidgetForItem( CurItem );
+					if ( !RowWidget.IsValid() )
+					{
+						// We couldn't find an existing widgets, meaning that this data item was not visible before.
+						// Make a new widget for it.
+						RowWidget = this->GenerateNewWidget( CurItem );
 
-					// Let the item generator know that we encountered the current Item and associated Widget.
-					WidgetGenerator.OnItemSeen( CurItem, RowWidget.ToSharedRef() );
+						// It is useful to know the item's index that the widget was generated from.
+						// Helps with even/odd coloring
+						RowWidget->SetIndexInList(ItemIndex);
 
-					RowWidget->AsWidget()->SlatePrepass();
-				}
+						// Let the item generator know that we encountered the current Item and associated Widget.
+						WidgetGenerator.OnItemSeen( CurItem, RowWidget.ToSharedRef() );
 
-				if ( ScrollByAmountInSlateUnits > 0 )
-				{
-					FVector2D DesiredSize = RowWidget->AsWidget()->GetDesiredSize();
-					const float RemainingHeight = DesiredSize.Y * ( 1.0 - FMath::Fractional( NewScrollOffset ) );
+						RowWidget->AsWidget()->SlatePrepass();
+					}
 
-					if ( AbsScrollByAmount > RemainingHeight )
+					if ( ScrollByAmountInSlateUnits > 0 )
 					{
-						if ( ItemIndex != SourceItems->Num() )
+						FVector2D DesiredSize = RowWidget->AsWidget()->GetDesiredSize();
+						const float RemainingHeight = DesiredSize.Y * ( 1.0 - FMath::Fractional( NewScrollOffset ) );
+
+						if ( AbsScrollByAmount > RemainingHeight )
+						{
+							if ( ItemIndex != SourceItems->Num() )
+							{
+								AbsScrollByAmount -= RemainingHeight;
+								NewScrollOffset = 1.0f + (int32)NewScrollOffset;
+								++ItemIndex;
+							}
+							else
+							{
+								NewScrollOffset = SourceItems->Num();
+								break;
+							}
+						} 
+						else if ( AbsScrollByAmount == RemainingHeight )
 						{
-							AbsScrollByAmount -= RemainingHeight;
 							NewScrollOffset = 1.0f + (int32)NewScrollOffset;
-							++ItemIndex;
+							break;
 						}
 						else
 						{
-							NewScrollOffset = SourceItems->Num();
+							NewScrollOffset = (int32)NewScrollOffset + ( 1.0f - ( ( RemainingHeight - AbsScrollByAmount ) / DesiredSize.Y ) );
 							break;
 						}
-					} 
-					else if ( AbsScrollByAmount == RemainingHeight )
-					{
-						NewScrollOffset = 1.0f + (int32)NewScrollOffset;
-						break;
 					}
 					else
 					{
-						NewScrollOffset = (int32)NewScrollOffset + ( 1.0f - ( ( RemainingHeight - AbsScrollByAmount ) / DesiredSize.Y ) );
-						break;
-					}
-				}
-				else
-				{
-					FVector2D DesiredSize = RowWidget->AsWidget()->GetDesiredSize();
+						FVector2D DesiredSize = RowWidget->AsWidget()->GetDesiredSize();
 
-					float Fractional = FMath::Fractional( NewScrollOffset );
-					if ( Fractional == 0 )
-					{
-						Fractional = 1.0f;
-						--NewScrollOffset;
-					}
+						float Fractional = FMath::Fractional( NewScrollOffset );
+						if ( Fractional == 0 )
+						{
+							Fractional = 1.0f;
+							--NewScrollOffset;
+						}
 
-					const float PrecedingHeight = DesiredSize.Y * Fractional;
+						const float PrecedingHeight = DesiredSize.Y * Fractional;
 
-					if ( AbsScrollByAmount > PrecedingHeight )
-					{
-						if ( ItemIndex != 0 )
+						if ( AbsScrollByAmount > PrecedingHeight )
+						{
+							if ( ItemIndex != 0 )
+							{
+								AbsScrollByAmount -= PrecedingHeight;
+								NewScrollOffset -= FMath::Fractional( NewScrollOffset );
+								--ItemIndex;
+							}
+							else
+							{
+								NewScrollOffset = 0;
+								break;
+							}
+						} 
+						else if ( AbsScrollByAmount == PrecedingHeight )
 						{
-							AbsScrollByAmount -= PrecedingHeight;
 							NewScrollOffset -= FMath::Fractional( NewScrollOffset );
-							--ItemIndex;
+							break;
 						}
 						else
 						{
-							NewScrollOffset = 0;
+							NewScrollOffset = (int32)NewScrollOffset + ( ( PrecedingHeight - AbsScrollByAmount ) / DesiredSize.Y );
 							break;
 						}
-					} 
-					else if ( AbsScrollByAmount == PrecedingHeight )
-					{
-						NewScrollOffset -= FMath::Fractional( NewScrollOffset );
-						break;
-					}
-					else
-					{
-						NewScrollOffset = (int32)NewScrollOffset + ( ( PrecedingHeight - AbsScrollByAmount ) / DesiredSize.Y );
-						break;
 					}
 				}
 			}
-		}
 
-		return ScrollTo( NewScrollOffset );
+
+			return ScrollTo( NewScrollOffset );
+		}
 	}
 
 protected:
diff --git a/Engine/Source/Runtime/Slate/Public/Widgets/Views/STableViewBase.h b/Engine/Source/Runtime/Slate/Public/Widgets/Views/STableViewBase.h
--- a/Engine/Source/Runtime/Slate/Public/Widgets/Views/STableViewBase.h
+++ b/Engine/Source/Runtime/Slate/Public/Widgets/Views/STableViewBase.h
@@ -7,6 +7,12 @@
 class SListPanel;
 class SHeaderRow;
 
+enum class EAllowOverscroll
+{
+	Yes,
+	No
+};
+
 /**
  * Contains ListView functionality that does not depend on the type of data being observed by the ListView.
  */
@@ -88,10 +94,11 @@ protected:
 	 *
 	 * @param MyGeometry      The geometry of the ListView at the time
 	 * @param ScrollByAmount  The amount to scroll by in Slate Screen Units.
+	 * @param AllowOverscroll Should we allow scrolling past the beginning/end of the list?
 	 *
 	 * @return The amount actually scrolled in items
 	 */
-	virtual float ScrollBy( const FGeometry& MyGeometry, float ScrollByAmount );
+	virtual float ScrollBy( const FGeometry& MyGeometry, float ScrollByAmount, EAllowOverscroll AllowOverscroll );
 
 	/**
 	 * Scroll the view to an offset
@@ -231,6 +238,40 @@ protected:
 	/**	Whether the software cursor should be drawn in the viewport */
 	bool bShowSoftwareCursor;
 
+	struct SLATE_API FOverscroll
+	{
+	public:
+
+		FOverscroll( const float InMaxOverscroll );
+
+		/** @return The Amount actually scrolled */
+		float ScrollBy( float Delta );
+
+		/** How far the user scrolled above/below the beginning/end of the list. */
+		float GetOverscroll() const;
+
+		/** Ticks the overscroll manager so it can animate. */
+		void UpdateOverscroll( float InDeltaTime );
+
+		/**
+		 * Should ScrollDelta be applied to overscroll or to regular item scrolling.
+		 *
+		 * @param bIsAtStartOfList  Are we at the very beginning of the list (i.e. showing the first item at the top of the view)?
+		 * @param bIsAtEndOfList    Are we showing the last item on the screen completely?
+		 * @param ScrollDelta       How much the user is trying to scroll in Slate Units.
+		 *
+		 * @return true if the user's scrolling should be applied toward overscroll.
+		 */
+		bool ShouldApplyOverscroll( const bool bIsAtStartOfList, const bool bIsAtEndOfList, const float ScrollDelta ) const;
+
+	private:
+		/** How much we've over-scrolled above/below the beginning/end of the list. */
+		float OverscrollAmount;
+
+		/** The maximum amount that we can scroll past the edge. */
+		const float MaxOverscroll;
+	} Overscroll;
+
 private:
 	
 	/** When true, a refresh should occur the next tick */
diff --git a/Engine/Source/Runtime/Slate/Public/Widgets/Views/STileView.h b/Engine/Source/Runtime/Slate/Public/Widgets/Views/STileView.h
--- a/Engine/Source/Runtime/Slate/Public/Widgets/Views/STileView.h
+++ b/Engine/Source/Runtime/Slate/Public/Widgets/Views/STileView.h
@@ -240,12 +240,34 @@ public:
 
 protected:
 
-	virtual float ScrollBy( const FGeometry& MyGeometry, float ScrollByAmountInSlateUnits ) OVERRIDE
+	virtual float ScrollBy( const FGeometry& MyGeometry, float ScrollByAmountInSlateUnits, EAllowOverscroll AllowOverscroll ) OVERRIDE
 	{
-		const float ItemHeight = this->GetItemHeight();
-		const double NewScrollOffset = this->ScrollOffset + ( ( ScrollByAmountInSlateUnits * GetNumItemsWide() ) / ItemHeight );
+		// Working around a CLANG bug, where all the base class
+		// members require an explicit namespace resolution.
+		typedef STableViewBase S;
 
-		return this->ScrollTo( NewScrollOffset );
+		const bool bWholeListVisible = S::ScrollOffset == 0 && S::bWasAtEndOfList;
+		if (bWholeListVisible)
+		{
+			return 0;
+		}
+		else if ( AllowOverscroll == EAllowOverscroll::Yes && S::Overscroll.ShouldApplyOverscroll( S::ScrollOffset == 0, S::bWasAtEndOfList, ScrollByAmountInSlateUnits ) )
+		{
+			const float UnclampedScrollDelta = ScrollByAmountInSlateUnits / (10 * GetNumItemsWide());
+			const float ActuallyScrolledBy = S::Overscroll.ScrollBy( UnclampedScrollDelta );
+			if (ActuallyScrolledBy != 0.0f)
+			{
+				this->RequestListRefresh();
+			}
+			return ActuallyScrolledBy;
+		}
+		else
+		{
+			const float ItemHeight = this->GetItemHeight();
+			const double NewScrollOffset = this->ScrollOffset + ( ( ScrollByAmountInSlateUnits * GetNumItemsWide() ) / ItemHeight );
+
+			return this->ScrollTo( NewScrollOffset );
+		}
 	}
 
 	virtual int32 GetNumItemsWide() const OVERRIDE
