Touch event handling added to TableRow in order to support selection and scrolling on touch devices.

diff --git a/Engine/Source/Runtime/Slate/Public/Widgets/Views/STableRow.h b/Engine/Source/Runtime/Slate/Public/Widgets/Views/STableRow.h
--- a/Engine/Source/Runtime/Slate/Public/Widgets/Views/STableRow.h
+++ b/Engine/Source/Runtime/Slate/Public/Widgets/Views/STableRow.h
@@ -82,8 +82,10 @@ public:
 	 *
 	 * @param	InArgs	The declaration data for this widget
 	 */
-	 void Construct(const typename STableRow<ItemType>::FArguments& InArgs, const TSharedRef<STableViewBase>& InOwnerTableView)
+	void Construct(const typename STableRow<ItemType>::FArguments& InArgs, const TSharedRef<STableViewBase>& InOwnerTableView)
 	{
+		bProcessingSelectionTouch = false;
+
 		ConstructInternal(InArgs, InOwnerTableView);
 
 		ConstructChildren(
@@ -393,9 +395,77 @@ public:
 		return FReply::Unhandled();
 	}
 
+	virtual FReply OnTouchStarted( const FGeometry& MyGeometry, const FPointerEvent& InTouchEvent ) OVERRIDE
+	{
+		bProcessingSelectionTouch = true;
+
+		return
+			FReply::Handled()
+			// Drag detect because if this tap turns into a drag, we stop processing
+			// the selection touch.
+			.DetectDrag( SharedThis(this), EKeys::LeftMouseButton );
+	}
+
+	virtual FReply OnTouchEnded( const FGeometry& MyGeometry, const FPointerEvent& InTouchEvent ) OVERRIDE
+	{
+		if ( bProcessingSelectionTouch )
+		{
+			bProcessingSelectionTouch = false;
+			const TSharedPtr< ITypedTableView<ItemType> > OwnerWidget = OwnerTablePtr.Pin();
+			const ItemType* MyItem = OwnerWidget->Private_ItemFromWidget( this );
+
+			switch( OwnerWidget->Private_GetSelectionMode() )
+			{
+				default:
+				case ESelectionMode::None:
+					return FReply::Unhandled();
+				break;
+
+				case ESelectionMode::Single:
+				{
+					OwnerWidget->Private_ClearSelection();
+					OwnerWidget->Private_SetItemSelection( *MyItem, true, true );
+					OwnerWidget->Private_SignalSelectionChanged(ESelectInfo::OnMouseClick);
+					return FReply::Handled();
+				}
+				break;
+
+				case ESelectionMode::SingleToggle:
+				{
+					const bool bShouldBecomeSelected = !OwnerWidget->Private_IsItemSelected(*MyItem);
+					OwnerWidget->Private_ClearSelection();
+					OwnerWidget->Private_SetItemSelection( *MyItem, bShouldBecomeSelected, true );
+					OwnerWidget->Private_SignalSelectionChanged(ESelectInfo::OnMouseClick);
+				}
+				break;
+
+				case ESelectionMode::Multi:
+				{
+					const bool bShouldBecomeSelected = !OwnerWidget->Private_IsItemSelected(*MyItem);
+					OwnerWidget->Private_SetItemSelection( *MyItem, bShouldBecomeSelected, true );
+					OwnerWidget->Private_SignalSelectionChanged(ESelectInfo::OnMouseClick);
+				}
+				break;
+			}
+			
+			return FReply::Handled();
+		}
+		else
+		{
+			return FReply::Unhandled();
+		}
+	}
+
 	virtual FReply OnDragDetected( const FGeometry& MyGeometry, const FPointerEvent& MouseEvent ) OVERRIDE
 	{
-		if ( HasMouseCapture() && ChangedSelectionOnMouseDown )
+		if (bProcessingSelectionTouch)
+		{
+			// With touch input, dragging scrolls the list while selection requires a tap.
+			// If we are processing a touch and it turned into a drag; pass it on to the 
+			bProcessingSelectionTouch = false;
+			return FReply::Handled().CaptureMouse( OwnerTablePtr.Pin()->AsWidget() );
+		}
+		else if ( HasMouseCapture() && ChangedSelectionOnMouseDown )
 		{
 			TSharedPtr< ITypedTableView<ItemType> > OwnerWidget = OwnerTablePtr.Pin();
 			OwnerWidget->Private_SignalSelectionChanged(ESelectInfo::OnMouseClick);
@@ -678,6 +748,10 @@ protected:
 	TWeakPtr<SWidget> Content;
 
 	bool ChangedSelectionOnMouseDown;
+
+	/** Did the current a touch interaction start in this item?*/
+	bool bProcessingSelectionTouch;
+
 };
 
 
