SLATE/UMG - Exposing additional setters for several attributes on borders and buttons.  Removing some getters for attributes that shouldn't have been exposed.

diff --git a/Engine/Source/Runtime/Slate/Private/Widgets/Input/SButton.cpp b/Engine/Source/Runtime/Slate/Private/Widgets/Input/SButton.cpp
--- a/Engine/Source/Runtime/Slate/Private/Widgets/Input/SButton.cpp
+++ b/Engine/Source/Runtime/Slate/Private/Widgets/Input/SButton.cpp
@@ -47,15 +47,7 @@ void SButton::Construct( const FArguments& InArgs )
 
 	ContentPadding = InArgs._ContentPadding;
 
-	/* Get pointer to the button style */
-	const FButtonStyle* const ButtonStyle = InArgs._ButtonStyle;
-
-	NormalImage = &ButtonStyle->Normal;
-	HoverImage = &ButtonStyle->Hovered;
-	PressedImage = &ButtonStyle->Pressed;
-
-	BorderPadding = ButtonStyle->NormalPadding;
-	PressedBorderPadding = ButtonStyle->PressedPadding;
+	SetButtonStyle(InArgs._ButtonStyle);
 
 	bIsFocusable = InArgs._IsFocusable;
 
@@ -63,8 +55,8 @@ void SButton::Construct( const FArguments& InArgs )
 
 	ClickMethod = InArgs._ClickMethod;
 
-	HoveredSound = InArgs._HoveredSoundOverride.Get(ButtonStyle->HoveredSlateSound);
-	PressedSound = InArgs._PressedSoundOverride.Get(ButtonStyle->PressedSlateSound);
+	HoveredSound = InArgs._HoveredSoundOverride.Get(Style->HoveredSlateSound);
+	PressedSound = InArgs._PressedSoundOverride.Get(Style->PressedSlateSound);
 }
 
 FMargin SButton::GetCombinedPadding() const
@@ -256,4 +248,37 @@ FVector2D SButton::ComputeDesiredSize() const
 	{
 		return SBorder::ComputeDesiredSize();
 	}
-}
\ No newline at end of file
+}
+
+void SButton::SetContentPadding(const TAttribute<FMargin>& InContentPadding)
+{
+	ContentPadding = InContentPadding;
+}
+
+void SButton::SetHoveredSound(TOptional<FSlateSound> InHoveredSound)
+{
+	HoveredSound = InHoveredSound.Get(Style->HoveredSlateSound);
+}
+
+void SButton::SetPressedSound(TOptional<FSlateSound> InPressedSound)
+{
+	PressedSound = InPressedSound.Get(Style->PressedSlateSound);
+}
+
+void SButton::SetOnClicked(FOnClicked InOnClicked)
+{
+	OnClicked = InOnClicked;
+}
+
+void SButton::SetButtonStyle(const FButtonStyle* ButtonStyle)
+{
+	/* Get pointer to the button style */
+	Style = ButtonStyle;
+
+	NormalImage = &Style->Normal;
+	HoverImage = &Style->Hovered;
+	PressedImage = &Style->Pressed;
+
+	BorderPadding = Style->NormalPadding;
+	PressedBorderPadding = Style->PressedPadding;
+}
diff --git a/Engine/Source/Runtime/Slate/Private/Widgets/Layout/SBorder.cpp b/Engine/Source/Runtime/Slate/Private/Widgets/Layout/SBorder.cpp
--- a/Engine/Source/Runtime/Slate/Private/Widgets/Layout/SBorder.cpp
+++ b/Engine/Source/Runtime/Slate/Private/Widgets/Layout/SBorder.cpp
@@ -65,7 +65,6 @@ void SBorder::ClearContent()
 	ChildSlot = FSimpleSlot();
 }
 
-
 /**
  * The widget should respond by populating the OutDrawElements array with FDrawElements 
  * that represent it and any of its children.
@@ -207,3 +206,38 @@ FVector2D SBorder::ComputeDesiredSize() const
 {
 	return DesiredSizeScale.Get() * SCompoundWidget::ComputeDesiredSize();
 }
+
+void SBorder::SetBorderBackgroundColor(const TAttribute<FSlateColor>& InColorAndOpacity)
+{
+	BorderBackgroundColor = InColorAndOpacity;
+}
+
+void SBorder::SetDesiredSizeScale(const TAttribute<FVector2D>& InDesiredSizeScale)
+{
+	DesiredSizeScale = InDesiredSizeScale;
+}
+
+void SBorder::SetHAlign(EHorizontalAlignment HAlign)
+{
+	ChildSlot.HAlignment = HAlign;
+}
+
+void SBorder::SetVAlign(EVerticalAlignment VAlign)
+{
+	ChildSlot.VAlignment = VAlign;
+}
+
+void SBorder::SetPadding(const TAttribute<FMargin>& InPadding)
+{
+	ChildSlot.SlotPadding = InPadding;
+}
+
+void SBorder::SetShowEffectWhenDisabled(const TAttribute<bool>& InShowEffectWhenDisabled)
+{
+	ShowDisabledEffect = InShowEffectWhenDisabled;
+}
+
+void SBorder::SetBorderImage(const TAttribute<const FSlateBrush*>& InBorderImage)
+{
+	BorderImage = InBorderImage;
+}
diff --git a/Engine/Source/Runtime/Slate/Public/Widgets/Input/SButton.h b/Engine/Source/Runtime/Slate/Public/Widgets/Input/SButton.h
--- a/Engine/Source/Runtime/Slate/Public/Widgets/Input/SButton.h
+++ b/Engine/Source/Runtime/Slate/Public/Widgets/Input/SButton.h
@@ -114,17 +114,20 @@ public:
 	 */
 	void Construct( const FArguments& InArgs );
 
-	/** Gets the content padding of the button */
-	FMargin GetContentPadding() const
-	{
-		return ContentPadding.Get();
-	}
+	/** See ContentPadding attribute */
+	void SetContentPadding(const TAttribute<FMargin>& InContentPadding);
 
-	/** Sets the content padding of the button */
-	void SetContentPadding(const TAttribute<FMargin>& InContentPadding)
-	{
-		ContentPadding = InContentPadding;
-	}
+	/** See HoveredSound attribute */
+	void SetHoveredSound(TOptional<FSlateSound> InHoveredSound);
+
+	/** See PressedSound attribute */
+	void SetPressedSound(TOptional<FSlateSound> InPressedSound);
+
+	/** See OnClicked event */
+	void SetOnClicked(FOnClicked InOnClicked);
+
+	/** See ButtonStyle attribute */
+	void SetButtonStyle(const FButtonStyle* ButtonStyle);
 
 protected:
 	/** @return combines the user-specified margin and the button's internal margin. */
@@ -145,6 +148,9 @@ protected:
 	/** The delegate to execute when the button is clicked */
 	FOnClicked OnClicked;
 
+	/** Style resource for the button */
+	const FButtonStyle* Style;
+
 	/** Brush resource that represents a button */
 	const FSlateBrush* NormalImage;
 	/** Brush resource that represents a button when it is hovered */
diff --git a/Engine/Source/Runtime/Slate/Public/Widgets/Layout/SBorder.h b/Engine/Source/Runtime/Slate/Public/Widgets/Layout/SBorder.h
--- a/Engine/Source/Runtime/Slate/Public/Widgets/Layout/SBorder.h
+++ b/Engine/Source/Runtime/Slate/Public/Widgets/Layout/SBorder.h
@@ -86,29 +86,26 @@ public:
 	/** Clears out the content for the border */
 	void ClearContent();
 
-	/** Gets the color and opacity of the background image of this border. */
-	FSlateColor GetBorderBackgroundColor()
-	{
-		return BorderBackgroundColor.Get();
-	}
-
 	/** Sets the color and opacity of the background image of this border. */
-	void SetBorderBackgroundColor( const TAttribute<FSlateColor>& InColorAndOpacity )
-	{
-		BorderBackgroundColor = InColorAndOpacity;
-	}
+	void SetBorderBackgroundColor(const TAttribute<FSlateColor>& InColorAndOpacity);
 
 	/** Set the desired size scale multiplier */
-	FVector2D GetDesiredSizeScale()
-	{
-		return DesiredSizeScale.Get();
-	}
+	void SetDesiredSizeScale(const TAttribute<FVector2D>& InDesiredSizeScale);
+	
+	/** See HAlign argument */
+	void SetHAlign(EHorizontalAlignment HAlign);
 
-	/** Set the desired size scale multiplier */
-	void SetDesiredSizeScale( const TAttribute<FVector2D>& InDesiredSizeScale )
-	{
-		DesiredSizeScale = InDesiredSizeScale;
-	}
+	/** See VAlign argument */
+	void SetVAlign(EVerticalAlignment VAlign);
+
+	/** See Padding attribute */
+	void SetPadding(const TAttribute<FMargin>& InPadding);
+
+	/** See ShowEffectWhenDisabled attribute */
+	void SetShowEffectWhenDisabled(const TAttribute<bool>& InShowEffectWhenDisabled);
+
+	/** See BorderImage attribute */
+	void SetBorderImage(const TAttribute<const FSlateBrush*>& InBorderImage);
 
 public:
 	// SWidget interface
