Initial support for multiline editable text

TTP# 336464 - Editor: Finish the Multiline Editable Text Block

Most of this work was done by JustinS; I just finished off some remaining whitespace, wrapping, cursor movement, and scaling issues.

diff --git a/Engine/Source/Runtime/Core/Private/Internationalization/ICUUtilities.cpp b/Engine/Source/Runtime/Core/Private/Internationalization/ICUUtilities.cpp
--- a/Engine/Source/Runtime/Core/Private/Internationalization/ICUUtilities.cpp
+++ b/Engine/Source/Runtime/Core/Private/Internationalization/ICUUtilities.cpp
@@ -10,95 +10,135 @@ namespace ICUUtilities
 {
 	void Convert(const FString& Source, icu::UnicodeString& Destination, const bool ShouldNullTerminate)
 	{
-		const char* const SourceEncoding = FPlatformString::GetEncodingName();
+		if( Source.Len() )
+		{
+			const char* const SourceEncoding = FPlatformString::GetEncodingName();
 	
-		UErrorCode ICUStatus = U_ZERO_ERROR;
+			UErrorCode ICUStatus = U_ZERO_ERROR;
 		
-		UConverter* ICUConverter = ucnv_open(SourceEncoding, &(ICUStatus));
-		if( U_FAILURE(ICUStatus) ) 
+			UConverter* ICUConverter = ucnv_open(SourceEncoding, &ICUStatus);
+			if( U_SUCCESS(ICUStatus) ) 
+			{
+				// Get the internal buffer of the string, we're going to use it as scratch space
+				const int32_t DestinationCapacityUChars = Source.Len() * 2;
+				UChar* InternalStringBuffer = Destination.getBuffer(DestinationCapacityUChars);
+
+				// Perform the conversion into the string buffer
+				const int32_t SourceSizeBytes = Source.Len() * sizeof(TCHAR);
+				const int32_t DestinationLength = ucnv_toUChars(ICUConverter, InternalStringBuffer, DestinationCapacityUChars, reinterpret_cast<const char*>(*Source), SourceSizeBytes, &ICUStatus);
+
+				// Optionally null terminate the string
+				if( ShouldNullTerminate )
+				{
+					InternalStringBuffer[DestinationLength] = 0;
+				}
+
+				// Size it back down to the correct size and release our lock on the string buffer
+				Destination.releaseBuffer(DestinationLength);
+
+				ucnv_close(ICUConverter);
+			}
+
+			check( U_SUCCESS(ICUStatus) );
+		}
+		else
 		{
-			return;
+			Destination.remove();
 		}
-
-		int32_t DestinationCapacity = Source.Len() * 2;
-		UChar* ICUChars = new UChar[DestinationCapacity];
-		const char* NativeChars = reinterpret_cast<const char*>(*Source);
-		int32_t NativeCharsLen = Source.Len() * sizeof(TCHAR);
-		int32_t DestinationLength = ucnv_toUChars(ICUConverter, ICUChars, DestinationCapacity, NativeChars, NativeCharsLen, &(ICUStatus));
-		Destination.setTo(ICUChars, DestinationLength);
-		delete[] ICUChars;
-
-		ucnv_close(ICUConverter);
 	}
 
 	icu::UnicodeString Convert(const FString& Source, const bool ShouldNullTerminate)
 	{
-		const char* const SourceEncoding = FPlatformString::GetEncodingName();
-
-		UErrorCode ICUStatus = U_ZERO_ERROR;
+		icu::UnicodeString Destination;
+		Convert(Source, Destination, ShouldNullTerminate);
+		return Destination;
+	}
 
-		UConverter* ICUConverter = ucnv_open(SourceEncoding, &(ICUStatus));
-		if( U_FAILURE(ICUStatus) ) 
+	void Convert(const icu::UnicodeString& Source, FString& Destination)
+	{
+		if( Source.length() )
 		{
-			return icu::UnicodeString();
+			const char* const DestinationEncoding = FPlatformString::GetEncodingName();
+
+			UErrorCode ICUStatus = U_ZERO_ERROR;
+
+			UConverter* ICUConverter = ucnv_open(DestinationEncoding, &ICUStatus);
+			if( U_SUCCESS(ICUStatus) ) 
+			{
+				// Get the internal buffer of the string, we're going to use it as scratch space
+				TArray<TCHAR>& InternalStringBuffer = Destination.GetCharArray();
+				
+				// Work out the maximum size required and resize the buffer so it can hold enough data
+				const int32_t DestinationCapacityBytes = UCNV_GET_MAX_BYTES_FOR_STRING(Source.length(), ucnv_getMaxCharSize(ICUConverter));
+				const int32 DestinationCapacityTCHARs = DestinationCapacityBytes / sizeof(TCHAR);
+				InternalStringBuffer.SetNumUninitialized(DestinationCapacityTCHARs);
+
+				// Perform the conversion into the string buffer, and then null terminate the FString and size it back down to the correct size
+				const int32_t DestinationSizeBytes = ucnv_fromUChars(ICUConverter, reinterpret_cast<char*>(InternalStringBuffer.GetData()), DestinationCapacityBytes, Source.getBuffer(), Source.length(), &ICUStatus);
+				const int32 DestinationSizeTCHARs = DestinationSizeBytes / sizeof(TCHAR);
+				InternalStringBuffer[DestinationSizeTCHARs] = 0;
+				InternalStringBuffer.SetNumUninitialized(DestinationSizeTCHARs + 1); // the array size includes null
+
+				ucnv_close(ICUConverter);
+			}
+
+			check( U_SUCCESS(ICUStatus) );
 		}
+		else
+		{
+			Destination.Empty();
+		}
+	}
 
-		int32_t DestinationCapacity = Source.Len() * 2;
-		UChar* ICUChars = new UChar[DestinationCapacity];
-		const char* NativeChars = reinterpret_cast<const char*>(*Source);
-		int32_t NativeCharsLen = Source.Len() * sizeof(TCHAR);
-		int32_t DestinationLength = ucnv_toUChars(ICUConverter, ICUChars, DestinationCapacity, NativeChars, NativeCharsLen, &(ICUStatus));
-		icu::UnicodeString Destination(ICUChars, DestinationLength);
-		delete[] ICUChars;
-
-		ucnv_close(ICUConverter);
-
+	FString Convert(const icu::UnicodeString& Source)
+	{
+		FString Destination;
+		Convert(Source, Destination);
 		return Destination;
 	}
 
-	void Convert(const icu::UnicodeString& Source, FString& Destination)
+	void Convert(const TCHAR Source, UChar32& Destination)
 	{
-		const char* const DestinationEncoding = FPlatformString::GetEncodingName();
-
+		const char* const SourceEncoding = FPlatformString::GetEncodingName();
+	
 		UErrorCode ICUStatus = U_ZERO_ERROR;
-
-		UConverter* ICUConverter = ucnv_open(DestinationEncoding, &(ICUStatus));
+		
+		UConverter* ICUConverter = ucnv_open(SourceEncoding, &ICUStatus);
 		if( U_FAILURE(ICUStatus) ) 
 		{
 			return;
 		}
 
-		int32_t DestinationCapacity = UCNV_GET_MAX_BYTES_FOR_STRING(Source.length(), ucnv_getMaxCharSize(ICUConverter));
-		char* NativeChars = reinterpret_cast<char*>(malloc(DestinationCapacity));
-		int32_t DestinationLength = ucnv_fromUChars(ICUConverter, NativeChars, DestinationCapacity, Source.getBuffer(), Source.length(), &(ICUStatus)) / sizeof(TCHAR);
 
-		Destination = FString(DestinationLength, reinterpret_cast<TCHAR*>(NativeChars));
-		free(NativeChars);
+		const char* NativeChars = reinterpret_cast<const char*>(&Source);
+		const int32_t NativeCharsLen = sizeof(TCHAR);
+		Destination = ucnv_getNextUChar(ICUConverter, &NativeChars, NativeChars + NativeCharsLen, &ICUStatus);
 
 		ucnv_close(ICUConverter);
 	}
 
-	FString Convert(const icu::UnicodeString& Source)
+	UChar32 Convert(const TCHAR Source)
 	{
-		const char* const DestinationEncoding = FPlatformString::GetEncodingName();
-
-		UErrorCode ICUStatus = U_ZERO_ERROR;
-
-		UConverter* ICUConverter = ucnv_open(DestinationEncoding, &(ICUStatus));
-		if( U_FAILURE(ICUStatus) ) 
-		{
-			return FString();
-		}
-
-		int32_t DestinationCapacity = UCNV_GET_MAX_BYTES_FOR_STRING(Source.length(), ucnv_getMaxCharSize(ICUConverter));
-		char* NativeChars = reinterpret_cast<char*>(malloc(DestinationCapacity));
-		int32_t DestinationLength = ucnv_fromUChars(ICUConverter, NativeChars, DestinationCapacity, Source.getBuffer(), Source.length(), &(ICUStatus)) / sizeof(TCHAR);
-
-		FString Destination(DestinationLength, reinterpret_cast<TCHAR*>(NativeChars));
-		free(NativeChars);
+		UChar32 Destination = 0;
+		Convert(Source, Destination);
+		return Destination;
+	}
 
-		ucnv_close(ICUConverter);
+	void Convert(const UChar32 Source, TCHAR& Destination)
+	{
+		// ICU has ucnv_getNextUChar as an optimization for converting a single character, but it doesn't seem to have an inverse of this
+		// Because of this we're using the string variants instead (which will be slower)
+		icu::UnicodeString SourceString;
+		SourceString.setTo(Source);
+		FString DestinationString;
+		Convert(SourceString, DestinationString);
+		Destination = DestinationString.Len() ? DestinationString[0] : 0;
+	}
 
+	TCHAR Convert(const UChar32 Source)
+	{
+		TCHAR Destination = 0;
+		Convert(Source, Destination);
 		return Destination;
 	}
 }
diff --git a/Engine/Source/Runtime/Core/Private/Internationalization/ICUUtilities.h b/Engine/Source/Runtime/Core/Private/Internationalization/ICUUtilities.h
--- a/Engine/Source/Runtime/Core/Private/Internationalization/ICUUtilities.h
+++ b/Engine/Source/Runtime/Core/Private/Internationalization/ICUUtilities.h
@@ -11,5 +11,10 @@ namespace ICUUtilities
 	icu::UnicodeString Convert(const FString& Source, const bool ShouldNullTerminate = true);
 	void Convert(const icu::UnicodeString& Source, FString& Destination);
 	FString Convert(const icu::UnicodeString& Source);
+
+	void Convert(const TCHAR Source, UChar32& Destination);
+	UChar32 Convert(const TCHAR Source);
+	void Convert(const UChar32 Source, TCHAR& Destination);
+	TCHAR Convert(const UChar32 Source);
 }
 #endif
\ No newline at end of file
diff --git a/Engine/Source/Runtime/Core/Private/Internationalization/ICUWordBreakIterator.cpp b/Engine/Source/Runtime/Core/Private/Internationalization/ICUWordBreakIterator.cpp
new file mode 100644
--- /dev/null
+++ b/Engine/Source/Runtime/Core/Private/Internationalization/ICUWordBreakIterator.cpp
@@ -0,0 +1,119 @@
+// Copyright 1998-2014 Epic Games, Inc. All Rights Reserved.
+
+#include "CorePrivate.h"
+#include "WordBreakIterator.h"
+
+#if UE_ENABLE_ICU
+#include "ICUUtilities.h"
+#include <unicode/brkiter.h>
+
+namespace
+{
+	TSharedRef<icu::BreakIterator> CreateWordBreakIterator()
+	{
+		UErrorCode ICUStatus = U_ZERO_ERROR;
+		return MakeShareable( icu::BreakIterator::createWordInstance( icu::Locale::getDefault(), ICUStatus ) );
+	}
+}
+
+class FWordBreakIterator::FImplementation
+{
+public:
+	FImplementation(const TCHAR* const String, const int32 StringLength)
+		: ICUBreakIterator(CreateWordBreakIterator())
+	{
+		ICUUtilities::Convert(FString(String, StringLength), ICUString);
+		ICUBreakIterator->setText(ICUString);
+	}
+
+	int32 GetCurrentPosition() const
+	{
+		return ICUBreakIterator->current();
+	}
+
+	int32 ResetToBeginning()
+	{
+		return ICUBreakIterator->first();
+	}
+
+	int32 ResetToEnd()
+	{
+		return ICUBreakIterator->last();
+	}
+
+	int32 MoveToPrevious()
+	{
+		return ICUBreakIterator->previous();
+	}
+
+	int32 MoveToNext()
+	{
+		return ICUBreakIterator->next();
+	}
+
+	int32 MoveToCandidateBefore(const int32 Index)
+	{
+		return ICUBreakIterator->preceding(Index);
+	}
+
+	int32 MoveToCandidateAfter(const int32 Index)
+	{
+		return ICUBreakIterator->following(Index);
+	}
+
+private:
+	TSharedRef<icu::BreakIterator> ICUBreakIterator;
+	icu::UnicodeString ICUString;
+};
+
+FWordBreakIterator::FWordBreakIterator(const FText& Text)
+	: Implementation( new FImplementation( *Text.ToString(), Text.ToString().Len() ) )
+{
+}
+
+FWordBreakIterator::FWordBreakIterator(const FString& String)
+	: Implementation( new FImplementation( *String, String.Len() ) )
+{
+}
+
+FWordBreakIterator::FWordBreakIterator(const TCHAR* const String, const int32 StringLength)
+	: Implementation( new FImplementation( String, StringLength ) )
+{
+}
+
+int32 FWordBreakIterator::GetCurrentPosition() const
+{
+	return Implementation->GetCurrentPosition();
+}
+
+int32 FWordBreakIterator::ResetToBeginning()
+{
+	return Implementation->ResetToBeginning();
+}
+
+int32 FWordBreakIterator::ResetToEnd()
+{
+	return Implementation->ResetToEnd();
+}
+
+int32 FWordBreakIterator::MoveToPrevious()
+{
+	return Implementation->MoveToPrevious();
+}
+
+int32 FWordBreakIterator::MoveToNext()
+{
+	return Implementation->MoveToNext();
+}
+
+int32 FWordBreakIterator::MoveToCandidateBefore(const int32 Index)
+{
+	return Implementation->MoveToCandidateBefore(Index);
+}
+
+int32 FWordBreakIterator::MoveToCandidateAfter(const int32 Index)
+{
+	return Implementation->MoveToCandidateAfter(Index);
+}
+
+#endif
\ No newline at end of file
diff --git a/Engine/Source/Runtime/Core/Private/Internationalization/LegacyWordBreakIterator.cpp b/Engine/Source/Runtime/Core/Private/Internationalization/LegacyWordBreakIterator.cpp
new file mode 100644
--- /dev/null
+++ b/Engine/Source/Runtime/Core/Private/Internationalization/LegacyWordBreakIterator.cpp
@@ -0,0 +1,80 @@
+// Copyright 1998-2014 Epic Games, Inc. All Rights Reserved.
+
+#include "CorePrivate.h"
+#include "WordBreakIterator.h"
+
+#if !UE_ENABLE_ICU
+
+FWordBreakIterator::FWordBreakIterator(const FText& Text)
+	: String( Text.ToString() ), CurrentPosition(0)
+{
+}
+
+FWordBreakIterator::FWordBreakIterator(const FString& InString)
+	: String( InString ), CurrentPosition(0)
+{
+}
+
+FWordBreakIterator::FWordBreakIterator(const TCHAR* const InString, const int32 StringLength)
+	: String( StringLength, InString ), CurrentPosition(0)
+{
+}
+
+int32 FWordBreakIterator::GetCurrentPosition() const
+{
+	return CurrentPosition;
+}
+
+int32 FWordBreakIterator::ResetToBeginning()
+{
+	return CurrentPosition = 0;
+}
+
+int32 FWordBreakIterator::ResetToEnd()
+{
+	return CurrentPosition = String.Len();
+}
+
+int32 FWordBreakIterator::MoveToPrevious()
+{
+	return MoveToCandidateBefore(CurrentPosition);
+}
+
+int32 FWordBreakIterator::MoveToNext()
+{
+	return MoveToCandidateAfter(CurrentPosition);
+}
+
+int32 FWordBreakIterator::MoveToCandidateBefore(const int32 Index)
+{
+	// Can break between char and whitespace.
+	for(CurrentPosition = FMath::Clamp( Index, 0, String.Len() ) - 1; CurrentPosition >= 0 + 1; --CurrentPosition)
+	{
+		bool bPreviousCharWasWhitespace = FChar::IsWhitespace(String[CurrentPosition - 1]);
+		bool bCurrentCharIsWhiteSpace = FChar::IsWhitespace(String[CurrentPosition]);
+		if(bPreviousCharWasWhitespace != bCurrentCharIsWhiteSpace)
+		{
+			break;
+		}
+	}
+
+	return CurrentPosition >= Index ? -1 : CurrentPosition;
+}
+
+int32 FWordBreakIterator::MoveToCandidateAfter(const int32 Index)
+{
+	// Can break between char and whitespace.
+	for(CurrentPosition = FMath::Clamp( Index, 0, String.Len() ) + 1; CurrentPosition <= String.Len() - 1; ++CurrentPosition)
+	{
+		bool bPreviousCharWasWhitespace = FChar::IsWhitespace(String[CurrentPosition - 1]);
+		bool bCurrentCharIsWhiteSpace = FChar::IsWhitespace(String[CurrentPosition]);
+		if(bPreviousCharWasWhitespace != bCurrentCharIsWhiteSpace)
+		{
+			break;
+		}
+	}
+
+	return CurrentPosition <= Index ? -1 : CurrentPosition;
+}
+
+#endif
\ No newline at end of file
diff --git a/Engine/Source/Runtime/Core/Private/Internationalization/Text.cpp b/Engine/Source/Runtime/Core/Private/Internationalization/Text.cpp
--- a/Engine/Source/Runtime/Core/Private/Internationalization/Text.cpp
+++ b/Engine/Source/Runtime/Core/Private/Internationalization/Text.cpp
@@ -296,7 +296,20 @@ FText::FText( FString InSourceString, FString InNamespace, FString InKey, int32
 FText FText::TrimPreceding( const FText& InText )
 {
 	FString TrimmedString = InText.ToString();
-	TrimmedString.Trim();
+	{
+		int32 StartPos = 0;
+		while ( StartPos < TrimmedString.Len() )
+		{
+			if( !FText::IsWhitespace( TrimmedString[StartPos] ) )
+			{
+				break;
+			}
+
+			++StartPos;
+		}
+
+		TrimmedString = TrimmedString.Right( TrimmedString.Len() - StartPos );
+	}
 
 	FText NewText = FText( MoveTemp( TrimmedString ) );
 
@@ -318,7 +331,20 @@ FText FText::TrimPreceding( const FText& InText )
 FText FText::TrimTrailing( const FText& InText )
 {
 	FString TrimmedString = InText.ToString();
-	TrimmedString.TrimTrailing();
+	{
+		int32 EndPos = TrimmedString.Len() - 1;
+		while( EndPos >= 0 )
+		{
+			if( !FText::IsWhitespace( TrimmedString[EndPos] ) )
+			{
+				break;
+			}
+
+			EndPos--;
+		}
+
+		TrimmedString = TrimmedString.Left( EndPos + 1 );
+	}
 
 	FText NewText = FText( MoveTemp ( TrimmedString ) );
 
@@ -340,8 +366,31 @@ FText FText::TrimTrailing( const FText& InText )
 FText FText::TrimPrecedingAndTrailing( const FText& InText )
 {
 	FString TrimmedString = InText.ToString();
-	TrimmedString.Trim();
-	TrimmedString.TrimTrailing();
+	{
+		int32 StartPos = 0;
+		while ( StartPos < TrimmedString.Len() )
+		{
+			if( !FText::IsWhitespace( TrimmedString[StartPos] ) )
+			{
+				break;
+			}
+
+			++StartPos;
+		}
+
+		int32 EndPos = TrimmedString.Len() - 1;
+		while( EndPos >= 0 )
+		{
+			if( !FText::IsWhitespace( TrimmedString[EndPos] ) )
+			{
+				break;
+			}
+
+			--EndPos;
+		}
+
+		TrimmedString = TrimmedString.Mid( StartPos, TrimmedString.Len() - EndPos );
+	}
 
 	FText NewText = FText( MoveTemp( TrimmedString ) );
 
@@ -360,6 +409,16 @@ FText FText::TrimPrecedingAndTrailing( const FText& InText )
 	return NewText;
 }
 
+bool FText::IsWhitespace( const TCHAR Char )
+{
+#if UE_ENABLE_ICU
+	const UChar32 ICUChar = ICUUtilities::Convert(Char);
+	return u_isWhitespace(ICUChar) != 0;
+#else
+	return FChar::IsWhitespace(Char);
+#endif
+}
+
 FText FText::Format(const FText& Fmt,const FText& v1)
 {
 	FFormatOrderedArguments Arguments;
diff --git a/Engine/Source/Runtime/Core/Private/Internationalization/WordWrapper.cpp b/Engine/Source/Runtime/Core/Private/Internationalization/WordWrapper.cpp
--- a/Engine/Source/Runtime/Core/Private/Internationalization/WordWrapper.cpp
+++ b/Engine/Source/Runtime/Core/Private/Internationalization/WordWrapper.cpp
@@ -65,8 +65,7 @@ bool FWordWrapper::ProcessLine()
 			--BreakIndex;			
 		}
 
-		//@todo Use ICU's whitespace functionality [12/5/2013 justin.sargent]
-		while (TChar<TCHAR>::IsWhitespace(String[BreakIndex - 1]))
+		while (FText::IsWhitespace(String[BreakIndex - 1]))
 		{
 			--BreakIndex;
 		}
@@ -78,8 +77,7 @@ bool FWordWrapper::ProcessLine()
 			bHasAddedLine = true;
 		}
 
-		//@todo Use ICU's whitespace functionality [12/5/2013 justin.sargent]
-		while (TChar<TCHAR>::IsWhitespace(String[NextStartIndex]))
+		while (FText::IsWhitespace(String[NextStartIndex]))
 		{
 			++NextStartIndex;
 		}
diff --git a/Engine/Source/Runtime/Core/Private/Tests/Internationalization/TextTest.cpp b/Engine/Source/Runtime/Core/Private/Tests/Internationalization/TextTest.cpp
--- a/Engine/Source/Runtime/Core/Private/Tests/Internationalization/TextTest.cpp
+++ b/Engine/Source/Runtime/Core/Private/Tests/Internationalization/TextTest.cpp
@@ -2,6 +2,7 @@
 
 #include "CorePrivate.h"
 #include "AutomationTest.h"
+#include "ICUUtilities.h"
 
 // Disable optimization for TextTest as it compiles very slowly in development builds
 PRAGMA_DISABLE_OPTIMIZATION
@@ -575,6 +576,80 @@ bool FTextTest::RunTest (const FString& Parameters)
 	return true;
 }
 
+
+#if UE_ENABLE_ICU
+
+IMPLEMENT_SIMPLE_AUTOMATION_TEST(FICUTextTest, "Core.Misc.ICUText", EAutomationTestFlags::ATF_SmokeTest)
+
+bool FICUTextTest::RunTest (const FString& Parameters)
+{
+	// Test to make sure that ICUUtilities converts strings correctly
+
+	const FString SourceString(TEXT("This is a test"));
+	const FString SourceString2(TEXT("This is another test"));
+	icu::UnicodeString ICUString;
+	FString ConversionBackStr;
+
+	// ---------------------------------------------------------------------
+
+	ICUUtilities::Convert(SourceString, ICUString);
+	if (SourceString.Len() != ICUString.length())
+	{
+		AddError(FString::Printf(TEXT("icu::UnicodeString is the incorrect length (%d; expected %d)."), ICUString.length(), SourceString.Len()));
+	}
+
+	ICUUtilities::Convert(ICUString, ConversionBackStr);
+	if (ICUString.length() != ConversionBackStr.Len())
+	{
+		AddError(FString::Printf(TEXT("FString is the incorrect length (%d; expected %d)."), ConversionBackStr.Len(), ICUString.length()));
+	}
+	if (SourceString != ConversionBackStr)
+	{
+		AddError(FString::Printf(TEXT("FString is has the incorrect converted value ('%s'; expected '%s')."), *ConversionBackStr, *SourceString));
+	}
+
+	// ---------------------------------------------------------------------
+
+	ICUUtilities::Convert(SourceString2, ICUString);
+	if (SourceString2.Len() != ICUString.length())
+	{
+		AddError(FString::Printf(TEXT("icu::UnicodeString is the incorrect length (%d; expected %d)."), ICUString.length(), SourceString2.Len()));
+	}
+
+	ICUUtilities::Convert(ICUString, ConversionBackStr);
+	if (ICUString.length() != ConversionBackStr.Len())
+	{
+		AddError(FString::Printf(TEXT("FString is the incorrect length (%d; expected %d)."), ConversionBackStr.Len(), ICUString.length()));
+	}
+	if (SourceString2 != ConversionBackStr)
+	{
+		AddError(FString::Printf(TEXT("FString is has the incorrect converted value ('%s'; expected '%s')."), *ConversionBackStr, *SourceString2));
+	}
+
+	// ---------------------------------------------------------------------
+
+	ICUUtilities::Convert(SourceString, ICUString);
+	if (SourceString.Len() != ICUString.length())
+	{
+		AddError(FString::Printf(TEXT("icu::UnicodeString is the incorrect length (%d; expected %d)."), ICUString.length(), SourceString.Len()));
+	}
+
+	ICUUtilities::Convert(ICUString, ConversionBackStr);
+	if (ICUString.length() != ConversionBackStr.Len())
+	{
+		AddError(FString::Printf(TEXT("FString is the incorrect length (%d; expected %d)."), ConversionBackStr.Len(), ICUString.length()));
+	}
+	if (SourceString != ConversionBackStr)
+	{
+		AddError(FString::Printf(TEXT("FString is has the incorrect converted value ('%s'; expected '%s')."), *ConversionBackStr, *SourceString));
+	}
+
+	return true;
+}
+
+#endif // #if UE_ENABLE_ICU
+
+
 #undef LOCTEXT_NAMESPACE 
 
 
diff --git a/Engine/Source/Runtime/Core/Public/Containers/UnrealString.h b/Engine/Source/Runtime/Core/Public/Containers/UnrealString.h
--- a/Engine/Source/Runtime/Core/Public/Containers/UnrealString.h
+++ b/Engine/Source/Runtime/Core/Public/Containers/UnrealString.h
@@ -400,6 +400,38 @@ public:
 		return *this;
 	}
 
+	FORCEINLINE FString& Append(const FString& Text)
+	{
+		*this += Text;
+		return *this;
+	}
+
+	FString& Append(const TCHAR* Text, int32 Count)
+	{
+		CheckInvariants();
+
+		if (Count != 0)
+		{
+			// position to insert the character.  
+			// At the end of the string if we have existing characters, otherwise at the 0 position
+			int32 InsertIndex = (Data.Num() > 0) ? Data.Num() - 1 : 0;
+
+			// number of characters to add.  If we don't have any existing characters, 
+			// we'll need to append the terminating zero as well.
+			int32 FinalCount = (Data.Num() > 0) ? Count : Count + 1;
+
+			Data.AddUninitialized(FinalCount);
+
+			for (int32 Index = 0; Index < Count; Index++)
+			{
+				Data[InsertIndex + Index] = Text[Index];
+			}
+
+			Data[Data.Num() - 1] = 0;
+		}
+		return *this;
+	}
+
 	/**
 	 * Removes characters within the string.
 	 *
@@ -412,6 +444,30 @@ public:
 		Data.RemoveAt(Index, Count, bAllowShrinking);
 	}
 
+	FORCEINLINE void InsertAt(int32 Index, TCHAR Character)
+	{
+		if (Data.Num() == 0)
+		{
+			*this += Character;
+		}
+		else
+		{
+			Data.Insert(Character, Index);
+		}
+	}
+
+	FORCEINLINE void InsertAt(int32 Index, const FString& Characters)
+	{
+		if (Data.Num() == 0)
+		{
+			*this += Characters;
+		}
+		else
+		{
+			Data.Insert(Characters.Data, Index);
+		}
+	}
+
 	/**
 	 * Removes the text from the start of the string if it exists.
 	 *
diff --git a/Engine/Source/Runtime/Core/Public/Internationalization/Text.h b/Engine/Source/Runtime/Core/Public/Internationalization/Text.h
--- a/Engine/Source/Runtime/Core/Public/Internationalization/Text.h
+++ b/Engine/Source/Runtime/Core/Public/Internationalization/Text.h
@@ -232,6 +232,11 @@ public:
 	 */
 	static FText TrimPrecedingAndTrailing( const FText& );
 
+	/**
+	 * Check to see if the given character is considered whitespace by the current culture
+	 */
+	static bool IsWhitespace( const TCHAR Char );
+
 	static void GetFormatPatternParameters(const FText& Pattern, TArray<FString>& ParameterNames);
 
 	static FText Format(const FText& Pattern, const FFormatNamedArguments& Arguments);
diff --git a/Engine/Source/Runtime/Core/Public/Internationalization/WordBreakIterator.h b/Engine/Source/Runtime/Core/Public/Internationalization/WordBreakIterator.h
new file mode 100644
--- /dev/null
+++ b/Engine/Source/Runtime/Core/Public/Internationalization/WordBreakIterator.h
@@ -0,0 +1,28 @@
+// Copyright 1998-2014 Epic Games, Inc. All Rights Reserved.
+#pragma once
+
+class CORE_API FWordBreakIterator
+{
+public:
+	FWordBreakIterator(const FText& Text);
+	FWordBreakIterator(const FString& String);
+	FWordBreakIterator(const TCHAR* const String, const int32 StringLength);
+
+	int32 GetCurrentPosition() const;
+
+	int32 ResetToBeginning();
+	int32 ResetToEnd();
+
+	int32 MoveToPrevious();
+	int32 MoveToNext();
+	int32 MoveToCandidateBefore(const int32 Index);
+	int32 MoveToCandidateAfter(const int32 Index);
+private:
+#if UE_ENABLE_ICU
+	class FImplementation;
+	TSharedRef<FImplementation> Implementation;
+#else
+	const FString String;
+	int32 CurrentPosition;
+#endif
+};
\ No newline at end of file
diff --git a/Engine/Source/Runtime/Slate/Private/Framework/Text/SlateHyperlinkRun.cpp b/Engine/Source/Runtime/Slate/Private/Framework/Text/SlateHyperlinkRun.cpp
--- a/Engine/Source/Runtime/Slate/Private/Framework/Text/SlateHyperlinkRun.cpp
+++ b/Engine/Source/Runtime/Slate/Private/Framework/Text/SlateHyperlinkRun.cpp
@@ -18,6 +18,11 @@ FTextRange FSlateHyperlinkRun::GetTextRange() const
 	return Range;
 }
 
+void FSlateHyperlinkRun::SetTextRange( const FTextRange& Value )
+{
+	Range = Value;
+}
+
 int16 FSlateHyperlinkRun::GetBaseLine( float Scale ) const 
 {
 	const TSharedRef< FSlateFontMeasure > FontMeasure = FSlateApplication::Get().GetRenderer()->GetFontMeasureService();
@@ -41,7 +46,7 @@ FVector2D FSlateHyperlinkRun::Measure( int32 StartIndex, int32 EndIndex, float S
 	return FontMeasure->Measure( **Text, StartIndex, EndIndex, Style.TextStyle.Font, true, Scale );
 }
 
-uint8 FSlateHyperlinkRun::GetKerning( int32 CurrentIndex, float Scale ) const 
+int8 FSlateHyperlinkRun::GetKerning( int32 CurrentIndex, float Scale ) const 
 {
 	return 0;
 }
@@ -82,6 +87,40 @@ void FSlateHyperlinkRun::ArrangeChildren( const TSharedRef< ILayoutBlock >& Bloc
 	ArrangedChildren.AddWidget( AllottedGeometry.MakeChild( WidgetBlock->GetWidget(), Block->GetLocationOffset() * ( 1 / AllottedGeometry.Scale ), Block->GetSize() * ( 1 / AllottedGeometry.Scale ), 1.0f ) );
 }
 
+int32 FSlateHyperlinkRun::GetTextIndexAt( const TSharedRef< ILayoutBlock >& Block, const FVector2D& Location, float Scale ) const
+{
+	return INDEX_NONE;
+}
+
+FVector2D FSlateHyperlinkRun::GetLocationAt( const TSharedRef< ILayoutBlock >& Block, int32 Offset, float Scale ) const
+{
+	return Block->GetLocationOffset();
+}
+
+void FSlateHyperlinkRun::Move(const TSharedRef<FString>& NewText, const FTextRange& NewRange)
+{
+	Text = NewText;
+	Range = NewRange;
+}
+
+TSharedRef<IRun> FSlateHyperlinkRun::Clone() const
+{
+	return FSlateHyperlinkRun::Create(Text, Style, Metadata, NavigateDelegate, Range);
+}
+
+void FSlateHyperlinkRun::AppendText(FString& AppendToText) const
+{
+	AppendToText.Append(**Text + Range.BeginIndex, Range.Len());
+}
+
+void FSlateHyperlinkRun::AppendText(FString& AppendToText, const FTextRange& PartialRange) const
+{
+	check(Range.BeginIndex <= PartialRange.BeginIndex);
+	check(Range.EndIndex >= PartialRange.EndIndex);
+
+	AppendToText.Append(**Text + PartialRange.BeginIndex, PartialRange.Len());
+}
+
 FSlateHyperlinkRun::FSlateHyperlinkRun( const TSharedRef< const FString >& InText, const FHyperlinkStyle& InStyle, const FMetadata& InMetadata, FOnClick InNavigateDelegate ) 
 	: Text( InText )
 	, Range( 0, Text->Len() )
diff --git a/Engine/Source/Runtime/Slate/Private/Framework/Text/SlateImageRun.cpp b/Engine/Source/Runtime/Slate/Private/Framework/Text/SlateImageRun.cpp
--- a/Engine/Source/Runtime/Slate/Private/Framework/Text/SlateImageRun.cpp
+++ b/Engine/Source/Runtime/Slate/Private/Framework/Text/SlateImageRun.cpp
@@ -25,7 +25,8 @@ TSharedRef< FSlateImageRun > FSlateImageRun::Create( const TSharedRef< const FSt
 	return MakeShareable( new FSlateImageRun( InText, InImage, InBaseline, InRange ) );
 }
 
-FSlateImageRun::FSlateImageRun( const TSharedRef< const FString >& InText, const FSlateBrush* InImage, int16 InBaseline, const FTextRange& InRange ) : Text( InText )
+FSlateImageRun::FSlateImageRun( const TSharedRef< const FString >& InText, const FSlateBrush* InImage, int16 InBaseline, const FTextRange& InRange ) 
+	: Text( InText )
 	, Range( InRange )
 	, Image( InImage )
 	, Baseline( InBaseline )
@@ -33,7 +34,8 @@ FSlateImageRun::FSlateImageRun( const TSharedRef< const FString >& InText, const
 	check( Image != NULL );
 }
 
-FSlateImageRun::FSlateImageRun( const TSharedRef< const FString >& InText, const FSlateBrush* InImage, int16 InBaseline ) : Text( InText )
+FSlateImageRun::FSlateImageRun( const TSharedRef< const FString >& InText, const FSlateBrush* InImage, int16 InBaseline ) 
+	: Text( InText )
 	, Range( 0, Text->Len() )
 	, Image( InImage )
 	, Baseline( InBaseline )
@@ -51,6 +53,16 @@ void FSlateImageRun::ArrangeChildren( const TSharedRef< ILayoutBlock >& Block, c
 	// no widgets
 }
 
+int32 FSlateImageRun::GetTextIndexAt( const TSharedRef< ILayoutBlock >& Block, const FVector2D& Location, float Scale ) const
+{
+	return INDEX_NONE;
+}
+
+FVector2D FSlateImageRun::GetLocationAt( const TSharedRef< ILayoutBlock >& Block, int32 Offset, float Scale ) const
+{
+	return Block->GetLocationOffset();
+}
+
 int32 FSlateImageRun::OnPaint( const FTextLayout::FLineView& Line, const TSharedRef< ILayoutBlock >& Block, const FTextBlockStyle& DefaultStyle, const FGeometry& AllottedGeometry, const FSlateRect& MyClippingRect, FSlateWindowElementList& OutDrawElements, int32 LayerId, const FWidgetStyle& InWidgetStyle, bool bParentEnabled ) const 
 {
 	if ( Image->DrawAs != ESlateBrushDrawType::NoDrawType )
@@ -68,7 +80,7 @@ TSharedRef< ILayoutBlock > FSlateImageRun::CreateBlock( int32 BeginIndex, int32
 	return FDefaultLayoutBlock::Create( SharedThis( this ), FTextRange( BeginIndex, EndIndex ), Size, Highlighter );
 }
 
-uint8 FSlateImageRun::GetKerning( int32 CurrentIndex, float Scale ) const 
+int8 FSlateImageRun::GetKerning( int32 CurrentIndex, float Scale ) const 
 {
 	return 0;
 }
@@ -94,3 +106,29 @@ FTextRange FSlateImageRun::GetTextRange() const
 	return Range;
 }
 
+void FSlateImageRun::SetTextRange( const FTextRange& Value )
+{
+	Range = Value;
+}
+
+void FSlateImageRun::Move(const TSharedRef<FString>& NewText, const FTextRange& NewRange)
+{
+	Text = NewText;
+	Range = NewRange;
+}
+
+TSharedRef<IRun> FSlateImageRun::Clone() const
+{
+	return FSlateImageRun::Create(Text, Image, Baseline, Range);
+}
+
+void FSlateImageRun::AppendText(FString& Text) const
+{
+	// Do nothing
+}
+
+void FSlateImageRun::AppendText(FString& AppendToText, const FTextRange& PartialRange) const
+{
+	check(Range.BeginIndex <= PartialRange.BeginIndex);
+	check(Range.EndIndex >= PartialRange.EndIndex);
+}
\ No newline at end of file
diff --git a/Engine/Source/Runtime/Slate/Private/Framework/Text/SlateTextRun.cpp b/Engine/Source/Runtime/Slate/Private/Framework/Text/SlateTextRun.cpp
--- a/Engine/Source/Runtime/Slate/Private/Framework/Text/SlateTextRun.cpp
+++ b/Engine/Source/Runtime/Slate/Private/Framework/Text/SlateTextRun.cpp
@@ -20,6 +20,11 @@ FTextRange FSlateTextRun::GetTextRange() const
 	return Range;
 }
 
+void FSlateTextRun::SetTextRange( const FTextRange& Value )
+{
+	Range = Value;
+}
+
 int16 FSlateTextRun::GetBaseLine( float Scale ) const 
 {
 	const TSharedRef< FSlateFontMeasure > FontMeasure = FSlateApplication::Get().GetRenderer()->GetFontMeasureService();
@@ -43,7 +48,7 @@ FVector2D FSlateTextRun::Measure( int32 BeginIndex, int32 EndIndex, float Scale
 	return FontMeasure->Measure( **Text, BeginIndex, EndIndex, Style.Font, true, Scale );
 }
 
-uint8 FSlateTextRun::GetKerning( int32 CurrentIndex, float Scale ) const 
+int8 FSlateTextRun::GetKerning(int32 CurrentIndex, float Scale) const
 {
 	const int32 PreviousIndex = CurrentIndex - 1;
 	if ( PreviousIndex < 0 )
@@ -113,7 +118,75 @@ void FSlateTextRun::ArrangeChildren( const TSharedRef< ILayoutBlock >& Block, co
 	// no widgets
 }
 
-FSlateTextRun::FSlateTextRun( const TSharedRef< const FString >& InText, const FTextBlockStyle& InStyle ) : Text( InText )
+int32 FSlateTextRun::GetTextIndexAt( const TSharedRef< ILayoutBlock >& Block, const FVector2D& Location, float Scale ) const
+{
+	const FVector2D& BlockOffset = Block->GetLocationOffset();
+	const FVector2D& BlockSize = Block->GetSize();
+
+	const float Left = BlockOffset.X;
+	const float Top = BlockOffset.Y;
+	const float Right = BlockOffset.X + BlockSize.X;
+	const float Bottom = BlockOffset.Y + BlockSize.Y;
+
+	const bool ContainsPoint = Location.X >= Left && Location.X < Right && Location.Y >= Top && Location.Y < Bottom;
+
+	if ( !ContainsPoint )
+	{
+		return INDEX_NONE;
+	}
+
+	const FTextRange BlockRange = Block->GetTextRange();
+	const TSharedRef< FSlateFontMeasure > FontMeasure = FSlateApplication::Get().GetRenderer()->GetFontMeasureService();
+
+	return FontMeasure->FindCharacterIndexAtOffset( Text.Get(), BlockRange.BeginIndex, BlockRange.EndIndex, Style.Font, Location.X - BlockOffset.X, true, Scale );
+}
+
+FVector2D FSlateTextRun::GetLocationAt( const TSharedRef< ILayoutBlock >& Block, int32 Offset, float Scale ) const
+{
+	const FVector2D& BlockOffset = Block->GetLocationOffset();
+	const FTextRange& BlockRange = Block->GetTextRange();
+
+	if (BlockRange.BeginIndex + Offset == BlockRange.EndIndex)
+	{
+		return BlockOffset + Block->GetSize();
+	}
+
+	const TSharedRef< FSlateFontMeasure > FontMeasure = FSlateApplication::Get().GetRenderer()->GetFontMeasureService();
+	const FVector2D OffsetLocation = FontMeasure->Measure( Text.Get(), BlockRange.BeginIndex, BlockRange.BeginIndex + Offset, Style.Font, true, Scale );
+
+	return BlockOffset + OffsetLocation;
+}
+
+void FSlateTextRun::Move(const TSharedRef<FString>& NewText, const FTextRange& NewRange)
+{
+	Text = NewText;
+	Range = NewRange;
+
+#if TEXT_LAYOUT_DEBUG
+	DebugSlice = FString( InRange.EndIndex - InRange.BeginIndex, (**Text) + InRange.BeginIndex );
+#endif
+}
+
+TSharedRef<IRun> FSlateTextRun::Clone() const
+{
+	return FSlateTextRun::Create(Text, Style, Range);
+}
+
+void FSlateTextRun::AppendText(FString& AppendToText) const
+{
+	AppendToText.Append(**Text + Range.BeginIndex, Range.Len());
+}
+
+void FSlateTextRun::AppendText(FString& AppendToText, const FTextRange& PartialRange) const
+{
+	check(Range.BeginIndex <= PartialRange.BeginIndex);
+	check(Range.EndIndex >= PartialRange.EndIndex);
+
+	AppendToText.Append(**Text + PartialRange.BeginIndex, PartialRange.Len());
+}
+
+FSlateTextRun::FSlateTextRun( const TSharedRef< const FString >& InText, const FTextBlockStyle& InStyle ) 
+	: Text( InText )
 	, Style( InStyle )
 	, Range( 0, Text->Len() )
 #if TEXT_LAYOUT_DEBUG
@@ -123,7 +196,8 @@ FSlateTextRun::FSlateTextRun( const TSharedRef< const FString >& InText, const F
 
 }
 
-FSlateTextRun::FSlateTextRun( const TSharedRef< const FString >& InText, const FTextBlockStyle& InStyle, const FTextRange& InRange ) : Text( InText )
+FSlateTextRun::FSlateTextRun( const TSharedRef< const FString >& InText, const FTextBlockStyle& InStyle, const FTextRange& InRange ) 
+	: Text( InText )
 	, Style( InStyle )
 	, Range( InRange )
 #if TEXT_LAYOUT_DEBUG
@@ -133,7 +207,8 @@ FSlateTextRun::FSlateTextRun( const TSharedRef< const FString >& InText, const F
 
 }
 
-FSlateTextRun::FSlateTextRun( const FSlateTextRun& Run ) : Text( Run.Text )
+FSlateTextRun::FSlateTextRun( const FSlateTextRun& Run ) 
+	: Text( Run.Text )
 	, Style( Run.Style )
 	, Range( Run.Range )
 #if TEXT_LAYOUT_DEBUG
diff --git a/Engine/Source/Runtime/Slate/Private/Framework/Text/SlateWidgetRun.cpp b/Engine/Source/Runtime/Slate/Private/Framework/Text/SlateWidgetRun.cpp
--- a/Engine/Source/Runtime/Slate/Private/Framework/Text/SlateWidgetRun.cpp
+++ b/Engine/Source/Runtime/Slate/Private/Framework/Text/SlateWidgetRun.cpp
@@ -18,6 +18,11 @@ FTextRange FSlateWidgetRun::GetTextRange() const
 	return Range;
 }
 
+void FSlateWidgetRun::SetTextRange( const FTextRange& Value )
+{
+	Range = Value;
+}
+
 int16 FSlateWidgetRun::GetBaseLine( float Scale ) const 
 {
 	return Info.Baseline * Scale;
@@ -33,7 +38,7 @@ FVector2D FSlateWidgetRun::Measure( int32 StartIndex, int32 EndIndex, float Scal
 	return Info.Size.Get( Info.Widget->GetDesiredSize() ) * Scale;
 }
 
-uint8 FSlateWidgetRun::GetKerning( int32 CurrentIndex, float Scale ) const 
+int8 FSlateWidgetRun::GetKerning( int32 CurrentIndex, float Scale ) const 
 {
 	return 0;
 }
@@ -59,6 +64,38 @@ void FSlateWidgetRun::ArrangeChildren( const TSharedRef< ILayoutBlock >& Block,
 	ArrangedChildren.AddWidget( AllottedGeometry.MakeChild( Info.Widget, Block->GetLocationOffset() * ( 1 / AllottedGeometry.Scale ), Block->GetSize() * ( 1 / AllottedGeometry.Scale ), 1.0f ) );
 }
 
+int32 FSlateWidgetRun::GetTextIndexAt( const TSharedRef< ILayoutBlock >& Block, const FVector2D& Location, float Scale ) const
+{
+	return INDEX_NONE;
+}
+
+FVector2D FSlateWidgetRun::GetLocationAt( const TSharedRef< ILayoutBlock >& Block, int32 Offset, float Scale ) const
+{
+	return Block->GetLocationOffset();
+}
+
+void FSlateWidgetRun::Move(const TSharedRef<FString>& NewText, const FTextRange& NewRange)
+{
+	Text = NewText;
+	Range = NewRange;
+}
+
+TSharedRef<IRun> FSlateWidgetRun::Clone() const
+{
+	return FSlateWidgetRun::Create(Text, Info, Range);
+}
+
+void FSlateWidgetRun::AppendText(FString& AppendToText) const
+{
+	//Do nothing
+}
+
+void FSlateWidgetRun::AppendText(FString& AppendToText, const FTextRange& PartialRange) const
+{
+	check(Range.BeginIndex <= PartialRange.BeginIndex);
+	check(Range.EndIndex >= PartialRange.EndIndex);
+}
+
 FSlateWidgetRun::FSlateWidgetRun( const TSharedRef< const FString >& InText, const FWidgetRunInfo& InWidgetInfo ) 
 	: Text( InText )
 	, Range( 0, Text->Len() )
diff --git a/Engine/Source/Runtime/Slate/Private/Framework/Text/SlateWordWrapper.cpp b/Engine/Source/Runtime/Slate/Private/Framework/Text/SlateWordWrapper.cpp
--- a/Engine/Source/Runtime/Slate/Private/Framework/Text/SlateWordWrapper.cpp
+++ b/Engine/Source/Runtime/Slate/Private/Framework/Text/SlateWordWrapper.cpp
@@ -3,6 +3,31 @@
 #include "SlatePrivatePCH.h"
 #include "SlateWordWrapper.h"
 
+/** Information about the substring as it would appear as a result of soft-wrapping a line of text. */
+struct FWrappedStringSlice
+{
+	/** 
+		* Create a substring info that sources InSourceString, begins at InFirstCharIndex and ends at InLastCharIndex (exclusive).
+		* The substring measures InSize.X by InSize.Y.
+		*/
+	FWrappedStringSlice( FString& InSourceString, int32 InFirstCharIndex, int32 InLastCharIndex, FVector2D InSize )
+	: SourceString( InSourceString )
+	, FirstCharIndex( InFirstCharIndex )
+	, LastCharIndex( InLastCharIndex )
+	, Size( InSize )
+	{
+	}
+
+	/** The substring sources this string. */
+	FString& SourceString;
+	/** Index of the first character in the substring */
+	int32 FirstCharIndex;
+	/** Index of the last character (exclusive) */
+	int32 LastCharIndex;
+	/** The measured size of this substring given the current font and point size. */
+	FVector2D Size;
+};
+
 namespace
 {
 	class FSlateWordWrapperBase : public FWordWrapper
diff --git a/Engine/Source/Runtime/Slate/Private/Framework/Text/TextEditHelper.cpp b/Engine/Source/Runtime/Slate/Private/Framework/Text/TextEditHelper.cpp
--- a/Engine/Source/Runtime/Slate/Private/Framework/Text/TextEditHelper.cpp
+++ b/Engine/Source/Runtime/Slate/Private/Framework/Text/TextEditHelper.cpp
@@ -9,21 +9,21 @@
  */
 struct FScopedTextTransaction
 {
-	explicit FScopedTextTransaction( ITextEditorWidget& InTextEditor )
+	explicit FScopedTextTransaction( const TSharedRef< ITextEditorWidget >& InTextEditor )
 	: TextEditor( InTextEditor )
 	{
-		TextEditor.StartChangingText();
+		TextEditor->StartChangingText();
 	}
 
 	~FScopedTextTransaction()
 	{
-		TextEditor.FinishChangingText();
+		TextEditor->FinishChangingText();
 	}
 
-	ITextEditorWidget& TextEditor;
+	const TSharedRef< ITextEditorWidget > TextEditor;
 };
 
-FReply FTextEditHelper::OnKeyChar( const FGeometry& MyGeometry, const FCharacterEvent& InCharacterEvent, class ITextEditorWidget& TextEditor )
+FReply FTextEditHelper::OnKeyChar( const FCharacterEvent& InCharacterEvent, const TSharedRef< ITextEditorWidget >& TextEditor )
 {
 	FReply Reply = FReply::Unhandled();
 
@@ -35,11 +35,11 @@ FReply FTextEditHelper::OnKeyChar( const FGeometry& MyGeometry, const FCharacter
 		// Backspace
 	case TCHAR( 8 ):
 		{
-			if( !TextEditor.GetIsReadOnly() )
+			if( !TextEditor->GetIsReadOnly() )
 			{
 				FScopedTextTransaction TextTransaction(TextEditor);
 
-				TextEditor.BackspaceChar();
+				TextEditor->BackspaceChar();
 				
 				Reply = FReply::Handled();
 			}
@@ -73,10 +73,10 @@ FReply FTextEditHelper::OnKeyChar( const FGeometry& MyGeometry, const FCharacter
 	default:
 		{
 			// Type the character, but only if it is allowed.
-			if (!TextEditor.GetIsReadOnly() && TextEditor.CanTypeCharacter(Character))
+			if (!TextEditor->GetIsReadOnly() && TextEditor->CanTypeCharacter(Character))
 			{
 				FScopedTextTransaction TextTransaction(TextEditor);
-				TextEditor.TypeChar( Character );
+				TextEditor->TypeChar( Character );
 				Reply = FReply::Handled();
 			}
 			else
@@ -92,48 +92,48 @@ FReply FTextEditHelper::OnKeyChar( const FGeometry& MyGeometry, const FCharacter
 	return Reply;
 }
 
-FReply FTextEditHelper::OnKeyDown( const FGeometry& MyGeometry, const FKeyboardEvent& InKeyboardEvent, ITextEditorWidget& TextEditor )
+FReply FTextEditHelper::OnKeyDown( const FKeyboardEvent& InKeyboardEvent, const TSharedRef< ITextEditorWidget >& TextEditor )
 {
 	const FKey Key = InKeyboardEvent.GetKey();
 
 	if( Key == EKeys::Left )
 	{
-		return TextEditor.MoveCursor(
+		return TextEditor->MoveCursor(
 			// Ctrl moves a whole word instead of one character.	
 			InKeyboardEvent.IsControlDown() ? ECursorMoveMethod::Word : ECursorMoveMethod::CharacterHorizontal,
 			// Move left
-			-1,
+			FVector2D( -1, 0 ),
 			// Shift selects text.	
 			InKeyboardEvent.IsShiftDown() ? ECursorAction::SelectText : ECursorAction::MoveCursor
 		);
 	}
 	else if( Key == EKeys::Right )
 	{
-		return TextEditor.MoveCursor(
+		return TextEditor->MoveCursor(
 			// Ctrl moves a whole word instead of one character.	
 			InKeyboardEvent.IsControlDown() ? ECursorMoveMethod::Word : ECursorMoveMethod::CharacterHorizontal,
 			// Move right
-			+1,
+			FVector2D( +1, 0 ),
 			// Shift selects text.	
 			InKeyboardEvent.IsShiftDown() ? ECursorAction::SelectText : ECursorAction::MoveCursor
 		);
 	}
 	else if ( Key == EKeys::Up )
 	{
-		return TextEditor.MoveCursor(
+		return TextEditor->MoveCursor(
 			ECursorMoveMethod::CharacterVertical,
 			// Move up
-			-1,
+			FVector2D( 0, -1 ),
 			// Shift selects text.	
 			InKeyboardEvent.IsShiftDown() ? ECursorAction::SelectText : ECursorAction::MoveCursor
 		);
 	}
 	else if ( Key == EKeys::Down )
 	{
-		return TextEditor.MoveCursor(
+		return TextEditor->MoveCursor(
 			ECursorMoveMethod::CharacterVertical,
 			// Move down
-			+1,
+			FVector2D( 0, +1 ),
 			// Shift selects text.	
 			InKeyboardEvent.IsShiftDown() ? ECursorAction::SelectText : ECursorAction::MoveCursor
 		);
@@ -141,7 +141,7 @@ FReply FTextEditHelper::OnKeyDown( const FGeometry& MyGeometry, const FKeyboardE
 	else if( Key == EKeys::Home )
 	{
 		// Go to the beginning of the document; select text if Shift is down.
-		TextEditor.JumpTo(
+		TextEditor->JumpTo(
 			(InKeyboardEvent.IsControlDown() ) ? ETextLocation::BeginningOfDocument : ETextLocation::BeginningOfLine,
 			(InKeyboardEvent.IsShiftDown()) ? ECursorAction::SelectText : ECursorAction::MoveCursor );
 
@@ -150,7 +150,7 @@ FReply FTextEditHelper::OnKeyDown( const FGeometry& MyGeometry, const FKeyboardE
 	else if ( Key == EKeys::End )
 	{
 		// Go to the end of the document; select text if Shift is down.
-		TextEditor.JumpTo(
+		TextEditor->JumpTo(
 			(InKeyboardEvent.IsControlDown() ) ? ETextLocation::EndOfDocument : ETextLocation::EndOfLine,
 			(InKeyboardEvent.IsShiftDown()) ? ECursorAction::SelectText : ECursorAction::MoveCursor );
 
@@ -158,44 +158,55 @@ FReply FTextEditHelper::OnKeyDown( const FGeometry& MyGeometry, const FKeyboardE
 	}
 	else if( Key == EKeys::Enter )
 	{
-		TextEditor.OnEnter();
+		FScopedTextTransaction TextTransaction(TextEditor);
+
+		TextEditor->OnEnter();
 
 		return FReply::Handled();
 	}
-	else if( Key == EKeys::Delete && !TextEditor.GetIsReadOnly() )
+	else if( Key == EKeys::Delete && !TextEditor->GetIsReadOnly() )
 	{
 		// @Todo: Slate keybindings support more than one set of keys. 
 		// Delete to next word boundary (Ctrl+Delete)
 		if (InKeyboardEvent.IsControlDown() && !InKeyboardEvent.IsAltDown() && !InKeyboardEvent.IsShiftDown())
 		{
-			TextEditor.MoveCursor(ECursorMoveMethod::Word, +1, ECursorAction::SelectText);
+			TextEditor->MoveCursor(
+				ECursorMoveMethod::Word, 
+				// Move right
+				FVector2D(+1, 0),
+				// selects text.	
+				ECursorAction::SelectText
+			);
 		}
 
-		TextEditor.DeleteChar();
+		FScopedTextTransaction TextTransaction(TextEditor);
+
+		// Delete selected text
+		TextEditor->DeleteChar();
 		
 		return FReply::Handled();
 	}	
 	else if( Key == EKeys::Escape )
 	{
-		return TextEditor.OnEscape();;
+		return TextEditor->OnEscape();;
 	}
 
 	// @Todo: Slate keybindings support more than one set of keys. 
 	//Alternate key for cut (Shift+Delete)
-	else if( Key == EKeys::Delete && InKeyboardEvent.IsShiftDown() && TextEditor.CanExecuteCut() )
+	else if( Key == EKeys::Delete && InKeyboardEvent.IsShiftDown() && TextEditor->CanExecuteCut() )
 	{
 		// Cut text to clipboard
-		TextEditor.CutSelectedTextToClipboard();
+		TextEditor->CutSelectedTextToClipboard();
 		
 		return FReply::Handled();
 	}
 
 	// @Todo: Slate keybindings support more than one set of keys. 
 	// Alternate key for copy (Ctrl+Insert) 
-	else if( Key == EKeys::Insert && InKeyboardEvent.IsControlDown() && TextEditor.CanExecuteCopy() ) 
+	else if( Key == EKeys::Insert && InKeyboardEvent.IsControlDown() && TextEditor->CanExecuteCopy() ) 
 	{
 		// Copy text to clipboard
-		TextEditor.CopySelectedTextToClipboard();
+		TextEditor->CopySelectedTextToClipboard();
 		
 		return FReply::Handled();
 	}
@@ -203,44 +214,49 @@ FReply FTextEditHelper::OnKeyDown( const FGeometry& MyGeometry, const FKeyboardE
 
 	// @Todo: Slate keybindings support more than one set of keys. 
 	// Alternate key for paste (Shift+Insert) 
-	else if( Key == EKeys::Insert && InKeyboardEvent.IsShiftDown() && TextEditor.CanExecutePaste() )
+	else if( Key == EKeys::Insert && InKeyboardEvent.IsShiftDown() && TextEditor->CanExecutePaste() )
 	{
 		// Paste text from clipboard
-		TextEditor.PasteTextFromClipboard();
+		TextEditor->PasteTextFromClipboard();
 		
 		return FReply::Handled();
 	}
 
 	// @Todo: Slate keybindings support more than one set of keys. 
 	//Alternate key for undo (Alt+Backspace)
-	else if( Key == EKeys::BackSpace && InKeyboardEvent.IsAltDown() && !InKeyboardEvent.IsShiftDown() && TextEditor.CanExecuteUndo() )
+	else if( Key == EKeys::BackSpace && InKeyboardEvent.IsAltDown() && !InKeyboardEvent.IsShiftDown() && TextEditor->CanExecuteUndo() )
 	{
 		// Undo
-		TextEditor.Undo();
+		TextEditor->Undo();
 		
 		return FReply::Handled();
 	}
 
 	// @Todo: Slate keybindings support more than one set of keys. 
 	// Delete to previous word boundary (Ctrl+Backspace)
-	else if( Key == EKeys::BackSpace && InKeyboardEvent.IsControlDown() && !InKeyboardEvent.IsAltDown() && !InKeyboardEvent.IsShiftDown() && !TextEditor.GetIsReadOnly() )
+	else if( Key == EKeys::BackSpace && InKeyboardEvent.IsControlDown() && !InKeyboardEvent.IsAltDown() && !InKeyboardEvent.IsShiftDown() && !TextEditor->GetIsReadOnly() )
 	{
 		FScopedTextTransaction TextTransaction(TextEditor);
 
-		TextEditor.MoveCursor(ECursorMoveMethod::Word, -1, ECursorAction::SelectText);
-		TextEditor.BackspaceChar();
+		TextEditor->MoveCursor(
+			ECursorMoveMethod::Word, 
+			// Move left
+			FVector2D(-1, 0), 
+			ECursorAction::SelectText
+		);
+		TextEditor->BackspaceChar();
 
 		return FReply::Handled();
 	}
 
 
 	// Ctrl+Y (or Ctrl+Shift+Z, or Alt+Shift+Backspace) to redo
-	else if( !TextEditor.GetIsReadOnly() && ( ( Key == EKeys::Y && InKeyboardEvent.IsControlDown() ) ||
+	else if( !TextEditor->GetIsReadOnly() && ( ( Key == EKeys::Y && InKeyboardEvent.IsControlDown() ) ||
 		( Key == EKeys::Z && InKeyboardEvent.IsControlDown() && InKeyboardEvent.IsShiftDown() ) ||
 		( Key == EKeys::BackSpace && InKeyboardEvent.IsAltDown() && InKeyboardEvent.IsShiftDown() ) ) )
 	{
 		// Redo
-		TextEditor.Redo();
+		TextEditor->Redo();
 		
 		return FReply::Handled();
 	}
@@ -256,4 +272,166 @@ FReply FTextEditHelper::OnKeyDown( const FGeometry& MyGeometry, const FKeyboardE
 	{
 		return FReply::Unhandled();
 	}
-}
\ No newline at end of file
+}
+FReply FTextEditHelper::OnMouseButtonDown( const FGeometry& MyGeometry, const FPointerEvent& InMouseEvent, const TSharedRef< ITextEditorWidget >& TextEditor )
+{
+	FReply Reply = FReply::Unhandled();
+
+	// If the mouse is already captured, then don't allow a new action to be taken
+	if( !TextEditor->GetWidget()->HasMouseCapture() )
+	{
+		if( InMouseEvent.GetEffectingButton() == EKeys::LeftMouseButton ||
+			InMouseEvent.GetEffectingButton() == EKeys::RightMouseButton )
+		{
+			// Am I getting focus right now?
+			const bool bIsGettingFocus = !TextEditor->GetWidget()->HasKeyboardFocus();
+			if( bIsGettingFocus )
+			{
+				// We might be receiving keyboard focus due to this event.  Because the keyboard focus received callback
+				// won't fire until after this function exits, we need to make sure our widget's state is in order early
+
+				// Assume we'll be given keyboard focus, so load text for editing
+				TextEditor->LoadText();
+
+				// Reset 'mouse has moved' state.  We'll use this in OnMouseMove to determine whether we
+				// should reset the selection range to the caret's position.
+				TextEditor->SetWasFocusedByLastMouseDown( true );
+			}
+
+			if( InMouseEvent.GetEffectingButton() == EKeys::LeftMouseButton )
+			{
+				if( InMouseEvent.IsShiftDown() )
+				{
+					TextEditor->MoveCursor( ECursorMoveMethod::ScreenPosition, MyGeometry.AbsoluteToLocal( InMouseEvent.GetScreenSpacePosition() ), ECursorAction::SelectText );
+				}
+				else
+				{
+					// Deselect any text that was selected
+					TextEditor->ClearSelection();
+					TextEditor->MoveCursor( ECursorMoveMethod::ScreenPosition, MyGeometry.AbsoluteToLocal( InMouseEvent.GetScreenSpacePosition() ), ECursorAction::MoveCursor );
+				}
+
+				// Start drag selection
+				TextEditor->BeginDragSelection();
+			}
+			else if( InMouseEvent.GetEffectingButton() == EKeys::RightMouseButton )
+			{
+				// If the user right clicked on a character that wasn't already selected, we'll clear
+				// the selection
+				if (TextEditor->AnyTextSelected() && !TextEditor->IsTextSelectedAt(MyGeometry.AbsoluteToLocal(InMouseEvent.GetScreenSpacePosition())))
+				{
+					// Deselect any text that was selected
+					TextEditor->ClearSelection();
+				}
+			}
+
+			// Right clicking to summon context menu, but we'll do that on mouse-up.
+			Reply = FReply::Handled();
+			Reply.CaptureMouse( TextEditor->GetWidget() );
+			Reply.SetKeyboardFocus( TextEditor->GetWidget(), EKeyboardFocusCause::Mouse );
+		}
+	}
+
+	return Reply;
+}
+
+FReply FTextEditHelper::OnMouseMove( const FGeometry& InMyGeometry, const FPointerEvent& InMouseEvent, const TSharedRef< ITextEditorWidget >& TextEditor )
+{
+	FReply Reply = FReply::Unhandled();
+
+	if( TextEditor->IsDragSelecting() && TextEditor->GetWidget()->HasMouseCapture() )
+	{
+		TextEditor->MoveCursor( ECursorMoveMethod::ScreenPosition, InMyGeometry.AbsoluteToLocal( InMouseEvent.GetScreenSpacePosition() ), ECursorAction::SelectText );
+		TextEditor->SetHasDragSelectedSinceFocused( true );
+		Reply = FReply::Handled();
+	}
+
+	return Reply;
+}
+
+FReply FTextEditHelper::OnMouseButtonUp( const FGeometry& MyGeometry, const FPointerEvent& InMouseEvent, const TSharedRef< ITextEditorWidget >& TextEditor )
+{
+	FReply Reply = FReply::Unhandled();
+
+	// The mouse must have been captured by either left or right button down before we'll process mouse ups
+	if( TextEditor->GetWidget()->HasMouseCapture() )
+	{
+		if( InMouseEvent.GetEffectingButton() == EKeys::LeftMouseButton &&
+			TextEditor->IsDragSelecting() )
+		{
+			// No longer drag-selecting
+			TextEditor->EndDragSelection();
+
+			// If we received keyboard focus on this click, then we'll want to select all of the text
+			// when the user releases the mouse button, unless the user actually dragged the mouse
+			// while holding the button down, in which case they've already selected some text and
+			// we'll leave things alone!
+			if( TextEditor->WasFocusedByLastMouseDown() )
+			{
+				if( !TextEditor->HasDragSelectedSinceFocused() )
+				{
+					if( TextEditor->SelectAllTextWhenFocused() )
+					{
+						// User wasn't dragging the mouse, so go ahead and select all of the text now
+						// that we've become focused
+						TextEditor->SelectAllText();
+
+						// Move the cursor to the end of the string
+						TextEditor->JumpTo( ETextLocation::EndOfDocument, ECursorAction::MoveCursor );
+
+						// @todo Slate: In this state, the caret should actually stay hidden (until the user interacts again), and we should not move the caret
+					}
+				}
+
+				TextEditor->SetWasFocusedByLastMouseDown( false );
+			}
+
+			// Release mouse capture
+			Reply = FReply::Handled();
+			Reply.ReleaseMouseCapture();
+		}
+		else if( InMouseEvent.GetEffectingButton() == EKeys::RightMouseButton )
+		{
+			if ( MyGeometry.IsUnderLocation( InMouseEvent.GetScreenSpacePosition() ) )
+			{
+				// Right clicked, so summon a context menu if the cursor is within the widget
+				TextEditor->SummonContextMenu( InMouseEvent.GetScreenSpacePosition() );
+			}
+
+			// Release mouse capture
+			Reply = FReply::Handled();
+			Reply.ReleaseMouseCapture();
+		}
+	}
+
+	return Reply;
+}
+
+FReply FTextEditHelper::OnMouseButtonDoubleClick(const FGeometry& InMyGeometry, const FPointerEvent& InMouseEvent, const TSharedRef< ITextEditorWidget >& TextEditor)
+{
+	FReply Reply = FReply::Unhandled();
+
+	if (InMouseEvent.GetEffectingButton() == EKeys::LeftMouseButton)
+	{
+		TextEditor->SelectWordAt(InMyGeometry.AbsoluteToLocal(InMouseEvent.GetScreenSpacePosition()));
+
+		Reply = FReply::Handled();
+	}
+
+	return Reply;
+}
+
+float FTextEditHelper::GetFontHeight(const FSlateFontInfo& FontInfo)
+{
+	const TSharedRef< FSlateFontMeasure > FontMeasure = FSlateApplication::Get().GetRenderer()->GetFontMeasureService();
+	return FontMeasure->GetMaxCharacterHeight(FontInfo);
+}
+
+float FTextEditHelper::CalculateCaretWidth(const float FontMaxCharHeight)
+{
+	// We adjust the width of the cursor because on smaller fonts the cursor would overlap too much of the characters it is next too.
+	// Other programs get around this by inverting the colors under the cursor. This causes a single character to be rendered as two characters.
+
+	// The caret with should be at least one pixel. For very small fonts the width could be < 1 which makes it not visible
+	return FMath::Max(1.0f, EditableTextDefs::CaretWidthPercent * FontMaxCharHeight);
+}
diff --git a/Engine/Source/Runtime/Slate/Private/Framework/Text/TextEditHelper.h b/Engine/Source/Runtime/Slate/Private/Framework/Text/TextEditHelper.h
--- a/Engine/Source/Runtime/Slate/Private/Framework/Text/TextEditHelper.h
+++ b/Engine/Source/Runtime/Slate/Private/Framework/Text/TextEditHelper.h
@@ -62,9 +62,30 @@ class FTextEditHelper
 public:
 
 	/** Someone typed a character, translate the input event into a call to a ITextEditorWidget action */
-	static FReply OnKeyChar( const FGeometry& MyGeometry, const FCharacterEvent& InCharacterEvent, class ITextEditorWidget& TextEditor );
+	static FReply OnKeyChar( const FCharacterEvent& InCharacterEvent, const TSharedRef< class ITextEditorWidget >& TextEditor );
 
 	/** Someone pressed a key, translate the input event into a call to a ITextEditorWidget action */
-	static FReply OnKeyDown( const FGeometry& MyGeometry, const FKeyboardEvent& InKeyboardEvent, ITextEditorWidget& TextEditor );
+	static FReply OnKeyDown( const FKeyboardEvent& InKeyboardEvent, const TSharedRef< ITextEditorWidget >& TextEditor );
 
-};
\ No newline at end of file
+	static FReply OnMouseButtonDown( const FGeometry& MyGeometry, const FPointerEvent& InMouseEvent, const TSharedRef< ITextEditorWidget >& TextEditor );
+
+	static FReply OnMouseMove( const FGeometry& InMyGeometry, const FPointerEvent& InMouseEvent, const TSharedRef< ITextEditorWidget >& TextEditor );
+
+	static FReply OnMouseButtonUp( const FGeometry& MyGeometry, const FPointerEvent& InMouseEvent, const TSharedRef< ITextEditorWidget >& TextEditor );
+
+	static FReply OnMouseButtonDoubleClick(const FGeometry& InMyGeometry, const FPointerEvent& InMouseEvent, const TSharedRef< ITextEditorWidget >& TextEditor);
+
+	/**
+	 * Gets the height of the largest character in the font
+	 *
+	 * @return  The fonts height
+	 */
+	static float GetFontHeight(const FSlateFontInfo& FontInfo);
+
+	/** 
+	 * Calculate the width of the caret 
+	 * @param FontMaxCharHeight The height of the font to calculate the caret width for
+	 * @return The width of the caret (might be clamped for very small fonts)
+	 */
+	static float CalculateCaretWidth(const float FontMaxCharHeight);
+};
diff --git a/Engine/Source/Runtime/Slate/Private/Framework/Text/TextLayout.cpp b/Engine/Source/Runtime/Slate/Private/Framework/Text/TextLayout.cpp
--- a/Engine/Source/Runtime/Slate/Private/Framework/Text/TextLayout.cpp
+++ b/Engine/Source/Runtime/Slate/Private/Framework/Text/TextLayout.cpp
@@ -1,6 +1,8 @@
 // Copyright 1998-2014 Epic Games, Inc. All Rights Reserved.
 
 #include "SlatePrivatePCH.h"
+#include "LineBreakIterator.h"
+#include "WordBreakIterator.h"
 
 FTextLayout::FBreakCandidate FTextLayout::CreateBreakCandidate( int32& OutRunIndex, FLineModel& Line, int32 PreviousBreak, int32 CurrentBreak )
 {
@@ -38,21 +40,32 @@ FTextLayout::FBreakCandidate FTextLayout::CreateBreakCandidate( int32& OutRunInd
 		WhitespaceStopIndex = StopIndex = FMath::Min( Range.EndIndex, CurrentBreak );
 		int32 BeginIndex = FMath::Max( PreviousBreak, Range.BeginIndex );
 
-		//@todo Use ICU's whitespace functionality [12/5/2013 justin.sargent]
-		while( WhitespaceStopIndex > BeginIndex && TChar<TCHAR>::IsWhitespace( (*Line.Text)[ WhitespaceStopIndex - 1 ] ) )
+		while( WhitespaceStopIndex > BeginIndex && FText::IsWhitespace( (*Line.Text)[ WhitespaceStopIndex - 1 ] ) )
 		{
 			--WhitespaceStopIndex;
 		}
 
-		if ( WhitespaceStopIndex != StopIndex && BeginIndex != WhitespaceStopIndex )
+		if ( BeginIndex == StopIndex )
 		{
+			// This slice is empty, no need to adjust anything
+			SliceSize = SliceSizeWithoutTrailingWhitespace = FVector2D::ZeroVector;
+		}
+		else if ( BeginIndex == WhitespaceStopIndex )
+		{
+			// This slice contains only whitespace, no need to adjust SliceSizeWithoutTrailingWhitespace
+			SliceSize = Run.Measure( BeginIndex, StopIndex, Scale );
+			SliceSizeWithoutTrailingWhitespace = FVector2D::ZeroVector;
+		}
+		else if ( WhitespaceStopIndex != StopIndex )
+		{
+			// This slice contains trailing whitespace, measure the text size, then add on the whitespace size
 			SliceSize = SliceSizeWithoutTrailingWhitespace = Run.Measure( BeginIndex, WhitespaceStopIndex, Scale );
 			SliceSize.X += Run.Measure( WhitespaceStopIndex, StopIndex, Scale ).X;
 		}
 		else
 		{
-			SliceSize = Run.Measure( BeginIndex, StopIndex, Scale );
-			SliceSizeWithoutTrailingWhitespace = SliceSize;
+			// This slice contains no whitespace, both sizes are the same and can use the same measurement
+			SliceSize = SliceSizeWithoutTrailingWhitespace = Run.Measure( BeginIndex, StopIndex, Scale );
 		}
 
 		BreakSize.X += SliceSize.X; // We accumulate the slice widths
@@ -81,10 +94,10 @@ FTextLayout::FBreakCandidate FTextLayout::CreateBreakCandidate( int32& OutRunInd
 	BreakSize.Y = BreakSizeWithoutTrailingWhitespace.Y = MaxAboveBaseline + MaxBelowBaseline;
 
 	FBreakCandidate BreakCandidate;
-	BreakCandidate.Size = BreakSize;
-	BreakCandidate.SizeWithoutTrailingWhitespace = BreakSizeWithoutTrailingWhitespace;
-	BreakCandidate.Range = FTextRange( PreviousBreak, CurrentBreak );
-	BreakCandidate.RangeWithoutTrailingWhitespace = FTextRange( PreviousBreak, WhitespaceStopIndex );
+	BreakCandidate.ActualSize = BreakSize;
+	BreakCandidate.TrimmedSize = BreakSizeWithoutTrailingWhitespace;
+	BreakCandidate.ActualRange = FTextRange( PreviousBreak, CurrentBreak );
+	BreakCandidate.TrimmedRange = FTextRange( PreviousBreak, WhitespaceStopIndex );
 	BreakCandidate.MaxAboveBaseline = MaxAboveBaseline;
 	BreakCandidate.MaxBelowBaseline = MaxBelowBaseline;
 	BreakCandidate.Kerning = Kerning;
@@ -96,15 +109,16 @@ FTextLayout::FBreakCandidate FTextLayout::CreateBreakCandidate( int32& OutRunInd
 	return BreakCandidate;
 }
 
-void FTextLayout::CreateLineViewBlocks( const FLineModel& Line, const int32 StopIndex, const int32 VisualStopIndex, int32& OutRunIndex, int32& OutHighlightIndex, int32& OutPreviousBlockEnd, TArray< TSharedRef< ILayoutBlock > >& OutSoftLine )
+void FTextLayout::CreateLineViewBlocks( int32 LineModelIndex, const int32 StopIndex, int32& OutRunIndex, int32& OutHighlightIndex, int32& OutPreviousBlockEnd, TArray< TSharedRef< ILayoutBlock > >& OutSoftLine )
 {
-	const bool VisualStopIndexIsValid = VisualStopIndex != INDEX_NONE;
+	const FLineModel& LineModel = LineModels[ LineModelIndex ];
+
 	int16 MaxAboveBaseline = 0;
 	int16 MaxBelowBaseline = 0;
 
-	for (; OutRunIndex < Line.Runs.Num(); )
+	for (; OutRunIndex < LineModel.Runs.Num(); )
 	{
-		const FRunModel& Run = Line.Runs[ OutRunIndex ];
+		const FRunModel& Run = LineModel.Runs[ OutRunIndex ];
 		const FTextRange RunRange = Run.GetTextRange();
 
 		TSharedPtr< IRunHighlighter > BlockHighlighter = NULL;
@@ -112,39 +126,54 @@ void FTextLayout::CreateLineViewBlocks( const FLineModel& Line, const int32 Stop
 		int32 BlockStopIndex = RunRange.EndIndex;
 		if ( OutHighlightIndex != INDEX_NONE )
 		{
-			const FTextHighlight& Highlight = Line.Highlights[ OutHighlightIndex ];
+			//Grab the currently active highlighter
+			const FTextHighlight& Highlight = LineModel.Highlights[OutHighlightIndex];
+
+			//Check to see if the last block was highlighted with the same highlighter
 			if ( OutPreviousBlockEnd >= Highlight.Range.BeginIndex )
 			{
+				//If the highlighter ends before our run...
 				if ( Highlight.Range.EndIndex <= RunRange.EndIndex )
 				{
+					//Adjust the stopping point of the block to be the end of the highlight,
+					//since highlights need their own block segments
 					BlockStopIndex = Highlight.Range.EndIndex;
 					BlockHighlighter = Highlight.Highlighter;
 				}
+				else
+				{
+					//This whole run is encompassed by the highlighter
+					BlockHighlighter = Highlight.Highlighter;
+				}
 			}
 			else
 			{
+				//Does the highlight begin before our run ends?
 				if ( Highlight.Range.BeginIndex <= RunRange.EndIndex )
 				{
+					// then adjust the current block stopping point to just before the highlight begins,
+					//since highlights need their own block segments
 					BlockStopIndex = Highlight.Range.BeginIndex;
 					BlockHighlighter = NULL;
 				}
 			}
 		}
 
-		if ( VisualStopIndex != INDEX_NONE )
+		if (StopIndex != INDEX_NONE)
 		{
-			BlockStopIndex = FMath::Min( VisualStopIndex, BlockStopIndex );
+			BlockStopIndex = FMath::Min(StopIndex, BlockStopIndex);
 		}
 
 		int32 BlockBeginIndex = FMath::Max( OutPreviousBlockEnd, RunRange.BeginIndex );
-		const bool IsLastBlock = BlockStopIndex == VisualStopIndex;
+		const bool IsLastBlock = BlockStopIndex == StopIndex;
 
+		// if this new block will be entirely encapsulated by the run...
 		if ( RunRange.BeginIndex < BlockStopIndex && RunRange.EndIndex > BlockBeginIndex )
 		{
 			check( BlockBeginIndex <= BlockStopIndex );
 
 			FBlockDefinition BlockDefine;
-			BlockDefine.Range = BlockDefine.VisualRange = FTextRange( BlockBeginIndex, BlockStopIndex );
+			BlockDefine.ActualRange = FTextRange(BlockBeginIndex, BlockStopIndex);
 			BlockDefine.Highlighter = BlockHighlighter;
 
 			OutSoftLine.Add( Run.CreateBlock( BlockDefine, Scale ) );
@@ -153,8 +182,7 @@ void FTextLayout::CreateLineViewBlocks( const FLineModel& Line, const int32 Stop
 		else
 		{
 			FBlockDefinition BlockDefine;
-			BlockDefine.Range = RunRange;
-			BlockDefine.VisualRange = FTextRange( RunRange.BeginIndex, ( IsLastBlock && VisualStopIndexIsValid ) ? VisualStopIndex : RunRange.EndIndex );
+			BlockDefine.ActualRange = RunRange;
 			BlockDefine.Highlighter = BlockHighlighter;
 
 			OutSoftLine.Add( Run.CreateBlock( BlockDefine, Scale ) );
@@ -173,11 +201,11 @@ void FTextLayout::CreateLineViewBlocks( const FLineModel& Line, const int32 Stop
 			++OutRunIndex;
 		}
 
-		if ( OutHighlightIndex != INDEX_NONE && BlockStopIndex == Line.Highlights[ OutHighlightIndex ].Range.EndIndex )
+		if ( OutHighlightIndex != INDEX_NONE && BlockStopIndex == LineModel.Highlights[ OutHighlightIndex ].Range.EndIndex )
 		{
 			++OutHighlightIndex;
 
-			if ( OutHighlightIndex >= Line.Highlights.Num() )
+			if ( OutHighlightIndex >= LineModel.Highlights.Num() )
 			{
 				OutHighlightIndex = INDEX_NONE;
 			}
@@ -189,17 +217,6 @@ void FTextLayout::CreateLineViewBlocks( const FLineModel& Line, const int32 Stop
 		}
 	}
 
-	//Skip any highlighted ranges which we have skipped due to ignoring trailing whitespace
-	while ( OutHighlightIndex != INDEX_NONE && Line.Highlights[ OutHighlightIndex ].Range.EndIndex <= StopIndex )
-	{
-		OutHighlightIndex++;
-
-		if ( OutHighlightIndex >= Line.Highlights.Num() )
-		{
-			OutHighlightIndex = INDEX_NONE;
-		}
-	}
-
 	FVector2D LineSize( ForceInitToZero );
 	FVector2D CurrentOffset( Margin.Left, Margin.Top + DrawSize.Y );
 
@@ -209,28 +226,35 @@ void FTextLayout::CreateLineViewBlocks( const FLineModel& Line, const int32 Stop
 		for (int Index = 0; Index < OutSoftLine.Num(); Index++)
 		{
 			const TSharedRef< ILayoutBlock > Block = OutSoftLine[ Index ];
-			const int16 BlockBaseline = Block->GetRun()->GetBaseLine( Scale );
+			const TSharedRef< IRun > Run = Block->GetRun();
+
+			const int16 BlockBaseline = Run->GetBaseLine(Scale);
 			const int16 VerticalOffset = MaxAboveBaseline - Block->GetSize().Y - BlockBaseline;
+			const int8 BlockKerning = Run->GetKerning(Block->GetTextRange().BeginIndex, Scale);
 
-			Block->SetLocationOffset( FVector2D( CurrentOffset.X + CurrentHorizontalPos, CurrentOffset.Y + VerticalOffset ) );
+			Block->SetLocationOffset(FVector2D(CurrentOffset.X + CurrentHorizontalPos + BlockKerning, CurrentOffset.Y + VerticalOffset));
 
 			CurrentHorizontalPos += Block->GetSize().X;
 		}
 
+		const float UnscaleLineHeight = MaxAboveBaseline + MaxBelowBaseline;
+
 		LineSize.X = CurrentHorizontalPos;
-		LineSize.Y = MaxAboveBaseline + MaxBelowBaseline;
-	}
+		LineSize.Y = UnscaleLineHeight * LineHeightPercentage;
 
-	FTextLayout::FLineView LineView;
-	LineView.Size = LineSize;
-	LineView.Offset = CurrentOffset;
-	LineView.Blocks.Append( OutSoftLine );
+		FTextLayout::FLineView LineView;
+		LineView.Offset = CurrentOffset;
+		LineView.Size = LineSize;
+		LineView.TextSize = FVector2D(CurrentHorizontalPos, UnscaleLineHeight);
+		LineView.Range = FTextRange(OutSoftLine[0]->GetTextRange().BeginIndex, OutSoftLine.Last()->GetTextRange().EndIndex);
+		LineView.ModelIndex = LineModelIndex;
+		LineView.Blocks.Append( OutSoftLine );
 
-	LineViews.Add( LineView );
+		LineViews.Add( LineView );
+	}
 
 	DrawSize.X = FMath::Max( DrawSize.X, LineSize.X ); //DrawSize.X is the size of the longest line + the Margin
-	const float LineHeight = LineSize.Y * LineHeightPercentage; 
-	DrawSize.Y += LineHeight; //DrawSize.Y is the total height of all lines
+	DrawSize.Y += LineSize.Y; //DrawSize.Y is the total height of all lines
 }
 
 void FTextLayout::JustifyLayout()
@@ -269,12 +293,9 @@ void FTextLayout::FlowLayout()
 {
 	check( WrappingWidth >= 0 );
 	const bool IsWrapping = WrappingWidth > 0;
-	const float WrappingDrawWidth = FMath::Max( 0.01f, ( WrappingWidth - ( Margin.Left + Margin.Right ) ) * Scale );
+	const float WrappingDrawWidth = FMath::Max( 0.01f, ( WrappingWidth - Margin.GetTotalSpaceAlong<Orient_Horizontal>() ) * Scale );
 
-	if ( IsWrapping )
-	{
-		ReadyForWrapping();
-	}
+	CreateWrappingCache();
 
 	TArray< TSharedRef< ILayoutBlock > > SoftLine;
 
@@ -296,7 +317,7 @@ void FTextLayout::FlowLayout()
 		if ( LineModel.BreakCandidates.Num() == 0 )
 		{
 			//Then iterate over all of it's runs
-			CreateLineViewBlocks( LineModel, INDEX_NONE, INDEX_NONE, /*OUT*/CurrentRunIndex, /*OUT*/CurrentHighlightIndex, /*OUT*/PreviousBlockEnd, SoftLine );
+			CreateLineViewBlocks( LineModelIndex, INDEX_NONE, /*OUT*/CurrentRunIndex, /*OUT*/CurrentHighlightIndex, /*OUT*/PreviousBlockEnd, SoftLine );
 			check( CurrentRunIndex == LineModel.Runs.Num() );
 			CurrentWidth = 0;
 			SoftLine.Empty();
@@ -310,30 +331,30 @@ void FTextLayout::FlowLayout()
 				const bool IsLastBreak = BreakIndex + 1 == LineModel.BreakCandidates.Num();
 				const bool IsFirstBreakOnSoftLine = CurrentWidth == 0;
 				const uint8 Kerning = ( IsFirstBreakOnSoftLine ) ? Break.Kerning : 0;
-				const bool BreakDoesFit = !IsWrapping || CurrentWidth + Break.Size.X + Kerning <= WrappingDrawWidth;
+				const bool BreakDoesFit = !IsWrapping || CurrentWidth + Break.ActualSize.X + Kerning <= WrappingDrawWidth;
 
 				if ( !BreakDoesFit || IsLastBreak )
 				{
-					const bool BreakWithoutTrailingWhitespaceDoesFit = !IsWrapping || CurrentWidth + Break.SizeWithoutTrailingWhitespace.X + Kerning <= WrappingDrawWidth;
+					const bool BreakWithoutTrailingWhitespaceDoesFit = !IsWrapping || CurrentWidth + Break.TrimmedSize.X + Kerning <= WrappingDrawWidth;
 					const bool IsFirstBreak = BreakIndex == 0;
 
 					const FBreakCandidate& FinalBreakOnSoftLine = ( !IsFirstBreak && !IsFirstBreakOnSoftLine && !BreakWithoutTrailingWhitespaceDoesFit ) ? LineModel.BreakCandidates[ --BreakIndex ] : Break;
 
-					CreateLineViewBlocks( LineModel, FinalBreakOnSoftLine.Range.EndIndex, FinalBreakOnSoftLine.RangeWithoutTrailingWhitespace.EndIndex, /*OUT*/CurrentRunIndex, /*OUT*/CurrentHighlightIndex, /*OUT*/PreviousBlockEnd, SoftLine );
+					CreateLineViewBlocks( LineModelIndex, FinalBreakOnSoftLine.ActualRange.EndIndex, /*OUT*/CurrentRunIndex, /*OUT*/CurrentHighlightIndex, /*OUT*/PreviousBlockEnd, SoftLine );
 
-					if ( CurrentRunIndex < LineModel.Runs.Num() && FinalBreakOnSoftLine.Range.EndIndex == LineModel.Runs[ CurrentRunIndex ].GetTextRange().EndIndex )
+					if ( CurrentRunIndex < LineModel.Runs.Num() && FinalBreakOnSoftLine.ActualRange.EndIndex == LineModel.Runs[ CurrentRunIndex ].GetTextRange().EndIndex )
 					{
 						++CurrentRunIndex;
 					}
 
-					PreviousBlockEnd = FinalBreakOnSoftLine.Range.EndIndex;
+					PreviousBlockEnd = FinalBreakOnSoftLine.ActualRange.EndIndex;
 
 					CurrentWidth = 0;
 					SoftLine.Empty();
 				} 
 				else
 				{
-					CurrentWidth += Break.Size.X;
+					CurrentWidth += Break.ActualSize.X;
 				}
 			}
 		}
@@ -386,43 +407,48 @@ void FTextLayout::ClearView()
 	LineViews.Empty();
 }
 
-void FTextLayout::ReadyForWrapping()
+void FTextLayout::CreateWrappingCache()
 {
-	if ( HasWrapInformation )
-	{
-		return;
-	}
-
-	const TSharedRef< FSlateFontMeasure > FontMeasure = FSlateApplication::Get().GetRenderer()->GetFontMeasureService();
-
 	for (int LineModelIndex = 0; LineModelIndex < LineModels.Num(); LineModelIndex++)
 	{
 		FLineModel& LineModel = LineModels[ LineModelIndex ];
 
-		for (int RunIndex = 0; RunIndex < LineModel.Runs.Num(); RunIndex++)
+		if ( !LineModel.HasWrappingInformation )
 		{
-			LineModel.Runs[ RunIndex ].ClearCache();
-		}
+			LineModel.BreakCandidates.Empty();
+			LineModel.HasWrappingInformation = true;
+
+			for (int RunIndex = 0; RunIndex < LineModel.Runs.Num(); RunIndex++)
+			{
+				LineModel.Runs[ RunIndex ].ClearCache();
+			}
 
-		FLineBreakIterator LineBreakIterator( **LineModel.Text );
+			FLineBreakIterator LineBreakIterator( **LineModel.Text );
 
-		int32 PreviousBreak = 0;
-		int32 CurrentBreak = 0;
-		int32 CurrentRunIndex = 0;
+			int32 PreviousBreak = 0;
+			int32 CurrentBreak = 0;
+			int32 CurrentRunIndex = 0;
 
-		while( ( CurrentBreak = LineBreakIterator.MoveToNext() ) != INDEX_NONE )
-		{
-			LineModel.BreakCandidates.Add( CreateBreakCandidate(/*OUT*/CurrentRunIndex, LineModel, PreviousBreak, CurrentBreak) );
-			PreviousBreak = CurrentBreak;
+			while( ( CurrentBreak = LineBreakIterator.MoveToNext() ) != INDEX_NONE )
+			{
+				LineModel.BreakCandidates.Add( CreateBreakCandidate(/*OUT*/CurrentRunIndex, LineModel, PreviousBreak, CurrentBreak) );
+				PreviousBreak = CurrentBreak;
+			}
 		}
 	}
+}
 
-	HasWrapInformation = true;
+void FTextLayout::ClearWrappingCache()
+{
+	for (int LineModelIndex = 0; LineModelIndex < LineModels.Num(); LineModelIndex++)
+	{
+		FLineModel& LineModel = LineModels[ LineModelIndex ];
+		LineModel.HasWrappingInformation = false;
+	}
 }
 
 FTextLayout::FTextLayout() : LineModels()
 	, LineViews()
-	, HasWrapInformation( false )
 	, HasChanged( false )
 	, Scale( 1.0f )
 	, WrappingWidth( 0 )
@@ -438,10 +464,15 @@ void FTextLayout::UpdateIfNeeded()
 	if ( HasChanged )
 	{
 		// if something has changed then create a new View
-		CreateView();
+		Update();
 	}
 }
 
+void FTextLayout::Update()
+{
+	CreateView();
+}
+
 void FTextLayout::SetHighlights( const TArray< FTextHighlight >& Highlights )
 {
 	ClearHighlights();
@@ -493,6 +524,408 @@ void FTextLayout::ClearHighlights()
 	HasChanged = true;
 }
 
+FTextLocation FTextLayout::GetTextLocationAt( const FLineView& LineView, const FVector2D& Relative )
+{
+	for (int32 BlockIndex = 0; BlockIndex < LineView.Blocks.Num(); BlockIndex++)
+	{
+		const TSharedRef< ILayoutBlock >& Block = LineView.Blocks[ BlockIndex ];
+		const int32 TextIndex = Block->GetRun()->GetTextIndexAt(Block, FVector2D(Relative.X, Block->GetLocationOffset().Y), Scale);
+
+		if ( TextIndex == INDEX_NONE )
+		{
+			continue;
+		}
+
+		return FTextLocation( LineView.ModelIndex, TextIndex );
+	}
+
+	const int32 LineTextLength = LineModels[ LineView.ModelIndex ].Text->Len();
+	if ( LineTextLength == 0 )
+	{
+		return FTextLocation( LineView.ModelIndex, 0 );
+	}
+	else if (Relative.X < LineView.Offset.X)
+	{
+		return FTextLocation(LineView.ModelIndex, LineView.Range.BeginIndex);
+	}
+
+	return FTextLocation( LineView.ModelIndex, LineView.Range.EndIndex );
+}
+
+int32 FTextLayout::GetLineViewIndexForTextLocation(const TArray< FTextLayout::FLineView >& LineViews, const FTextLocation& Location)
+{
+	const int32 LineModelIndex = Location.GetLineIndex();
+	const int32 Offset = Location.GetOffset();
+
+	const FLineModel& LineModel = LineModels[LineModelIndex];
+	for (int Index = 0; Index < LineViews.Num(); Index++)
+	{
+		const FTextLayout::FLineView& LineView = LineViews[Index];
+
+		if (LineView.ModelIndex == LineModelIndex && (Offset == 0 || LineView.Range.InclusiveContains(Offset) || LineModel.Text->IsEmpty()))
+		{
+			return Index;
+		}
+	}
+
+	return INDEX_NONE;
+}
+
+FTextLocation FTextLayout::GetTextLocationAt( const FVector2D& Relative )
+{
+	// Early out if we have no LineViews
+	if (LineViews.Num() == 0)
+	{
+		return FTextLocation(0,0);
+	}
+
+	// Iterate until we find a LineView that is below our expected Y location
+	int32 ViewIndex;
+	for (ViewIndex = 0; ViewIndex < LineViews.Num(); ViewIndex++)
+	{
+		const FLineView& LineView = LineViews[ ViewIndex ];
+
+		if (LineView.Offset.Y > Relative.Y)
+		{
+			// Set the ViewIndex back to the previous line, but not lower than the first line
+			ViewIndex = FMath::Max( 0, ViewIndex - 1 );
+			break;
+		}
+	}
+
+	//if none of the lines are below our expected Y location then...
+	if (ViewIndex >= LineViews.Num())
+	{
+		// just use the very last line
+		ViewIndex = LineViews.Num() - 1;
+		const FLineView& LineView = LineViews[ ViewIndex ];
+		return GetTextLocationAt( LineView, Relative );
+	}
+	else 
+	{
+		// if the current LineView does not encapsulate our expected Y location then jump to the next LineView
+		// if we aren't already at the last LineView
+		const FLineView& LineView = LineViews[ ViewIndex ];
+		if ((LineView.Offset.Y + LineView.Size.Y) < Relative.Y && ViewIndex < LineViews.Num() - 1)
+		{
+			++ViewIndex;
+		}
+	}
+
+	const FLineView& LineView = LineViews[ ViewIndex ];
+	return GetTextLocationAt(LineView, FVector2D(Relative.X, LineView.Offset.Y));
+}
+
+FVector2D FTextLayout::GetLocationAt( const FTextLocation& Location )
+{
+	const int32 LineModelIndex = Location.GetLineIndex();
+	const int32 Offset = Location.GetOffset();
+
+	// Find the LineView which encapsulates the location's offset
+	int LineViewIndex;
+	for (LineViewIndex = 0; LineViewIndex < LineViews.Num(); LineViewIndex++)
+	{
+		const FTextLayout::FLineView& LineView = LineViews[ LineViewIndex ];
+
+		if ( LineView.ModelIndex == LineModelIndex && ( Offset == 0 || LineView.Range.InclusiveContains(Offset) ) )
+		{
+			break;
+		}
+	}
+
+	// if we failed to find a LineView for the location, early out
+	if ( !LineViews.IsValidIndex( LineViewIndex ) )
+	{
+		return FVector2D( ForceInitToZero );
+	}
+
+	const FTextLayout::FLineView& LineView = LineViews[ LineViewIndex ];
+
+	//Iterate over the LineView's blocks...
+	for (int32 BlockIndex = 0; BlockIndex < LineView.Blocks.Num(); BlockIndex++)
+	{
+		const TSharedRef< ILayoutBlock >& Block = LineView.Blocks[ BlockIndex ];
+		const FTextRange& BlockRange = Block->GetTextRange();
+
+		//if the block's range contains the specified locations offset...
+		if (BlockRange.InclusiveContains(Offset))
+		{
+			// Ask the block for the exact screen location
+			const FVector2D ScreenLocation = Block->GetRun()->GetLocationAt( Block, Offset - BlockRange.BeginIndex, Scale );
+
+			// if the block was unable to provide a location continue iterating
+			if ( ScreenLocation.IsZero() )
+			{
+				continue;
+			}
+
+			return ScreenLocation;
+		}
+	}
+
+	// Failed to find the screen location
+	return FVector2D( ForceInitToZero );
+}
+
+bool FTextLayout::InsertAt(const FTextLocation& Location, TCHAR Character)
+{
+	const int32 InsertLocation = Location.GetOffset();
+	const int32 LineIndex = Location.GetLineIndex();
+
+	if (!LineModels.IsValidIndex(LineIndex))
+	{
+		return false;
+	}
+
+	FLineModel& LineModel = LineModels[LineIndex];
+
+	LineModel.Text->InsertAt(InsertLocation, Character);
+	LineModel.HasWrappingInformation = false;
+
+	bool RunIsAfterInsertLocation = false;
+	for (int RunIndex = 0; RunIndex < LineModel.Runs.Num(); RunIndex++)
+	{
+		FRunModel& RunModel = LineModel.Runs[ RunIndex ];
+		const FTextRange& RunRange = RunModel.GetTextRange();
+
+		const bool bIsLastRun = RunIndex == LineModel.Runs.Num() - 1;
+		if (RunRange.Contains(InsertLocation) || (bIsLastRun && !RunIsAfterInsertLocation))
+		{
+			check(RunIsAfterInsertLocation == false);
+			RunIsAfterInsertLocation = true;
+
+			RunModel.SetTextRange( FTextRange( RunRange.BeginIndex, RunRange.EndIndex + 1 ) );
+		}
+		else if (RunIsAfterInsertLocation)
+		{
+			FTextRange NewRange = RunRange;
+			NewRange.Offset(1);
+			RunModel.SetTextRange(NewRange);
+		}
+	}
+
+	HasChanged = true;
+	return true;
+}
+
+bool FTextLayout::InsertAt( const FTextLocation& Location, const FString& Text )
+{
+	const int32 InsertLocation = Location.GetOffset();
+	const int32 LineIndex = Location.GetLineIndex();
+
+	if (!LineModels.IsValidIndex(LineIndex))
+	{
+		return false;
+	}
+
+	FLineModel& LineModel = LineModels[LineIndex];
+
+	LineModel.Text->InsertAt(InsertLocation, Text);
+	LineModel.HasWrappingInformation = false;
+
+	bool RunIsAfterInsertLocation = false;
+	for (int RunIndex = 0; RunIndex < LineModel.Runs.Num(); RunIndex++)
+	{
+		FRunModel& RunModel = LineModel.Runs[ RunIndex ];
+		const FTextRange& RunRange = RunModel.GetTextRange();
+
+		const bool bIsLastRun = RunIndex == LineModel.Runs.Num() - 1;
+		if (RunRange.Contains(InsertLocation) || (bIsLastRun && !RunIsAfterInsertLocation))
+		{
+			check(RunIsAfterInsertLocation == false);
+			RunIsAfterInsertLocation = true;
+
+			RunModel.SetTextRange( FTextRange( RunRange.BeginIndex, RunRange.EndIndex + Text.Len() ) );
+		}
+		else if (RunIsAfterInsertLocation)
+		{
+			FTextRange NewRange = RunRange;
+			NewRange.Offset(Text.Len());
+			RunModel.SetTextRange(NewRange);
+		}
+	}
+
+	HasChanged = true;
+	return true;
+}
+
+bool FTextLayout::JoinLineWithNextLine(int32 LineIndex)
+{
+	if (!LineModels.IsValidIndex(LineIndex) || !LineModels.IsValidIndex(LineIndex + 1))
+	{
+		return false;
+	}
+
+	//Grab both lines we are supposed to join
+	FLineModel& LineModel = LineModels[LineIndex];
+	FLineModel& NextLineModel = LineModels[LineIndex + 1];
+
+	//If the next line is empty we'll just remove it
+	if (NextLineModel.Text->Len() == 0)
+	{
+		return RemoveLine(LineIndex + 1);
+	}
+
+	const int32 LineLength = LineModel.Text->Len();
+
+	//Append the next line to the current line
+	LineModel.Text->InsertAt(LineLength, NextLineModel.Text.Get());
+
+	//Clear the current lines wrapping information
+	LineModel.HasWrappingInformation = false;
+
+	//Iterate through all of the next lines runs and bring them over to the current line
+	for (int RunIndex = 0; RunIndex < NextLineModel.Runs.Num(); RunIndex++)
+	{
+		const TSharedRef<IRun> Run = NextLineModel.Runs[RunIndex].GetRun();
+		FTextRange NewRange = Run->GetTextRange();
+
+		if (!NewRange.IsEmpty())
+		{
+			NewRange.Offset(LineLength);
+
+			Run->Move(LineModel.Text, NewRange);
+			LineModel.Runs.Add(FRunModel(Run));
+		}
+	}
+
+	//Remove the next line from the list of line models
+	LineModels.RemoveAt(LineIndex + 1);
+
+	HasChanged = true;
+	return true;
+}
+
+bool FTextLayout::SplitLineAt(const FTextLocation& Location)
+{
+	int32 BreakLocation = Location.GetOffset();
+	int32 LineIndex = Location.GetLineIndex();
+
+	if (!LineModels.IsValidIndex(LineIndex))
+	{
+		return false;
+	}
+
+	FLineModel& LineModel = LineModels[LineIndex];
+
+	FLineModel LeftLineModel(MakeShareable(new FString(BreakLocation, **LineModel.Text)));
+	FLineModel RightLineModel(MakeShareable(new FString(LineModel.Text->Len() - BreakLocation, **LineModel.Text + BreakLocation)));
+
+	check(LeftLineModel.Text->Len() == BreakLocation);
+
+	bool RunIsToTheLeftOfTheBreakLocation = true;
+	for (int RunIndex = 0; RunIndex < LineModel.Runs.Num(); RunIndex++)
+	{
+		const TSharedRef<IRun> Run = LineModel.Runs[RunIndex].GetRun();
+		const FTextRange& RunRange = Run->GetTextRange();
+		
+		if ((LineModel.Runs.Num() == 1 && RunRange.InclusiveContains(BreakLocation)) || RunRange.Contains(BreakLocation))
+		{
+			check(RunIsToTheLeftOfTheBreakLocation == true);
+			RunIsToTheLeftOfTheBreakLocation = false;
+
+			const TSharedRef<IRun> LeftRun = Run;
+			const TSharedRef<IRun> RightRun = Run->Clone();
+			
+			LeftRun->Move(LeftLineModel.Text, FTextRange(RunRange.BeginIndex, LeftLineModel.Text->Len()));
+			RightRun->Move(RightLineModel.Text, FTextRange(0, RunRange.EndIndex - LeftLineModel.Text->Len()));
+
+			LeftLineModel.Runs.Add(FRunModel(LeftRun));
+			RightLineModel.Runs.Add(FRunModel(RightRun));
+		}
+		else if (RunIsToTheLeftOfTheBreakLocation)
+		{
+			Run->Move(LeftLineModel.Text, RunRange);
+			LeftLineModel.Runs.Add(FRunModel(Run));
+		}
+		else
+		{
+			// This run is after the break, so adjust the range to match that of RHS of the split (RightLineModel)
+			// We can do this by subtracting the LeftLineModel text size, since that's the LHS of the split
+			FTextRange NewRange = RunRange;
+			NewRange.Offset(-LeftLineModel.Text->Len());
+
+			Run->Move(RightLineModel.Text, NewRange);
+			RightLineModel.Runs.Add(FRunModel(Run));
+		}
+	}
+
+	LineModels.RemoveAt(LineIndex);
+	LineModels.Insert(LeftLineModel, LineIndex);
+	LineModels.Insert(RightLineModel, LineIndex + 1);
+
+	HasChanged = true;
+	return true;
+}
+
+bool FTextLayout::RemoveAt( const FTextLocation& Location, int32 Count )
+{
+	int32 RemoveLocation = Location.GetOffset();
+	int32 LineIndex = Location.GetLineIndex();
+
+	if (!LineModels.IsValidIndex(LineIndex))
+	{
+		return false;
+	}
+
+	FLineModel& LineModel = LineModels[LineIndex];
+
+	//Make sure we aren't trying to remove more characters then we have
+	Count = RemoveLocation + Count > LineModel.Text->Len() ? Count - ((RemoveLocation + Count) - LineModel.Text->Len()) : Count;
+
+	if (Count == 0)
+	{
+		return false;
+	}
+
+	LineModel.Text->RemoveAt(RemoveLocation, Count);
+	LineModel.HasWrappingInformation = false;
+
+	bool RunIsAfterRemoveLocation = false;
+	for (int RunIndex = 0; RunIndex < LineModel.Runs.Num(); RunIndex++)
+	{
+		FRunModel& RunModel = LineModel.Runs[RunIndex];
+		const FTextRange& RunRange = RunModel.GetTextRange();
+
+		if (RunRange.Contains(RemoveLocation))
+		{
+			check(RunIsAfterRemoveLocation == false);
+			RunIsAfterRemoveLocation = true;
+
+			FTextRange NewRange(RunRange.BeginIndex, RunRange.EndIndex - Count);
+			if (NewRange.IsEmpty() && LineModel.Runs.Num() > 1)
+			{
+				LineModel.Runs.RemoveAt(RunIndex--);
+			}
+			else
+			{
+				RunModel.SetTextRange(NewRange);
+			}
+		}
+		else if (RunIsAfterRemoveLocation)
+		{
+			FTextRange NewRange = RunRange;
+			NewRange.Offset(-Count);
+			RunModel.SetTextRange(NewRange);
+		}
+	}
+
+	HasChanged = true;
+	return true;
+}
+
+bool FTextLayout::RemoveLine(int32 LineIndex)
+{
+	if (!LineModels.IsValidIndex(LineIndex))
+	{
+		return false;
+	}
+
+	LineModels.RemoveAt(LineIndex);
+
+	return true;;
+}
+
 void FTextLayout::AddLine( const TSharedRef< FString >& Text, const TArray< TSharedRef< IRun > >& Runs )
 {
 	FLineModel LineModel( Text );
@@ -512,6 +945,145 @@ void FTextLayout::ClearLines()
 	HasChanged = true;
 }
 
+bool FTextLayout::IsEmpty() const
+{
+	return (LineModels.Num() == 0 || (LineModels.Num() == 1 && LineModels[0].Text->Len() == 0));
+}
+
+void FTextLayout::GetAsText(FString& DisplayText) const
+{
+	for (int LineModelIndex = 0; LineModelIndex < LineModels.Num(); LineModelIndex++)
+	{
+		const FLineModel& LineModel = LineModels[LineModelIndex];
+
+		for (int RunIndex = 0; RunIndex < LineModel.Runs.Num(); RunIndex++)
+		{
+			const FRunModel& Run = LineModel.Runs[RunIndex];
+			Run.AppendText(DisplayText);
+		}
+
+		DisplayText.AppendChar('\n');
+	}
+}
+
+void GetRangeAsTextFromLine(FString& DisplayText, const FTextLayout::FLineModel& LineModel, const FTextRange& Range)
+{
+	for (int RunIndex = 0; RunIndex < LineModel.Runs.Num(); RunIndex++)
+	{
+		const FTextLayout::FRunModel& Run = LineModel.Runs[RunIndex];
+		const FTextRange& RunRange = Run.GetTextRange();
+
+		const FTextRange IntersectRange = RunRange.Intersect(Range);
+
+		if (!IntersectRange.IsEmpty())
+		{
+			Run.AppendText(DisplayText, IntersectRange);
+		}
+		else if (RunRange.BeginIndex > Range.EndIndex)
+		{
+			//We're past the selection range so we can stop
+			break;
+		}
+	}
+}
+
+void FTextLayout::GetSelectionAsText(FString& DisplayText, const FTextSelection& Selection) const
+{
+	int32 SelectionBeginningLineIndex = Selection.GetBeginning().GetLineIndex();
+	int32 SelectionBeginningLineOffset = Selection.GetBeginning().GetOffset();
+
+	int32 SelectionEndLineIndex = Selection.GetEnd().GetLineIndex();
+	int32 SelectionEndLineOffset = Selection.GetEnd().GetOffset();
+
+	if (LineModels.IsValidIndex(SelectionBeginningLineIndex) && LineModels.IsValidIndex(SelectionEndLineIndex))
+	{
+		if (SelectionBeginningLineIndex == SelectionEndLineIndex)
+		{
+			const FTextRange SelectionRange(SelectionBeginningLineOffset, SelectionEndLineOffset);
+			const FLineModel& LineModel = LineModels[SelectionBeginningLineIndex];
+
+			GetRangeAsTextFromLine(DisplayText, LineModel, SelectionRange);
+		}
+		else
+		{
+			for (int LineIndex = SelectionBeginningLineIndex; LineIndex <= SelectionEndLineIndex; LineIndex++)
+			{
+				if (LineIndex == SelectionBeginningLineIndex)
+				{
+					const FLineModel& LineModel = LineModels[SelectionBeginningLineIndex];
+					const FTextRange SelectionRange(SelectionBeginningLineOffset, LineModel.Text->Len());
+
+					GetRangeAsTextFromLine(DisplayText, LineModel, SelectionRange);
+				}
+				else if (LineIndex == SelectionEndLineIndex)
+				{
+					const FLineModel& LineModel = LineModels[SelectionEndLineIndex];
+					const FTextRange SelectionRange(0, SelectionEndLineOffset);
+
+					GetRangeAsTextFromLine(DisplayText, LineModel, SelectionRange);
+				}
+				else
+				{
+					const FLineModel& LineModel = LineModels[LineIndex];
+					const FTextRange SelectionRange(0, LineModel.Text->Len());
+
+					GetRangeAsTextFromLine(DisplayText, LineModel, SelectionRange);
+				}
+
+				if (LineIndex != SelectionEndLineIndex)
+				{
+					DisplayText.AppendChar('\n');
+				}
+			}
+		}
+	}
+}
+
+FTextSelection FTextLayout::GetWordAt(const FTextLocation& Location) const
+{
+	const int32 LineIndex = Location.GetLineIndex();
+	const int32 Offset = Location.GetOffset();
+
+	if (!LineModels.IsValidIndex(LineIndex))
+	{
+		return FTextSelection();
+	}
+
+	const FLineModel& LineModel = LineModels[LineIndex];
+
+	FWordBreakIterator WordBreakIterator(**LineModel.Text);
+
+	int32 PreviousBreak = WordBreakIterator.MoveToCandidateAfter(Offset);
+	int32 CurrentBreak = 0;
+
+	while ((CurrentBreak = WordBreakIterator.MoveToPrevious()) != INDEX_NONE)
+	{
+		bool HasLetter = false;
+		for (int32 Index = CurrentBreak; Index < PreviousBreak; Index++)
+		{
+			if (!FText::IsWhitespace((*LineModel.Text)[Index]))
+			{
+				HasLetter = true;
+				break;
+			}
+		}
+
+		if (HasLetter)
+		{
+			break;
+		}
+
+		PreviousBreak = CurrentBreak;
+	}
+
+	if (PreviousBreak == CurrentBreak || CurrentBreak == INDEX_NONE)
+	{
+		return FTextSelection();
+	}
+
+	return FTextSelection(FTextLocation(LineIndex, CurrentBreak), FTextLocation(LineIndex, PreviousBreak));
+}
+
 void FTextLayout::SetMargin( const FMargin& InMargin )
 {
 	if ( Margin == InMargin )
@@ -562,32 +1134,9 @@ void FTextLayout::SetScale( float Value )
 	Scale = Value;
 	HasChanged = true;
 
-	if ( HasWrapInformation )
-	{
-		for (int LineModelIndex = 0; LineModelIndex < LineModels.Num(); LineModelIndex++)
-		{
-			FLineModel& LineModel = LineModels[ LineModelIndex ];
-			int32 CurrentRunIndex = 0;
-			int32 PreviousBreak = 0;
-			TArray< FBreakCandidate > UpdatedBreakCandidates;
-
-			for (int BreakIndex = 0; BreakIndex < LineModel.BreakCandidates.Num(); BreakIndex++)
-			{
-				const FBreakCandidate& Break = LineModel.BreakCandidates[ BreakIndex ];
-				UpdatedBreakCandidates.Add( CreateBreakCandidate( /*OUT*/CurrentRunIndex, LineModel, PreviousBreak, Break.Range.EndIndex ) );
-				PreviousBreak = Break.Range.EndIndex;
-			}
-
-			check( LineModel.BreakCandidates.Num() == UpdatedBreakCandidates.Num() );
-			LineModel.BreakCandidates = UpdatedBreakCandidates;
-
-			for (int RunIndex = 0; RunIndex < LineModel.Runs.Num(); RunIndex++)
-			{
-				FRunModel& Run = LineModel.Runs[ RunIndex ];
-				Run.ClearCache();
-			}
-		}
-	}
+	// Changing the scale will affect the wrapping information for *all lines*
+	// Clear out the entire cache so it gets regenerated on the text call to FlowLayout
+	ClearWrappingCache();
 }
 
 float FTextLayout::GetScale() const
@@ -664,10 +1213,12 @@ FTextLayout::~FTextLayout()
 
 }
 
-FTextLayout::FLineModel::FLineModel( const TSharedRef< FString >& InText ) : Text( InText )
+FTextLayout::FLineModel::FLineModel( const TSharedRef< FString >& InText ) 
+	: Text( InText )
 	, Runs()
 	, BreakCandidates()
 	, Highlights()
+	, HasWrappingInformation( false )
 {
 
 }
@@ -678,17 +1229,24 @@ void FTextLayout::FRunModel::ClearCache()
 	MeasuredRangeSizes.Empty();
 }
 
-TSharedRef< ILayoutBlock > FTextLayout::FRunModel::CreateBlock( const FBlockDefinition& BlockDefine, float Scale ) const
+void FTextLayout::FRunModel::AppendText(FString& Text) const
 {
-	check( BlockDefine.Range.EndIndex >= BlockDefine.VisualRange.EndIndex );
-	check( BlockDefine.Range.BeginIndex <= BlockDefine.VisualRange.BeginIndex );
+	Run->AppendText(Text);
+}
+
+void FTextLayout::FRunModel::AppendText(FString& Text, const FTextRange& Range) const
+{
+	Run->AppendText(Text, Range);
+}
 
-	const FTextRange& SizeRange = BlockDefine.VisualRange;
+TSharedRef< ILayoutBlock > FTextLayout::FRunModel::CreateBlock( const FBlockDefinition& BlockDefine, float Scale ) const
+{
+	const FTextRange& SizeRange = BlockDefine.ActualRange;
 
 	if ( MeasuredRanges.Num() == 0 )
 	{
 		FTextRange RunRange = Run->GetTextRange();
-		return Run->CreateBlock( BlockDefine.Range.BeginIndex, BlockDefine.Range.EndIndex, Run->Measure( SizeRange.BeginIndex, SizeRange.EndIndex, Scale ), BlockDefine.Highlighter );
+		return Run->CreateBlock( BlockDefine.ActualRange.BeginIndex, BlockDefine.ActualRange.EndIndex, Run->Measure( SizeRange.BeginIndex, SizeRange.EndIndex, Scale ), BlockDefine.Highlighter );
 	}
 
 	int32 StartRangeIndex = 0;
@@ -786,7 +1344,7 @@ TSharedRef< ILayoutBlock > FTextLayout::FRunModel::CreateBlock( const FBlockDefi
 		}
 	}
 
-	return Run->CreateBlock( BlockDefine.Range.BeginIndex, BlockDefine.Range.EndIndex, BlockSize, BlockDefine.Highlighter );
+	return Run->CreateBlock( BlockDefine.ActualRange.BeginIndex, BlockDefine.ActualRange.EndIndex, BlockSize, BlockDefine.Highlighter );
 }
 
 int FTextLayout::FRunModel::BinarySearchForEndIndex( const TArray< FTextRange >& Ranges, int32 RangeBeginIndex, int32 EndIndex )
@@ -869,6 +1427,11 @@ FTextRange FTextLayout::FRunModel::GetTextRange() const
 	return Run->GetTextRange();
 }
 
+void FTextLayout::FRunModel::SetTextRange( const FTextRange& Value )
+{
+	Run->SetTextRange( Value );
+}
+
 void FTextLayout::FRunModel::EndLayout()
 {
 	Run->EndLayout();
diff --git a/Engine/Source/Runtime/Slate/Private/Widgets/Input/SEditableText.cpp b/Engine/Source/Runtime/Slate/Private/Widgets/Input/SEditableText.cpp
--- a/Engine/Source/Runtime/Slate/Private/Widgets/Input/SEditableText.cpp
+++ b/Engine/Source/Runtime/Slate/Private/Widgets/Input/SEditableText.cpp
@@ -64,7 +64,7 @@ void SEditableText::Construct( const FArguments& InArgs )
 	IsReadOnly = InArgs._IsReadOnly;
 	IsPassword = InArgs._IsPassword;
 	IsCaretMovedWhenGainFocus = InArgs._IsCaretMovedWhenGainFocus;
-	SelectAllTextWhenFocused = InArgs._SelectAllTextWhenFocused;
+	bSelectAllTextWhenFocused = InArgs._SelectAllTextWhenFocused;
 	RevertTextOnEscape = InArgs._RevertTextOnEscape;
 	ClearKeyboardFocusOnCommit = InArgs._ClearKeyboardFocusOnCommit;
 	OnIsTypedCharValid = InArgs._OnIsTypedCharValid;
@@ -233,7 +233,7 @@ void SEditableText::Tick( const FGeometry& AllottedGeometry, const double InCurr
 	// Make sure the caret is scrolled into view.  Note that even if we don't have keyboard focus (and the caret
 	// is hidden), we still do this to make sure the beginning of the string is in view.
 	{
-		const float CaretWidth = CalculateCaretWidth(FontMaxCharHeight);
+		const float CaretWidth = FTextEditHelper::CalculateCaretWidth(FontMaxCharHeight);
 
 		// Grab the caret position in the scrolled area space
 		const float CaretLeft = CaretVisualPositionSpring.GetPosition();
@@ -437,7 +437,7 @@ bool SEditableText::FTextInputMethodContext::GetTextBounds(const uint32 BeginInd
 			IsClipped = true;
 		}
 
-		const float FontHeight = This->GetFontHeight();
+		const float FontHeight = FTextEditHelper::GetFontHeight(FontInfo);
 		const float DrawPositionY = ( CachedGeometry.Size.Y / 2 ) - ( FontHeight / 2 );
 
 		const float TextVertOffset = EditableTextDefs::TextVertOffsetPercent * FontHeight;
@@ -506,22 +506,6 @@ void SEditableText::FTextInputMethodContext::EndComposition()
 	}
 }
 
-/**
- * Gets the height of the largest character in the font
- *
- * @return  The fonts height
- */
-float SEditableText::GetFontHeight() const
-{
-	// Take a local copy of the FontInfo and remove any flags that could affect the outcome in a negative way
-	FSlateFontInfo FontCopy = Font.Get();
-
-	// Aesthetic choice: editable text can use a little more room, and should never be smaller than a W of the font.
-	const TSharedRef< FSlateFontMeasure > FontMeasureService = FSlateApplication::Get().GetRenderer()->GetFontMeasureService();
-	const FVector2D FontSize = FontMeasureService->Measure( FString(TEXT("W")), FontCopy );
-	return FontSize.Y;
-}
-
 void SEditableText::StartChangingText()
 {
 	// Never change text on read only controls! 
@@ -592,8 +576,6 @@ void SEditableText::DeleteChar()
 {
 	if( !IsReadOnly.Get() )
 	{
-		StartChangingText();
-
 		if( AnyTextSelected() )
 		{
 			// Delete selected text
@@ -608,8 +590,6 @@ void SEditableText::DeleteChar()
 				EditedText = FText::FromString(OldText.Left( CaretPosition ) + OldText.Mid( CaretPosition + 1 ));
 			}
 		}
-
-		FinishChangingText();
 	}
 }
 
@@ -654,7 +634,7 @@ static int32 OffsetToTestBasedOnDirection( int8 Direction )
 	return (Direction >= 0) ? 0 : -1;
 }
 
-FReply SEditableText::MoveCursor( ECursorMoveMethod::Type Method, int8 Direction, ECursorAction::Type Action )
+FReply SEditableText::MoveCursor( ECursorMoveMethod::Type Method, const FVector2D& Direction, ECursorAction::Type Action )
 {
 	const int32 OldCaretPosition = CaretPosition;
 	int32 NewCaretPosition = OldCaretPosition;
@@ -662,54 +642,65 @@ FReply SEditableText::MoveCursor( ECursorMoveMethod::Type Method, int8 Direction
 	const int32 StringLength = EditedText.ToString().Len();
 	const FString& EditedTextAsString = EditedText.ToString();
 
-
 	// If control is held down, jump to beginning of the current word (or, the previous word, if
 	// the caret is currently on a whitespace character
-	if( Method == ECursorMoveMethod::Word )
+	switch( Method ) 
 	{
-		if (Direction > 0)
+	case ECursorMoveMethod::Word:
 		{
-			// Scan right for text
-			int32 CurCharIndex = CaretPosition;
-			while( CurCharIndex < StringLength && !FChar::IsWhitespace( EditedText.ToString().GetCharArray()[ CurCharIndex ] ) )
-			{
-				++CurCharIndex;
-			}
+			const FString EditedTextString = EditedText.ToString();
 
-			// Scan right for whitespace
-			while( CurCharIndex < StringLength && FChar::IsWhitespace( EditedText.ToString().GetCharArray()[ CurCharIndex ] ) )
+			if (Direction.X > 0)
 			{
-				++CurCharIndex;
-			}
+				// Scan right for text
+				int32 CurCharIndex = CaretPosition;
+				while( CurCharIndex < StringLength && !FText::IsWhitespace( EditedTextString[ CurCharIndex ] ) )
+				{
+					++CurCharIndex;
+				}
 
-			NewCaretPosition = CurCharIndex;
-		}
-		else
-		{
-			// Scan left for whitespace
-			int32 CurCharIndex = CaretPosition;
-			while( CurCharIndex > 0 && FChar::IsWhitespace( EditedText.ToString().GetCharArray()[ CurCharIndex - 1 ] ) )
-			{
-				--CurCharIndex;
-			}
+				// Scan right for whitespace
+				while( CurCharIndex < StringLength && FText::IsWhitespace( EditedTextString[ CurCharIndex ] ) )
+				{
+					++CurCharIndex;
+				}
 
-			// Scan left for text
-			while( CurCharIndex > 0 && !FChar::IsWhitespace( EditedText.ToString().GetCharArray()[ CurCharIndex - 1 ] ) )
-			{
-				--CurCharIndex;
+				NewCaretPosition = CurCharIndex;
 			}
+			else
+			{
+				// Scan left for whitespace
+				int32 CurCharIndex = CaretPosition;
+				while( CurCharIndex > 0 && FText::IsWhitespace( EditedTextString[ CurCharIndex - 1 ] ) )
+				{
+					--CurCharIndex;
+				}
 
-			NewCaretPosition = CurCharIndex;
+				// Scan left for text
+				while( CurCharIndex > 0 && !FText::IsWhitespace( EditedTextString[ CurCharIndex - 1 ] ) )
+				{
+					--CurCharIndex;
+				}
+
+				NewCaretPosition = CurCharIndex;
+			}
 		}
-	}
-	else if ( Method == ECursorMoveMethod::CharacterHorizontal )
-	{
-		NewCaretPosition = FMath::Clamp( CaretPosition + Direction, 0, StringLength );
-	}
-	else
-	{
-		ensure( Method == ECursorMoveMethod::CharacterVertical );
+		break;
+
+	case ECursorMoveMethod::CharacterHorizontal:
+		NewCaretPosition = FMath::Clamp<int32>( CaretPosition + Direction.X, 0, StringLength );
+		break;
+	
+	case ECursorMoveMethod::CharacterVertical:
 		return FReply::Unhandled();
+	
+	case ECursorMoveMethod::ScreenPosition:
+		NewCaretPosition = FindClickedCharacterIndex( Direction );
+		break;
+
+	default:
+		checkSlow(false, "Unknown ECursorMoveMethod value");
+		break;
 	}
 
 	SetCaretPosition( NewCaretPosition );
@@ -781,7 +772,6 @@ void SEditableText::ClearSelection()
 	Selection.Clear();
 }
 
-
 void SEditableText::SelectAllText()
 {
 	// NOTE: Caret position is left unchanged
@@ -799,6 +789,93 @@ void SEditableText::SelectAllText()
 	SelectionTargetRightSpring.SetPosition( SelectionTargetX );
 }
 
+bool SEditableText::SelectAllTextWhenFocused()
+{
+	return bSelectAllTextWhenFocused.Get();
+}
+
+void SEditableText::SelectWordAt(const FVector2D& LocalPosition)
+{
+	// Figure out where in the string the user clicked
+	const int32 ClickedCharIndex = FindClickedCharacterIndex( LocalPosition );
+
+	// Deselect any text that was selected
+	ClearSelection();
+
+	// Select the word that the user double-clicked on, and move the caret to the end of that word
+	{
+		const FString EditedTextString = EditedText.ToString();
+
+		// Find beginning of the word
+		int32 FirstWordCharIndex = ClickedCharIndex;
+		while( FirstWordCharIndex > 0 && !FText::IsWhitespace( EditedTextString[ FirstWordCharIndex - 1 ] ) )
+		{
+			--FirstWordCharIndex;
+		}
+
+		// Find end of the word
+		int32 LastWordCharIndex = ClickedCharIndex;
+		while( LastWordCharIndex < EditedText.ToString().Len() && !FText::IsWhitespace( EditedTextString[ LastWordCharIndex ] ) )
+		{
+			++LastWordCharIndex;
+		}
+
+		// Move the caret to the end of the selected word
+		SetCaretPosition( ClickedCharIndex );
+
+		// Select the word!
+		Selection.StartIndex = FirstWordCharIndex;
+		Selection.FinishIndex = LastWordCharIndex;
+		RestartSelectionTargetAnimation();
+	}
+}
+
+void SEditableText::BeginDragSelection() 
+{
+	bIsDragSelecting = true;
+}
+
+bool SEditableText::IsDragSelecting() const
+{
+	return bIsDragSelecting;
+}
+
+void SEditableText::EndDragSelection() 
+{
+	bIsDragSelecting = false;
+}
+
+bool SEditableText::AnyTextSelected() const
+{
+	return Selection.StartIndex != INDEX_NONE && Selection.StartIndex != Selection.FinishIndex;
+}
+
+bool SEditableText::IsTextSelectedAt(const FVector2D& LocalPosition) const
+{
+	const int32 ClickedCharIndex = FindClickedCharacterIndex(LocalPosition);
+	return ClickedCharIndex >= Selection.GetMinIndex() && ClickedCharIndex < Selection.GetMaxIndex();
+}
+
+void SEditableText::SetWasFocusedByLastMouseDown( bool Value )
+{
+	bWasFocusedByLastMouseDown = Value;
+}
+
+bool SEditableText::WasFocusedByLastMouseDown() const
+{
+	return bWasFocusedByLastMouseDown;
+}
+
+bool SEditableText::HasDragSelectedSinceFocused() const
+{
+	return bHasDragSelectedSinceFocused;
+}
+
+void SEditableText::SetHasDragSelectedSinceFocused( bool Value )
+{
+	bHasDragSelectedSinceFocused = Value;
+}
+
 
 FReply SEditableText::OnEscape()
 {
@@ -1000,46 +1077,46 @@ bool SEditableText::CanExecuteUndo() const
 void SEditableText::Undo()
 {
 	if( !IsReadOnly.Get() && UndoStates.Num() > 0 && !TextInputMethodContext->IsComposing )
-	{
-		// Restore from undo state
-		int32 UndoStateIndex;
-		if( CurrentUndoLevel == INDEX_NONE )
 		{
-			// We haven't undone anything since the last time a new undo state was added
-			UndoStateIndex = UndoStates.Num() - 1;
+			// Restore from undo state
+			int32 UndoStateIndex;
+			if( CurrentUndoLevel == INDEX_NONE )
+			{
+				// We haven't undone anything since the last time a new undo state was added
+				UndoStateIndex = UndoStates.Num() - 1;
 
-			// Store an undo state for the current state (before undo was pressed)
-			FUndoState NewUndoState;
-			MakeUndoState( NewUndoState );
-			PushUndoState( NewUndoState );
-		}
-		else
-		{
-			// Move down to the next undo level
-			UndoStateIndex = CurrentUndoLevel - 1;
-		}
+				// Store an undo state for the current state (before undo was pressed)
+				FUndoState NewUndoState;
+				MakeUndoState( NewUndoState );
+				PushUndoState( NewUndoState );
+			}
+			else
+			{
+				// Move down to the next undo level
+				UndoStateIndex = CurrentUndoLevel - 1;
+			}
 
-		// Is there anything else to undo?
-		if( UndoStateIndex >= 0 )
-		{
+			// Is there anything else to undo?
+			if( UndoStateIndex >= 0 )
 			{
-				// NOTE: It's important the no code called here creates or destroys undo states!
-				const FUndoState& UndoState = UndoStates[ UndoStateIndex ];
+				{
+					// NOTE: It's important the no code called here creates or destroys undo states!
+					const FUndoState& UndoState = UndoStates[ UndoStateIndex ];
 		
-				EditedText = UndoState.Text;
-				CaretPosition = UndoState.CaretPosition;
-				Selection = UndoState.Selection;
+					EditedText = UndoState.Text;
+					CaretPosition = UndoState.CaretPosition;
+					Selection = UndoState.Selection;
 
-				SaveText();
-			}
+					SaveText();
+				}
 
-			CurrentUndoLevel = UndoStateIndex;
+				CurrentUndoLevel = UndoStateIndex;
 
-			// Let outsiders know that the text content has been changed
-			OnTextChanged.ExecuteIfBound( EditedText );
+				// Let outsiders know that the text content has been changed
+				OnTextChanged.ExecuteIfBound( EditedText );
+			}
 		}
 	}
-}
 
 
 void SEditableText::Redo()
@@ -1080,13 +1157,18 @@ void SEditableText::Redo()
 	}
 }
 
+TSharedRef< SWidget > SEditableText::GetWidget()
+{
+	return SharedThis( this );
+}
+
 FVector2D SEditableText::ComputeDesiredSize() const
 {
 	const TSharedRef< FSlateFontMeasure > FontMeasureService = FSlateApplication::Get().GetRenderer()->GetFontMeasureService();
 	const FSlateFontInfo& FontInfo = Font.Get();
 	
-	const float FontMaxCharHeight = GetFontHeight();
-	const float CaretWidth = CalculateCaretWidth(FontMaxCharHeight);
+	const float FontMaxCharHeight = FTextEditHelper::GetFontHeight(FontInfo);
+	const float CaretWidth = FTextEditHelper::CalculateCaretWidth(FontMaxCharHeight);
 
 	FVector2D TextSize;
 
@@ -1125,7 +1207,7 @@ int32 SEditableText::OnPaint( const FGeometry& AllottedGeometry, const FSlateRec
 	const FString VisibleText = GetStringToRender();
 	const FLinearColor ThisColorAndOpacity = ColorAndOpacity.Get().GetColor(InWidgetStyle);
 	const FColor ColorAndOpacitySRGB = ThisColorAndOpacity * InWidgetStyle.GetColorAndOpacityTint();
-	const float FontMaxCharHeight = GetFontHeight();
+	const float FontMaxCharHeight = FTextEditHelper::GetFontHeight(FontInfo);
 	const double CurrentTime = FSlateApplication::Get().GetCurrentTime();
 
 	// We'll draw with the 'focused' look if we're either focused or we have a context menu summoned
@@ -1301,7 +1383,7 @@ int32 SEditableText::OnPaint( const FGeometry& AllottedGeometry, const FSlateRec
 
 		const float CaretX = CaretVisualPositionSpring.GetPosition();
 		const float CaretY = DrawPositionY + EditableTextDefs::CaretVertOffset;
-		const float CaretWidth = CalculateCaretWidth(FontMaxCharHeight);
+		const float CaretWidth = FTextEditHelper::CalculateCaretWidth(FontMaxCharHeight);
 		const float CaretHeight = EditableTextDefs::CaretHeightPercent * FontMaxCharHeight;
 
 		const FVector2D CaretDrawPosition = ScrollHelper.FromScrollerSpace( FVector2D( CaretX, CaretY ) );
@@ -1388,7 +1470,7 @@ FReply SEditableText::OnKeyboardFocusReceived( const FGeometry& MyGeometry, cons
 			if( InKeyboardFocusEvent.GetCause() != EKeyboardFocusCause::OtherWidgetLostFocus )
 			{
 				// Select all text on keyboard focus.  This mirrors mouse focus functionality.
-				if( SelectAllTextWhenFocused.Get() )
+				if( bSelectAllTextWhenFocused.Get() )
 				{
 					SelectAllText();
 				}
@@ -1457,13 +1539,13 @@ void SEditableText::OnKeyboardFocusLost( const FKeyboardFocusEvent& InKeyboardFo
 
 FReply SEditableText::OnKeyChar( const FGeometry& MyGeometry, const FCharacterEvent& InCharacterEvent )
 {
-	return FTextEditHelper::OnKeyChar( MyGeometry, InCharacterEvent, *this );
+	return FTextEditHelper::OnKeyChar( InCharacterEvent, SharedThis( this ) );
 }
 
 
 FReply SEditableText::OnKeyDown( const FGeometry& MyGeometry, const FKeyboardEvent& InKeyboardEvent )
 {
-	FReply Reply = FTextEditHelper::OnKeyDown( MyGeometry, InKeyboardEvent, *this );
+	FReply Reply = FTextEditHelper::OnKeyDown( InKeyboardEvent, SharedThis( this ) );
 
 	// Process keybindings if the event wasn't already handled
 	if( !Reply.IsEventHandled() && UICommandList->ProcessCommandBindings( InKeyboardEvent ) )
@@ -1487,117 +1569,27 @@ FReply SEditableText::OnKeyUp( const FGeometry& MyGeometry, const FKeyboardEvent
 
 FReply SEditableText::OnMouseButtonDown( const FGeometry& InMyGeometry, const FPointerEvent& InMouseEvent )
 {
-	FReply Reply = FReply::Unhandled();
-
-	// If the mouse is already captured, then don't allow a new action to be taken
 	if( !HasMouseCapture() )
 	{
 		if( InMouseEvent.GetEffectingButton() == EKeys::LeftMouseButton ||
 			InMouseEvent.GetEffectingButton() == EKeys::RightMouseButton )
 		{
-			const int32 OldCaretPosition = CaretPosition;
-
-			// Am I getting focus right now?
-			const bool bIsGettingFocus = !HasKeyboardFocus();
-			if( bIsGettingFocus )
-			{
-				// We might be receiving keyboard focus due to this event.  Because the keyboard focus received callback
-				// won't fire until after this function exits, we need to make sure our widget's state is in order early
-
-				// Assume we'll be given keyboard focus, so load text for editing
-				LoadText();
-
-				// Reset 'mouse has moved' state.  We'll use this in OnMouseMove to determine whether we
-				// should reset the selection range to the caret's position.
-				bWasFocusedByLastMouseDown = true;
-			}
-
-			// Figure out where in the string the user clicked
-			const int32 ClickedCharIndex = FindClickedCharacterIndex( InMyGeometry, InMouseEvent.GetScreenSpacePosition() );
-
-			// Move the caret to this position
-			SetCaretPosition( ClickedCharIndex );
-
 			// Reset 'is drag selecting' state
 			bHasDragSelectedSinceMouseDown = false;
-
-			if( InMouseEvent.GetEffectingButton() == EKeys::LeftMouseButton )
-			{
-				if( InMouseEvent.IsShiftDown() )
-				{
-					SelectText( OldCaretPosition );
-				}
-				else
-				{
-					// Deselect any text that was selected
-					ClearSelection();
-				}
-
-				// Start drag selection
-				bIsDragSelecting = true;
-			}
-			else if( InMouseEvent.GetEffectingButton() == EKeys::RightMouseButton )
-			{
-				// If the user clicked right clicked on a character that wasn't already selected, we'll clear
-				// the selection
-				if( AnyTextSelected() &&
-					( CaretPosition < Selection.GetMinIndex() || CaretPosition > Selection.GetMaxIndex() ) )
-				{
-					// Deselect any text that was selected
-					ClearSelection();
-				}
-			}
-
-
-			// Right clicking to summon context menu, but we'll do that on mouse-up.
-			Reply = FReply::Handled();
-			Reply.CaptureMouse( AsShared() );
-			Reply.SetKeyboardFocus( AsShared(), EKeyboardFocusCause::Mouse );
 		}
 	}
 
-	return Reply;
+	return FTextEditHelper::OnMouseButtonDown( InMyGeometry, InMouseEvent, SharedThis( this ) );
 }
 
 
 FReply SEditableText::OnMouseButtonUp( const FGeometry& InMyGeometry, const FPointerEvent& InMouseEvent )
 {
-	FReply Reply = FReply::Unhandled();
-
 	// The mouse must have been captured by either left or right button down before we'll process mouse ups
 	if( HasMouseCapture() )
 	{
-		if( InMouseEvent.GetEffectingButton() == EKeys::LeftMouseButton &&
-			bIsDragSelecting )
+		if( InMouseEvent.GetEffectingButton() == EKeys::LeftMouseButton && bIsDragSelecting )
 		{
-			// No longer drag-selecting
-			bIsDragSelecting = false;
-
-
-			// If we received keyboard focus on this click, then we'll want to select all of the text
-			// when the user releases the mouse button, unless the user actually dragged the mouse
-			// while holding the button down, in which case they've already selected some text and
-			// we'll leave things alone!
-			if( bWasFocusedByLastMouseDown )
-			{
-				if( !bHasDragSelectedSinceFocused )
-				{
-					if( SelectAllTextWhenFocused.Get() )
-					{
-						// User wasn't dragging the mouse, so go ahead and select all of the text now
-						// that we've become focused
-						SelectAllText();
-
-						// Move the cursor to the end of the string
-						SetCaretPosition( EditedText.ToString().Len() );
-
-						// @todo Slate: In this state, the caret should actually stay hidden (until the user interacts again), and we should not move the caret
-					}
-				}
-
-				bWasFocusedByLastMouseDown = false;
-			}
-
 			// If we selected any text while dragging, then kick off the selection target animation
 			if( bHasDragSelectedSinceMouseDown && Selection.StartIndex != Selection.FinishIndex )
 			{
@@ -1609,103 +1601,27 @@ FReply SEditableText::OnMouseButtonUp( const FGeometry& InMyGeometry, const FPoi
 
 				RestartSelectionTargetAnimation();
 			}
-
-
-			// Release mouse capture
-			Reply = FReply::Handled();
-			Reply.ReleaseMouseCapture();
-		}
-		else if( InMouseEvent.GetEffectingButton() == EKeys::RightMouseButton )
-		{
-			if ( InMyGeometry.IsUnderLocation(InMouseEvent.GetScreenSpacePosition()) )
-			{
-				// Right clicked, so summon a context menu if the cursor is within the widget
-				SummonContextMenu( InMouseEvent.GetScreenSpacePosition() );
-			}
-
-			// Release mouse capture
-			Reply = FReply::Handled();
-			Reply.ReleaseMouseCapture();
 		}
 	}
-	
-	return Reply;
+
+	return FTextEditHelper::OnMouseButtonUp( InMyGeometry, InMouseEvent, SharedThis( this ) );
 }
 
 
 FReply SEditableText::OnMouseMove( const FGeometry& InMyGeometry, const FPointerEvent& InMouseEvent )
 {
-	FReply Reply = FReply::Unhandled();
-
 	if( bIsDragSelecting && HasMouseCapture() )
 	{
-		// Figure out where in the string the user clicked
-		const int32 ClickedCharIndex = FindClickedCharacterIndex( InMyGeometry, InMouseEvent.GetScreenSpacePosition() );
-
-		// Move the caret to this position
-		const int32 OldCaretPosition = CaretPosition;
-		SetCaretPosition( ClickedCharIndex );
-
-
-		// Check to see if the user dragged far enough to select any characters
-		if( OldCaretPosition != CaretPosition )
-		{
-			// OK, the user has dragged over a character since we received focus
-			bHasDragSelectedSinceFocused = true;
-			bHasDragSelectedSinceMouseDown = true;
-
-			// Select text that the user dragged over
-			SelectText( OldCaretPosition );
-		}
-
-		Reply = FReply::Handled();
+		bHasDragSelectedSinceMouseDown = true;
 	}
 
-	return Reply;
+	return FTextEditHelper::OnMouseMove( InMyGeometry, InMouseEvent, SharedThis( this ) );;
 }
 
 
 FReply SEditableText::OnMouseButtonDoubleClick( const FGeometry& InMyGeometry, const FPointerEvent& InMouseEvent )
 {
-	FReply Reply = FReply::Unhandled();
-
-	if( InMouseEvent.GetEffectingButton() == EKeys::LeftMouseButton )
-	{
-		// Figure out where in the string the user clicked
-		const int32 ClickedCharIndex = FindClickedCharacterIndex( InMyGeometry, InMouseEvent.GetScreenSpacePosition() );
-
-		// Deselect any text that was selected
-		ClearSelection();
-
-		// Select the word that the user double-clicked on, and move the caret to the end of that word
-		{
-			// Find beginning of the word
-			int32 FirstWordCharIndex = ClickedCharIndex;
-			while( FirstWordCharIndex > 0 && !FChar::IsWhitespace( EditedText.ToString().GetCharArray()[ FirstWordCharIndex - 1 ] ) )
-			{
-				--FirstWordCharIndex;
-			}
-
-			// Find end of the word
-			int32 LastWordCharIndex = ClickedCharIndex;
-			while( LastWordCharIndex < EditedText.ToString().Len() && !FChar::IsWhitespace( EditedText.ToString().GetCharArray()[ LastWordCharIndex ] ) )
-			{
-				++LastWordCharIndex;
-			}
-
-			// Select the word!
-			Selection.StartIndex = FirstWordCharIndex;
-			Selection.FinishIndex = LastWordCharIndex;
-			RestartSelectionTargetAnimation();
-
-			// Move the caret to the end of the selected word
-			SetCaretPosition( ClickedCharIndex );
-		}
-
-		Reply = FReply::Handled();
-	}
-
-	return Reply;
+	return FTextEditHelper::OnMouseButtonDoubleClick( InMyGeometry, InMouseEvent, SharedThis( this ) );
 }
 
 
@@ -1766,8 +1682,14 @@ int32 SEditableText::FindClickedCharacterIndex( const FGeometry& InMyGeometry, c
 	const FVector2D& ScreenCursorPosition = InScreenCursorPosition;
 	const FVector2D& LocalCursorPosition = InMyGeometry.AbsoluteToLocal( ScreenCursorPosition );
 
+	return FindClickedCharacterIndex(LocalCursorPosition);
+}
+
+
+int32 SEditableText::FindClickedCharacterIndex( const FVector2D& InLocalCursorPosition ) const
+{
 	// Convert the click position from the local widget space to the scrolled area
-	const FVector2D& PositionInScrollableArea = ScrollHelper.ToScrollerSpace( LocalCursorPosition );
+	const FVector2D& PositionInScrollableArea = ScrollHelper.ToScrollerSpace( InLocalCursorPosition );
 	const float ClampedCursorX = FMath::Max( PositionInScrollableArea.X, 0.0f );
 
 	// Figure out where in the string the user clicked
@@ -2029,12 +1951,6 @@ FString SEditableText::GetStringToRender() const
 	return VisibleText;
 }
 
-float SEditableText::CalculateCaretWidth(const float FontMaxCharHeight) const
-{
-	// The caret with should be at least one pixel. For very small fonts the width could be < 1 which makes it not visible
-	return FMath::Max(1.0f, EditableTextDefs::CaretWidthPercent * FontMaxCharHeight);
-}
-
 void SEditableText::RestartSelectionTargetAnimation()
 {
 	// Don't bother with this unless we're running at a decent frame rate
diff --git a/Engine/Source/Runtime/Slate/Private/Widgets/Input/SMultiLineEditableTextBox.cpp b/Engine/Source/Runtime/Slate/Private/Widgets/Input/SMultiLineEditableTextBox.cpp
new file mode 100644
--- /dev/null
+++ b/Engine/Source/Runtime/Slate/Private/Widgets/Input/SMultiLineEditableTextBox.cpp
@@ -0,0 +1,149 @@
+// Copyright 1998-2014 Epic Games, Inc. All Rights Reserved.
+
+#include "SlatePrivatePCH.h"
+
+
+/**
+ * Construct this widget
+ *
+ * @param	InArgs	The declaration data for this widget
+ */
+void SMultiLineEditableTextBox::Construct( const FArguments& InArgs )
+{
+	check (InArgs._Style);
+	BorderImageNormal = &InArgs._Style->BackgroundImageNormal;
+	BorderImageHovered = &InArgs._Style->BackgroundImageHovered;
+	BorderImageFocused = &InArgs._Style->BackgroundImageFocused;
+	BorderImageReadOnly = &InArgs._Style->BackgroundImageReadOnly;
+	TAttribute<FMargin> Padding = InArgs._Padding.IsSet() ? InArgs._Padding : InArgs._Style->Padding;
+	TAttribute<FSlateFontInfo> Font = InArgs._Font.IsSet() ? InArgs._Font : InArgs._Style->Font;
+	TAttribute<FSlateColor> ForegroundColor = InArgs._ForegroundColor.IsSet() ? InArgs._ForegroundColor : InArgs._Style->ForegroundColor;
+	TAttribute<FSlateColor> BackgroundColor = InArgs._BackgroundColor.IsSet() ? InArgs._BackgroundColor : InArgs._Style->BackgroundColor;
+	ReadOnlyForegroundColor = InArgs._ReadOnlyForegroundColor.IsSet() ? InArgs._ReadOnlyForegroundColor : InArgs._Style->ReadOnlyForegroundColor;
+
+	SBorder::Construct( SBorder::FArguments()
+		.BorderImage( this, &SMultiLineEditableTextBox::GetBorderImage )
+		.BorderBackgroundColor( BackgroundColor )
+		.ForegroundColor( ForegroundColor )
+		.Padding( 0 )
+		[
+			SAssignNew( Box, SHorizontalBox)
+			+SHorizontalBox::Slot()
+			.VAlign(VAlign_Fill)
+			.HAlign(HAlign_Fill)
+			.FillWidth(1)
+			[
+				SNew( SBox )
+				.Padding( Padding )
+				[
+					SAssignNew( EditableText, SMultiLineEditableText )
+					.Text( InArgs._Text )
+					.Font( Font )
+					.IsReadOnly( InArgs._IsReadOnly )
+					.OnTextChanged( InArgs._OnTextChanged )
+					.OnTextCommitted( InArgs._OnTextCommitted )
+					.Justification(InArgs._Justification)
+					.LineHeightPercentage(InArgs._LineHeightPercentage)
+					.WrapTextAt(InArgs._WrapTextAt)
+					.AutoWrapText(InArgs._AutoWrapText)
+				]
+			]
+		]
+	);
+
+	ErrorReporting = InArgs._ErrorReporting;
+	if ( ErrorReporting.IsValid() )
+	{
+		Box->AddSlot()
+		.AutoWidth()
+		.Padding(3,0)
+		[
+			ErrorReporting->AsWidget()
+		];
+	}
+
+}
+
+
+/**
+ * Sets the text string currently being edited 
+ *
+ * @param  InNewText  The new text string
+ */
+void SMultiLineEditableTextBox::SetText( const TAttribute< FText >& InNewText )
+{
+	EditableText->SetText( InNewText );
+}
+
+void SMultiLineEditableTextBox::SetError( const FText& InError )
+{
+	SetError( InError.ToString() );
+}
+
+void SMultiLineEditableTextBox::SetError( const FString& InError )
+{
+	const bool bHaveError = !InError.IsEmpty();
+
+	if ( !ErrorReporting.IsValid() )
+	{
+		// No error reporting was specified; make a default one
+		TSharedPtr<SPopupErrorText> ErrorTextWidget;
+		Box->AddSlot()
+		.AutoWidth()
+		.Padding(3,0)
+		[
+			SAssignNew( ErrorTextWidget, SPopupErrorText )
+		];
+		ErrorReporting = ErrorTextWidget;
+	}
+
+	ErrorReporting->SetError( InError );
+}
+
+/** @return Border image for the text box based on the hovered and focused state */
+const FSlateBrush* SMultiLineEditableTextBox::GetBorderImage() const
+{
+	if ( EditableText->GetIsReadOnly() )
+	{
+		return BorderImageReadOnly;
+	}
+	else if ( EditableText->HasKeyboardFocus() )
+	{
+		return BorderImageFocused;
+	}
+	else
+	{
+		if ( EditableText->IsHovered() )
+		{
+			return BorderImageHovered;
+		}
+		else
+		{
+			return BorderImageNormal;
+		}
+	}
+}
+
+bool SMultiLineEditableTextBox::SupportsKeyboardFocus() const
+{
+	return StaticCastSharedPtr<SWidget>(EditableText)->SupportsKeyboardFocus();
+}
+
+bool SMultiLineEditableTextBox::HasKeyboardFocus() const
+{
+	// Since keyboard focus is forwarded to our editable text, we will test it instead
+	return SBorder::HasKeyboardFocus() || EditableText->HasKeyboardFocus();
+}
+
+FReply SMultiLineEditableTextBox::OnKeyboardFocusReceived( const FGeometry& MyGeometry, const FKeyboardFocusEvent& InKeyboardFocusEvent )
+{
+	FReply Reply = FReply::Handled();
+
+	if ( InKeyboardFocusEvent.GetCause() != EKeyboardFocusCause::Cleared )
+	{
+		// Forward keyboard focus to our editable text widget
+		Reply.SetKeyboardFocus( EditableText.ToSharedRef(), InKeyboardFocusEvent.GetCause() );
+	}
+
+	return Reply;
+}
diff --git a/Engine/Source/Runtime/Slate/Private/Widgets/Input/SVirtualKeyboardEntry.cpp b/Engine/Source/Runtime/Slate/Private/Widgets/Input/SVirtualKeyboardEntry.cpp
--- a/Engine/Source/Runtime/Slate/Private/Widgets/Input/SVirtualKeyboardEntry.cpp
+++ b/Engine/Source/Runtime/Slate/Private/Widgets/Input/SVirtualKeyboardEntry.cpp
@@ -1,7 +1,7 @@
 // Copyright 1998-2014 Epic Games, Inc. All Rights Reserved.
 
 #include "SlatePrivatePCH.h"
-
+#include "TextEditHelper.h"
 
 /** Constructor */
 SVirtualKeyboardEntry::SVirtualKeyboardEntry()
@@ -122,22 +122,6 @@ void SVirtualKeyboardEntry::Tick( const FGeometry& AllottedGeometry, const doubl
 	const bool bShouldAppearFocused = HasKeyboardFocus();
 }
 
-/**
- * Gets the height of the largest character in the font
- *
- * @return  The fonts height
- */
-float SVirtualKeyboardEntry::GetFontHeight() const
-{
-	// Take a local copy of the FontInfo and remove any flags that could affect the outcome in a negative way
-	FSlateFontInfo FontCopy = Font.Get();
-
-	// Aesthetic choice: editable text can use a little more room, and should never be smaller than a W of the font.
-	const TSharedRef< FSlateFontMeasure > FontMeasureService = FSlateApplication::Get().GetRenderer()->GetFontMeasureService();
-	const FVector2D FontSize = FontMeasureService->Measure( FString(TEXT("W")), FontCopy );
-	return FontSize.Y;
-}
-
 /**
  * Computes the desired size of this widget (SWidget)
  *
@@ -148,7 +132,7 @@ FVector2D SVirtualKeyboardEntry::ComputeDesiredSize() const
 	const FSlateFontInfo& FontInfo = Font.Get();
 	const TSharedRef< FSlateFontMeasure > FontMeasureService = FSlateApplication::Get().GetRenderer()->GetFontMeasureService();
 
-	const float FontMaxCharHeight = GetFontHeight();
+	const float FontMaxCharHeight = FTextEditHelper::GetFontHeight(FontInfo);
 
 	FVector2D TextSize;
 
@@ -195,7 +179,7 @@ int32 SVirtualKeyboardEntry::OnPaint( const FGeometry& AllottedGeometry, const F
 	const FString VisibleText = GetStringToRender();
 	const FLinearColor ThisColorAndOpacity = ColorAndOpacity.Get().GetColor(InWidgetStyle);
 	const FColor ColorAndOpacitySRGB = ThisColorAndOpacity * InWidgetStyle.GetColorAndOpacityTint();
-	const float FontMaxCharHeight = GetFontHeight();
+	const float FontMaxCharHeight = FTextEditHelper::GetFontHeight(FontInfo);
 	const double CurrentTime = FSlateApplication::Get().GetCurrentTime();
 
 	// We'll draw with the 'focused' look if we're either focused or we have a context menu summoned
diff --git a/Engine/Source/Runtime/Slate/Private/Widgets/Testing/STestSuite.cpp b/Engine/Source/Runtime/Slate/Private/Widgets/Testing/STestSuite.cpp
--- a/Engine/Source/Runtime/Slate/Private/Widgets/Testing/STestSuite.cpp
+++ b/Engine/Source/Runtime/Slate/Private/Widgets/Testing/STestSuite.cpp
@@ -1560,7 +1560,7 @@ public:
 	{
 		InlineEditableText = LOCTEXT( "TestingInlineEditableTextBlock", "Testing inline editable text block!" );
 
-		const FText MultiLineText = LOCTEXT("MultiLineTextTest", "He has refused his Assent to Laws, the most wholesome and necessary for the public good.\nHe has forbidden his Governors to pass Laws of immediate and pressing importance, unless suspended in their operation till his Assent should be obtained; and when so suspended, he has utterly neglected to attend to them.\nHe has refused to pass other Laws for the accommodation of large districts of people, unless those people would relinquish the right of Representation in the Legislature, a right inestimable to them and formidable to tyrants only.\n\nHe has called together legislative bodies at places unusual, uncomfortable, and distant from the depository of their public Records, for the sole purpose of fatiguing them into compliance with his measures.\nHe has dissolved Representative Houses repeatedly, for opposing with manly firmness his invasions on the rights of the people.\nHe has refused for a long time, after such dissolutions, to cause others to be elected; whereby the Legislative powers, incapable of Annihilation, have returned to the People at large for their exercise; the State remaining in the mean time exposed to all the dangers of invasion from without, and convulsions within.\nHe has endeavoured to prevent the population of these States; for that purpose obstructing the Laws for Naturalization of Foreigners; refusing to pass others to encourage their migrations hither, and raising the conditions of new Appropriations of Lands.\n");
+		MultilineEditableText = LOCTEXT("MultiLineTextTest", "He has refused his Assent to Laws, the most wholesome and necessary for the public good.\nHe has forbidden his Governors to pass Laws of immediate and pressing importance, unless suspended in their operation till his Assent should be obtained; and when so suspended, he has utterly neglected to attend to them.\nHe has refused to pass other Laws for the accommodation of large districts of people, unless those people would relinquish the right of Representation in the Legislature, a right inestimable to them and formidable to tyrants only.\n\nHe has called together legislative bodies at places unusual, uncomfortable, and distant from the depository of their public Records, for the sole purpose of fatiguing them into compliance with his measures.\nHe has dissolved Representative Houses repeatedly, for opposing with manly firmness his invasions on the rights of the people.\nHe has refused for a long time, after such dissolutions, to cause others to be elected; whereby the Legislative powers, incapable of Annihilation, have returned to the People at large for their exercise; the State remaining in the mean time exposed to all the dangers of invasion from without, and convulsions within.\nHe has endeavoured to prevent the population of these States; for that purpose obstructing the Laws for Naturalization of Foreigners; refusing to pass others to encourage their migrations hither, and raising the conditions of new Appropriations of Lands.\n");
 
 		Animation = FCurveSequence(0, 5);
 		Animation.Play();
@@ -1594,14 +1594,14 @@ public:
 				[
 					SNew( SEditableTextBox )
 					.Text( LOCTEXT( "TestingLongText", "Here is an editable text box with a very long initial string.  Useful to test scrolling.  Remember, this editable text box has many features, such as cursor navigation, text selection with either the mouse or keyboard, and cut, copy and paste.  You can even undo and redo just how you'd expect to." ) )
-					.Font( FSlateFontInfo( FPaths::EngineContentDir() / TEXT("Slate/Fonts/Roboto-BoldCondensed.ttf" ), 13 ) )
+					.Font( FSlateFontInfo( FPaths::EngineContentDir() / TEXT("Slate/Fonts/Roboto-Bold.ttf" ), 13 ) )
 					.HintText(LOCTEXT("TestingLongTextHint", "Hint Text"))
 				]
 				+ SVerticalBox::Slot().AutoHeight() .HAlign(HAlign_Center) .Padding(5)
 				[
 					SNew( SEditableTextBox )
 					.Text( LOCTEXT( "TestingBigTextBigMargin", "Big text, big margin!" ) )
-					.Font( FSlateFontInfo( FPaths::EngineContentDir() / TEXT("Slate/Fonts/Roboto-BoldCondensed.ttf" ), 40 ) )	// Lucida Console
+					.Font( FSlateFontInfo( FPaths::EngineContentDir() / TEXT("Slate/Fonts/Roboto-Bold.ttf" ), 40 ) )
 					.RevertTextOnEscape( true )
 					.BackgroundColor( this, &STextEditTest::GetLoopingColor )
 					.HintText(LOCTEXT("TestingBigTextMarginHint", "Hint Text"))
@@ -1618,16 +1618,16 @@ public:
 				[
 					SAssignNew( SearchBox, SSearchBox )
 					//.Text( LOCTEXT("TestingSearchBox", "Testing search boxes tool tip") )
-					//.Font( FSlateFontInfo( FPaths::EngineContentDir() / TEXT("Slate/Fonts/Roboto-BoldCondensed.ttf" ), 12 ) )
+					//.Font( FSlateFontInfo( FPaths::EngineContentDir() / TEXT("Slate/Fonts/Roboto-Bold.ttf" ), 12 ) )
 					.ToolTipText( LOCTEXT("TestingSearchBox_Tooltip", "Testing search boxes") )
 				]
 				+ SVerticalBox::Slot().AutoHeight() .HAlign(HAlign_Center) .VAlign(VAlign_Center) .Padding( 5 )
 				[
 					SAssignNew( NumericInput, SEditableTextBox )
-					.Text( LOCTEXT( "NumericInput", "This should be an integer" ) )
+					.Text( LOCTEXT( "NumericInput", "This should be a number" ) )
 					.OnTextChanged( this, &STextEditTest::OnNumericInputTextChanged )
 					.RevertTextOnEscape( true )
-					.HintText(LOCTEXT("NumericInputHint", "Enter an integer"))
+					.HintText(LOCTEXT("NumericInputHint", "Enter a number"))
 				]
 				+ SVerticalBox::Slot().AutoHeight() .HAlign(HAlign_Center) .VAlign(VAlign_Center) .Padding( 5 )
 				[
@@ -1650,17 +1650,28 @@ public:
 					SNew(SHorizontalBox)
 					+SHorizontalBox::Slot().AutoWidth() .VAlign(VAlign_Top) .Padding(2)
 					[
-						SNew(SBorder)
+						SNew(SVerticalBox)
+						+ SVerticalBox::Slot()
+						.AutoHeight()
+						[
+							SNew(STextBlock)
+							.Text(LOCTEXT("MultiLineTextWrapAt600px", "Multi-line editable text wrapping at 600px"))
+						]
+						+ SVerticalBox::Slot()
+						.AutoHeight()
 						[
 							SNew(SHorizontalBox)
-							+SHorizontalBox::Slot()
+							+ SHorizontalBox::Slot()
 							[
-								SNew(SMultiLineEditableText)
-								.Text(MultiLineText)
-								.ScrollBar( FirstScrollBar )
-								.PreferredSize( FVector2D(100,175) )
+								SNew(SMultiLineEditableTextBox)
+								.Text(this, &STextEditTest::GetMultilineEditableText)
+								.Font(FSlateFontInfo(FPaths::EngineContentDir() / TEXT("Slate/Fonts/Roboto-Regular.ttf"), 12))
+								.Justification(ETextJustify::Center)
+								.LineHeightPercentage(2.0f)
+								.OnTextCommitted(this, &STextEditTest::HandleMultilineEditableTextCommitted)
+								.WrapTextAt(600.0f)
 							]
-							+SHorizontalBox::Slot() .AutoWidth()
+							+ SHorizontalBox::Slot().AutoWidth()
 							[
 								FirstScrollBar
 							]
@@ -1668,14 +1679,24 @@ public:
 					]
 					+SHorizontalBox::Slot().FillWidth(1).Padding(2)
 					[
-						SNew(SBorder)
+						SNew(SVerticalBox)
+						+ SVerticalBox::Slot()
+						.AutoHeight()
+						[
+							SNew(STextBlock)
+							.Text(LOCTEXT("MultiLineTextAutoWrap", "Multi-line editable text auto-wrapping"))
+						]
+						+ SVerticalBox::Slot()
+						.AutoHeight()
 						[
 							SNew(SHorizontalBox)
 							+SHorizontalBox::Slot()
 							[
-								SNew(SMultiLineEditableText)
-								.Text(MultiLineText)
-								.ScrollBar( SecondScrollBar )
+								SNew(SMultiLineEditableTextBox)
+								.Text(MultilineEditableText)
+								.Font(FSlateFontInfo(FPaths::EngineContentDir() / TEXT("Slate/Fonts/Roboto-Regular.ttf"), 12))
+								.AutoWrapText(true)
+								.HintText(LOCTEXT("TypehereTextHint", "Type Here"))
 							]
 							+SHorizontalBox::Slot() .AutoWidth()
 							[
@@ -1684,11 +1705,46 @@ public:
 						]
 					]
 				]
+				+SVerticalBox::Slot().AutoHeight()
+				[
+					SNew(SHorizontalBox)
+					+SHorizontalBox::Slot().AutoWidth() .VAlign(VAlign_Top) .Padding(2)
+					[
+						SNew(SVerticalBox)
+						+ SVerticalBox::Slot()
+						.AutoHeight()
+						[
+							SNew(STextBlock)
+							.Text(LOCTEXT("MultiLineTextNoWrap", "Multi-line editable text without wrapping"))
+						]
+						+ SVerticalBox::Slot()
+						.AutoHeight()
+						[
+							SNew(SHorizontalBox)
+							+ SHorizontalBox::Slot()
+							[
+								SNew(SMultiLineEditableTextBox)
+								.Font(FSlateFontInfo(FPaths::EngineContentDir() / TEXT("Slate/Fonts/Roboto-Regular.ttf"), 12))
+								//.WrapTextAt(300.0f)
+							]
+						]
+					]
+				]
 			]
 		];
 	}
 	END_SLATE_FUNCTION_BUILD_OPTIMIZATION
 
+	FText GetMultilineEditableText() const
+	{
+		return MultilineEditableText;
+	}
+
+	void HandleMultilineEditableTextCommitted(const FText& Text, ETextCommit::Type Type)
+	{
+		MultilineEditableText = Text;
+	}
+
 	void FocusDefaultWidget()
 	{
 		// Set focus to the editable text, so the user doesn't have to click initially to type
@@ -1787,6 +1843,8 @@ protected:
 
 	TSharedPtr<SInlineEditableTextBlock> InlineEditableTextBlock;
 	FText InlineEditableText;
+
+	FText MultilineEditableText;
  };
 
 /** Demonstrates the brokenness of our current approach to trading smoothness for sharpness. */
diff --git a/Engine/Source/Runtime/Slate/Private/Widgets/Text/SMultiLineEditableText.cpp b/Engine/Source/Runtime/Slate/Private/Widgets/Text/SMultiLineEditableText.cpp
--- a/Engine/Source/Runtime/Slate/Private/Widgets/Text/SMultiLineEditableText.cpp
+++ b/Engine/Source/Runtime/Slate/Private/Widgets/Text/SMultiLineEditableText.cpp
@@ -4,139 +4,1124 @@
 #include "TextEditHelper.h"
 #include "SlateWordWrapper.h"
 
+void SMultiLineEditableText::FCursorInfo::SetCursorLocationAndCalculateAlignment(const TSharedPtr<FTextLayout>& TextLayout, const FTextLocation& InCursorPosition, const ECursorEOLMode CursorEOLMode)
+{
+	FTextLocation NewCursorPosition = InCursorPosition;
+	ECursorAlignment NewAlignment = ECursorAlignment::Left;
+
+	const int32 CursorLineIndex = InCursorPosition.GetLineIndex();
+	const int32 CursorOffset = InCursorPosition.GetOffset();
+
+	// A CursorOffset of zero could mark the end of an empty line, but we don't need to adjust the cursor for an empty line
+	if (TextLayout.IsValid() && CursorOffset > 0)
+	{
+		bool bHasFoundEOL = false;
+
+		if (CursorEOLMode == ECursorEOLMode::HardLines)
+		{
+			const TArray< FTextLayout::FLineModel >& Lines = TextLayout->GetLineModels();
+			const FTextLayout::FLineModel& Line = Lines[CursorLineIndex];
+			bHasFoundEOL = (Line.Text->Len() == CursorOffset);
+		}
+		else
+		{
+			const TArray< FTextLayout::FLineView >& LineViews = TextLayout->GetLineViews();
+			const int32 LineViewIndex = TextLayout->GetLineViewIndexForTextLocation(LineViews, InCursorPosition);
+
+			if (LineViews.IsValidIndex(LineViewIndex))
+			{
+				const FTextLayout::FLineView& CurrentLineView = LineViews[LineViewIndex];
+				bHasFoundEOL = (CurrentLineView.Range.EndIndex == CursorOffset);
+			}
+		}
+
+		if (bHasFoundEOL)
+		{
+			// We need to move the cursor back one from where it currently is; this keeps the interaction point the same 
+			// (since the cursor is aligned to the right), but visually draws the cursor in the correct place
+			NewCursorPosition = FTextLocation(NewCursorPosition, -1);
+			NewAlignment = ECursorAlignment::Right;
+		}
+	}
+
+	SetCursorLocationAndAlignment(NewCursorPosition, NewAlignment);
+}
+
+void SMultiLineEditableText::FCursorInfo::SetCursorLocationAndAlignment(const FTextLocation& InCursorPosition, const ECursorAlignment InCursorAlignment)
+{
+	CursorPosition = InCursorPosition;
+	CursorAlignment = InCursorAlignment;
+	LastCursorInteractionTime = FSlateApplication::Get().GetCurrentTime();
+}
+
+SMultiLineEditableText::FCursorInfo SMultiLineEditableText::FCursorInfo::CreateUndo() const
+{
+	FCursorInfo UndoData;
+	UndoData.CursorPosition = CursorPosition;
+	UndoData.CursorAlignment = CursorAlignment;
+	return UndoData;
+}
+
+void SMultiLineEditableText::FCursorInfo::RestoreFromUndo(const FCursorInfo& UndoData)
+{
+	// Use SetCursorLocationAndAlignment since it updates LastCursorInteractionTime
+	SetCursorLocationAndAlignment(UndoData.CursorPosition, UndoData.CursorAlignment);
+}
+
+FNoChildren SMultiLineEditableText::FSlateCursorRunHighlighter::NoChildrenInstance;
+
+SMultiLineEditableText::FSlateCursorRunHighlighter::FSlateCursorRunHighlighter(const FCursorInfo* InCursorInfo)
+	: CursorInfo(InCursorInfo)
+{
+	check(CursorInfo);
+}
+
+void SMultiLineEditableText::FSlateCursorRunHighlighter::OnArrangeChildren( const TSharedRef< ILayoutBlock >& Block, const FGeometry& AllottedGeometry, FArrangedChildren& ArrangedChildren ) const 
+{
+	//No Widgets
+}
+
+FChildren* SMultiLineEditableText::FSlateCursorRunHighlighter::GetChildren()
+{
+	return &NoChildrenInstance;
+}
+
+int32 SMultiLineEditableText::FSlateCursorRunHighlighter::OnPaint( const FTextLayout::FLineView& Line, const TSharedRef< ISlateRun >& Run, const TSharedRef< ILayoutBlock >& Block, const FTextBlockStyle& DefaultStyle, const FGeometry& AllottedGeometry, const FSlateRect& MyClippingRect, FSlateWindowElementList& OutDrawElements, int32 LayerId, const FWidgetStyle& InWidgetStyle, bool bParentEnabled ) const 
+{
+	FVector2D Location( Block->GetLocationOffset() );
+	Location.Y = Line.Offset.Y;
+
+	if ( Mode == EMode::Selection )
+	{
+		const FColor SelectionBackgroundColorAndOpacity = ((FLinearColor::White - DefaultStyle.ColorAndOpacity.GetColor(InWidgetStyle))*0.5f + FLinearColor(-0.2f, -0.05f, 0.15f)) * InWidgetStyle.GetColorAndOpacityTint();
+
+		float HighlightWidth = Block->GetSize().X;
+		if (HighlightWidth == 0)
+		{
+			const TSharedRef< FSlateFontMeasure > FontMeasureService = FSlateApplication::Get().GetRenderer()->GetFontMeasureService();
+			const FVector2D SpaceSize = FontMeasureService->Measure(FString(TEXT("  ")), DefaultStyle.Font);
+			HighlightWidth = SpaceSize.X;
+		}
+
+		// Draw the actual highlight rectangle
+		FSlateDrawElement::MakeBox(
+			OutDrawElements,
+			++LayerId,
+			FPaintGeometry(AllottedGeometry.AbsolutePosition + Location, FVector2D(HighlightWidth, Line.Size.Y), AllottedGeometry.Scale),
+			&DefaultStyle.HighlightShape,
+			MyClippingRect,
+			bParentEnabled ? ESlateDrawEffect::None : ESlateDrawEffect::DisabledEffect,
+			SelectionBackgroundColorAndOpacity
+		);
+
+		//FLinearColor InvertedForeground = FLinearColor::White - InWidgetStyle.GetForegroundColor();
+		//InvertedForeground.A = InWidgetStyle.GetForegroundColor().A;
+
+		//FWidgetStyle WidgetStyle( InWidgetStyle );
+		//WidgetStyle.SetForegroundColor( InvertedForeground );
+
+		return Run->OnPaint( Line, Block, DefaultStyle, AllottedGeometry, MyClippingRect, OutDrawElements, LayerId, InWidgetStyle, bParentEnabled );
+	}
+	else
+	{
+		FLinearColor CursorColorAndOpacity = InWidgetStyle.GetForegroundColor();
+
+		const float FontMaxCharHeight = FTextEditHelper::GetFontHeight(DefaultStyle.Font);
+		const float CursorWidth = FTextEditHelper::CalculateCaretWidth(FontMaxCharHeight);
+		const double CurrentTime = FSlateApplication::Get().GetCurrentTime();
+
+		// The cursor is always visible (i.e. not blinking) when we're interacting with it; otherwise it might get lost.
+		const bool bForceCursorVisible = (CurrentTime - CursorInfo->GetLastCursorInteractionTime()) < EditableTextDefs::CaretBlinkPauseTime;
+		float CursorOpacity = (bForceCursorVisible)
+			? 1.0f
+			: FMath::RoundToFloat( FMath::MakePulsatingValue( CurrentTime, EditableTextDefs::BlinksPerSecond ));
+
+		CursorOpacity *= CursorOpacity;	// Squared falloff, because it looks more interesting
+		CursorColorAndOpacity.A = CursorOpacity;
+
+		// @todo: Slate Styles - make this brush part of the widget style
+		const FSlateBrush* StyleInfo = FCoreStyle::Get().GetBrush("EditableText.SelectionBackground");
+
+		const FVector2D OptionalWidth = CursorInfo->GetCursorAlignment() == ECursorAlignment::Right ? FVector2D(Block->GetSize().X, 0) : FVector2D::ZeroVector;
+
+		FSlateDrawElement::MakeBox(
+			OutDrawElements,
+			LayerId,
+			FPaintGeometry(AllottedGeometry.AbsolutePosition + Location + OptionalWidth, FVector2D(CursorWidth * AllottedGeometry.Scale, Line.TextSize.Y), AllottedGeometry.Scale),
+			StyleInfo,
+			MyClippingRect,
+			bParentEnabled ? ESlateDrawEffect::None : ESlateDrawEffect::DisabledEffect,
+			CursorColorAndOpacity*InWidgetStyle.GetColorAndOpacityTint());
+
+		return Run->OnPaint( Line, Block, DefaultStyle, AllottedGeometry, MyClippingRect, OutDrawElements, LayerId, InWidgetStyle, bParentEnabled );
+	}
+}
+
+TSharedRef< SMultiLineEditableText::FSlateCursorRunHighlighter > SMultiLineEditableText::FSlateCursorRunHighlighter::Create(const FCursorInfo* InCursorInfo)
+{
+	return MakeShareable( new SMultiLineEditableText::FSlateCursorRunHighlighter(InCursorInfo) );
+}
+
+void SMultiLineEditableText::FSlateCursorRunHighlighter::SetMode( EMode InMode )
+{
+	Mode = InMode;
+}
+
 SMultiLineEditableText::SMultiLineEditableText()
-: CursorPosition( FTextLocation() )
-, FontInfo( FPaths::EngineContentDir() / TEXT("Slate/Fonts/Roboto-Regular.ttf"), 10 ) 
-, NumLinesScrollOffset( 0.0f )
-, PreferredCursorOffsetInLine( 0 )
-, LastCursorInteractionTime( 0 )
-, PreferredSize( FVector2D::ZeroVector )
+	: PreferredCursorScreenOffsetInLine(0)
+	, bSelectAllTextWhenFocused(false)
+	, bIsDragSelecting(false)
+	, bWasFocusedByLastMouseDown(false)
+	, bHasDragSelectedSinceFocused(false)
+	, CurrentUndoLevel(INDEX_NONE)
+	, bIsChangingText(false)
+	, IsReadOnly(false)
+	, UICommandList(new FUICommandList())
 {
+	
 }
 
 void SMultiLineEditableText::Construct( const FArguments& InArgs )
 {
-	NumLinesScrollOffset = 0.0f;
-	OptionalScrollBarPtr = InArgs._ScrollBar;
-	PreferredSize = InArgs._PreferredSize;
-	if (InArgs._ScrollBar.IsValid())
+	CursorInfo = FCursorInfo();
+	SelectionStart = TOptional<FTextLocation>();
+	PreferredCursorScreenOffsetInLine = 0;
+
+	WrapTextAt = InArgs._WrapTextAt;
+	AutoWrapText = InArgs._AutoWrapText;
+	CachedAutoWrapTextWidth = 0.0f;
+
+	Margin = InArgs._Margin;
+	LineHeightPercentage = InArgs._LineHeightPercentage;
+	Justification = InArgs._Justification;
+	TextStyle = *InArgs._TextStyle;
+	
+	if (InArgs._Font.IsSet() || InArgs._Font.IsBound())
 	{
-		InArgs._ScrollBar->SetOnUserScrolled( FOnUserScrolled::CreateSP( this, &SMultiLineEditableText::OnScrollbarScrolled ) );
+		TextStyle.SetFont(InArgs._Font.Get());
+	}
+
+	IsReadOnly = InArgs._IsReadOnly;
+
+	OnTextChanged = InArgs._OnTextChanged;
+	OnTextCommitted = InArgs._OnTextCommitted;
+
+	CursorSelectionHighlighter = FSlateCursorRunHighlighter::Create(&CursorInfo);
+	TextLayout = FSlateTextLayout::Create();
+
+	BoundText = InArgs._Text;
+
+	SetText(BoundText.Get(FText::GetEmpty()));
+
+	TextLayout->SetWrappingWidth( WrapTextAt.Get() );
+	TextLayout->SetMargin( Margin.Get() );
+	TextLayout->SetLineHeightPercentage( LineHeightPercentage.Get() );
+	TextLayout->SetJustification( Justification.Get() );
+
+	TextLayout->UpdateIfNeeded();
+
+	// Map UI commands to delegates which are called when the command should be executed
+	UICommandList->MapAction(FGenericCommands::Get().Undo,
+		FExecuteAction::CreateSP(this, &SMultiLineEditableText::Undo),
+		FCanExecuteAction::CreateSP(this, &SMultiLineEditableText::CanExecuteUndo));
+
+	UICommandList->MapAction(FGenericCommands::Get().Cut,
+		FExecuteAction::CreateSP(this, &SMultiLineEditableText::CutSelectedTextToClipboard),
+		FCanExecuteAction::CreateSP(this, &SMultiLineEditableText::CanExecuteCut));
+
+	UICommandList->MapAction(FGenericCommands::Get().Paste,
+		FExecuteAction::CreateSP(this, &SMultiLineEditableText::PasteTextFromClipboard),
+		FCanExecuteAction::CreateSP(this, &SMultiLineEditableText::CanExecutePaste));
+
+	UICommandList->MapAction(FGenericCommands::Get().Copy,
+		FExecuteAction::CreateSP(this, &SMultiLineEditableText::CopySelectedTextToClipboard),
+		FCanExecuteAction::CreateSP(this, &SMultiLineEditableText::CanExecuteCopy));
+
+	UICommandList->MapAction(FGenericCommands::Get().Delete,
+		FExecuteAction::CreateSP(this, &SMultiLineEditableText::ExecuteDeleteAction),
+		FCanExecuteAction::CreateSP(this, &SMultiLineEditableText::CanExecuteDelete));
+
+	UICommandList->MapAction(FGenericCommands::Get().SelectAll,
+		FExecuteAction::CreateSP(this, &SMultiLineEditableText::SelectAllText),
+		FCanExecuteAction::CreateSP(this, &SMultiLineEditableText::CanExecuteSelectAll));
+
+	// build context menu extender
+	MenuExtender = MakeShareable(new FExtender);
+	MenuExtender->AddMenuExtension("EditText", EExtensionHook::Before, TSharedPtr<FUICommandList>(), InArgs._ContextMenuExtender);
+}
+
+void SMultiLineEditableText::SetText(const TAttribute< FText >& InText)
+{
+	const bool bHasPreviouslySetText = LastSetText.IsSet();
+	LastSetText = InText.Get();
+
+	FString EditedText;
+	TextLayout->GetAsText(EditedText);
+
+	check(LastSetText.IsSet());
+	const FText& LastSetTextValue = LastSetText.GetValue();
+	const FString& TextString = LastSetTextValue.ToString();
+	const bool HasTextChanged = !bHasPreviouslySetText || EditedText != TextString;
+
+	if (HasTextChanged)
+	{
+		const int32 TextLength = TextString.Len();
+
+		TextLayout->ClearLines();
+
+		if(TextLength)
+		{
+			// Iterate over each line break candidate, adding ranges from after the end of the last line added to before the newline or end of string.
+			FLineBreakIterator LBI(TextString);
+			int32 RangeBegin = LBI.GetCurrentPosition();
+			TSharedPtr< FString > LineText;
+
+			for (;;)
+			{
+				const int32 BreakIndex = LBI.MoveToNext();
+
+				// If beginning or potential end is invalid, cease.
+				if (RangeBegin == INDEX_NONE || BreakIndex == INDEX_NONE)
+				{
+					break;
+				}
+
+				const int32 BreakingCharacterIndex = BreakIndex - 1;
+				if (BreakingCharacterIndex >= RangeBegin) // Valid index check.
+				{
+					const TCHAR BreakingCharacter = TextString[BreakingCharacterIndex];
+					// If line break candidate is after a newline, add the range as a new line.
+					if (FChar::IsLinebreak(BreakingCharacter))
+					{
+						LineText = MakeShareable(new FString(BreakingCharacterIndex - RangeBegin, (*TextString) + RangeBegin));
+					}
+					// If the line break candidate is the end of the string, add the range as a new line.
+					else if (BreakIndex == TextLength)
+					{
+						LineText = MakeShareable(new FString(BreakIndex - RangeBegin, (*TextString) + RangeBegin));
+					}
+					else
+					{
+						continue;
+					}
+
+					if (LineText.IsValid())
+					{
+						// Create run.
+						TArray< TSharedRef< IRun > > Runs;
+						Runs.Add(FSlateTextRun::Create(LineText.ToSharedRef(), TextStyle));
+
+						TextLayout->AddLine(LineText.ToSharedRef(), Runs);
+
+						// Reset the LineText so we don't add it multiple times.
+						LineText.Reset();
+					}
+
+					RangeBegin = BreakIndex; // The next line begins after the end of the current line.
+				}
+			}
+		}
+		else
+		{
+			TSharedPtr< FString > LineText = MakeShareable(new FString());
+
+			// Create an empty run
+			TArray< TSharedRef< IRun > > Runs;
+			Runs.Add(FSlateTextRun::Create(LineText.ToSharedRef(), TextStyle));
+
+			TextLayout->AddLine(LineText.ToSharedRef(), Runs);
+		}
+
+		// Let outsiders know that the text content has been changed
+		OnTextChanged.ExecuteIfBound(LastSetTextValue);
 	}
+}
 
-	TArray<FString> TempTextLines;
-	InArgs._Text.ParseIntoArray( &TempTextLines, TEXT("\n"), false );
-	TextLines.Reserve( TempTextLines.Num() );
+FReply SMultiLineEditableText::OnKeyboardFocusReceived( const FGeometry& MyGeometry, const FKeyboardFocusEvent& InKeyboardFocusEvent )
+{
+	UpdateCursorHighlight();
+	return SWidget::OnKeyboardFocusReceived( MyGeometry, InKeyboardFocusEvent );
+}
 
-	for ( int32 LineIndex=0; LineIndex < TempTextLines.Num(); ++LineIndex )
+void SMultiLineEditableText::OnKeyboardFocusLost( const FKeyboardFocusEvent& InKeyboardFocusEvent )
+{
+	// Skip the focus lost code if it's due to the context menu opening
+	if (!ContextMenuWindow.IsValid())
 	{
-		TextLines.Add( MakeShareable( new FString( TempTextLines[LineIndex] ) ) );
+		//ITextInputMethodSystem* const TextInputMethodSystem = FSlateApplication::Get().GetTextInputMethodSystem();
+		//if (TextInputMethodSystem)
+		//{
+		//	TextInputMethodSystem->DeactivateContext(TextInputMethodContext.ToSharedRef());
+		//}
+
+		// When focus is lost let anyone who is interested that text was committed
+		if (InKeyboardFocusEvent.GetCause() != EKeyboardFocusCause::WindowActivate) //Don't clear selection when activating a new window or cant copy and paste on right click )
+		{
+			ClearSelection();
+		}
+
+		// See if user explicitly tabbed away or moved focus
+		ETextCommit::Type TextAction;
+		switch (InKeyboardFocusEvent.GetCause())
+		{
+		case EKeyboardFocusCause::Keyboard:
+		case EKeyboardFocusCause::Mouse:
+			TextAction = ETextCommit::OnUserMovedFocus;
+			break;
+
+		case EKeyboardFocusCause::Cleared:
+			TextAction = ETextCommit::OnCleared;
+			break;
+
+		default:
+			TextAction = ETextCommit::Default;
+			break;
+		}
+
+		//Always clear the local undo chain on commit.
+		ClearUndoStates();
+
+		FString EditedText;
+		TextLayout->GetAsText(EditedText);
+
+		OnTextCommitted.ExecuteIfBound(FText::FromString(EditedText), TextAction);
+		// todo: jdale - testing
+		//RemoveCursorHighlight();
 	}
 }
 
 void SMultiLineEditableText::StartChangingText()
 {
+	// Never change text on read only controls! 
+	check(!IsReadOnly.Get());
+
+	// We're starting to (potentially) change text
+	check(!bIsChangingText);
+	bIsChangingText = true;
+
+	// Save off an undo state in case we actually change the text
+	MakeUndoState(StateBeforeChangingText);
 }
 
 void SMultiLineEditableText::FinishChangingText()
 {
+	// We're no longer changing text
+	check(bIsChangingText);
+	bIsChangingText = false;
+
+	FString EditedTextString;
+	TextLayout->GetAsText(EditedTextString);
+	FText EditedText = FText::FromString(EditedTextString);
+
+	// Has the text changed?
+	if (EditedText.ToString() != StateBeforeChangingText.Text.ToString())
+	{
+		// Save text state
+		SaveText(EditedText);
+
+		// Text was actually changed, so push the undo state we previously saved off
+		PushUndoState(StateBeforeChangingText);
+
+		// We're done with this state data now.  Clear out any old data.
+		StateBeforeChangingText = FUndoState();
+		
+		LastSetText = EditedText;
+
+		// Let outsiders know that the text content has been changed
+		OnTextChanged.ExecuteIfBound(EditedText);
+	}
+}
+
+void SMultiLineEditableText::SaveText(const FText& TextToSave)
+{
+	// Don't set text if the text attribute has a 'getter' binding on it, otherwise we'd blow away
+	// that binding.  If there is a getter binding, then we'll assume it will provide us with
+	// updated text after we've fired our 'text changed' callbacks
+	if (!BoundText.IsBound())
+	{
+		BoundText.Set(TextToSave);
+	}
+}
+
+bool SMultiLineEditableText::GetIsReadOnly() const
+{
+	return IsReadOnly.Get();
+}
+
+void SMultiLineEditableText::BackspaceChar()
+{
+	if( AnyTextSelected() )
+	{
+		// Delete selected text
+		DeleteSelectedText();
+	}
+	else
+	{
+		const FTextLocation CursorInteractionPosition = CursorInfo.GetCursorInteractionLocation();
+		FTextLocation FinalCursorLocation = CursorInteractionPosition;
+
+		const TArray< FTextLayout::FLineModel >& Lines = TextLayout->GetLineModels();
+		const FTextLayout::FLineModel& Line = Lines[CursorInteractionPosition.GetLineIndex()];
+
+		//If we are at the very beginning of the line...
+		if (CursorInteractionPosition.GetOffset() == 0)
+		{
+			//And if the current line isn't the very first line then...
+			if (CursorInteractionPosition.GetLineIndex() > 0)
+			{
+				int32 CachePreviousLinesCurrentLength = Lines[CursorInteractionPosition.GetLineIndex() - 1].Text->Len();
+				if (TextLayout->JoinLineWithNextLine(CursorInteractionPosition.GetLineIndex() - 1))
+				{
+					//Update the cursor so it appears at the end of the previous line,
+					//as we're going to delete the imaginary \n separating them
+					FinalCursorLocation = FTextLocation(CursorInteractionPosition.GetLineIndex() - 1, CachePreviousLinesCurrentLength);
+				}
+			}
+			//else do nothing as the FinalCursorLocation is already correct
+		}
+		//else if (CursorInteractionPosition.GetOffset() == Line.Text->Len())
+		//{
+		//	//And if the current line isn't the very last line then...
+		//	if (Lines.IsValidIndex(CursorInteractionPosition.GetLineIndex() + 1))
+		//	{
+		//		//do nothing to the cursor as the FinalCursorLocation is already correct
+		//		TextLayout->JoinLineWithNextLine(CursorInteractionPosition.GetLineIndex());
+		//	}
+		//	//else do nothing as the FinalCursorLocation is already correct
+		//}
+		else
+		{
+			// Delete character to the left of the caret
+			if (TextLayout->RemoveAt(FTextLocation(CursorInteractionPosition, -1)))
+			{
+				// Adjust caret position one to the left
+				FinalCursorLocation = FTextLocation(CursorInteractionPosition, -1);
+			}
+		}
+
+		TextLayout->UpdateIfNeeded();
+
+		CursorInfo.SetCursorLocationAndCalculateAlignment(TextLayout, FinalCursorLocation);
+
+		ClearSelection();
+		UpdateCursorHighlight();
+	}
+}
+
+/**
+ * This gets executed by the context menu and should only attempt to delete any selected text
+ */
+void SMultiLineEditableText::ExecuteDeleteAction()
+{
+	if (!IsReadOnly.Get() && AnyTextSelected())
+	{
+		StartChangingText();
+
+		// Delete selected text
+		DeleteSelectedText();
+
+		FinishChangingText();
+	}
+}
+
+bool SMultiLineEditableText::CanExecuteDelete() const
+{
+	bool bCanExecute = true;
+
+	// Can't execute if this is a read-only control
+	if (IsReadOnly.Get())
+	{
+		bCanExecute = false;
+	}
+
+	// Can't execute unless there is some text selected
+	if (!AnyTextSelected())
+	{
+		bCanExecute = false;
+	}
+
+	return bCanExecute;
+}
+
+void SMultiLineEditableText::DeleteChar()
+{
+	if( !IsReadOnly.Get() )
+	{
+		if( AnyTextSelected() )
+		{
+			// Delete selected text
+			DeleteSelectedText();
+		}
+		else
+		{
+			const FTextLocation CursorInteractionPosition = CursorInfo.GetCursorInteractionLocation();
+			FTextLocation FinalCursorLocation = CursorInteractionPosition;
+
+			const TArray< FTextLayout::FLineModel >& Lines = TextLayout->GetLineModels();
+			const FTextLayout::FLineModel& Line = Lines[CursorInteractionPosition.GetLineIndex()];
+
+			//If we are at the very beginning of the line...
+			if (Line.Text->Len() == 0)
+			{
+				//And if the current line isn't the very first line then...
+				if (CursorInteractionPosition.GetLineIndex() > 0)
+				{
+					if (TextLayout->RemoveLine(CursorInteractionPosition.GetLineIndex()))
+					{
+						//Update the cursor so it appears at the end of the previous line,
+						//as we're going to delete the imaginary \n separating them
+						FinalCursorLocation = FTextLocation(CursorInteractionPosition.GetLineIndex() - 1, Lines[CursorInteractionPosition.GetLineIndex() - 1].Text->Len());
+					}
+				}
+				//else do nothing as the FinalCursorLocation is already correct
+			}
+			else if (CursorInteractionPosition.GetOffset() >= Line.Text->Len())
+			{
+				//And if the current line isn't the very last line then...
+				if (Lines.IsValidIndex(CursorInteractionPosition.GetLineIndex() + 1))
+				{
+					if (TextLayout->JoinLineWithNextLine(CursorInteractionPosition.GetLineIndex()))
+					{
+						//else do nothing as the FinalCursorLocation is already correct
+					}
+				}
+				//else do nothing as the FinalCursorLocation is already correct
+			}
+			else
+			{
+				// Delete character to the right of the caret
+				TextLayout->RemoveAt(CursorInteractionPosition);
+				//do nothing to the cursor as the FinalCursorLocation is already correct
+			}
+
+			TextLayout->UpdateIfNeeded();
+
+			CursorInfo.SetCursorLocationAndCalculateAlignment(TextLayout, FinalCursorLocation);
+
+			ClearSelection();
+			UpdateCursorHighlight();
+		}
+	}
+}
+
+bool SMultiLineEditableText::CanTypeCharacter(const TCHAR CharInQuestion) const
+{
+	return true;
+}
+
+void SMultiLineEditableText::TypeChar( const int32 Character )
+{
+	if( AnyTextSelected() )
+	{
+		// Delete selected text
+		DeleteSelectedText();
+	}
+
+	// Certain characters are not allowed
+	bool bIsCharAllowed = true;
+	{
+		if( Character <= 0x1F )
+		{
+			bIsCharAllowed = false;
+		}
+	}
+
+	if( bIsCharAllowed )
+	{
+		const FTextLocation CursorInteractionPosition = CursorInfo.GetCursorInteractionLocation();
+		const TArray< FTextLayout::FLineModel >& Lines = TextLayout->GetLineModels();
+		const FTextLayout::FLineModel& Line = Lines[ CursorInteractionPosition.GetLineIndex() ];
+
+		// Insert character at caret position
+		TextLayout->InsertAt( CursorInteractionPosition, Character );
+
+		// Advance caret position
+		ClearSelection();
+		const FTextLocation FinalCursorLocation = FTextLocation( CursorInteractionPosition.GetLineIndex(), FMath::Min( CursorInteractionPosition.GetOffset() + 1, Line.Text->Len() ) );
+
+		TextLayout->UpdateIfNeeded();
+
+		CursorInfo.SetCursorLocationAndCalculateAlignment(TextLayout, FinalCursorLocation);
+		UpdateCursorHighlight();
+	}
+}
+
+void SMultiLineEditableText::DeleteSelectedText()
+{
+	if( AnyTextSelected() )
+	{
+		const FTextLocation CursorInteractionPosition = CursorInfo.GetCursorInteractionLocation();
+		FTextLocation SelectionLocation = SelectionStart.Get(CursorInteractionPosition);
+		FTextSelection Selection(SelectionLocation, CursorInteractionPosition);
+
+		int32 SelectionBeginningLineIndex = Selection.GetBeginning().GetLineIndex();
+		int32 SelectionBeginningLineOffset = Selection.GetBeginning().GetOffset();
+
+		int32 SelectionEndLineIndex = Selection.GetEnd().GetLineIndex();
+		int32 SelectionEndLineOffset = Selection.GetEnd().GetOffset();
+
+		if (SelectionBeginningLineIndex == SelectionEndLineIndex)
+		{
+			TextLayout->RemoveAt(FTextLocation(SelectionBeginningLineIndex, SelectionBeginningLineOffset), SelectionEndLineOffset - SelectionBeginningLineOffset);
+			//Do nothing to the cursor as it is already in the correct location
+		}
+		else
+		{
+			const TArray< FTextLayout::FLineModel >& Lines = TextLayout->GetLineModels();
+			const FTextLayout::FLineModel& EndLine = Lines[SelectionEndLineIndex];
+
+			if (EndLine.Text->Len() == SelectionEndLineOffset)
+			{ 
+				TextLayout->RemoveLine(SelectionEndLineIndex);
+			}
+			else
+			{
+				TextLayout->RemoveAt(FTextLocation(SelectionEndLineIndex, 0), SelectionEndLineOffset);
+			}
+
+			for (int32 LineIndex = SelectionEndLineIndex - 1; LineIndex > SelectionBeginningLineIndex; LineIndex--)
+			{
+				TextLayout->RemoveLine(LineIndex);
+			}
+
+			const FTextLayout::FLineModel& BeginLine = Lines[SelectionBeginningLineIndex];
+
+			if (SelectionBeginningLineOffset == 0)
+			{
+				TextLayout->RemoveLine(SelectionBeginningLineIndex);
+			}
+			else
+			{
+				TextLayout->RemoveAt(FTextLocation(SelectionBeginningLineIndex, SelectionBeginningLineOffset), BeginLine.Text->Len() - SelectionBeginningLineOffset);
+			}
+
+			TextLayout->JoinLineWithNextLine(SelectionBeginningLineIndex);
+
+			if (Lines.Num() == 0)
+			{
+				const TSharedRef< FString > EmptyText = MakeShareable(new FString());
+				TArray< TSharedRef< IRun > > Runs;
+				Runs.Add(FSlateTextRun::Create(EmptyText, TextStyle));
+
+				TextLayout->AddLine(EmptyText, Runs);
+			}
+		}
+
+		// Clear selection
+		ClearSelection();
+		const FTextLocation FinalCursorLocation = FTextLocation(SelectionBeginningLineIndex, SelectionBeginningLineOffset);
+		CursorInfo.SetCursorLocationAndCalculateAlignment(TextLayout, FinalCursorLocation);
+		UpdateCursorHighlight();
+
+		TextLayout->UpdateIfNeeded();
+	}
+}
+
+FReply SMultiLineEditableText::MoveCursor( ECursorMoveMethod::Type Method, const FVector2D& Direction, ECursorAction::Type Action )
+{
+	FTextLocation NewCursorPosition;
+	const FTextLocation CursorPosition = CursorInfo.GetCursorInteractionLocation();
+	ECursorEOLMode CursorEOLMode = ECursorEOLMode::HardLines;
+
+	switch (Method)
+	{
+	case ECursorMoveMethod::CharacterHorizontal:
+		NewCursorPosition = TranslatedLocation( CursorPosition, Direction.X );
+		PreferredCursorScreenOffsetInLine = TextLayout->GetLocationAt( NewCursorPosition ).X;
+		break;
+
+	case ECursorMoveMethod::Word:
+		NewCursorPosition = ScanForWordBoundary( CursorPosition, Direction.X );
+		PreferredCursorScreenOffsetInLine = TextLayout->GetLocationAt( NewCursorPosition ).X;
+		break;
+	
+	case ECursorMoveMethod::CharacterVertical:
+		NewCursorPosition = TranslateLocationVertical( CursorPosition, Direction.Y );
+		break;
+	
+	case ECursorMoveMethod::ScreenPosition:
+		NewCursorPosition = TextLayout->GetTextLocationAt( Direction * TextLayout->GetScale() );
+		CursorEOLMode = ECursorEOLMode::SoftLines;
+		PreferredCursorScreenOffsetInLine = Direction.X;
+		break;
+
+	default:
+		checkSlow(false, "Unknown ECursorMoveMethod value");
+		break;
+	}
+
+	if (Action == ECursorAction::SelectText)
+	{
+		// We are selecting text. Just remember where the selection started.
+		// The cursor is implicitly the other endpoint.
+		if (!SelectionStart.IsSet())
+		{
+			this->SelectionStart = CursorPosition;
+		}
+	}
+	else
+	{
+		// No longer selection text; clear the selection!
+		this->ClearSelection();
+	}
+
+	CursorInfo.SetCursorLocationAndCalculateAlignment(TextLayout, NewCursorPosition, CursorEOLMode);
+	UpdateCursorHighlight();
+
+	return FReply::Handled();
+}
+
+void SMultiLineEditableText::UpdateCursorHighlight()
+{
+	TextLayout->ClearHighlights();
+
+	const FTextLocation CursorInteractionPosition = CursorInfo.GetCursorInteractionLocation();
+	FTextLocation SelectionLocation = SelectionStart.Get( CursorInteractionPosition );
+
+	if ( SelectionLocation != CursorInteractionPosition )
+	{
+		FTextSelection Selection( SelectionLocation, CursorInteractionPosition );
+
+		const int32 SelectionBeginningLineIndex = Selection.GetBeginning().GetLineIndex();
+		const int32 SelectionBeginningLineOffset = Selection.GetBeginning().GetOffset();
+		
+		const int32 SelectionEndLineIndex = Selection.GetEnd().GetLineIndex();
+		const int32 SelectionEndLineOffset = Selection.GetEnd().GetOffset();
+
+		if ( SelectionBeginningLineIndex == SelectionEndLineIndex )
+		{
+			const FTextRange Range(SelectionBeginningLineOffset, SelectionEndLineOffset);
+			//if (!Range.IsEmpty())
+			//{
+				TextLayout->AddHighlight(FTextHighlight(SelectionBeginningLineIndex, Range, CursorSelectionHighlighter.ToSharedRef()));
+			//}
+		}
+		else
+		{
+			const TArray< FTextLayout::FLineModel >& Lines = TextLayout->GetLineModels();
+
+			for (int LineIndex = SelectionBeginningLineIndex; LineIndex <= SelectionEndLineIndex; LineIndex++)
+			{
+				if ( LineIndex == SelectionBeginningLineIndex )
+				{
+					const FTextRange Range(SelectionBeginningLineOffset, Lines[LineIndex].Text->Len());
+					//if (!Range.IsEmpty())
+					//{
+						TextLayout->AddHighlight(FTextHighlight(LineIndex, Range, CursorSelectionHighlighter.ToSharedRef()));
+					//}
+				}
+				else if ( LineIndex == SelectionEndLineIndex )
+				{
+					const FTextRange Range(0, SelectionEndLineOffset);
+					//if (!Range.IsEmpty())
+					//{
+						TextLayout->AddHighlight(FTextHighlight(LineIndex, Range, CursorSelectionHighlighter.ToSharedRef()));
+					//}
+				}
+				else
+				{
+					const FTextRange Range(0, Lines[LineIndex].Text->Len());
+					//if (!Range.IsEmpty())
+					//{
+						TextLayout->AddHighlight(FTextHighlight(LineIndex, Range, CursorSelectionHighlighter.ToSharedRef()));
+					//}
+				}
+			}
+		}
+
+		CursorSelectionHighlighter->SetMode( FSlateCursorRunHighlighter::EMode::Selection );
+	}
+	else
+	{
+		// The cursor mode uses the literal position rather than the interaction position
+		const FTextLocation CursorPosition = CursorInfo.GetCursorLocation();
+
+		const TArray< FTextLayout::FLineModel >& Lines = TextLayout->GetLineModels();
+		const int32 LineTextLength = Lines[CursorPosition.GetLineIndex()].Text->Len();
+
+		if (LineTextLength == 0)
+		{
+			TextLayout->AddHighlight( FTextHighlight( CursorPosition.GetLineIndex(), FTextRange( 0, 0 ), CursorSelectionHighlighter.ToSharedRef() ) );
+		}
+		else if (CursorPosition.GetOffset() == LineTextLength)
+		{
+			TextLayout->AddHighlight(FTextHighlight(CursorPosition.GetLineIndex(), FTextRange(LineTextLength - 1, LineTextLength), CursorSelectionHighlighter.ToSharedRef()));
+		}
+		else
+		{
+			TextLayout->AddHighlight(FTextHighlight(CursorPosition.GetLineIndex(), FTextRange(CursorPosition.GetOffset(), CursorPosition.GetOffset() + 1), CursorSelectionHighlighter.ToSharedRef()));
+		}
+
+		CursorSelectionHighlighter->SetMode( FSlateCursorRunHighlighter::EMode::Cursor );
+	}
+};
+
+void SMultiLineEditableText::RemoveCursorHighlight()
+{
+	TextLayout->ClearHighlights();
+}
+
+void SMultiLineEditableText::JumpTo(ETextLocation::Type JumpLocation, ECursorAction::Type Action)
+{
+	switch (JumpLocation)
+	{
+		case ETextLocation::BeginningOfLine:
+		{
+			const FTextLocation CursorInteractionPosition = CursorInfo.GetCursorInteractionLocation();
+			const TArray< FTextLayout::FLineView >& LineViews = TextLayout->GetLineViews();
+			const int32 CurrentLineViewIndex = TextLayout->GetLineViewIndexForTextLocation(LineViews, CursorInteractionPosition );
+
+			if (LineViews.IsValidIndex(CurrentLineViewIndex))
+			{
+				const FTextLayout::FLineView& CurrentLineView = LineViews[CurrentLineViewIndex];
+				
+				const FTextLocation OldCursorPosition = CursorInteractionPosition;
+				const FTextLocation NewCursorPosition = FTextLocation(OldCursorPosition.GetLineIndex(), CurrentLineView.Range.BeginIndex);
+				
+				if (Action == ECursorAction::SelectText)
+				{
+					this->SelectionStart = OldCursorPosition;
+				}
+				else
+				{
+					ClearSelection();
+				}
+
+				PreferredCursorScreenOffsetInLine = TextLayout->GetLocationAt(NewCursorPosition).X;
+
+				CursorInfo.SetCursorLocationAndCalculateAlignment(TextLayout, NewCursorPosition);
+				UpdateCursorHighlight();
+			}
+		}
+		break;
+
+		case ETextLocation::BeginningOfDocument:
+		{
+			const FTextLocation OldCursorPosition = CursorInfo.GetCursorInteractionLocation();
+			const FTextLocation NewCursorPosition = FTextLocation(0,0);
+
+			if (Action == ECursorAction::SelectText)
+			{
+				this->SelectionStart = OldCursorPosition;
+			}
+			else
+			{
+				ClearSelection();
+			}
+
+			PreferredCursorScreenOffsetInLine = TextLayout->GetLocationAt(NewCursorPosition).X;
+
+			CursorInfo.SetCursorLocationAndCalculateAlignment(TextLayout, NewCursorPosition);
+			UpdateCursorHighlight();
+		}
+		break;
+
+		case ETextLocation::EndOfLine:
+		{
+			const FTextLocation CursorInteractionPosition = CursorInfo.GetCursorInteractionLocation();
+			const TArray< FTextLayout::FLineView >& LineViews = TextLayout->GetLineViews();
+			const int32 CurrentLineViewIndex = TextLayout->GetLineViewIndexForTextLocation(LineViews, CursorInteractionPosition);
+
+			if (LineViews.IsValidIndex(CurrentLineViewIndex))
+			{
+				const FTextLayout::FLineView& CurrentLineView = LineViews[CurrentLineViewIndex];
+
+				const FTextLocation OldCursorPosition = CursorInteractionPosition;
+				const FTextLocation NewCursorPosition = FTextLocation(OldCursorPosition.GetLineIndex(), CurrentLineView.Range.EndIndex);
+
+				if (Action == ECursorAction::SelectText)
+				{
+					this->SelectionStart = OldCursorPosition;
+				}
+				else
+				{
+					ClearSelection();
+				}
+
+				PreferredCursorScreenOffsetInLine = TextLayout->GetLocationAt(NewCursorPosition).X;
+
+				CursorInfo.SetCursorLocationAndCalculateAlignment(TextLayout, NewCursorPosition);
+				UpdateCursorHighlight();
+			}
+		}
+		break;
+		case ETextLocation::EndOfDocument:
+		{
+			if (!TextLayout->IsEmpty())
+			{
+				const FTextLocation OldCursorPosition = CursorInfo.GetCursorInteractionLocation();
+				const TArray< FTextLayout::FLineModel >& Lines = TextLayout->GetLineModels();
+
+				const int32 LastLineIndex = Lines.Num() - 1;
+				const FTextLocation NewCursorPosition = FTextLocation(LastLineIndex, Lines[LastLineIndex].Text->Len());
+
+				if (Action == ECursorAction::SelectText)
+				{
+					this->SelectionStart = OldCursorPosition;
+				}
+				else
+				{
+					ClearSelection();
+				}
+
+				PreferredCursorScreenOffsetInLine = TextLayout->GetLocationAt(NewCursorPosition).X;
+
+				CursorInfo.SetCursorLocationAndCalculateAlignment(TextLayout, NewCursorPosition);
+				UpdateCursorHighlight();
+			}
+		}
+		break;
+	};
+}
+
+void SMultiLineEditableText::ClearSelection()
+{
+	SelectionStart = TOptional<FTextLocation>();
+}
+
+bool SMultiLineEditableText::CanExecuteSelectAll() const
+{
+	bool bCanExecute = true;
+
+	// Can't select all if string is empty
+	if (TextLayout->IsEmpty())
+	{
+		bCanExecute = false;
+	}
+
+	const TArray< FTextLayout::FLineModel >& Lines = TextLayout->GetLineModels();
+	const int32 NumberOfLines = Lines.Num();
+
+	// Can't select all if entire string is already selected
+	const FTextLocation CursorInteractionPosition = CursorInfo.GetCursorInteractionLocation();
+	if (SelectionStart.IsSet() &&
+		SelectionStart.GetValue() == FTextLocation(0, 0) &&
+		CursorInteractionPosition == FTextLocation(NumberOfLines - 1, Lines[NumberOfLines - 1].Text->Len()))
+	{
+		bCanExecute = false;
+	}
+
+	return bCanExecute;
+}
+
+void SMultiLineEditableText::SelectAllText()
+{
+	if (TextLayout->IsEmpty())
+	{
+		return;
+	}
+
+	const TArray< FTextLayout::FLineModel >& Lines = TextLayout->GetLineModels();
+	const int32 NumberOfLines = Lines.Num();
+
+	SelectionStart = FTextLocation(0,0);
+	const FTextLocation NewCursorPosition = FTextLocation( NumberOfLines - 1, Lines[ NumberOfLines - 1 ].Text->Len() );
+	CursorInfo.SetCursorLocationAndCalculateAlignment(TextLayout, NewCursorPosition);
+	UpdateCursorHighlight();
+}
+
+bool SMultiLineEditableText::SelectAllTextWhenFocused()
+{
+	return bSelectAllTextWhenFocused.Get();
 }
 
-bool SMultiLineEditableText::GetIsReadOnly() const
+void SMultiLineEditableText::SelectWordAt(const FVector2D& LocalPosition)
 {
-	return false;
+	FTextLocation InitialLocation = TextLayout->GetTextLocationAt(LocalPosition * TextLayout->GetScale());
+	FTextSelection WordSelection = TextLayout->GetWordAt(InitialLocation);
+
+	FTextLocation WordStart = WordSelection.GetBeginning();
+	FTextLocation WordEnd = WordSelection.GetEnd();
+
+	if (WordStart.IsValid() && WordEnd.IsValid())
+	{
+		// Deselect any text that was selected
+		ClearSelection();
+
+		if (WordStart != WordEnd)
+		{
+			SelectionStart = WordStart;
+		}
+
+		const FTextLocation NewCursorPosition = WordEnd;
+		CursorInfo.SetCursorLocationAndCalculateAlignment(TextLayout, NewCursorPosition);
+		UpdateCursorHighlight();
+	}
 }
 
-void SMultiLineEditableText::BackspaceChar()
+void SMultiLineEditableText::BeginDragSelection() 
 {
+	bIsDragSelecting = true;
 }
 
-void SMultiLineEditableText::DeleteChar()
+bool SMultiLineEditableText::IsDragSelecting() const
 {
+	return bIsDragSelecting;
 }
 
-bool SMultiLineEditableText::CanTypeCharacter(const TCHAR CharInQuestion) const
+void SMultiLineEditableText::EndDragSelection() 
 {
-	return true;
+	bIsDragSelecting = false;
 }
 
-void SMultiLineEditableText::TypeChar( const int32 Character )
+bool SMultiLineEditableText::AnyTextSelected() const
 {
+	const FTextLocation CursorInteractionPosition = CursorInfo.GetCursorInteractionLocation();
+	const FTextLocation SelectionPosition = SelectionStart.Get( CursorInteractionPosition );
+	return SelectionPosition != CursorInteractionPosition;
 }
 
-FReply SMultiLineEditableText::MoveCursor( ECursorMoveMethod::Type Method, int8 Direction, ECursorAction::Type Action )
+bool SMultiLineEditableText::IsTextSelectedAt(const FVector2D& LocalPosition) const
 {
-	FTextLocation NewCursorPosition;
-	if (Method == ECursorMoveMethod::CharacterHorizontal)
-	{
-		NewCursorPosition = TranslatedLocation(CursorPosition, Direction);
-		PreferredCursorOffsetInLine = AbsoluteToWrapped(NewCursorPosition).GetOffset();
-	}
-	else if (Method == ECursorMoveMethod::Word)
-	{
-		NewCursorPosition = ScanForWordBoundary( CursorPosition, Direction );
-		PreferredCursorOffsetInLine = AbsoluteToWrapped(NewCursorPosition).GetOffset();
-	}
-	else
-	{
-		ensure( Method == ECursorMoveMethod::CharacterVertical );
-		NewCursorPosition = TranslateLocationVertical(CursorPosition, Direction);
-	}
+	const FTextLocation CursorInteractionPosition = CursorInfo.GetCursorInteractionLocation();
+	const FTextLocation SelectionPosition = SelectionStart.Get(CursorInteractionPosition);
 
-	if (Action == ECursorAction::SelectText)
-	{
-		// We are selecting text. Just remember where the selection started.
-		// The cursor is implicitly the other endpoint.
-		if (!SelectionStart.IsSet())
-		{
-			this->SelectionStart = CursorPosition;
-		}
-	}
-	else
+	if (SelectionPosition == CursorInteractionPosition)
 	{
-		// No longer selection text; clear the selection!
-		this->ClearSelection();
+		return false;
 	}
 
-	if( WrappedText.Num() > 0 )
-	{
-		const int32 NewWrappedCursorPosition = AbsoluteToWrapped( NewCursorPosition ).GetLineIndex();
-		const float LineHeight = WrappedText[NewCursorPosition.GetLineIndex()].Size.Y;
-		const float NumLinesOnScreen = CachedLastFrameGeometry.Size.Y / LineHeight;
-		const float FractionOfTotalLinesOnScreen = FMath::Clamp(NumLinesOnScreen / WrappedText.Num(), 0.0f, 1.0f);
-		if ( NewWrappedCursorPosition < NumLinesScrollOffset )
-		{
-			NumLinesScrollOffset = NewWrappedCursorPosition;
-			UpdateScrollbar( CachedLastFrameGeometry );
-		}
-		else if ( FMath::FloorToInt(NumLinesScrollOffset + NumLinesOnScreen - 1) < NewWrappedCursorPosition )
-		{
-			NumLinesScrollOffset = FMath::Clamp( FMath::FloorToInt(NewWrappedCursorPosition - NumLinesOnScreen + 1 ), 0, WrappedText.Num());
-			UpdateScrollbar( CachedLastFrameGeometry );
-		}
-	}
+	const FTextLocation ClickedPosition = TextLayout->GetTextLocationAt(LocalPosition * TextLayout->GetScale());
 
-	CursorPosition = NewCursorPosition;
+	FTextLocation SelectionLocation = SelectionStart.Get(CursorInteractionPosition);
+	FTextSelection Selection(SelectionLocation, CursorInteractionPosition);
 
-	LastCursorInteractionTime = FSlateApplication::Get().GetCurrentTime();
+	int32 SelectionBeginningLineIndex = Selection.GetBeginning().GetLineIndex();
+	int32 SelectionBeginningLineOffset = Selection.GetBeginning().GetOffset();
 
-	return FReply::Handled();
+	int32 SelectionEndLineIndex = Selection.GetEnd().GetLineIndex();
+	int32 SelectionEndLineOffset = Selection.GetEnd().GetOffset();
 
+	if (SelectionBeginningLineIndex == ClickedPosition.GetLineIndex())
+	{
+		return SelectionBeginningLineOffset >= ClickedPosition.GetOffset();
+	}
+	else if (SelectionEndLineIndex == ClickedPosition.GetLineIndex())
+	{
+		return SelectionEndLineOffset < ClickedPosition.GetOffset();
+	}
+	
+	return SelectionBeginningLineIndex < ClickedPosition.GetLineIndex() && SelectionEndLineIndex > ClickedPosition.GetLineIndex();
 }
-void SMultiLineEditableText::JumpTo(ETextLocation::Type JumpLocation, ECursorAction::Type Action)
+
+void SMultiLineEditableText::SetWasFocusedByLastMouseDown( bool Value )
 {
-	
+	bWasFocusedByLastMouseDown = Value;
 }
 
-void SMultiLineEditableText::ClearSelection()
+bool SMultiLineEditableText::WasFocusedByLastMouseDown() const
 {
-	SelectionStart = TOptional<FTextLocation>();
+	return bWasFocusedByLastMouseDown;
 }
 
-void SMultiLineEditableText::SelectAllText()
+bool SMultiLineEditableText::HasDragSelectedSinceFocused() const
 {
-	SelectionStart = FTextLocation(0,0);
-	CursorPosition = FTextLocation( TextLines.Num() - 1, TextLines[TextLines.Num()-1]->Len() );
+	return bHasDragSelectedSinceFocused;
+}
+
+void SMultiLineEditableText::SetHasDragSelectedSinceFocused( bool Value )
+{
+	bHasDragSelectedSinceFocused = Value;
 }
 
 FReply SMultiLineEditableText::OnEscape()
@@ -146,33 +1131,152 @@ FReply SMultiLineEditableText::OnEscape()
 
 void SMultiLineEditableText::OnEnter()
 {
+	if (AnyTextSelected())
+	{
+		// Delete selected text
+		DeleteSelectedText();
+	}
+
+	const FTextLocation CursorInteractionPosition = CursorInfo.GetCursorInteractionLocation();
+	if (TextLayout->SplitLineAt(CursorInteractionPosition))
+	{
+		// Adjust the cursor position to be at the beginning of the new line
+		const FTextLocation NewCursorPosition = FTextLocation(CursorInteractionPosition.GetLineIndex() + 1, 0);
+		CursorInfo.SetCursorLocationAndCalculateAlignment(TextLayout, NewCursorPosition);
+	}
+
+	ClearSelection();
+	UpdateCursorHighlight();
 }
 
 bool SMultiLineEditableText::CanExecuteCut() const
 {
-	return true;
+	bool bCanExecute = true;
+
+	// Can't execute if this is a read-only control
+	if (IsReadOnly.Get())
+	{
+		bCanExecute = false;
+	}
+
+	// Can't execute if there is no text selected
+	if (!AnyTextSelected())
+	{
+		bCanExecute = false;
+	}
+
+	return bCanExecute;
 }
 
 void SMultiLineEditableText::CutSelectedTextToClipboard()
 {
+	if (AnyTextSelected())
+	{
+		StartChangingText();
+
+		const FTextLocation CursorInteractionPosition = CursorInfo.GetCursorInteractionLocation();
+		FTextLocation SelectionLocation = SelectionStart.Get(CursorInteractionPosition);
+		FTextSelection Selection(SelectionLocation, CursorInteractionPosition);
+
+		// Grab the selected substring
+		FString SelectedText;
+		TextLayout->GetSelectionAsText(SelectedText, Selection);
+
+		// Copy text to clipboard
+		FPlatformMisc::ClipboardCopy(*SelectedText);
+
+		DeleteSelectedText();
+
+		UpdateCursorHighlight();
+
+		FinishChangingText();
+	}
 }
 
 bool SMultiLineEditableText::CanExecuteCopy() const
 {
-	return true;
+	bool bCanExecute = true;
+
+	// Can't execute if there is no text selected
+	if (!AnyTextSelected())
+	{
+		bCanExecute = false;
+	}
+
+	return bCanExecute;
 }
 
 void SMultiLineEditableText::CopySelectedTextToClipboard()
 {
+	if (AnyTextSelected())
+	{
+		const FTextLocation CursorInteractionPosition = CursorInfo.GetCursorInteractionLocation();
+		FTextLocation SelectionLocation = SelectionStart.Get(CursorInteractionPosition);
+		FTextSelection Selection(SelectionLocation, CursorInteractionPosition);
+
+		// Grab the selected substring
+		FString SelectedText;
+		TextLayout->GetSelectionAsText(SelectedText, Selection);
+
+		// Copy text to clipboard
+		FPlatformMisc::ClipboardCopy(*SelectedText);
+	}
 }
 
 bool SMultiLineEditableText::CanExecutePaste() const
 {
-	return true;
+	bool bCanExecute = true;
+
+	// Can't execute if this is a read-only control
+	if (IsReadOnly.Get())
+	{
+		bCanExecute = false;
+	}
+
+	// Can't paste unless the clipboard has a string in it
+	if (!DoesClipboardHaveAnyText())
+	{
+		bCanExecute = false;
+	}
+
+	return bCanExecute;
+}
+
+bool SMultiLineEditableText::DoesClipboardHaveAnyText() const
+{
+	FString ClipboardContent;
+	FPlatformMisc::ClipboardPaste(ClipboardContent);
+	return ClipboardContent.Len() > 0;
 }
 
 void SMultiLineEditableText::PasteTextFromClipboard()
 {
+	if (!IsReadOnly.Get())
+	{
+		StartChangingText(); 
+
+		// Delete currently selected text, if any
+		DeleteSelectedText();
+
+		// Paste from the clipboard
+		FString PastedText;
+		FPlatformMisc::ClipboardPaste(PastedText);
+
+		for (int32 CurCharIndex = 0; CurCharIndex < PastedText.Len(); ++CurCharIndex)
+		{
+			const TCHAR& Character = PastedText.GetCharArray()[CurCharIndex];
+			if (Character == '\n')
+			{
+				OnEnter();
+			}
+			else
+			{
+				TypeChar(Character);
+			}
+		}
+
+		FinishChangingText();
+	}
 }
 
 bool SMultiLineEditableText::CanExecuteUndo() const
@@ -182,204 +1286,311 @@ bool SMultiLineEditableText::CanExecuteUndo() const
 
 void SMultiLineEditableText::Undo()
 {
+	if (!IsReadOnly.Get() && UndoStates.Num() > 0/* && !TextInputMethodContext->IsComposing*/)
+	{
+		// Restore from undo state
+		int32 UndoStateIndex;
+		if (CurrentUndoLevel == INDEX_NONE)
+		{
+			// We haven't undone anything since the last time a new undo state was added
+			UndoStateIndex = UndoStates.Num() - 1;
+
+			// Store an undo state for the current state (before undo was pressed)
+			FUndoState NewUndoState;
+			MakeUndoState(NewUndoState);
+			PushUndoState(NewUndoState);
+		}
+		else
+		{
+			// Move down to the next undo level
+			UndoStateIndex = CurrentUndoLevel - 1;
+		}
+
+		// Is there anything else to undo?
+		if (UndoStateIndex >= 0)
+		{
+			{
+				// NOTE: It's important the no code called here creates or destroys undo states!
+				const FUndoState& UndoState = UndoStates[UndoStateIndex];
+
+				SetText(UndoState.Text);
+				CursorInfo = UndoState.CursorInfo.CreateUndo();
+				SelectionStart = UndoState.SelectionStart;
+
+				UpdateCursorHighlight();
+
+				SaveText(UndoState.Text);
+			}
+
+			CurrentUndoLevel = UndoStateIndex;
+
+			// Let outsiders know that the text content has been changed
+			check(LastSetText.IsSet());
+			OnTextChanged.ExecuteIfBound(LastSetText.GetValue());
+		}
+	}
 }
 
 void SMultiLineEditableText::Redo()
 {
+	// Is there anything to redo?  If we've haven't tried to undo since the last time new
+	// undo state was added, then CurrentUndoLevel will be INDEX_NONE
+	if (!IsReadOnly.Get() && CurrentUndoLevel != INDEX_NONE /*&& !(TextInputMethodContext->IsComposing)*/)
+	{
+		const int32 NextUndoLevel = CurrentUndoLevel + 1;
+		if (UndoStates.Num() > NextUndoLevel)
+		{
+			// Restore from undo state
+			{
+				// NOTE: It's important the no code called here creates or destroys undo states!
+				const FUndoState& UndoState = UndoStates[NextUndoLevel];
+
+				SetText(UndoState.Text);
+				CursorInfo.RestoreFromUndo(UndoState.CursorInfo);
+				SelectionStart = UndoState.SelectionStart;
+
+				UpdateCursorHighlight();
+
+				SaveText(UndoState.Text);
+			}
+
+			CurrentUndoLevel = NextUndoLevel;
+
+			if (UndoStates.Num() <= CurrentUndoLevel + 1)
+			{
+				// We've redone all available undo states
+				CurrentUndoLevel = INDEX_NONE;
+
+				// Pop last undo state that we created on initial undo
+				UndoStates.RemoveAt(UndoStates.Num() - 1);
+			}
+
+			// Let outsiders know that the text content has been changed
+			check(LastSetText.IsSet());
+			OnTextChanged.ExecuteIfBound(LastSetText.GetValue());
+		}
+	}
 }
 
-void SMultiLineEditableText::Tick( const FGeometry& AllottedGeometry, const double InCurrentTime, const float InDeltaTime )
+void SMultiLineEditableText::PushUndoState(const SMultiLineEditableText::FUndoState& InUndoState)
 {
-	if ( CachedLastFrameGeometry.Size != AllottedGeometry.Size )
+	// If we've already undone some state, then we'll remove any undo state beyond the level that
+	// we've already undone up to.
+	if (CurrentUndoLevel != INDEX_NONE)
 	{
-		CachedLastFrameGeometry = AllottedGeometry;
-		ReWrap( AllottedGeometry );
-		UpdateScrollbar(AllottedGeometry);
+		UndoStates.RemoveAt(CurrentUndoLevel, UndoStates.Num() - CurrentUndoLevel);
 
+		// Reset undo level as we haven't undone anything since this latest undo
+		CurrentUndoLevel = INDEX_NONE;
+	}
+
+	// Cache new undo state
+	UndoStates.Add(InUndoState);
+
+	// If we've reached the maximum number of undo levels, then trim our array
+	if (UndoStates.Num() > EditableTextDefs::MaxUndoLevels)
+	{
+		UndoStates.RemoveAt(0);
 	}
 }
 
-int32 SMultiLineEditableText::OnPaint( const FGeometry& AllottedGeometry, const FSlateRect& MyClippingRect, FSlateWindowElementList& OutDrawElements, int32 LayerId, const FWidgetStyle& InWidgetStyle, bool bParentEnabled ) const
+void SMultiLineEditableText::ClearUndoStates()
 {
-	const int32 TextLayer = LayerId;
-	const int32 SelectionLayer = LayerId + 1;
-	const TSharedRef< FSlateFontMeasure > FontMeasureService = FSlateApplication::Get().GetRenderer()->GetFontMeasureService();
+	CurrentUndoLevel = INDEX_NONE;
+	UndoStates.Empty();
+}
 
-	const float CursorWidth = 2.0f;
+void SMultiLineEditableText::MakeUndoState(SMultiLineEditableText::FUndoState& OutUndoState)
+{
+	//@todo save and restoring the whole document is not ideal [3/31/2014 justin.sargent]
+	FString EditedText;
+	TextLayout->GetAsText(EditedText);
 
-	const double CurrentTime = FSlateApplication::Get().GetCurrentTime();
+	OutUndoState.Text = FText::FromString(EditedText);
+	OutUndoState.CursorInfo = CursorInfo;
+	OutUndoState.SelectionStart = SelectionStart;
+}
 
-	const bool bEnabled = ShouldBeEnabled( bParentEnabled );
-	const ESlateDrawEffect::Type DrawEffects = (bEnabled) ? ESlateDrawEffect::None : ESlateDrawEffect::DisabledEffect;
+TSharedRef< SWidget > SMultiLineEditableText::GetWidget()
+{
+	return SharedThis( this );
+}
 
-	const int32 WholeLinesOffset = FMath::FloorToInt( NumLinesScrollOffset );
-	const float PartialLineOffset = FMath::Fractional( NumLinesScrollOffset );
+void SMultiLineEditableText::SummonContextMenu(const FVector2D& InLocation)
+{
+	// Set the menu to automatically close when the user commits to a choice
+	const bool bShouldCloseWindowAfterMenuSelection = true;
 
+#define LOCTEXT_NAMESPACE "EditableTextContextMenu"
 
-	if (WrappedText.Num() > 0)
+	// This is a context menu which could be summoned from within another menu if this text block is in a menu
+	// it should not close the menu it is inside
+	bool bCloseSelfOnly = true;
+	FMenuBuilder MenuBuilder(bShouldCloseWindowAfterMenuSelection, UICommandList, MenuExtender, bCloseSelfOnly, &FCoreStyle::Get());
 	{
-		float HeightOffsetSoFar = - PartialLineOffset *  WrappedText[FMath::Clamp(WholeLinesOffset, 0, WrappedText.Num()-1)].Size.Y;
+		MenuBuilder.BeginSection("EditText", LOCTEXT("Heading", "Modify Text"));
+		{
+			// Undo
+			MenuBuilder.AddMenuEntry(FGenericCommands::Get().Undo);
+		}
+		MenuBuilder.EndSection();
+
+		MenuBuilder.BeginSection("EditableTextModify2");
+		{
+			// Cut
+			MenuBuilder.AddMenuEntry(FGenericCommands::Get().Cut);
 
-		bool FilledVerticalSpace = false;
+			// Copy
+			MenuBuilder.AddMenuEntry(FGenericCommands::Get().Copy);
 
-		const FTextRange Selection = FTextRange(SelectionStart.GetValue(), CursorPosition);
-		FTextRange SelectionInWrappedText = (SelectionStart.IsSet())
-			? AbsoluteToWrapped( Selection )
-			: FTextRange( FTextLocation(), FTextLocation() );
+			// Paste
+			MenuBuilder.AddMenuEntry(FGenericCommands::Get().Paste);
 
-		bool bInSelection = SelectionInWrappedText.First.GetLineIndex() < WholeLinesOffset && SelectionInWrappedText.First.GetLineIndex() > WholeLinesOffset;
+			// Delete
+			MenuBuilder.AddMenuEntry(FGenericCommands::Get().Delete);
+		}
+		MenuBuilder.EndSection();
 
-		for ( int32 LineIndex=WholeLinesOffset; !FilledVerticalSpace && LineIndex < WrappedText.Num(); ++LineIndex )
+		MenuBuilder.BeginSection("EditableTextModify3");
 		{
-			const FWrappedStringSlice& ThisWrappedLine = WrappedText[LineIndex];
-			// @todo FontService: Copy is undesirable?
-			const FString StringToPaint = ThisWrappedLine.SourceString.Mid( ThisWrappedLine.FirstCharIndex, ThisWrappedLine.LastCharIndex-ThisWrappedLine.FirstCharIndex );
+			// Select All
+			MenuBuilder.AddMenuEntry(FGenericCommands::Get().SelectAll);
+		}
+		MenuBuilder.EndSection();
+	}
 
-			// Draw this wrapped line of text
-			FSlateDrawElement::MakeText(
-				OutDrawElements,
-				LayerId,
-				AllottedGeometry.ToPaintGeometry(FVector2D(0, HeightOffsetSoFar), FVector2D::UnitVector),
-				StringToPaint,
-				FontInfo,
-				MyClippingRect,
-				ShouldBeEnabled(bParentEnabled),
-				InWidgetStyle.GetColorAndOpacityTint()
-			);
+#undef LOCTEXT_NAMESPACE
 
-			static const FLinearColor SelectionColor = FLinearColor(0,0,1.0f,0.25f);
+	const bool bFocusImmediatley = false;
+	TSharedPtr< SWindow > ContextMenuWindowPinned = FSlateApplication::Get().PushMenu(SharedThis(this), MenuBuilder.MakeWidget(), InLocation, FPopupTransitionEffect(FPopupTransitionEffect::ContextMenu), bFocusImmediatley);
+	ContextMenuWindow = ContextMenuWindowPinned;
 
-			// DRAW SELECTION
-			if ( SelectionStart.IsSet() )
-			{
-				const bool bIsFirstSelectionLine =
-					(&(TextLines[Selection.First.GetLineIndex()].Get()) == &ThisWrappedLine.SourceString) &&
-					Selection.First.GetOffset() >= ThisWrappedLine.FirstCharIndex &&
-					Selection.First.GetOffset() <= ThisWrappedLine.LastCharIndex;
+	// Make sure the window is valid.  It's possible for the parent to already be in the destroy queue, for example if the editable text was configured to dismiss it's window during OnTextCommitted.
+	if (ContextMenuWindowPinned.IsValid())
+	{
+		ContextMenuWindowPinned->SetOnWindowClosed(FOnWindowClosed::CreateSP(this, &SMultiLineEditableText::OnWindowClosed));
+	}
+}
 
-				const bool bIsLastSelectionLine =
-				 	(&(TextLines[Selection.Last.GetLineIndex()].Get()) == &ThisWrappedLine.SourceString) &&
-					Selection.Last.GetOffset() >= ThisWrappedLine.FirstCharIndex &&
-					Selection.Last.GetOffset() <= ThisWrappedLine.LastCharIndex;
+void SMultiLineEditableText::OnWindowClosed(const TSharedRef<SWindow>&)
+{
+	// Give this focus when the context menu has been dismissed
+	FSlateApplication::Get().SetKeyboardFocus(AsShared(), EKeyboardFocusCause::OtherWidgetLostFocus);
+}
 
-				if ( bIsFirstSelectionLine && bIsLastSelectionLine )
-				{
-					// Selection contained within a single line.
-					const FString StringBeforeSelection = ThisWrappedLine.SourceString.Left( SelectionInWrappedText.First.GetOffset() );
-					const FVector2D SizeBeforeSelection = FontMeasureService->Measure(StringBeforeSelection, FontInfo, AllottedGeometry.Scale);					
-					
-					const FString SelectedString = ThisWrappedLine.SourceString.Mid( SelectionInWrappedText.First.GetOffset(), SelectionInWrappedText.Last.GetOffset() - SelectionInWrappedText.First.GetOffset() );		
-					const FVector2D SelectionSize = FontMeasureService->Measure(SelectedString, FontInfo, AllottedGeometry.Scale);
-
-					FSlateDrawElement::MakeBox(
-						OutDrawElements,
-						SelectionLayer,
-						AllottedGeometry.ToPaintGeometry(FVector2D(SizeBeforeSelection.X, HeightOffsetSoFar), SelectionSize),
-						FCoreStyle::Get().GetBrush( "WhiteBrush" ),
-						MyClippingRect,
-						DrawEffects,
-						SelectionColor
-					);
+void SMultiLineEditableText::LoadText()
+{
+	SetText(BoundText.Get());
+	TextLayout->UpdateIfNeeded();
+}
 
-				}
-				else if ( bIsFirstSelectionLine )
-				{
-					// Selection starts on this line, ends on another.
-					bInSelection = true;
-
-					// @todo FontService: Copy is unnecessary
-					const FString StringBeforeSelection = StringToPaint.Left(SelectionInWrappedText.First.GetOffset());
-					const FVector2D UnselectedPortionSize = FontMeasureService->Measure(StringBeforeSelection, FontInfo, AllottedGeometry.Scale);
-					FSlateDrawElement::MakeBox(
-						OutDrawElements,
-						SelectionLayer,
-						AllottedGeometry.ToPaintGeometry(FVector2D(UnselectedPortionSize.X, HeightOffsetSoFar), ThisWrappedLine.Size - FVector2D(UnselectedPortionSize.X, 0)),
-						FCoreStyle::Get().GetBrush( "WhiteBrush" ),
-						MyClippingRect,
-						DrawEffects,
-						SelectionColor
-					);
-				}
-				else if ( bIsLastSelectionLine )
-				{
-					// Selection starts on some line above, ends on this one.
-					bInSelection = false;
-
-					// @todo FontService: Copy is unnecessary
-					const FString StringBeforeSelection = StringToPaint.Left(SelectionInWrappedText.Last.GetOffset());
-					const FVector2D SelectedPortionSize = FontMeasureService->Measure(StringBeforeSelection, FontInfo, AllottedGeometry.Scale);
-					FSlateDrawElement::MakeBox(
-						OutDrawElements,
-						SelectionLayer,
-						AllottedGeometry.ToPaintGeometry(FVector2D(0, HeightOffsetSoFar), FVector2D(SelectedPortionSize.X, ThisWrappedLine.Size.Y)),
-						FCoreStyle::Get().GetBrush( "WhiteBrush" ),
-						MyClippingRect,
-						DrawEffects,
-						SelectionColor
-					);
-				}
-				else if (bInSelection)
-				{
-					// The entirety of this line is selected
-					FSlateDrawElement::MakeBox(
-						OutDrawElements,
-						SelectionLayer,
-						AllottedGeometry.ToPaintGeometry(FVector2D(0, HeightOffsetSoFar), ThisWrappedLine.Size),
-						FCoreStyle::Get().GetBrush( "WhiteBrush" ),
-						MyClippingRect,
-						DrawEffects,
-						SelectionColor
-					);
-				}
+void SMultiLineEditableText::Tick( const FGeometry& AllottedGeometry, const double InCurrentTime, const float InDeltaTime )
+{
+	TextLayout->SetScale( AllottedGeometry.Scale );
 
-			}
+	const bool bShouldAppearFocused = HasKeyboardFocus() || ContextMenuWindow.IsValid();
 
-			// DRAW CURSOR
-			if (this->HasKeyboardFocus())
-			{
-				const bool CursorIsOnThisTextLine = ( &ThisWrappedLine.SourceString == &(TextLines[CursorPosition.GetLineIndex()].Get()) );
-				const bool bIsStartOfNewLine = ThisWrappedLine.FirstCharIndex == 0;
-				if ( CursorIsOnThisTextLine &&
-					((bIsStartOfNewLine && CursorPosition.GetOffset() == 0) || CursorPosition.GetOffset() > ThisWrappedLine.FirstCharIndex) &&
-					CursorPosition.GetOffset() <= ThisWrappedLine.LastCharIndex )
-				{
-					// The cursor is always visible (i.e. not blinking) when we're interacting with it; otherwise it might get lost.
-					const bool bForceCursorVisible = (CurrentTime - LastCursorInteractionTime) < EditableTextDefs::CaretBlinkPauseTime;
-					const float CursorOpacity = (bForceCursorVisible)
-						? 1.0f
-						: FMath::RoundToInt( FMath::MakePulsatingValue( CurrentTime, EditableTextDefs::BlinksPerSecond ));
-
-					const int32 CursorOffsetInWrappedLine = CursorPosition.GetOffset() - ThisWrappedLine.FirstCharIndex;
-					const FVector2D SlateUnitsCursorOffsetInLine = FontMeasureService->Measure( StringToPaint.Left(CursorOffsetInWrappedLine), FontInfo, AllottedGeometry.Scale );
-
-					const FVector2D DrawPosition = FVector2D( SlateUnitsCursorOffsetInLine.X, HeightOffsetSoFar );
-					const FVector2D DrawSize = FVector2D( CursorWidth, ThisWrappedLine.Size.Y );
-					// @todo: State Styles - make this brush part of the widget style
-					const FSlateBrush* StyleInfo = FCoreStyle::Get().GetBrush("EditableText.SelectionBackground");
-
-					FSlateDrawElement::MakeBox(
-						OutDrawElements,
-						LayerId,
-						AllottedGeometry.ToPaintGeometry(DrawPosition, DrawSize),
-						StyleInfo,
-						MyClippingRect,
-						DrawEffects,
-						FLinearColor(1,1,1,CursorOpacity) );
+	if (!bShouldAppearFocused && BoundText.IsBound())
+	{
+		const FText& UpdatedBoundText = BoundText.Get();
+		if (!LastSetText.IsSet() || UpdatedBoundText.CompareTo(LastSetText.GetValue()) != 0)
+		{
+			SetText(UpdatedBoundText);
+		}
+	}
+}
 
-				}
-			}
+int32 SMultiLineEditableText::OnPaint( const FGeometry& AllottedGeometry, const FSlateRect& MyClippingRect, FSlateWindowElementList& OutDrawElements, int32 LayerId, const FWidgetStyle& InWidgetStyle, bool bParentEnabled ) const
+{
+	// Update the auto-wrap size now that we have computed paint geometry; won't take affect until text frame
+	if ( AutoWrapText.Get() )
+	{
+		CachedAutoWrapTextWidth = AllottedGeometry.Size.X;
+	}
 
-			HeightOffsetSoFar += WrappedText[LineIndex].Size.Y;
+	const ETextJustify::Type TextJustification = TextLayout->GetJustification();
+
+	if ( TextJustification == ETextJustify::Right )
+	{
+		const FVector2D TextLayoutSize = TextLayout->GetSize();
+		const FVector2D Offset( ( AllottedGeometry.Size.X - TextLayoutSize.X ), 0 );
+		LayerId = TextLayout->OnPaint( TextStyle, AllottedGeometry.MakeChild( Offset, TextLayoutSize ), MyClippingRect, OutDrawElements, LayerId, InWidgetStyle, ShouldBeEnabled( bParentEnabled ) );
+	}
+	else if ( TextJustification == ETextJustify::Center )
+	{
+		FVector2D TextLayoutSize = TextLayout->GetSize();
+		float LayoutWidth = TextLayoutSize.X;
+		if (TextLayoutSize.X == 0)
+		{
+			const TSharedRef< FSlateFontMeasure > FontMeasureService = FSlateApplication::Get().GetRenderer()->GetFontMeasureService();
+			TextLayoutSize.X = FontMeasureService->Measure(FString(TEXT("  ")), TextStyle.Font).X;
 		}
+
+		const FVector2D Offset((AllottedGeometry.Size.X - TextLayoutSize.X) / 2, 0);
+		LayerId = TextLayout->OnPaint( TextStyle, AllottedGeometry.MakeChild( Offset, TextLayoutSize ), MyClippingRect, OutDrawElements, LayerId, InWidgetStyle, ShouldBeEnabled( bParentEnabled ) );
+	}
+	else
+	{
+		LayerId = TextLayout->OnPaint( TextStyle, AllottedGeometry, MyClippingRect, OutDrawElements, LayerId, InWidgetStyle, ShouldBeEnabled( bParentEnabled ) );
 	}
 
-	return SelectionLayer;
+	return LayerId;
 }
 
+void SMultiLineEditableText::CacheDesiredSize()
+{
+	// Get the wrapping width and font to see if they have changed
+	float WrappingWidth = WrapTextAt.Get();
+
+	const FMargin& OurMargin = Margin.Get();
+
+	// Text wrapping can either be used defined (WrapTextAt), automatic (AutoWrapText), or a mixture of both
+	// Take whichever has the smallest value (>1)
+	if(AutoWrapText.Get() && CachedAutoWrapTextWidth >= 1.0f)
+	{
+		WrappingWidth = (WrappingWidth >= 1.0f) ? FMath::Min(WrappingWidth, CachedAutoWrapTextWidth) : CachedAutoWrapTextWidth;
+	}
+
+	TextLayout->SetWrappingWidth( WrappingWidth );
+	TextLayout->SetMargin( OurMargin );
+	TextLayout->SetLineHeightPercentage( LineHeightPercentage.Get() );
+	TextLayout->SetJustification( Justification.Get() );
+	TextLayout->UpdateIfNeeded();
+
+	SWidget::CacheDesiredSize();
+}
 
 FVector2D SMultiLineEditableText::ComputeDesiredSize() const
 {
-	return PreferredSize;
+	const FSlateFontInfo& FontInfo = TextStyle.Font;
+	
+	const float FontMaxCharHeight = FTextEditHelper::GetFontHeight(FontInfo);
+	const float CaretWidth = FTextEditHelper::CalculateCaretWidth(FontMaxCharHeight);
+
+	// The layouts current margin size. We should not report a size smaller then the margins.
+	const FMargin Margin = TextLayout->GetMargin();
+	const FVector2D TextLayoutSize = TextLayout->GetSize();
+	const float WrappingWidth = WrapTextAt.Get();
+
+	// If a wrapping width has been provided that should be reported as the desired width.
+	float DesiredWidth = WrappingWidth > 0 ? WrappingWidth : TextLayoutSize.X;
+	DesiredWidth = FMath::Max(Margin.GetTotalSpaceAlong<Orient_Horizontal>(), DesiredWidth);
+	DesiredWidth += CaretWidth;
+
+	float DesiredHeight = FMath::Max(Margin.GetTotalSpaceAlong<Orient_Vertical>(), TextLayoutSize.Y);
+	DesiredHeight = FMath::Max(DesiredHeight, FontMaxCharHeight);
+	
+	return FVector2D(DesiredWidth, DesiredHeight);
+}
+
+FChildren* SMultiLineEditableText::GetChildren()
+{
+	return TextLayout->GetChildren();
+}
+
+void SMultiLineEditableText::OnArrangeChildren( const FGeometry& AllottedGeometry, FArrangedChildren& ArrangedChildren ) const
+{
+	TextLayout->ArrangeChildren( AllottedGeometry, ArrangedChildren );
 }
 
 bool SMultiLineEditableText::SupportsKeyboardFocus() const
@@ -389,18 +1600,18 @@ bool SMultiLineEditableText::SupportsKeyboardFocus() const
 
 FReply SMultiLineEditableText::OnKeyChar( const FGeometry& MyGeometry,const FCharacterEvent& InCharacterEvent )
 {
-	return FTextEditHelper::OnKeyChar( MyGeometry, InCharacterEvent, *this );
+	return FTextEditHelper::OnKeyChar( InCharacterEvent, SharedThis( this ) );
 }
 
 FReply SMultiLineEditableText::OnKeyDown( const FGeometry& MyGeometry, const FKeyboardEvent& InKeyboardEvent )
 {
-	FReply Reply = FTextEditHelper::OnKeyDown( MyGeometry, InKeyboardEvent, *this );
+	FReply Reply = FTextEditHelper::OnKeyDown( InKeyboardEvent, SharedThis( this ) );
 
-	//// Process keybindings if the event wasn't already handled
-	//if( !Reply.IsEventHandled() && UICommandList->ProcessCommandBindings( InKeyboardEvent ) )
-	//{
-	//	Reply = FReply::Handled();
-	//}
+	// Process keybindings if the event wasn't already handled
+	if (!Reply.IsEventHandled() && UICommandList->ProcessCommandBindings(InKeyboardEvent))
+	{
+		Reply = FReply::Handled();
+	}
 
 	return Reply;
 }
@@ -410,44 +1621,49 @@ FReply SMultiLineEditableText::OnKeyUp( const FGeometry& MyGeometry, const FKeyb
 	return FReply::Unhandled();
 }
 
+FReply SMultiLineEditableText::OnMouseButtonDown( const FGeometry& MyGeometry, const FPointerEvent& MouseEvent ) 
+{
+	FReply Reply = FTextEditHelper::OnMouseButtonDown( MyGeometry, MouseEvent, SharedThis( this ) );
+	return Reply;
+}
 
-void SMultiLineEditableText::ReWrap( const FGeometry& AllottedGeometry )
+FReply SMultiLineEditableText::OnMouseButtonUp( const FGeometry& MyGeometry, const FPointerEvent& MouseEvent ) 
 {
-	SlateWordWrapper::WrapText(TextLines, FontInfo, AllottedGeometry.Scale, AllottedGeometry.Size.X , WrappedText);
+	FReply Reply = FTextEditHelper::OnMouseButtonUp( MyGeometry, MouseEvent, SharedThis( this ) );
+	return Reply;
 }
 
-void SMultiLineEditableText::OnScrollbarScrolled( float OffsetFraction )
+FReply SMultiLineEditableText::OnMouseMove( const FGeometry& MyGeometry, const FPointerEvent& MouseEvent )
 {
-	NumLinesScrollOffset = FMath::Clamp(OffsetFraction * WrappedText.Num(), 0.0f, WrappedText.Num()-1.0f);
-	UpdateScrollbar( CachedLastFrameGeometry );
+	FReply Reply = FTextEditHelper::OnMouseMove( MyGeometry, MouseEvent, SharedThis( this ) );
+	return Reply;
 }
 
-/** Remember where the cursor was when we started selectin. */
-void SMultiLineEditableText::BeginSelecting( const FTextLocation& Endpoint )
+FReply SMultiLineEditableText::OnMouseButtonDoubleClick(const FGeometry& MyGeometry, const FPointerEvent& MouseEvent)
 {
-	SelectionStart = Endpoint;
+	FReply Reply = FTextEditHelper::OnMouseButtonDoubleClick(MyGeometry, MouseEvent, SharedThis(this));
+	return Reply;
 }
 
+FCursorReply SMultiLineEditableText::OnCursorQuery( const FGeometry& MyGeometry, const FPointerEvent& CursorEvent ) const
+{
+	return FCursorReply::Cursor( EMouseCursor::TextEditBeam );
+}
 
-void SMultiLineEditableText::UpdateScrollbar( const FGeometry &AllottedGeometry ) 
+/** Remember where the cursor was when we started selecting. */
+void SMultiLineEditableText::BeginSelecting( const FTextLocation& Endpoint )
 {
-	TSharedPtr<SScrollBar> OptionalScrollBar = OptionalScrollBarPtr.Pin();
-	if (OptionalScrollBar.IsValid() && WrappedText.Num() > 0)
-	{
-		const float LineHeight = WrappedText[0].Size.Y;
-		const float NumLinesOnScreen = AllottedGeometry.Size.Y / LineHeight;
-		const float FractionOfTotalLinesOnScreen = FMath::Clamp(NumLinesOnScreen / WrappedText.Num(), 0.0f, 1.0f);
-		const float MaxVerticalOffsetInLines = FMath::Max(0.0f, WrappedText.Num() - NumLinesOnScreen);
-		NumLinesScrollOffset = FMath::Clamp( NumLinesScrollOffset, 0.0f, MaxVerticalOffsetInLines );
-		OptionalScrollBar->SetState(NumLinesScrollOffset / WrappedText.Num(), FractionOfTotalLinesOnScreen);
-	}
+	SelectionStart = Endpoint;
 }
 
-SMultiLineEditableText::FTextLocation SMultiLineEditableText::TranslatedLocation( const FTextLocation& Location, int8 Direction ) const
+FTextLocation SMultiLineEditableText::TranslatedLocation( const FTextLocation& Location, int8 Direction ) const
 {
 	const int32 OffsetInLine = Location.GetOffset() + Direction;
+	const TArray< FTextLayout::FLineModel >& Lines = TextLayout->GetLineModels();
+	const int32 NumberOfLines = Lines.Num();
+	const int32 LineTextLength = Lines[ Location.GetLineIndex() ].Text->Len();
 		
-	if ( OffsetInLine > TextLines[Location.GetLineIndex()]->Len() && Location.GetLineIndex() < TextLines.Num()-1 )
+	if ( Location.GetLineIndex() < NumberOfLines - 1 && OffsetInLine > LineTextLength )
 	{
 		// We're going over the end of the line and we aren't on the last line
 		return FTextLocation( Location.GetLineIndex() + 1, 0 );
@@ -455,55 +1671,41 @@ SMultiLineEditableText::FTextLocation SMultiLineEditableText::TranslatedLocation
 	else if ( OffsetInLine < 0 && Location.GetLineIndex() > 0 )
 	{
 		// We're stepping before the beginning of the line, and we're not on the first line.
-		return FTextLocation( Location.GetLineIndex() - 1, TextLines[Location.GetLineIndex() - 1]->Len() );	
+		const int32 NewLineIndex = Location.GetLineIndex() - 1;
+		return FTextLocation( NewLineIndex, Lines[ NewLineIndex ].Text->Len() );	
 	}
 	else
 	{
-		return FTextLocation( Location.GetLineIndex(), FMath::Clamp(OffsetInLine, 0, TextLines[Location.GetLineIndex()]->Len()) );
+		if ( LineTextLength == 0 )
+		{
+			return FTextLocation( Location.GetLineIndex(), 0 );
+		}
+
+		return FTextLocation( Location.GetLineIndex(), FMath::Clamp( OffsetInLine, 0, LineTextLength ) );
 	}
 }
 
-/** Find wrapped line on which the cursor is currently located. */
-struct FMatchWrappedLine
+FTextLocation SMultiLineEditableText::TranslateLocationVertical( const FTextLocation& Location, int8 Direction ) const
 {
-	FMatchWrappedLine( const FString& InStringToMatch, int32 InOffset )
-		: StringToMatch( InStringToMatch )
-		, Offset( InOffset )
-		{}
-	bool Matches( const FWrappedStringSlice& CandidateLine  ) const
-	{
-		return
-			(&StringToMatch == &(CandidateLine.SourceString)) &&
-			Offset >= CandidateLine.FirstCharIndex &&
-			Offset <= CandidateLine.LastCharIndex;
-	}
-	const FString& StringToMatch;
-	const int32 Offset;
-};
+	const TArray< FTextLayout::FLineView >& LineViews = TextLayout->GetLineViews();
+	const int32 NumberOfLineViews = LineViews.Num();
 
-SMultiLineEditableText::FTextLocation SMultiLineEditableText::TranslateLocationVertical( const FTextLocation& Location, int8 Direction ) const
-{
-	const int32 CurrentWrappedLineIndex = WrappedText.FindMatch( FMatchWrappedLine( TextLines[Location.GetLineIndex()].Get(), Location.GetOffset() ) );
-	const FWrappedStringSlice& CurrentWrappedLine = WrappedText[CurrentWrappedLineIndex];
-	ensure(CurrentWrappedLineIndex != INDEX_NONE);
+	const int32 CurrentLineViewIndex = TextLayout->GetLineViewIndexForTextLocation(LineViews, Location);
+	ensure(CurrentLineViewIndex != INDEX_NONE);
+	const FTextLayout::FLineView& CurrentLineView = LineViews[ CurrentLineViewIndex ];
 
-	const int32 NewWrappedLinedIndex = FMath::Clamp(CurrentWrappedLineIndex + Direction, 0, WrappedText.Num()-1);
-	const FWrappedStringSlice& NewWrappedLine = WrappedText[NewWrappedLinedIndex];
+	const int32 NewLineViewIndex = FMath::Clamp( CurrentLineViewIndex + Direction, 0, NumberOfLineViews - 1 );
+	const FTextLayout::FLineView& NewLineView = LineViews[ NewLineViewIndex ];
 	
-	// We moved wrapped lines. Did we actually move to a different source line?
-	const bool bMovedToNewSourceLine = &CurrentWrappedLine.SourceString != &NewWrappedLine.SourceString;
-	const int32 NewSourceLineIndex = Location.GetLineIndex() + ( bMovedToNewSourceLine ? Direction : 0 );
 	// Our horizontal position is the clamped version of whatever the user explicitly set with horizontal movement.
-	const int32 NewOffsetInLine = FMath::Clamp(NewWrappedLine.FirstCharIndex + PreferredCursorOffsetInLine, 0, NewWrappedLine.LastCharIndex);
-	
-	return FTextLocation( NewSourceLineIndex, NewOffsetInLine );	
+	return TextLayout->GetTextLocationAt( NewLineView, FVector2D( PreferredCursorScreenOffsetInLine, NewLineView.Offset.Y ) * TextLayout->GetScale() );
 }
 
-SMultiLineEditableText::FTextLocation SMultiLineEditableText::ScanForWordBoundary( const FTextLocation& CurrentLocation, int8 Direction ) const
+FTextLocation SMultiLineEditableText::ScanForWordBoundary( const FTextLocation& CurrentLocation, int8 Direction ) const
 {
 	FTextLocation Location = TranslatedLocation(CurrentLocation, Direction);
 
-	while ( !IsAtStartOfDoc(Location) && !IsAtStartOfLine(Location) && !IsAtEndOfDoc(Location) && !IsAtEndOfLine(Location) && !IsAtWordStart(Location) )
+	while ( !IsAtBeginningOfDocument(Location) && !IsAtBeginningOfLine(Location) && !IsAtEndOfDocument(Location) && !IsAtEndOfLine(Location) && !IsAtWordStart(Location) )
 	{
 		Location = TranslatedLocation(Location, Direction);
 	}	
@@ -513,32 +1715,36 @@ SMultiLineEditableText::FTextLocation SMultiLineEditableText::ScanForWordBoundar
 
 TCHAR SMultiLineEditableText::GetCharacterAt( const FTextLocation& Location ) const
 {
-	const bool bIsLineEmpty = TextLines[Location.GetLineIndex()]->IsEmpty();
+	const TArray< FTextLayout::FLineModel >& Lines = TextLayout->GetLineModels();
+
+	const bool bIsLineEmpty = Lines[ Location.GetLineIndex() ].Text->IsEmpty();
 	return (bIsLineEmpty)
 		? '\n'
-		: (*TextLines[Location.GetLineIndex()])[Location.GetOffset()];
+		: (*Lines[ Location.GetLineIndex() ].Text )[Location.GetOffset()];
 }
 
-
-bool SMultiLineEditableText::IsAtStartOfDoc( const FTextLocation& Location ) const
+bool SMultiLineEditableText::IsAtBeginningOfDocument( const FTextLocation& Location ) const
 {
 	return Location.GetLineIndex() == 0 && Location.GetOffset() == 0;
 }
 	
-
-bool SMultiLineEditableText::IsAtEndOfDoc( const FTextLocation& Location ) const
+bool SMultiLineEditableText::IsAtEndOfDocument( const FTextLocation& Location ) const
 {
-	return TextLines.Num() == 0 || ( TextLines.Num()-1 == Location.GetLineIndex() && TextLines[Location.GetLineIndex()]->Len()-1 == Location.GetOffset());
+	const TArray< FTextLayout::FLineModel >& Lines = TextLayout->GetLineModels();
+	const int32 NumberOfLines = Lines.Num();
+
+	return NumberOfLines == 0 || ( NumberOfLines - 1 == Location.GetLineIndex() && Lines[ Location.GetLineIndex() ].Text->Len()-1 == Location.GetOffset() );
 }
 
-bool SMultiLineEditableText::IsAtStartOfLine( const FTextLocation& Location ) const
+bool SMultiLineEditableText::IsAtBeginningOfLine( const FTextLocation& Location ) const
 {
 	return Location.GetOffset() == 0;
 }
 
 bool SMultiLineEditableText::IsAtEndOfLine( const FTextLocation& Location ) const
 {
-	return TextLines[Location.GetLineIndex()]->Len()-1 == Location.GetOffset();
+	const TArray< FTextLayout::FLineModel >& Lines = TextLayout->GetLineModels();
+	return Lines[ Location.GetLineIndex() ].Text->Len()-1 == Location.GetOffset();
 }
 
 bool SMultiLineEditableText::IsAtWordStart( const FTextLocation& Location ) const
@@ -546,17 +1752,5 @@ bool SMultiLineEditableText::IsAtWordStart( const FTextLocation& Location ) cons
 	const TCHAR CharBeforeCursor = GetCharacterAt(TranslatedLocation(Location, -1));
 	const TCHAR CharAfterCursor = GetCharacterAt(TranslatedLocation(Location, +1));
 
-	return FChar::IsWhitespace(CharBeforeCursor) && !FChar::IsWhitespace(CharAfterCursor);
-}
-
-SMultiLineEditableText::FTextLocation SMultiLineEditableText::AbsoluteToWrapped( const FTextLocation& Location ) const
-{
-	const int32 CurrentWrappedLineIndex = WrappedText.FindMatch( FMatchWrappedLine( TextLines[Location.GetLineIndex()].Get(), Location.GetOffset() ) );
-	const int32 OffsetInWrappedLine = Location.GetOffset() - WrappedText[CurrentWrappedLineIndex].FirstCharIndex;
-	return FTextLocation(CurrentWrappedLineIndex, OffsetInWrappedLine);
-}
-
-SMultiLineEditableText::FTextRange SMultiLineEditableText::AbsoluteToWrapped( const FTextRange& Range ) const
-{
-	return FTextRange( AbsoluteToWrapped(Range.First), AbsoluteToWrapped(Range.Last) );
+	return FText::IsWhitespace(CharBeforeCursor) && !FText::IsWhitespace(CharAfterCursor);
 }
diff --git a/Engine/Source/Runtime/Slate/Public/Framework/Text/IRun.h b/Engine/Source/Runtime/Slate/Public/Framework/Text/IRun.h
--- a/Engine/Source/Runtime/Slate/Public/Framework/Text/IRun.h
+++ b/Engine/Source/Runtime/Slate/Public/Framework/Text/IRun.h
@@ -6,14 +6,27 @@ class SLATE_API IRun
 public:
 
 	virtual FTextRange GetTextRange() const = 0;
+	virtual void SetTextRange( const FTextRange& Value ) = 0;
+
 	virtual int16 GetBaseLine( float Scale ) const = 0;
 	virtual int16 GetMaxHeight( float Scale ) const = 0;
 
 	virtual FVector2D Measure( int32 StartIndex, int32 EndIndex, float Scale) const = 0;
-	virtual uint8 GetKerning( int32 CurrentIndex, float Scale ) const = 0;
+	virtual int8 GetKerning(int32 CurrentIndex, float Scale) const = 0;
 
 	virtual TSharedRef< class ILayoutBlock > CreateBlock( int32 StartIndex, int32 EndIndex, FVector2D Size, const TSharedPtr< class IRunHighlighter >& Highlighter ) = 0;
 
+	virtual int32 GetTextIndexAt( const TSharedRef< ILayoutBlock >& Block, const FVector2D& Location, float Scale ) const = 0;
+	
+	virtual FVector2D GetLocationAt(const TSharedRef< ILayoutBlock >& Block, int32 Offset, float Scale) const = 0;
+
 	virtual void BeginLayout() = 0;
 	virtual void EndLayout() = 0;
+
+	virtual void Move(const TSharedRef<FString>& NewText, const FTextRange& NewRange) = 0;
+	virtual TSharedRef<IRun> Clone() const = 0;
+
+	virtual void AppendText(FString& Text) const = 0;
+	virtual void AppendText(FString& Text, const FTextRange& Range) const = 0;
+
 };
diff --git a/Engine/Source/Runtime/Slate/Public/Framework/Text/SlateHyperlinkRun.h b/Engine/Source/Runtime/Slate/Public/Framework/Text/SlateHyperlinkRun.h
--- a/Engine/Source/Runtime/Slate/Public/Framework/Text/SlateHyperlinkRun.h
+++ b/Engine/Source/Runtime/Slate/Public/Framework/Text/SlateHyperlinkRun.h
@@ -38,13 +38,15 @@ public:
 
 	virtual FTextRange GetTextRange() const OVERRIDE;
 
+	virtual void SetTextRange( const FTextRange& Value ) OVERRIDE;
+
 	virtual int16 GetBaseLine( float Scale ) const OVERRIDE;
 
 	virtual int16 GetMaxHeight( float Scale ) const OVERRIDE;
 
 	virtual FVector2D Measure( int32 StartIndex, int32 EndIndex, float Scale ) const OVERRIDE;
 
-	virtual uint8 GetKerning( int32 CurrentIndex, float Scale ) const OVERRIDE;
+	virtual int8 GetKerning(int32 CurrentIndex, float Scale) const OVERRIDE;
 
 	virtual TSharedRef< ILayoutBlock > CreateBlock( int32 StartIndex, int32 EndIndex, FVector2D Size, const TSharedPtr< IRunHighlighter >& Highlighter ) OVERRIDE;
 
@@ -54,9 +56,19 @@ public:
 
 	virtual void ArrangeChildren( const TSharedRef< ILayoutBlock >& Block, const FGeometry& AllottedGeometry, FArrangedChildren& ArrangedChildren ) const OVERRIDE;
 
+	virtual int32 GetTextIndexAt( const TSharedRef< ILayoutBlock >& Block, const FVector2D& Location, float Scale ) const OVERRIDE;
+
+	virtual FVector2D GetLocationAt( const TSharedRef< ILayoutBlock >& Block, int32 Offset, float Scale ) const OVERRIDE;
+
 	virtual void BeginLayout() OVERRIDE { Children.Empty(); }
 	virtual void EndLayout() OVERRIDE {}
 
+	virtual void Move(const TSharedRef<FString>& NewText, const FTextRange& NewRange) OVERRIDE;
+	virtual TSharedRef<IRun> Clone() const OVERRIDE;
+
+	virtual void AppendText(FString& AppendToText) const OVERRIDE;
+	virtual void AppendText(FString& AppendToText, const FTextRange& PartialRange) const OVERRIDE;
+
 private:
 
 	FSlateHyperlinkRun( const TSharedRef< const FString >& InText, const FHyperlinkStyle& InStyle, const FMetadata& InMetadata, FOnClick InNavigateDelegate );
diff --git a/Engine/Source/Runtime/Slate/Public/Framework/Text/SlateImageRun.h b/Engine/Source/Runtime/Slate/Public/Framework/Text/SlateImageRun.h
--- a/Engine/Source/Runtime/Slate/Public/Framework/Text/SlateImageRun.h
+++ b/Engine/Source/Runtime/Slate/Public/Framework/Text/SlateImageRun.h
@@ -13,12 +13,14 @@ public:
 	virtual ~FSlateImageRun() {}
 
 	virtual FTextRange GetTextRange() const OVERRIDE;
+	virtual void SetTextRange( const FTextRange& Value ) OVERRIDE;
+
 	virtual int16 GetBaseLine( float Scale ) const OVERRIDE;
 	virtual int16 GetMaxHeight( float Scale ) const OVERRIDE;
 
 	virtual FVector2D Measure( int32 StartIndex, int32 EndIndex, float Scale ) const OVERRIDE;
 
-	virtual uint8 GetKerning( int32 CurrentIndex, float Scale ) const OVERRIDE;
+	virtual int8 GetKerning(int32 CurrentIndex, float Scale) const OVERRIDE;
 
 
 	virtual TSharedRef< ILayoutBlock > CreateBlock( int32 StartIndex, int32 EndIndex, FVector2D Size, const TSharedPtr< IRunHighlighter >& Highlighter ) OVERRIDE;
@@ -29,9 +31,19 @@ public:
 
 	virtual void ArrangeChildren( const TSharedRef< ILayoutBlock >& Block, const FGeometry& AllottedGeometry, FArrangedChildren& ArrangedChildren ) const OVERRIDE;
 
+	virtual int32 GetTextIndexAt( const TSharedRef< ILayoutBlock >& Block, const FVector2D& Location, float Scale ) const OVERRIDE;
+
+	virtual FVector2D GetLocationAt( const TSharedRef< ILayoutBlock >& Block, int32 Offset, float Scale ) const OVERRIDE;
+
 	virtual void BeginLayout() OVERRIDE {}
 	virtual void EndLayout() OVERRIDE {}
 
+	virtual void Move(const TSharedRef<FString>& NewText, const FTextRange& NewRange) OVERRIDE;
+	virtual TSharedRef<IRun> Clone() const OVERRIDE;
+
+	virtual void AppendText(FString& Text) const OVERRIDE;
+	virtual void AppendText(FString& AppendToText, const FTextRange& PartialRange) const OVERRIDE;
+
 private:
 
 	FSlateImageRun( const TSharedRef< const FString >& InText, const FSlateBrush* InImage, int16 InBaseline );
diff --git a/Engine/Source/Runtime/Slate/Public/Framework/Text/SlateTextRun.h b/Engine/Source/Runtime/Slate/Public/Framework/Text/SlateTextRun.h
--- a/Engine/Source/Runtime/Slate/Public/Framework/Text/SlateTextRun.h
+++ b/Engine/Source/Runtime/Slate/Public/Framework/Text/SlateTextRun.h
@@ -14,10 +14,12 @@ public:
 	virtual ~FSlateTextRun() {}
 
 	virtual FTextRange GetTextRange() const OVERRIDE;
+	virtual void SetTextRange( const FTextRange& Value ) OVERRIDE;
+
 	virtual int16 GetBaseLine( float Scale ) const OVERRIDE;
 	virtual int16 GetMaxHeight( float Scale ) const OVERRIDE;
 	virtual FVector2D Measure( int32 StartIndex, int32 EndIndex, float Scale ) const OVERRIDE;
-	virtual uint8 GetKerning( int32 CurrentIndex, float Scale ) const OVERRIDE;
+	virtual int8 GetKerning(int32 CurrentIndex, float Scale) const OVERRIDE;
 
 	virtual TSharedRef< ILayoutBlock > CreateBlock( int32 StartIndex, int32 EndIndex, FVector2D Size, const TSharedPtr< IRunHighlighter >& Highlighter ) OVERRIDE;
 
@@ -30,6 +32,16 @@ public:
 	virtual void BeginLayout() OVERRIDE {}
 	virtual void EndLayout() OVERRIDE {}
 
+	virtual int32 GetTextIndexAt( const TSharedRef< ILayoutBlock >& Block, const FVector2D& Location, float Scale ) const OVERRIDE;
+
+	virtual FVector2D GetLocationAt(const TSharedRef< ILayoutBlock >& Block, int32 Offset, float Scale) const OVERRIDE;
+
+	virtual void Move(const TSharedRef<FString>& NewText, const FTextRange& NewRange) OVERRIDE;
+	virtual TSharedRef<IRun> Clone() const OVERRIDE;
+
+	virtual void AppendText(FString& Text) const OVERRIDE;
+	virtual void AppendText(FString& AppendToText, const FTextRange& PartialRange) const OVERRIDE;
+
 private:
 
 	FSlateTextRun( const TSharedRef< const FString >& InText, const FTextBlockStyle& InStyle );
diff --git a/Engine/Source/Runtime/Slate/Public/Framework/Text/SlateWidgetRun.h b/Engine/Source/Runtime/Slate/Public/Framework/Text/SlateWidgetRun.h
--- a/Engine/Source/Runtime/Slate/Public/Framework/Text/SlateWidgetRun.h
+++ b/Engine/Source/Runtime/Slate/Public/Framework/Text/SlateWidgetRun.h
@@ -35,11 +35,13 @@ public:
 	virtual ~FSlateWidgetRun() {}
 
 	virtual FTextRange GetTextRange() const OVERRIDE;
+	virtual void SetTextRange( const FTextRange& Value ) OVERRIDE;
+
 	virtual int16 GetBaseLine( float Scale ) const OVERRIDE;
 	virtual int16 GetMaxHeight( float Scale ) const OVERRIDE;
 
 	virtual FVector2D Measure( int32 StartIndex, int32 EndIndex, float Scale ) const OVERRIDE;
-	virtual uint8 GetKerning( int32 CurrentIndex, float Scale ) const OVERRIDE;
+	virtual int8 GetKerning(int32 CurrentIndex, float Scale) const OVERRIDE;
 
 	virtual TSharedRef< ILayoutBlock > CreateBlock( int32 StartIndex, int32 EndIndex, FVector2D Size, const TSharedPtr< IRunHighlighter >& Highlighter ) OVERRIDE;
 
@@ -49,9 +51,19 @@ public:
 
 	virtual void ArrangeChildren( const TSharedRef< ILayoutBlock >& Block, const FGeometry& AllottedGeometry, FArrangedChildren& ArrangedChildren ) const;
 
+	virtual int32 GetTextIndexAt( const TSharedRef< ILayoutBlock >& Block, const FVector2D& Location, float Scale ) const OVERRIDE;
+
+	virtual FVector2D GetLocationAt( const TSharedRef< ILayoutBlock >& Block, int32 Offset, float Scale ) const OVERRIDE;
+
 	virtual void BeginLayout() OVERRIDE {}
 	virtual void EndLayout() OVERRIDE {}
 
+	virtual void Move(const TSharedRef<FString>& NewText, const FTextRange& NewRange) OVERRIDE;
+	virtual TSharedRef<IRun> Clone() const OVERRIDE;
+
+	virtual void AppendText(FString& Text) const OVERRIDE;
+	virtual void AppendText(FString& AppendToText, const FTextRange& PartialRange) const OVERRIDE;
+
 private:
 
 	FSlateWidgetRun( const TSharedRef< const FString >& InText, const FWidgetRunInfo& InWidgetInfo );
diff --git a/Engine/Source/Runtime/Slate/Public/Framework/Text/TextLayout.h b/Engine/Source/Runtime/Slate/Public/Framework/Text/TextLayout.h
--- a/Engine/Source/Runtime/Slate/Public/Framework/Text/TextLayout.h
+++ b/Engine/Source/Runtime/Slate/Public/Framework/Text/TextLayout.h
@@ -1,8 +1,6 @@
 // Copyright 1998-2014 Epic Games, Inc. All Rights Reserved.
 #pragma once
 
-#include "LineBreakIterator.h"
-
 #define TEXT_LAYOUT_DEBUG 0
 
 namespace ETextJustify
@@ -15,23 +13,136 @@ namespace ETextJustify
 	};
 }
 
+/** Location within the text model. */
+struct FTextLocation
+{
+public:
+	FTextLocation( const int32 InLineIndex = 0, const int32 InOffset = 0 )
+		: LineIndex( InLineIndex )
+		, Offset( InOffset )
+	{
+
+	}
+
+	FTextLocation( const FTextLocation& InLocation, const int32 InOffset )
+		: LineIndex( InLocation.GetLineIndex() )
+		, Offset(FMath::Max(InLocation.GetOffset() + InOffset, 0))
+	{
+
+	}
+
+	bool operator==( const FTextLocation& Other ) const
+	{
+		return
+			LineIndex == Other.LineIndex &&
+			Offset == Other.Offset;
+	}
+
+	bool operator!=( const FTextLocation& Other ) const
+	{
+		return
+			LineIndex != Other.LineIndex ||
+			Offset != Other.Offset;
+	}
+
+	bool operator<( const FTextLocation& Other ) const
+	{
+		return this->LineIndex < Other.LineIndex || (this->LineIndex == Other.LineIndex && this->Offset < Other.Offset);
+	}
+
+	int32 GetLineIndex() const { return LineIndex; }
+	int32 GetOffset() const { return Offset; }
+	bool IsValid() const { return LineIndex != INDEX_NONE && Offset != INDEX_NONE; }
+
+private:
+	int32 LineIndex;
+	int32 Offset;
+};
+
+class FTextSelection
+{
+
+public:
+
+	FTextLocation LocationA;
+
+	FTextLocation LocationB;
+
+public:
+
+	FTextSelection()
+		: LocationA(INDEX_NONE)
+		, LocationB(INDEX_NONE)
+	{
+	}
+
+	FTextSelection(const FTextLocation& InLocationA, const FTextLocation& InLocationB)
+		: LocationA(InLocationA)
+		, LocationB(InLocationB)
+	{
+	}
+
+	const FTextLocation& GetBeginning() const
+	{
+		if (LocationA.GetLineIndex() == LocationB.GetLineIndex())
+		{
+			if (LocationA.GetOffset() < LocationB.GetOffset())
+			{
+				return LocationA;
+			}
+
+			return LocationB;
+		}
+		else if (LocationA.GetLineIndex() < LocationB.GetLineIndex())
+		{
+			return LocationA;
+		}
+
+		return LocationB;
+	}
+
+	const FTextLocation& GetEnd() const
+	{
+		if (LocationA.GetLineIndex() == LocationB.GetLineIndex())
+		{
+			if (LocationA.GetOffset() > LocationB.GetOffset())
+			{
+				return LocationA;
+			}
+
+			return LocationB;
+		}
+		else if (LocationA.GetLineIndex() > LocationB.GetLineIndex())
+		{
+			return LocationA;
+		}
+
+		return LocationB;
+	}
+};
+
 class SLATE_API FTextLayout
 {
 public:
 
 	struct FBlockDefinition
 	{
-		FTextRange Range;
-		FTextRange VisualRange;
+		/** Range inclusive of trailing whitespace, as used to visually display and interact with the text */
+		FTextRange ActualRange;
+
 		TSharedPtr< IRunHighlighter > Highlighter;
 	};
 
 	struct FBreakCandidate
 	{
-		FTextRange Range; 
-		FTextRange RangeWithoutTrailingWhitespace; 
-		FVector2D Size;
-		FVector2D SizeWithoutTrailingWhitespace;
+		/** Range inclusive of trailing whitespace, as used to visually display and interact with the text */
+		FTextRange ActualRange;
+		/** Range exclusive of trailing whitespace, as used to perform wrapping on a word boundary */
+		FTextRange TrimmedRange;
+		/** Measured size inclusive of trailing whitespace, as used to visually display and interact with the text */
+		FVector2D ActualSize;
+		/** Measured size exclusive of trailing whitespace, as used to perform wrapping on a word boundary */
+		FVector2D TrimmedSize;
 
 		int16 MaxAboveBaseline;
 		int16 MaxBelowBaseline;
@@ -57,6 +168,8 @@ public:
 		void EndLayout();
 
 		FTextRange GetTextRange() const;
+		void SetTextRange( const FTextRange& Value );
+		
 		int16 GetBaseLine( float Scale ) const;
 		int16 GetMaxHeight( float Scale ) const;
 
@@ -72,6 +185,9 @@ public:
 
 		void ClearCache();
 
+		void AppendText(FString& Text) const;
+		void AppendText(FString& Text, const FTextRange& Range) const;
+
 	private:
 
 		TSharedRef< IRun > Run;
@@ -89,13 +205,17 @@ public:
 		TArray< FRunModel > Runs;
 		TArray< FBreakCandidate > BreakCandidates;
 		TArray< FTextHighlight > Highlights;
+		mutable bool HasWrappingInformation;
 	};
 
 	struct FLineView
 	{
 		TArray< TSharedRef< ILayoutBlock > > Blocks;
-		FVector2D Size;
 		FVector2D Offset;
+		FVector2D Size;
+		FVector2D TextSize;
+		FTextRange Range;
+		int32 ModelIndex;
 	};
 
 
@@ -148,14 +268,54 @@ public:
 	*/
 	virtual void UpdateIfNeeded();
 
+	virtual void Update();
+
+	int32 GetLineViewIndexForTextLocation(const TArray< FTextLayout::FLineView >& LineViews, const FTextLocation& Location);
+
+	/**
+	 * 
+	 */
+	FTextLocation GetTextLocationAt( const FVector2D& Relative );
+
+	FTextLocation GetTextLocationAt( const FLineView& LineView, const FVector2D& Relative );
+
+	FVector2D GetLocationAt( const FTextLocation& Location );
+
+	bool SplitLineAt(const FTextLocation& Location);
+
+	bool JoinLineWithNextLine(int32 LineIndex);
+
+	bool InsertAt(const FTextLocation& Location, TCHAR Character);
+
+	bool InsertAt( const FTextLocation& Location, const FString& Text );
+
+	bool RemoveAt( const FTextLocation& Location, int32 Count = 1 );
+
+	bool RemoveLine(int32 LineIndex);
+
+	bool IsEmpty() const;
+
+	void GetAsText(FString& DisplayText) const;
+
+	void GetSelectionAsText(FString& DisplayText, const FTextSelection& Selection) const;
+
+	FTextSelection GetWordAt(const FTextLocation& Location) const;
+
 protected:
 
 	FTextLayout();
 
 	/**
-	* Prepares the cache with the information needed to perform line wrapping if the information does not already exist.
+	* Create the wrapping cache for the current text based upon the current scale
+	* Each line keeps its own cached state, so needs to be cleared when changing the text within a line
+	* When changing the scale, all the lines need to be cleared (see ClearWrappingCache)
 	*/
-	void ReadyForWrapping();
+	void CreateWrappingCache();
+
+	/**
+	 * Clears the current wrapping cache for all lines
+	 */
+	void ClearWrappingCache();
 
 	/**
 	* Clears the current layouts view information.
@@ -184,7 +344,7 @@ private:
 
 	void JustifyLayout();
 
-	void CreateLineViewBlocks( const FLineModel& Line, const int32 StopIndex, const int32 VisualStopIndex, int32& OutRunIndex, int32& OutHighlightIndex, int32& OutPreviousBlockEnd, TArray< TSharedRef< ILayoutBlock > >& OutSoftLine );
+	void CreateLineViewBlocks( int32 LineModelIndex, const int32 StopIndex, int32& OutRunIndex, int32& OutHighlightIndex, int32& OutPreviousBlockEnd, TArray< TSharedRef< ILayoutBlock > >& OutSoftLine );
 
 	FBreakCandidate CreateBreakCandidate( int32& OutRunIndex, FLineModel& Line, int32 PreviousBreak, int32 CurrentBreak );
 
@@ -196,9 +356,6 @@ protected:
 	/** The views for the lines of text. A LineView represents a single visual line of text. Multiple LineViews can map to the same LineModel, if for example wrapping occurs. */
 	TArray< FLineView > LineViews;
 
-	/** Whether or not we currently have the information cached to perform line wrapping. */
-	bool HasWrapInformation;
-
 	/** Whether parameters on the layout have changed which requires the view be updated. */
 	bool HasChanged;
 
diff --git a/Engine/Source/Runtime/Slate/Public/Framework/Text/TextRange.h b/Engine/Source/Runtime/Slate/Public/Framework/Text/TextRange.h
--- a/Engine/Source/Runtime/Slate/Public/Framework/Text/TextRange.h
+++ b/Engine/Source/Runtime/Slate/Public/Framework/Text/TextRange.h
@@ -19,6 +19,20 @@ struct SLATE_API FTextRange
 
 	int32 Len() const { return EndIndex - BeginIndex; }
 	bool IsEmpty() const { return (EndIndex - BeginIndex) <= 0; }
+	void Offset(int32 Amount) { BeginIndex += Amount; BeginIndex = FMath::Max(0, BeginIndex);  EndIndex += Amount; EndIndex = FMath::Max(0, EndIndex); }
+	bool Contains(int32 Index) const { return Index >= BeginIndex && Index < EndIndex; }
+	bool InclusiveContains(int32 Index) const { return Index >= BeginIndex && Index <= EndIndex; }
+
+	FTextRange Intersect(const FTextRange& Other) const
+	{
+		FTextRange Intersected(FMath::Max(BeginIndex, Other.BeginIndex), FMath::Min(EndIndex, Other.EndIndex));
+		if (Intersected.EndIndex <= Intersected.BeginIndex)
+		{
+			return FTextRange(0, 0);
+		}
+
+		return Intersected;
+	}
 
 	int32 BeginIndex;
 	int32 EndIndex;
diff --git a/Engine/Source/Runtime/Slate/Public/Slate.h b/Engine/Source/Runtime/Slate/Public/Slate.h
--- a/Engine/Source/Runtime/Slate/Public/Slate.h
+++ b/Engine/Source/Runtime/Slate/Public/Slate.h
@@ -99,6 +99,7 @@ DECLARE_CYCLE_STAT_EXTERN( TEXT("Slate Misc Time"), STAT_SlateMiscTime, STATGROU
 #include "MultiBoxBuilder.h"
 #include "MultiBoxExtender.h"
 #include "SMultiLineEditableText.h"
+#include "SMultiLineEditableTextBox.h"
 #include "SEditableText.h"
 #include "SEditableTextBox.h"
 #include "SSearchBox.h"
diff --git a/Engine/Source/Runtime/Slate/Public/Widgets/Input/SEditableText.h b/Engine/Source/Runtime/Slate/Public/Widgets/Input/SEditableText.h
--- a/Engine/Source/Runtime/Slate/Public/Widgets/Input/SEditableText.h
+++ b/Engine/Source/Runtime/Slate/Public/Widgets/Input/SEditableText.h
@@ -29,8 +29,8 @@ public:
 		, _IsCaretMovedWhenGainFocus( true )
 		, _SelectAllTextWhenFocused( false )
 		, _RevertTextOnEscape( false )
-		, _ClearKeyboardFocusOnCommit( true )
-		, _MinDesiredWidth( 0.0f )
+		, _ClearKeyboardFocusOnCommit(true)
+		, _MinDesiredWidth(0.0f)
 		, _SelectAllTextOnCommit( false )
 		{}
 
@@ -132,16 +132,6 @@ public:
 	 */
 	void SetText( const TAttribute< FText >& InNewText );
 
-	/**
-	 * Checks to see if there is any text selected
-	 *
-	 * @return  True if any text is selected
-	 */
-	bool AnyTextSelected() const
-	{
-		return Selection.StartIndex != INDEX_NONE && Selection.StartIndex != Selection.FinishIndex;
-	}
-
 	/**
 	 * Restores the text to the original state
 	 */
@@ -279,13 +269,6 @@ protected:
 		}
 	};
 
-	/**
-	 * Gets the height of the largest character in the font
-	 *
-	 * @return  The fonts height
-	 */
-	virtual float GetFontHeight() const;
-
 public:
 	// BEGIN ITextEditorWidget interface
 	virtual void StartChangingText() OVERRIDE;
@@ -295,10 +278,21 @@ public:
 	virtual void DeleteChar() OVERRIDE;
 	virtual bool CanTypeCharacter(const TCHAR CharInQuestion) const OVERRIDE;
 	virtual void TypeChar( const int32 Character ) OVERRIDE;
-	virtual FReply MoveCursor( ECursorMoveMethod::Type Method, int8 Direction, ECursorAction::Type Action ) OVERRIDE;
+	virtual FReply MoveCursor( ECursorMoveMethod::Type Method, const FVector2D& Direction, ECursorAction::Type Action ) OVERRIDE;
 	virtual void JumpTo(ETextLocation::Type JumpLocation, ECursorAction::Type Action) OVERRIDE;
 	virtual void ClearSelection() OVERRIDE;
 	virtual void SelectAllText() OVERRIDE;
+	virtual bool SelectAllTextWhenFocused() OVERRIDE;
+	virtual void SelectWordAt(const FVector2D& LocalPosition) OVERRIDE;
+	virtual void BeginDragSelection() OVERRIDE;
+	virtual bool IsDragSelecting() const OVERRIDE;
+	virtual void EndDragSelection() OVERRIDE;
+	virtual bool AnyTextSelected() const OVERRIDE;
+	virtual bool IsTextSelectedAt(const FVector2D& LocalPosition) const OVERRIDE;
+	virtual void SetWasFocusedByLastMouseDown( bool Value ) OVERRIDE;
+	virtual bool WasFocusedByLastMouseDown() const OVERRIDE;
+	virtual void SetHasDragSelectedSinceFocused( bool Value ) OVERRIDE;
+	virtual bool HasDragSelectedSinceFocused() const OVERRIDE;
 	virtual FReply OnEscape() OVERRIDE;
 	virtual void OnEnter() OVERRIDE;
 	virtual bool CanExecuteCut() const OVERRIDE;
@@ -310,6 +304,9 @@ public:
 	virtual bool CanExecuteUndo() const OVERRIDE;
 	virtual void Undo() OVERRIDE;
 	virtual void Redo() OVERRIDE;
+	virtual TSharedRef< SWidget > GetWidget() OVERRIDE;
+	virtual void SummonContextMenu( const FVector2D& InLocation ) OVERRIDE;
+	virtual void LoadText() OVERRIDE;
 	// END ITextEditorWidget interface
 
 protected:
@@ -348,13 +345,22 @@ protected:
 	/**
 	 * Determines which character was clicked on
 	 *
-	 * @param  InMyGeometry  Widget geometry
-	 * @param  InMouseEvent  Mouse button event
+	 * @param  InMyGeometry				Widget geometry
+	 * @param  InScreenCursorPosition	Click position (in screen space)
 	 *
 	 * @return  The clicked character index
 	 */
 	int32 FindClickedCharacterIndex( const FGeometry& InMyGeometry, const FVector2D& InScreenCursorPosition ) const;
 
+	/**
+	 * Determines which character was clicked on
+	 *
+	 * @param  InLocalCursorPosition	Click position (in local space)
+	 *
+	 * @return  The clicked character index
+	 */
+	int32 FindClickedCharacterIndex( const FVector2D& InLocalCursorPosition ) const;
+
 	/**
 	 * Adds the specified undo state to the undo stack
 	 *
@@ -376,25 +382,11 @@ protected:
 	 */
 	void MakeUndoState( SEditableText::FUndoState& OutUndoState );
 
-
-	/**
-	 * Loads text for editing, from our text attribute
-	 */
-	void LoadText();
-
 	/**
 	 * Saves edited text into our text attribute, if appropriate
 	 */
 	void SaveText();
 
-	/**
-	 * Summons a context menu that allows the user to perform certain text-related actions
-	 *
-	 * @param	InLocation	Screen space position to show the menu
-	 */
-	void SummonContextMenu( const FVector2D& InLocation );
-
-
 	/**
 	 * Executes action to delete selected text
 	 */
@@ -431,13 +423,6 @@ protected:
 	 */
 	void OnWindowClosed(const TSharedRef<SWindow>&);
 
-	/** 
-	 * Calculate the width of the caret 
-	 * @param FontMaxCharHeight The height of the font to calculate the caret width for
-	 * @return The width of the caret (might be clamped for very small fonts)
-	 */
-	float CalculateCaretWidth(const float FontMaxCharHeight) const;
-
 private:
 	/** 
 	 * @return whether there is anything in the clipboard
@@ -487,7 +472,7 @@ private:
 	TAttribute< bool > IsCaretMovedWhenGainFocus;
 	
 	/** Whether to select all text when the user clicks to give focus on the widget */
-	TAttribute< bool > SelectAllTextWhenFocused;
+	TAttribute< bool > bSelectAllTextWhenFocused;
 
 	/** Whether to allow the user to back out of changes when they press the escape key */
 	TAttribute< bool > RevertTextOnEscape;
diff --git a/Engine/Source/Runtime/Slate/Public/Widgets/Input/SMultiLineEditableTextBox.h b/Engine/Source/Runtime/Slate/Public/Widgets/Input/SMultiLineEditableTextBox.h
new file mode 100644
--- /dev/null
+++ b/Engine/Source/Runtime/Slate/Public/Widgets/Input/SMultiLineEditableTextBox.h
@@ -0,0 +1,170 @@
+// Copyright 1998-2014 Epic Games, Inc. All Rights Reserved.
+
+#pragma once
+
+/**
+ * Editable text box widget
+ */
+class SLATE_API SMultiLineEditableTextBox : public SBorder
+{
+
+public:
+
+	SLATE_BEGIN_ARGS( SMultiLineEditableTextBox )
+		: _Style(&FCoreStyle::Get().GetWidgetStyle< FEditableTextBoxStyle >("NormalEditableTextBox"))
+		, _Text()
+		, _HintText()
+		, _Font()
+		, _ForegroundColor()
+		, _ReadOnlyForegroundColor()
+		, _Justification(ETextJustify::Left)
+		, _LineHeightPercentage(1.0f)
+		, _IsReadOnly( false )
+		, _IsPassword( false )
+		, _IsCaretMovedWhenGainFocus ( true )
+		, _SelectAllTextWhenFocused( false )
+		, _RevertTextOnEscape( false )
+		, _ClearKeyboardFocusOnCommit( true )
+		, _WrapTextAt(0.0f)
+		, _AutoWrapText(false)
+		, _SelectAllTextOnCommit( false )
+		, _BackgroundColor()		
+		, _Padding()
+		, _ErrorReporting()
+		{}
+
+		/** The styling of the textbox */
+		SLATE_STYLE_ARGUMENT( FEditableTextBoxStyle, Style )
+
+		/** Sets the text content for this editable text box widget */
+ 		SLATE_ATTRIBUTE( FText, Text )
+
+		/** Hint text that appears when there is no text in the text box */
+ 		SLATE_ATTRIBUTE( FText, HintText )
+
+		/** Font color and opacity (overrides Style) */
+		SLATE_ATTRIBUTE( FSlateFontInfo, Font )
+
+		/** Text color and opacity (overrides Style) */
+		SLATE_ATTRIBUTE( FSlateColor, ForegroundColor )
+		
+		/** Text color and opacity when read-only (overrides Style) */
+		SLATE_ATTRIBUTE( FSlateColor, ReadOnlyForegroundColor )
+
+		/** How the text should be aligned with the margin. */
+		SLATE_ATTRIBUTE(ETextJustify::Type, Justification)
+
+		/** The amount to scale each lines height by. */
+		SLATE_ATTRIBUTE(float, LineHeightPercentage)
+
+		/** Sets whether this text box can actually be modified interactively by the user */
+		SLATE_ATTRIBUTE( bool, IsReadOnly )
+
+		/** Sets whether this text box is for storing a password */
+		SLATE_ATTRIBUTE( bool, IsPassword )
+
+		/** Workaround as we loose focus when the auto completion closes. */
+		SLATE_ATTRIBUTE( bool, IsCaretMovedWhenGainFocus )
+
+		/** Whether to select all text when the user clicks to give focus on the widget */
+		SLATE_ATTRIBUTE( bool, SelectAllTextWhenFocused )
+
+		/** Whether to allow the user to back out of changes when they press the escape key */
+		SLATE_ATTRIBUTE( bool, RevertTextOnEscape )
+
+		/** Whether to clear keyboard focus when pressing enter to commit changes */
+		SLATE_ATTRIBUTE( bool, ClearKeyboardFocusOnCommit )
+
+		/** Called whenever the text is changed interactively by the user */
+		SLATE_EVENT( FOnTextChanged, OnTextChanged )
+
+		/** Called whenever the text is committed.  This happens when the user presses enter or the text box loses focus. */
+		SLATE_EVENT( FOnTextCommitted, OnTextCommitted )
+
+		/** Whether text wraps onto a new line when it's length exceeds this width; if this value is zero or negative, no wrapping occurs. */
+		SLATE_ATTRIBUTE( float, WrapTextAt )
+
+		/** Whether to wrap text automatically based on the widget's computed horizontal space.  IMPORTANT: Using automatic wrapping can result
+		    in visual artifacts, as the the wrapped size will computed be at least one frame late!  Consider using WrapTextAt instead.  The initial 
+			desired size will not be clamped.  This works best in cases where the text block's size is not affecting other widget's layout. */
+		SLATE_ATTRIBUTE( bool, AutoWrapText )
+
+		/** Whether to select all text when pressing enter to commit changes */
+		SLATE_ATTRIBUTE( bool, SelectAllTextOnCommit )
+
+		/** The color of the background/border around the editable text (overrides Style) */
+		SLATE_ATTRIBUTE( FSlateColor, BackgroundColor )
+
+		/** Padding between the box/border and the text widget inside (overrides Style) */
+		SLATE_ATTRIBUTE( FMargin, Padding )
+
+		/** Provide a alternative mechanism for error reporting. */
+		SLATE_ARGUMENT( TSharedPtr<class IErrorReportingWidget>, ErrorReporting )
+
+	SLATE_END_ARGS()
+	
+	/**
+	 * Construct this widget
+	 *
+	 * @param	InArgs	The declaration data for this widget
+	 */
+	void Construct( const FArguments& InArgs );
+	
+	/**
+	 * Returns the text string
+	 *
+	 * @return  Text string
+	 */
+	const FText& GetText() const
+	{
+		return EditableText->GetText();
+	}
+
+	/**
+	 * Sets the text string currently being edited 
+	 *
+	 * @param  InNewText  The new text string
+	 */
+	void SetText( const TAttribute< FText >& InNewText );
+
+	/**
+	 * If InError is a non-empty string the TextBox will the ErrorReporting provided during construction
+	 * If no error reporting was provided, the TextBox will create a default error reporter.
+	 */
+	void SetError( const FText& InError );
+	void SetError( const FString& InError );
+
+	// SWidget overrides
+	virtual bool SupportsKeyboardFocus() const OVERRIDE;
+	virtual bool HasKeyboardFocus() const OVERRIDE;
+	virtual FReply OnKeyboardFocusReceived( const FGeometry& MyGeometry, const FKeyboardFocusEvent& InKeyboardFocusEvent ) OVERRIDE;
+
+protected:
+
+	/** Editable text widget */
+	TSharedPtr< SMultiLineEditableText > EditableText;
+
+	/** Read-only foreground color */
+	TAttribute<FSlateColor> ReadOnlyForegroundColor;
+
+	/** Allows for inserting additional widgets that extend the functionality of the text box */
+	TSharedPtr<SHorizontalBox> Box;
+
+	/** SomeWidget reporting */
+	TSharedPtr<class IErrorReportingWidget> ErrorReporting;
+
+private:
+
+	/** Styling: border image to draw when not hovered or focused */
+	const FSlateBrush* BorderImageNormal;
+	/** Styling: border image to draw when hovered */
+	const FSlateBrush* BorderImageHovered;
+	/** Styling: border image to draw when focused */
+	const FSlateBrush* BorderImageFocused;
+	/** Styling: border image to draw when read only */
+	const FSlateBrush* BorderImageReadOnly;
+
+	/** @return Border image for the text box based on the hovered and focused state */
+	const FSlateBrush* GetBorderImage() const;
+
+};
diff --git a/Engine/Source/Runtime/Slate/Public/Widgets/Input/SVirtualKeyboardEntry.h b/Engine/Source/Runtime/Slate/Public/Widgets/Input/SVirtualKeyboardEntry.h
--- a/Engine/Source/Runtime/Slate/Public/Widgets/Input/SVirtualKeyboardEntry.h
+++ b/Engine/Source/Runtime/Slate/Public/Widgets/Input/SVirtualKeyboardEntry.h
@@ -133,13 +133,6 @@ public:
 
 protected:
 
-	/**
-	 * Gets the height of the largest character in the font
-	 *
-	 * @return  The fonts height
-	 */
-	virtual float GetFontHeight() const;
-
 	// SWidget interface
 	virtual void Tick( const FGeometry& AllottedGeometry, const double InCurrentTime, const float InDeltaTime ) OVERRIDE;
 	virtual FVector2D ComputeDesiredSize() const OVERRIDE;
diff --git a/Engine/Source/Runtime/Slate/Public/Widgets/Text/ITextEditorWidget.h b/Engine/Source/Runtime/Slate/Public/Widgets/Text/ITextEditorWidget.h
--- a/Engine/Source/Runtime/Slate/Public/Widgets/Text/ITextEditorWidget.h
+++ b/Engine/Source/Runtime/Slate/Public/Widgets/Text/ITextEditorWidget.h
@@ -12,6 +12,8 @@ namespace ECursorMoveMethod
 		CharacterVertical,
 		/** Move one word at a time (e.g. ctrl is held down). */
 		Word,
+		/** Move the cursor to the correct character based on the given screen position. */
+		ScreenPosition,
 	};
 }
 
@@ -41,7 +43,7 @@ namespace ETextLocation
 	};
 }
 
-class SLATE_API ITextEditorWidget
+class SLATE_API ITextEditorWidget 
 {
 	public:
 
@@ -91,7 +93,7 @@ class SLATE_API ITextEditorWidget
 	 * @param  Direction    Move forward/backward by number of units specified; usuatlly +/-1.
 	 * @param  Action       Just move or select text.
 	 */
-	virtual FReply MoveCursor( ECursorMoveMethod::Type Method, int8 Direction, ECursorAction::Type Action ) = 0;
+	virtual FReply MoveCursor( ECursorMoveMethod::Type Method, const FVector2D& Direction, ECursorAction::Type Action ) = 0;
 	
 	/**
 	 * Jump the cursor to a special location: e.g. end of line, beginning of document, etc..
@@ -108,6 +110,26 @@ class SLATE_API ITextEditorWidget
 	 */
 	virtual void SelectAllText() = 0;
 
+	virtual bool SelectAllTextWhenFocused() = 0;
+
+	virtual void SelectWordAt(const FVector2D& LocalPosition) = 0;
+
+	virtual void BeginDragSelection() = 0;
+
+	virtual bool IsDragSelecting() const = 0;
+
+	virtual void EndDragSelection() = 0;
+
+	virtual bool AnyTextSelected() const = 0;
+
+	virtual bool IsTextSelectedAt(const FVector2D& LocalPosition) const = 0;
+
+	virtual void SetWasFocusedByLastMouseDown( bool Value ) = 0;
+	virtual bool WasFocusedByLastMouseDown() const = 0;
+
+	virtual void SetHasDragSelectedSinceFocused( bool Value ) = 0;
+	virtual bool HasDragSelectedSinceFocused() const = 0;
+
 	/**
 	 * Executed when the user presses ESC
 	 */
@@ -170,4 +192,11 @@ class SLATE_API ITextEditorWidget
 	 * Restores cached undo state that was previously rolled back
 	 */
 	virtual void Redo() = 0;
+
+	virtual TSharedRef< SWidget > GetWidget() = 0;
+
+	virtual void SummonContextMenu( const FVector2D& InLocation ) = 0;
+
+	virtual void LoadText() = 0;
+
 };
diff --git a/Engine/Source/Runtime/Slate/Public/Widgets/Text/SMultiLineEditableText.h b/Engine/Source/Runtime/Slate/Public/Widgets/Text/SMultiLineEditableText.h
--- a/Engine/Source/Runtime/Slate/Public/Widgets/Text/SMultiLineEditableText.h
+++ b/Engine/Source/Runtime/Slate/Public/Widgets/Text/SMultiLineEditableText.h
@@ -2,50 +2,64 @@
 
 #pragma once
 
+#include "MultiBox/MultiBoxExtender.h"
 #include "ITextEditorWidget.h"
 
-/** Information about the substring as it would appear as a result of soft-wrapping a line of text. */
-struct FWrappedStringSlice
-{
-	/** 
-	 * Create a substring info that sources InSourceString, begins at InFirstCharIndex and ends at InLastCharIndex (exclusive).
-	 * The substring measures InSize.X by InSize.Y.
-	 */
-	FWrappedStringSlice( FString& InSourceString, int32 InFirstCharIndex, int32 InLastCharIndex, FVector2D InSize )
-	: SourceString( InSourceString )
-	, FirstCharIndex( InFirstCharIndex )
-	, LastCharIndex( InLastCharIndex )
-	, Size( InSize )
-	{
-	}
-
-	/** The substring sources this string. */
-	FString& SourceString;
-	/** Index of the first character in the substring */
-	int32 FirstCharIndex;
-	/** Index of the last character (exclusive) */
-	int32 LastCharIndex;
-	/** The measured size of this substring given the current font and point size. */
-	FVector2D Size;
-};
-
 /** An editable text widget that supports multiple lines and soft word-wrapping. */
-class SLATE_API SMultiLineEditableText : public SLeafWidget, public ITextEditorWidget
+class SLATE_API SMultiLineEditableText : public SWidget, public ITextEditorWidget
 {
 public:
 	SLATE_BEGIN_ARGS( SMultiLineEditableText )
-	: _Text()
-	, _ScrollBar()
-	, _PreferredSize(FVector2D(100,100))
+		: _Text()
+		, _WrapTextAt( 0.0f )
+		, _AutoWrapText(false)
+		, _TextStyle( &FCoreStyle::Get().GetWidgetStyle<FTextBlockStyle>( "NormalText" ) )
+		, _Font()
+		, _Margin( FMargin() )
+		, _LineHeightPercentage( 1.0f )
+		, _Justification( ETextJustify::Left )
+		, _IsReadOnly(false)
+		, _OnTextChanged()
+		, _OnTextCommitted()
+		, _ContextMenuExtender()
 	{}
 		/** The initial text that will appear in the widget. */
-		SLATE_TEXT_ARGUMENT( Text )
+		SLATE_ATTRIBUTE(FText, Text)
+
+		/** Whether text wraps onto a new line when it's length exceeds this width; if this value is zero or negative, no wrapping occurs. */
+		SLATE_ATTRIBUTE(float, WrapTextAt)
+
+		/** Whether to wrap text automatically based on the widget's computed horizontal space.  IMPORTANT: Using automatic wrapping can result
+		    in visual artifacts, as the the wrapped size will computed be at least one frame late!  Consider using WrapTextAt instead.  The initial 
+			desired size will not be clamped.  This works best in cases where the text block's size is not affecting other widget's layout. */
+		SLATE_ATTRIBUTE(bool, AutoWrapText)
+
+		/** Pointer to a style of the text block, which dictates the font, color, and shadow options. */
+		SLATE_STYLE_ARGUMENT(FTextBlockStyle, TextStyle)
+
+		/** Font color and opacity (overrides Style) */
+		SLATE_ATTRIBUTE(FSlateFontInfo, Font)
+
+		/** The amount of blank space left around the edges of text area. */
+		SLATE_ATTRIBUTE(FMargin, Margin)
+
+		/** The amount to scale each lines height by. */
+		SLATE_ATTRIBUTE(float, LineHeightPercentage)
+
+		/** How the text should be aligned with the margin. */
+		SLATE_ATTRIBUTE(ETextJustify::Type, Justification)
 
-		/** Scroll bar that is used to visualize and manipulate the scrolling of the text */
-		SLATE_ARGUMENT( TSharedPtr<class SScrollBar>, ScrollBar )
+		/** Sets whether this text box can actually be modified interactively by the user */
+		SLATE_ATTRIBUTE(bool, IsReadOnly)
 
-		/** How big should this text field be for optimum operation. The value is irrelevant when the editable text fills all available space. */
-		SLATE_ARGUMENT( FVector2D, PreferredSize)
+		/** Called whenever the text is changed interactively by the user */
+		SLATE_EVENT(FOnTextChanged, OnTextChanged)
+
+		/** Called whenever the text is committed.  This happens when the user presses enter or the text box loses focus. */
+		SLATE_EVENT(FOnTextCommitted, OnTextCommitted)
+
+		/** Menu extender for the right-click context menu */
+		SLATE_EVENT(FMenuExtensionDelegate, ContextMenuExtender)
 
 	SLATE_END_ARGS()
 
@@ -53,69 +67,190 @@ public:
 
 	void Construct( const FArguments& InArgs );
 
+	/**
+	 * Gets the text assigned to this text block
+	 */
+	const FText& GetText() const
+	{
+		return BoundText.Get();
+	}
+
+	/**
+	 * Sets the text for this text block
+	 */
+	void SetText(const TAttribute< FText >& InText);
+
+	virtual bool GetIsReadOnly() const OVERRIDE;
+	virtual void ClearSelection() OVERRIDE;
+
 private:
 	
-	/** Location within the text model. */
-	struct FTextLocation
+	enum class ECursorAlignment : uint8
+	{
+		/** Visually align the cursor to the left of the character its placed at, and insert text before the character */
+		Left,
+
+		/** Visually align the cursor to the right of the character its placed at, and insert text after the character */
+		Right,
+	};
+
+	enum class ECursorEOLMode : uint8
+	{
+		/** The cursor considers hard-line endings when looking for an EOL to align the cursor to (ignores line wrapping) */
+		HardLines,
+
+		/** The cursor considers soft-line endings when looking for an EOL to align the cursor to (checks for line wrapping) */
+		SoftLines,
+	};
+
+	/** Store the information about the current cursor position */
+	class FCursorInfo
 	{
 	public:
-		FTextLocation( const int32 InLineIndex = 0, const int32 InOffset = 0 )
-		: LineIndex( InLineIndex )
-		, Offset( InOffset )
+		FCursorInfo()
+			: CursorPosition()
+			, CursorAlignment(ECursorAlignment::Left)
+			, LastCursorInteractionTime(0)
+		{
+		}
+
+		/** Get the literal position of the cursor (note: this may not be the correct place to insert text to, use GetCursorInteractionLocation for that) */
+		FORCEINLINE FTextLocation GetCursorLocation() const
 		{
+			return CursorPosition;
+		}
 
+		/** Get the alignment of the cursor */
+		FORCEINLINE ECursorAlignment GetCursorAlignment() const
+		{
+			return CursorAlignment;
 		}
 
-		bool operator==( const FTextLocation& Other ) const
+		/** Get the interaction position of the cursor (where to insert, delete, etc, text from/to) */
+		FORCEINLINE FTextLocation GetCursorInteractionLocation() const
 		{
-			return
-				LineIndex == Other.LineIndex &&
-				Offset == Other.Offset;
+			// If the cursor is right aligned, we treat it as if it were one character further along for insert/delete purposes
+			return FTextLocation(CursorPosition, (CursorAlignment == ECursorAlignment::Right) ? 1 : 0);
 		}
 
-		bool operator<( const FTextLocation& Other ) const
+		FORCEINLINE double GetLastCursorInteractionTime() const
 		{
-			return this->LineIndex < Other.LineIndex || (this->LineIndex == Other.LineIndex && this->Offset < Other.Offset);
+			return LastCursorInteractionTime;
 		}
 
-		int32 GetLineIndex() const { return LineIndex; }
-		int32 GetOffset() const { return Offset; }
+		/** Set the position of the cursor, and then work out the correct alignment based on the current text layout */
+		void SetCursorLocationAndCalculateAlignment(const TSharedPtr<FTextLayout>& TextLayout, const FTextLocation& InCursorPosition, const ECursorEOLMode CursorEOLMode = ECursorEOLMode::HardLines);
 
-		private:
-		int32 LineIndex;
-		int32 Offset;
+		/** Set the literal position and alignment of the cursor */
+		void SetCursorLocationAndAlignment(const FTextLocation& InCursorPosition, const ECursorAlignment InCursorAlignment);
+
+		/** Create an undo for this cursor data */
+		FCursorInfo CreateUndo() const;
+
+		/** Restore this cursor data from an undo */
+		void RestoreFromUndo(const FCursorInfo& UndoData);
+
+	private:
+		/** Current cursor position; there is always one. */
+		FTextLocation CursorPosition;
+
+		/** Cursor alignment (horizontal) within its position. This affects the rendering and behavior of the cursor when adding new characters */
+		ECursorAlignment CursorAlignment;
+
+		/** Last time the user did anything with the cursor.*/
+		double LastCursorInteractionTime;
 	};
 
-	/** Range between two locations. Guarantees that the First location is before the Last. */
-	struct FTextRange
+	/**
+	* Stores a single undo level for editable text
+	*/
+	class FUndoState
 	{
-		/**
-		 * Construct a Range from two locations: A and B.
-		 * Guarantees that First will be the earlier of the two.
-		 */
-		FTextRange( const FTextLocation& A, const FTextLocation& B )
-		: First( (A < B) ? A : B )
-		, Last( (A < B) ? B : A )
+
+	public:
+
+		/** Text */
+		FText Text;
+
+		/** Selection state */
+		TOptional<FTextLocation> SelectionStart;
+
+		/** Cursor data */
+		FCursorInfo CursorInfo;
+
+	public:
+
+		FUndoState()
+			: Text()
+			, SelectionStart()
+			, CursorInfo()
 		{
 		}
+	};
+
+	class FSlateCursorRunHighlighter : public ISlateRunHighlighter
+	{
+	public:
+
+		static TSharedRef< FSlateCursorRunHighlighter > Create(const FCursorInfo* InCursorInfo);
+
+	public:
+
+		enum class EMode : uint8
+		{
+			Cursor,
+			Selection,
+		};
+
+	public:
+
+		virtual ~FSlateCursorRunHighlighter() {}
+
+		virtual int32 OnPaint( const FTextLayout::FLineView& Line, const TSharedRef< ISlateRun >& Run, const TSharedRef< ILayoutBlock >& Block, const FTextBlockStyle& DefaultStyle, const FGeometry& AllottedGeometry, const FSlateRect& MyClippingRect, FSlateWindowElementList& OutDrawElements, int32 LayerId, const FWidgetStyle& InWidgetStyle, bool bParentEnabled ) const OVERRIDE;
 
-		const FTextLocation First;
-		const FTextLocation Last;
+		virtual FChildren* GetChildren() OVERRIDE;
+
+		virtual void OnArrangeChildren( const TSharedRef< ILayoutBlock >& Block, const FGeometry& AllottedGeometry, FArrangedChildren& ArrangedChildren ) const OVERRIDE;
+
+		void SetMode( EMode Mode );
+
+	private:
+
+		FSlateCursorRunHighlighter(const FCursorInfo* InCursorInfo);
+
+	private:
+
+		/** Cursor data that this highlighter is tracking */
+		const FCursorInfo* CursorInfo;
+
+		EMode Mode;
+
+		static FNoChildren NoChildrenInstance;
 	};
 
 private:
+
 	// BEGIN ITextEditorWidget interface
 	virtual void StartChangingText() OVERRIDE;
 	virtual void FinishChangingText() OVERRIDE;
-	virtual bool GetIsReadOnly() const OVERRIDE;
 	virtual void BackspaceChar() OVERRIDE;
 	virtual void DeleteChar() OVERRIDE;
 	virtual bool CanTypeCharacter(const TCHAR CharInQuestion) const OVERRIDE;
 	virtual void TypeChar( const int32 Character ) OVERRIDE;
-	virtual FReply MoveCursor( ECursorMoveMethod::Type Method, int8 Direction, ECursorAction::Type Action ) OVERRIDE;
+	virtual FReply MoveCursor( ECursorMoveMethod::Type Method, const FVector2D& Direction, ECursorAction::Type Action ) OVERRIDE;
 	virtual void JumpTo(ETextLocation::Type JumpLocation, ECursorAction::Type Action) OVERRIDE;
-	virtual void ClearSelection() OVERRIDE;
 	virtual void SelectAllText() OVERRIDE;
+	virtual bool SelectAllTextWhenFocused() OVERRIDE;
+	virtual void SelectWordAt(const FVector2D& LocalPosition) OVERRIDE;
+	virtual void BeginDragSelection() OVERRIDE;
+	virtual bool IsDragSelecting() const OVERRIDE;
+	virtual void EndDragSelection() OVERRIDE;
+	virtual bool AnyTextSelected() const OVERRIDE;
+	virtual bool IsTextSelectedAt(const FVector2D& LocalPosition) const OVERRIDE;
+	virtual void SetWasFocusedByLastMouseDown( bool Value ) OVERRIDE;
+	virtual bool WasFocusedByLastMouseDown() const OVERRIDE;
+	virtual void SetHasDragSelectedSinceFocused( bool Value ) OVERRIDE;
+	virtual bool HasDragSelectedSinceFocused() const OVERRIDE;
 	virtual FReply OnEscape() OVERRIDE;
 	virtual void OnEnter() OVERRIDE;
 	virtual bool CanExecuteCut() const OVERRIDE;
@@ -127,32 +262,48 @@ private:
 	virtual bool CanExecuteUndo() const OVERRIDE;
 	virtual void Undo() OVERRIDE;
 	virtual void Redo() OVERRIDE;
+	virtual TSharedRef< SWidget > GetWidget() OVERRIDE;
+	virtual void SummonContextMenu( const FVector2D& InLocation ) OVERRIDE;
+	virtual void LoadText() OVERRIDE;
 	// END ITextEditorWidget interface
 
+
 private:
 	// BEGIN SWidget interface
 	virtual void Tick( const FGeometry& AllottedGeometry, const double InCurrentTime, const float InDeltaTime ) OVERRIDE;
 	virtual int32 OnPaint( const FGeometry& AllottedGeometry, const FSlateRect& MyClippingRect, FSlateWindowElementList& OutDrawElements, int32 LayerId, const FWidgetStyle& InWidgetStyle, bool bParentEnabled ) const OVERRIDE;
+	virtual void CacheDesiredSize() OVERRIDE;
 	virtual FVector2D ComputeDesiredSize() const OVERRIDE;
+	virtual FChildren* GetChildren() OVERRIDE;
+	virtual void OnArrangeChildren( const FGeometry& AllottedGeometry, FArrangedChildren& ArrangedChildren ) const OVERRIDE;
 	virtual bool SupportsKeyboardFocus() const OVERRIDE;
 	virtual FReply OnKeyChar( const FGeometry& MyGeometry,const FCharacterEvent& InCharacterEvent ) OVERRIDE;
 	virtual FReply OnKeyDown( const FGeometry& MyGeometry, const FKeyboardEvent& InKeyboardEvent ) OVERRIDE;
 	virtual FReply OnKeyUp( const FGeometry& MyGeometry, const FKeyboardEvent& InKeyboardEvent ) OVERRIDE;
+	virtual FReply OnKeyboardFocusReceived( const FGeometry& MyGeometry, const FKeyboardFocusEvent& InKeyboardFocusEvent ) OVERRIDE;
+	virtual void OnKeyboardFocusLost( const FKeyboardFocusEvent& InKeyboardFocusEvent ) OVERRIDE;
+	virtual FReply OnMouseButtonDown( const FGeometry& MyGeometry, const FPointerEvent& MouseEvent ) OVERRIDE;
+	virtual FReply OnMouseButtonUp( const FGeometry& MyGeometry, const FPointerEvent& MouseEvent ) OVERRIDE;
+	virtual FReply OnMouseMove( const FGeometry& MyGeometry, const FPointerEvent& MouseEvent ) OVERRIDE;
+	virtual FReply OnMouseButtonDoubleClick(const FGeometry& MyGeometry, const FPointerEvent& MouseEvent) OVERRIDE;
+	virtual FCursorReply OnCursorQuery( const FGeometry& MyGeometry, const FPointerEvent& CursorEvent ) const OVERRIDE;
 	// END SWidget interface
 
 private:
 	
-	/** Re-compute the wrapping information about the current text. */
-	void ReWrap(const FGeometry& AllottedGeometry);
-
-	/** Invoked when the user manipulates the scrollbar*/
-	void OnScrollbarScrolled(float OffsetFraction);
+	void OnWindowClosed(const TSharedRef<SWindow>&);
 
 	/** Remember where the cursor was when we started selecting. */
 	void BeginSelecting( const FTextLocation& Endpoint );
 
-	/** Given the current geometry, make the scrollbar reflect the current offset and fraction of the visible text. */
-	void UpdateScrollbar( const FGeometry &AllottedGeometry );
+	void DeleteSelectedText();
+
+	void ExecuteDeleteAction();
+	bool CanExecuteDelete() const;
+
+	bool CanExecuteSelectAll() const;
+
+	bool DoesClipboardHaveAnyText() const;
 
 	/**
 	 * Given a location and a Direction to offset, return a new location.
@@ -177,13 +328,13 @@ private:
 	TCHAR GetCharacterAt( const FTextLocation& Location ) const;
 
 	/** Are we at the beginning of all the text. */
-	bool IsAtStartOfDoc( const FTextLocation& Location ) const;
+	bool IsAtBeginningOfDocument( const FTextLocation& Location ) const;
 	
 	/** Are we at the end of all the text. */
-	bool IsAtEndOfDoc( const FTextLocation& Location ) const;
+	bool IsAtEndOfDocument( const FTextLocation& Location ) const;
 
 	/** Is this location the beginning of a line */
-	bool IsAtStartOfLine( const FTextLocation& Location ) const;
+	bool IsAtBeginningOfLine( const FTextLocation& Location ) const;
 
 	/** Is this location the end of a line. */
 	bool IsAtEndOfLine( const FTextLocation& Location ) const;
@@ -191,42 +342,95 @@ private:
 	/** Are we currently at the beginning of a word */
 	bool IsAtWordStart( const FTextLocation& Location ) const;
 
-	/** Given a location in unwrapped text, compute the location in wrapped line. */
-	FTextLocation AbsoluteToWrapped( const FTextLocation& Location ) const;
+	void UpdateCursorHighlight();
+
+	void RemoveCursorHighlight();
 
-	/** Given a range in unwrapped text, compute the range in wrapped text. */
-	FTextRange AbsoluteToWrapped( const FTextRange& Range ) const;
+	void PushUndoState(const SMultiLineEditableText::FUndoState& InUndoState);
 
-	/** If we have a pointer to this scroll bar, use it to visualize and manipulate the scroll offset. */
-	TWeakPtr<class SScrollBar> OptionalScrollBarPtr;
+	void ClearUndoStates();
 
-	/** We maintain our own version of the text as an array of strings. This is done for performance reasons. */
-	TArray< TSharedRef<FString> > TextLines;
+	void MakeUndoState(SMultiLineEditableText::FUndoState& OutUndoState);
+
+	void SaveText(const FText& TextToSave);
+
+private:
+
+	/** The text displayed in this text block */
+	TAttribute<FText> BoundText;
+
+	/** The text that is currently in TextLayout; used to avoid un-needed updates */
+	TOptional<FText> LastSetText;
+
+	/** In control of the layout and wrapping of the BoundText */
+	TSharedPtr< FSlateTextLayout > TextLayout;
+
+	/** Default style used by the TextLayout */
+	FTextBlockStyle TextStyle;
 	
-	/** The information about how the text is currently wrapped. */
-	TArray< FWrappedStringSlice > WrappedText;
+	/** Whether text wraps onto a new line when it's length exceeds this width; if this value is zero or negative, no wrapping occurs. */
+	TAttribute< float > WrapTextAt;
+
+	/** True if we're wrapping text automatically based on the computed horizontal space for this widget */
+	TAttribute< bool > AutoWrapText;
+
+	/** Cached auto-wrap width that this text is using. This is used when determining if the cached string size should be updated */
+	mutable float CachedAutoWrapTextWidth;
 
-	/** Remember last frame's geometry for comparison and  */
-	FGeometry CachedLastFrameGeometry;
+	TAttribute< FMargin > Margin;
+	TAttribute< ETextJustify::Type > Justification; 
+	TAttribute< float > LineHeightPercentage;
+
+	TSharedPtr< FSlateCursorRunHighlighter > CursorSelectionHighlighter;
 
 	/** That start of the selection when there is a selection. The end is implicitly wherever the cursor happens to be. */
 	TOptional<FTextLocation> SelectionStart;
 
-	/** Current cursor position; there is always one. */
-	FTextLocation CursorPosition;
+	/** Current cursor data */
+	FCursorInfo CursorInfo;
+
+	/** The user probably wants the cursor where they last explicitly positioned it horizontally. */
+	float PreferredCursorScreenOffsetInLine;
 
-	/** The font being used to render the text. */
-	FSlateFontInfo FontInfo;
+	/** Whether to select all text when the user clicks to give focus on the widget */
+	TAttribute< bool > bSelectAllTextWhenFocused;
 
-	/** Floating point offset in number of lines from the top of the document. */
-	float NumLinesScrollOffset;
+	/** True if we're currently selecting text by dragging the mouse cursor with the left button held down */
+	bool bIsDragSelecting;
 
-	/** The user probably wants the cursor where they last explicitly positioned it horizontally. */
-	int32 PreferredCursorOffsetInLine;
+	/** True if the last mouse down caused us to receive keyboard focus */
+	bool bWasFocusedByLastMouseDown;
+
+	/** True if characters were selected by dragging since the last keyboard focus.  Used for text selection. */
+	bool bHasDragSelectedSinceFocused;
+
+	/** Undo states */
+	TArray< FUndoState > UndoStates;
+
+	/** Current undo state level that we've rolled back to, or INDEX_NONE if we haven't undone.  Used for 'Redo'. */
+	int32 CurrentUndoLevel;
+
+	/** Undo state that will be pushed if text is actually changed between calls to StartChangingText() and FinishChangingText() */
+	FUndoState StateBeforeChangingText;
+
+	/** True if we're currently (potentially) changing the text string */
+	bool bIsChangingText;
+
+	/** Sets whether this text box can actually be modified interactively by the user */
+	TAttribute< bool > IsReadOnly;
+
+	/** A list commands to execute if a user presses the corresponding keybinding in the text box */
+	TSharedRef< FUICommandList > UICommandList;
+
+	/** Called whenever the text is changed interactively by the user */
+	FOnTextChanged OnTextChanged;
+
+	/** Called whenever the text is committed.  This happens when the user presses enter or the text box loses focus. */
+	FOnTextCommitted OnTextCommitted;
 
-	/** Last time the user did anything with the cursor.*/
-	double LastCursorInteractionTime;
+	/** Menu extender for right-click context menu */
+	TSharedPtr<FExtender> MenuExtender;
 
-	/** How big should this text field be for optimum operation. The value is irrelevant when the editable text fills all available space. */
-	FVector2D PreferredSize;
+	/** Weak pointer to context menu window that's currently open, if there is one */
+	TWeakPtr< SWindow > ContextMenuWindow;
 };
diff --git a/Engine/Source/Runtime/SlateCore/Private/Fonts/FontMeasure.cpp b/Engine/Source/Runtime/SlateCore/Private/Fonts/FontMeasure.cpp
--- a/Engine/Source/Runtime/SlateCore/Private/Fonts/FontMeasure.cpp
+++ b/Engine/Source/Runtime/SlateCore/Private/Fonts/FontMeasure.cpp
@@ -260,7 +260,7 @@ uint16 FSlateFontMeasure::GetMaxCharacterHeight( const FSlateFontInfo& InFontInf
 	return MaxHeight;
 }
 
-uint8 FSlateFontMeasure::GetKerning( const FSlateFontInfo& InFontInfo, float FontScale, TCHAR PreviousCharacter, TCHAR CurrentCharacter ) const
+int8 FSlateFontMeasure::GetKerning(const FSlateFontInfo& InFontInfo, float FontScale, TCHAR PreviousCharacter, TCHAR CurrentCharacter) const
 {
 	FCharacterList& CharacterList = FontCache->GetCharacterList( InFontInfo, FontScale );
 	return CharacterList.GetKerning( PreviousCharacter, CurrentCharacter );
diff --git a/Engine/Source/Runtime/SlateCore/Private/Rendering/ElementBatcher.cpp b/Engine/Source/Runtime/SlateCore/Private/Rendering/ElementBatcher.cpp
--- a/Engine/Source/Runtime/SlateCore/Private/Rendering/ElementBatcher.cpp
+++ b/Engine/Source/Runtime/SlateCore/Private/Rendering/ElementBatcher.cpp
@@ -494,7 +494,7 @@ void FSlateElementBatcher::AddTextElement( const FVector2D& Position, const FVec
 				InvTextureSizeY = 1.0f/FontTexture->GetHeight();
 			}
 
-			const bool bIsWhitespace = FChar::IsWhitespace(CurrentChar);
+			const bool bIsWhitespace = FText::IsWhitespace(CurrentChar);
 
 			if( !bIsWhitespace && CharIndex > 0 )
 			{
diff --git a/Engine/Source/Runtime/SlateCore/Public/Fonts/FontMeasure.h b/Engine/Source/Runtime/SlateCore/Public/Fonts/FontMeasure.h
--- a/Engine/Source/Runtime/SlateCore/Public/Fonts/FontMeasure.h
+++ b/Engine/Source/Runtime/SlateCore/Public/Fonts/FontMeasure.h
@@ -176,7 +176,7 @@ public:
 	 *
 	 * @return the kerning value used between the two specified characters
 	 */
-	uint8 GetKerning( const FSlateFontInfo& InFontInfo, float FontScale, TCHAR PreviousCharacter, TCHAR CurrentCharacter ) const;
+	int8 GetKerning(const FSlateFontInfo& InFontInfo, float FontScale, TCHAR PreviousCharacter, TCHAR CurrentCharacter) const;
 
 	/**
 	 * Returns the baseline for the specified font.
