Fixed keyboard shortcuts not being right-aligned on menu blocks

This was a regressions caused by CL#2088161 and addresses TTP#337982 - Regression: UE4: Slate menu keybinding strings are no longer right aligned

diff --git a/Engine/Source/Runtime/Slate/Private/Framework/MultiBox/SMenuEntryBlock.cpp b/Engine/Source/Runtime/Slate/Private/Framework/MultiBox/SMenuEntryBlock.cpp
--- a/Engine/Source/Runtime/Slate/Private/Framework/MultiBox/SMenuEntryBlock.cpp
+++ b/Engine/Source/Runtime/Slate/Private/Framework/MultiBox/SMenuEntryBlock.cpp
@@ -416,7 +416,6 @@ TSharedRef< SWidget > SMenuEntryBlock::BuildMenuEntryWidget( const FMenuEntryBui
 			]
 		]
 		+SHorizontalBox::Slot()
-		.AutoWidth()
 		[
 			ButtonContent.ToSharedRef()
 		];
