SLATE / UMG - Working on undo/redo.  All the attributes of STextBlock can now be set through method calls after construction.

diff --git a/Engine/Plugins/Experimental/LiveEditor/Source/LiveEditor/Private/LiveEditorConfigWindow.cpp b/Engine/Plugins/Experimental/LiveEditor/Source/LiveEditor/Private/LiveEditorConfigWindow.cpp
--- a/Engine/Plugins/Experimental/LiveEditor/Source/LiveEditor/Private/LiveEditorConfigWindow.cpp
+++ b/Engine/Plugins/Experimental/LiveEditor/Source/LiveEditor/Private/LiveEditorConfigWindow.cpp
@@ -120,7 +120,7 @@ public:
 			UE_LOG( LiveEditorConfigWindow, Log, TEXT("%s Activated"), *(*n) );
 
 			bIsActive = (Name.Get())? FLiveEditorManager::Get().CheckActive( *Name.Get() ) : false;
-			NameView->SetForegroundColor( GetFontColor() );
+			NameView->SetColorAndOpacity( GetFontColor() );
 		}
 		return FReply::Handled();
 	}
@@ -133,7 +133,7 @@ public:
 			UE_LOG( LiveEditorConfigWindow, Log, TEXT("%s DeActivated"), *(*n) );
 
 			bIsActive = (Name.Get())? FLiveEditorManager::Get().CheckActive( *Name.Get() ) : false;
-			NameView->SetForegroundColor( GetFontColor() );
+			NameView->SetColorAndOpacity(GetFontColor());
 		}
 		return FReply::Handled();
 	}
diff --git a/Engine/Source/Editor/Kismet/Private/SBlueprintPalette.cpp b/Engine/Source/Editor/Kismet/Private/SBlueprintPalette.cpp
--- a/Engine/Source/Editor/Kismet/Private/SBlueprintPalette.cpp
+++ b/Engine/Source/Editor/Kismet/Private/SBlueprintPalette.cpp
@@ -372,25 +372,25 @@ static TSharedRef<IToolTip> ConstructToolTipWithActionPath(TSharedPtr<FEdGraphSc
 	FFavoritedBlueprintPaletteItem ActionItem(ActionIn);
 	if (ActionItem.IsValid())
 	{
-		FTextBlockStyle PathStyle = FTextBlockStyle()
+		static FTextBlockStyle PathStyle = FTextBlockStyle()
 			.SetFont(FSlateFontInfo(FPaths::EngineContentDir() / TEXT("Slate/Fonts/Roboto-Regular.ttf"), 8))
 			.SetColorAndOpacity(FLinearColor(0.4f, 0.4f, 0.4f));
 
 		NewToolTip = SNew(SToolTip)
 		[
 			SNew(SVerticalBox)
-			+SVerticalBox::Slot()
+			+ SVerticalBox::Slot()
 			[
 				OldToolTip->GetContentWidget()
 			]
 
-			+SVerticalBox::Slot()
-				.HAlign(EHorizontalAlignment::HAlign_Right)
+			+ SVerticalBox::Slot()
+			.HAlign(EHorizontalAlignment::HAlign_Right)
 			[
 				SNew(STextBlock)
-					.ColorAndOpacity(FSlateColor::UseSubduedForeground())
-					.Text(FText::FromString(ActionItem.ToString()))
-					.TextStyle(&PathStyle)
+				.ColorAndOpacity(FSlateColor::UseSubduedForeground())
+				.Text(FText::FromString(ActionItem.ToString()))
+				//.TextStyle(&PathStyle)
 			]
 		];
 	}
diff --git a/Engine/Source/Editor/Kismet/Private/SBlueprintSubPalette.cpp b/Engine/Source/Editor/Kismet/Private/SBlueprintSubPalette.cpp
--- a/Engine/Source/Editor/Kismet/Private/SBlueprintSubPalette.cpp
+++ b/Engine/Source/Editor/Kismet/Private/SBlueprintSubPalette.cpp
@@ -356,44 +356,46 @@ TSharedRef<SVerticalBox> SBlueprintSubPalette::ConstructHeadingWidget(FSlateBrus
 	TSharedPtr<SToolTip> ToolTip;
 	SAssignNew(ToolTip, SToolTip).Text(ToolTipText);
 
-	FTextBlockStyle TitleStyle = FTextBlockStyle()
+	static FTextBlockStyle TitleStyle = FTextBlockStyle()
 		.SetFont(FSlateFontInfo(FPaths::EngineContentDir() / TEXT("Slate/Fonts/Roboto-Bold.ttf"), 10))
 		.SetColorAndOpacity(FLinearColor(0.4f, 0.4f, 0.4f));
 
 	return SNew(SVerticalBox)
-			.ToolTip(ToolTip)
-			// so we still get tooltip text for an empty SHorizontalBox
-			.Visibility(EVisibility::Visible) 
-		+SVerticalBox::Slot()
-			.AutoHeight()
+		.ToolTip(ToolTip)
+		// so we still get tooltip text for an empty SHorizontalBox
+		.Visibility(EVisibility::Visible) 
+		+ SVerticalBox::Slot()
+		.AutoHeight()
 		[
 			SNew(SHorizontalBox)
-			+SHorizontalBox::Slot()
-				.AutoWidth()
-				.VAlign(VAlign_Center)
-				.Padding(2.f, 2.f)
+			
+			+ SHorizontalBox::Slot()
+			.AutoWidth()
+			.VAlign(VAlign_Center)
+			.Padding(2.f, 2.f)
 			[
 				SNew(SImage).Image(Icon)
 			]
 
-			+SHorizontalBox::Slot()
-				.AutoWidth()
-				.VAlign(VAlign_Center)
-				.Padding(2.f, 2.f)
+			+ SHorizontalBox::Slot()
+			.AutoWidth()
+			.VAlign(VAlign_Center)
+			.Padding(2.f, 2.f)
 			[
 				SNew(STextBlock)
-					.Text(TitleText)
-					.TextStyle(&TitleStyle)
+				.Text(TitleText)
+				.TextStyle(&TitleStyle)
 			]
 		]
-		+SVerticalBox::Slot()
-			.AutoHeight()
-			.Padding(0.f, 2.f, 0.f, 5.f)
+
+		+ SVerticalBox::Slot()
+		.AutoHeight()
+		.Padding(0.f, 2.f, 0.f, 5.f)
 		[
 			SNew(SBorder)
-				// use the border's padding to actually create the horizontal line
-				.Padding(1.f)
-				.BorderImage(FEditorStyle::GetBrush(TEXT("Menu.Separator")))
+			// use the border's padding to actually create the horizontal line
+			.Padding(1.f)
+			.BorderImage(FEditorStyle::GetBrush(TEXT("Menu.Separator")))
 		];	
 }
 
diff --git a/Engine/Source/Editor/LandscapeEditor/Private/LandscapeEditorDetailWidgets.cpp b/Engine/Source/Editor/LandscapeEditor/Private/LandscapeEditorDetailWidgets.cpp
--- a/Engine/Source/Editor/LandscapeEditor/Private/LandscapeEditorDetailWidgets.cpp
+++ b/Engine/Source/Editor/LandscapeEditor/Private/LandscapeEditorDetailWidgets.cpp
@@ -146,8 +146,9 @@ void SToolSelector::BuildMultiBlockWidget(const ISlateStyle* StyleSet, const FNa
 
 	Label = ToolSelectorButtonBlock->Label;
 	SmallText = ToolSelectorButtonBlock->SmallText;
-	FTextBlockStyle LabelStyle = FTextBlockStyle(FEditorStyle::GetWidgetStyle< FTextBlockStyle >(FEditorStyle::Join(StyleName, ".Label"))).SetShadowOffset(FVector2D::UnitVector);
-	FTextBlockStyle SmallTextStyle = FTextBlockStyle(LabelStyle).SetFont(LabelStyle.Font.FontName, LabelStyle.Font.Size - 1).SetColorAndOpacity(FSlateColor::UseSubduedForeground());
+
+	static FTextBlockStyle LabelStyle = FTextBlockStyle(FEditorStyle::GetWidgetStyle< FTextBlockStyle >(FEditorStyle::Join(StyleName, ".Label"))).SetShadowOffset(FVector2D::UnitVector);
+	static FTextBlockStyle SmallTextStyle = FTextBlockStyle(LabelStyle).SetFont(LabelStyle.Font.FontName, LabelStyle.Font.Size - 1).SetColorAndOpacity(FSlateColor::UseSubduedForeground());
 
 	// Create the content for our button
 	TSharedRef< SWidget > ButtonContent =
diff --git a/Engine/Source/Editor/UMGEditor/Private/SUMGDesigner.cpp b/Engine/Source/Editor/UMGEditor/Private/SUMGDesigner.cpp
--- a/Engine/Source/Editor/UMGEditor/Private/SUMGDesigner.cpp
+++ b/Engine/Source/Editor/UMGEditor/Private/SUMGDesigner.cpp
@@ -212,7 +212,7 @@ FReply SUMGDesigner::OnMouseButtonDown(const FGeometry& MyGeometry, const FPoint
 			//@TODO UMG primary FBlueprintEditor needs to be inherited and selection control needs to be centralized.
 			// Set the template as selected in the details panel
 			TArray<UWidget*> SelectedTemplates;
-			//SelectedTemplates.Add(CurrentSelection.Template);
+			//SelectedTemplates.Add(CurrentSelection.GetTemplate());
 			SelectedTemplates.Add(CurrentSelection.GetPreview());
 			ShowDetailsForObjects(SelectedTemplates);
 
diff --git a/Engine/Source/Editor/UMGEditor/Private/SUMGEditorTree.cpp b/Engine/Source/Editor/UMGEditor/Private/SUMGEditorTree.cpp
--- a/Engine/Source/Editor/UMGEditor/Private/SUMGEditorTree.cpp
+++ b/Engine/Source/Editor/UMGEditor/Private/SUMGEditorTree.cpp
@@ -100,7 +100,9 @@ void SUMGEditorTree::ShowDetailsForObjects(TArray<UWidget*> TemplateWidgets)
 	TArray<UWidget*> PreviewWidgets;
 	for ( UWidget* TemplateWidget : TemplateWidgets )
 	{
-		PreviewWidgets.Add(FSelectedWidget::FromTemplate(BlueprintEditor.Pin(), TemplateWidget).GetPreview());
+		FSelectedWidget Selection = FSelectedWidget::FromTemplate(BlueprintEditor.Pin(), TemplateWidget);
+		PreviewWidgets.Add(Selection.GetPreview());
+		//PreviewWidgets.Add(Selection.GetTemplate());
 	}
 
 	BlueprintEditor.Pin()->SelectWidgets(PreviewWidgets);
diff --git a/Engine/Source/Editor/UMGEditor/Private/WidgetBlueprintEditor.cpp b/Engine/Source/Editor/UMGEditor/Private/WidgetBlueprintEditor.cpp
--- a/Engine/Source/Editor/UMGEditor/Private/WidgetBlueprintEditor.cpp
+++ b/Engine/Source/Editor/UMGEditor/Private/WidgetBlueprintEditor.cpp
@@ -121,13 +121,21 @@ static bool MigratePropertyValue(UObject* SourceObject, UObject* DestinationObje
 	return true;
 }
 
-static bool MigratePropertyValue(UObject* SourceObject, UObject* DestinationObject, FEditPropertyChain::TDoubleLinkedListNode* PropertyChainNode, UProperty* MemberProperty)
+static bool MigratePropertyValue(UObject* SourceObject, UObject* DestinationObject, FEditPropertyChain::TDoubleLinkedListNode* PropertyChainNode, UProperty* MemberProperty, bool bIsModify)
 {
 	UProperty* CurrentProperty = PropertyChainNode->GetValue();
 
 	if ( PropertyChainNode->GetNextNode() == NULL )
 	{
-		return MigratePropertyValue(SourceObject, DestinationObject, MemberProperty);
+		if ( bIsModify )
+		{
+			DestinationObject->Modify();
+			return true;
+		}
+		else
+		{
+			return MigratePropertyValue(SourceObject, DestinationObject, MemberProperty);
+		}
 	}
 	
 	if ( UObjectProperty* CurrentObjectProperty = Cast<UObjectProperty>(CurrentProperty) )
@@ -136,14 +144,14 @@ static bool MigratePropertyValue(UObject* SourceObject, UObject* DestinationObje
 		UObject* SourceObjectProperty = CurrentObjectProperty->GetObjectPropertyValue(CurrentObjectProperty->ContainerPtrToValuePtr<void>(SourceObject));
 		UObject* DestionationObjectProperty = CurrentObjectProperty->GetObjectPropertyValue(CurrentObjectProperty->ContainerPtrToValuePtr<void>(DestinationObject));
 
-		return MigratePropertyValue(SourceObjectProperty, DestionationObjectProperty, PropertyChainNode->GetNextNode(), PropertyChainNode->GetNextNode()->GetValue());
+		return MigratePropertyValue(SourceObjectProperty, DestionationObjectProperty, PropertyChainNode->GetNextNode(), PropertyChainNode->GetNextNode()->GetValue(), bIsModify);
 	}
 	else if ( UArrayProperty* CurrentArrayProperty = Cast<UArrayProperty>(CurrentProperty) )
 	{
 		// Arrays!
 	}
 
-	return MigratePropertyValue(SourceObject, DestinationObject, PropertyChainNode->GetNextNode(), MemberProperty);
+	return MigratePropertyValue(SourceObject, DestinationObject, PropertyChainNode->GetNextNode(), MemberProperty, bIsModify);
 }
 
 void FWidgetBlueprintEditor::AddReferencedObjects( FReferenceCollector& Collector )
@@ -155,32 +163,42 @@ void FWidgetBlueprintEditor::AddReferencedObjects( FReferenceCollector& Collecto
 	}
 }
 
-void FWidgetBlueprintEditor::NotifyPostChange(const FPropertyChangedEvent& PropertyChangedEvent, class FEditPropertyChain* PropertyThatChanged)
+void FWidgetBlueprintEditor::NotifyPreChange(FEditPropertyChain* PropertyAboutToChange)
 {
-	//Super::NotifyPostChange(PropertyChangedEvent, PropertyThatChanged);
+	MigrateFromChain(PropertyAboutToChange, true);
+}
 
-	if ( PropertyChangedEvent.ChangeType == EPropertyChangeType::ValueSet )
-	{
-		UWidgetBlueprint* Blueprint = GetWidgetBlueprintObj();
+void FWidgetBlueprintEditor::MigrateFromChain(FEditPropertyChain* PropertyThatChanged, bool bIsModify)
+{
+	UWidgetBlueprint* Blueprint = GetWidgetBlueprintObj();
 
-		UUserWidget* PreviewActor = GetPreview();
-		if ( PreviewActor != NULL )
+	UUserWidget* PreviewActor = GetPreview();
+	if ( PreviewActor != NULL )
+	{
+		for ( UWidget* PreviewWidget : SelectedPreviewWidgets )
 		{
-			for ( UWidget* PreviewWidget : SelectedPreviewWidgets )
+			FString PreviewWidgetName = PreviewWidget->GetName();
+			UWidget* TemplateWidget = Blueprint->WidgetTree->FindWidget(PreviewWidgetName);
+
+			if ( TemplateWidget )
 			{
-				FString PreviewWidgetName = PreviewWidget->GetName();
-				UWidget* TemplateWidget = Blueprint->WidgetTree->FindWidget(PreviewWidgetName);
-
-				if ( TemplateWidget )
-				{
-					FEditPropertyChain::TDoubleLinkedListNode* PropertyChainNode = PropertyThatChanged->GetHead();
-					MigratePropertyValue(PreviewWidget, TemplateWidget, PropertyChainNode, PropertyChainNode->GetValue());
-				}
+				FEditPropertyChain::TDoubleLinkedListNode* PropertyChainNode = PropertyThatChanged->GetHead();
+				MigratePropertyValue(PreviewWidget, TemplateWidget, PropertyChainNode, PropertyChainNode->GetValue(), bIsModify);
 			}
 		}
 	}
 }
 
+void FWidgetBlueprintEditor::NotifyPostChange(const FPropertyChangedEvent& PropertyChangedEvent, FEditPropertyChain* PropertyThatChanged)
+{
+	//Super::NotifyPostChange(PropertyChangedEvent, PropertyThatChanged);
+	
+	if ( PropertyChangedEvent.ChangeType == EPropertyChangeType::ValueSet )
+	{
+		MigrateFromChain(PropertyThatChanged, false);
+	}
+}
+
 UWidgetBlueprint* FWidgetBlueprintEditor::GetWidgetBlueprintObj() const
 {
 	return Cast<UWidgetBlueprint>(GetBlueprintObj());
@@ -267,7 +285,7 @@ void FWidgetBlueprintEditor::UpdatePreview(UBlueprint* InBlueprint, bool bInForc
 		PreviewBlueprint = InBlueprint;
 
 		PreviewActor = ConstructObject<UUserWidget>(PreviewBlueprint->GeneratedClass, PreviewScene.GetWorld()->GetCurrentLevel());
-		PreviewActor->SetFlags(RF_Standalone);
+		PreviewActor->SetFlags(RF_Standalone | RF_Transactional);
 
 		PreviewWidgetActorPtr = PreviewActor;
 	}
diff --git a/Engine/Source/Editor/UMGEditor/Private/WidgetBlueprintEditor.h b/Engine/Source/Editor/UMGEditor/Private/WidgetBlueprintEditor.h
--- a/Engine/Source/Editor/UMGEditor/Private/WidgetBlueprintEditor.h
+++ b/Engine/Source/Editor/UMGEditor/Private/WidgetBlueprintEditor.h
@@ -21,6 +21,7 @@ public:
 	void InitWidgetBlueprintEditor(const EToolkitMode::Type Mode, const TSharedPtr< IToolkitHost >& InitToolkitHost, const TArray<UBlueprint*>& InBlueprints, bool bShouldOpenInDefaultsMode);
 
 	virtual void Tick(float DeltaTime) OVERRIDE;
+	virtual void NotifyPreChange(class FEditPropertyChain* PropertyAboutToChange) OVERRIDE;
 	virtual void NotifyPostChange(const FPropertyChangedEvent& PropertyChangedEvent, class FEditPropertyChain* PropertyThatChanged) OVERRIDE;
 
 	/** FGCObjectInterface */
@@ -41,6 +42,8 @@ private:
 	void DestroyPreview();
 	void UpdatePreview(UBlueprint* InBlueprint, bool bInForceFullUpdate);
 
+	void MigrateFromChain(FEditPropertyChain* PropertyThatChanged, bool bIsModify);
+
 	/**
 	 * Gets the default movie scene which is used when there is no 
 	 * animation data present on the widget blueprint
diff --git a/Engine/Source/Runtime/Slate/Private/Widgets/Notifications/SProgressBar.cpp b/Engine/Source/Runtime/Slate/Private/Widgets/Notifications/SProgressBar.cpp
--- a/Engine/Source/Runtime/Slate/Private/Widgets/Notifications/SProgressBar.cpp
+++ b/Engine/Source/Runtime/Slate/Private/Widgets/Notifications/SProgressBar.cpp
@@ -21,11 +21,6 @@ void SProgressBar::Construct( const FArguments& InArgs )
 	CurveSequence.Play();
 }
 
-TOptional<float> SProgressBar::GetPercent() const
-{
-	return Percent.Get();
-}
-
 void SProgressBar::SetPercent(TAttribute< TOptional<float> > InPercent)
 {
 	Percent = InPercent;
diff --git a/Engine/Source/Runtime/Slate/Private/Widgets/Text/STextBlock.cpp b/Engine/Source/Runtime/Slate/Private/Widgets/Text/STextBlock.cpp
--- a/Engine/Source/Runtime/Slate/Private/Widgets/Text/STextBlock.cpp
+++ b/Engine/Source/Runtime/Slate/Private/Widgets/Text/STextBlock.cpp
@@ -13,7 +13,7 @@ void STextBlock::Construct( const FArguments& InArgs )
 {
 	Text = InArgs._Text;
 
-	const FTextBlockStyle* MyStyle = InArgs._TextStyle;
+	TextStyle = InArgs._TextStyle;
 
 	HighlightText = InArgs._HighlightText;
 	WrapTextAt = InArgs._WrapTextAt;
@@ -22,34 +22,47 @@ void STextBlock::Construct( const FArguments& InArgs )
 	CachedWrapTextWidth = 0.0f;
 	CachedAutoWrapTextWidth = 0.0f;
 
-	Font = (InArgs._Font.IsSet())
-		? InArgs._Font
-		: MyStyle->Font;
+	Font = InArgs._Font;
+	ColorAndOpacity = InArgs._ColorAndOpacity;
+	ShadowOffset = InArgs._ShadowOffset;
+	ShadowColorAndOpacity = InArgs._ShadowColorAndOpacity;
+	HighlightColor = InArgs._HighlightColor;
+	HighlightShape = InArgs._HighlightShape;
 
-	ForegroundColor = InArgs._ColorAndOpacity.IsSet()
-		? InArgs._ColorAndOpacity
-		: MyStyle->ColorAndOpacity;
+	OnDoubleClicked = InArgs._OnDoubleClicked;
 
-	ShadowOffset = InArgs._ShadowOffset.IsSet()
-		? InArgs._ShadowOffset
-		: MyStyle->ShadowOffset;
+	// Request text size be cached
+	bRequestCache = true;
+}
 
-	ShadowColorAndOpacity = InArgs._ShadowColorAndOpacity.IsSet()
-		? InArgs._ShadowColorAndOpacity
-		: MyStyle->ShadowColorAndOpacity;
+FSlateFontInfo STextBlock::GetFont() const
+{
+	return Font.IsSet() ? Font.Get() : TextStyle->Font;
+}
 
-	HighlightColor = InArgs._HighlightColor.IsSet()
-		? InArgs._HighlightColor
-		: MyStyle->HighlightColor;
+FSlateColor STextBlock::GetColorAndOpacity() const
+{
+	return ColorAndOpacity.IsSet() ? ColorAndOpacity.Get() : TextStyle->ColorAndOpacity;
+}
 
-	HighlightShape = InArgs._HighlightShape.IsSet()
-		? InArgs._HighlightShape
-		: &MyStyle->HighlightShape;
+FVector2D STextBlock::GetShadowOffset() const
+{
+	return ShadowOffset.IsSet() ? ShadowOffset.Get() : TextStyle->ShadowOffset;
+}
 
-	OnDoubleClicked = InArgs._OnDoubleClicked;
+FLinearColor STextBlock::GetShadowColorAndOpacity() const
+{
+	return ShadowColorAndOpacity.IsSet() ? ShadowColorAndOpacity.Get() : TextStyle->ShadowColorAndOpacity;
+}
 
-	// Request text size be cached
-	bRequestCache = true;
+FLinearColor STextBlock::GetHighlightColor() const
+{
+	return HighlightColor.IsSet() ? HighlightColor.Get() : TextStyle->HighlightColor;
+}
+
+const FSlateBrush* STextBlock::GetHighlightShape() const
+{
+	return HighlightShape.IsSet() ? HighlightShape.Get() : &TextStyle->HighlightShape;
 }
 
 void STextBlock::SetText( const TAttribute< FText >& InText )
@@ -66,11 +79,6 @@ void STextBlock::SetText( const TAttribute< FText >& InText )
 	bRequestCache = true;
 }
 
-void STextBlock::SetForegroundColor( const TAttribute<FSlateColor>& InSlateColor )
-{
-	ForegroundColor = InSlateColor;
-}
-
 int32 STextBlock::OnPaint( const FGeometry& AllottedGeometry, const FSlateRect& MyClippingRect, FSlateWindowElementList& OutDrawElements, int32 LayerId, const FWidgetStyle& InWidgetStyle, bool bParentEnabled ) const
 {
 #if SLATE_HD_STATS
@@ -84,10 +92,10 @@ int32 STextBlock::OnPaint( const FGeometry& AllottedGeometry, const FSlateRect&
 	bool bEnabled = ShouldBeEnabled( bParentEnabled );
 	ESlateDrawEffect::Type DrawEffects = bEnabled ? ESlateDrawEffect::None : ESlateDrawEffect::DisabledEffect;
 
-	const FVector2D CurShadowOffset = ShadowOffset.Get();
+	const FVector2D CurShadowOffset = GetShadowOffset();
 	const bool ShouldDropShadow = CurShadowOffset.Size() > 0;
 
-	const FSlateFontInfo& FontInfo = Font.Get();
+	FSlateFontInfo FontInfo = GetFont();
 
 	// Perform text auto-wrapping if that was enabled.
 	if( AutoWrapText.Get() )
@@ -147,10 +155,10 @@ int32 STextBlock::OnPaint( const FGeometry& AllottedGeometry, const FSlateRect&
 							OutDrawElements,
 							++LayerId,
 							AllottedGeometry.ToPaintGeometry(HighlightStartOffset, HighlightSize),
-							HighlightShape.Get(),
+							GetHighlightShape(),
 							ClippingRect,
 							DrawEffects,
-							InWidgetStyle.GetColorAndOpacityTint() * HighlightColor.Get()
+							InWidgetStyle.GetColorAndOpacityTint() * GetHighlightColor()
 							);
 					}
 				}
@@ -169,7 +177,7 @@ int32 STextBlock::OnPaint( const FGeometry& AllottedGeometry, const FSlateRect&
 			FontInfo,
 			ClippingRect,
 			DrawEffects,
-			ShadowColorAndOpacity.Get()*InWidgetStyle.GetColorAndOpacityTint()
+			GetShadowColorAndOpacity() * InWidgetStyle.GetColorAndOpacityTint()
 		);
 	}
 
@@ -182,7 +190,7 @@ int32 STextBlock::OnPaint( const FGeometry& AllottedGeometry, const FSlateRect&
 		FontInfo,
 		ClippingRect,
 		DrawEffects,
-		InWidgetStyle.GetColorAndOpacityTint() * ForegroundColor.Get().GetColor(InWidgetStyle)
+		InWidgetStyle.GetColorAndOpacityTint() * GetColorAndOpacity().GetColor(InWidgetStyle)
 	);
 
 	return LayerId;
@@ -219,7 +227,7 @@ void STextBlock::CacheDesiredSize()
 
 	// Get the wrapping width and font to see if they have changed
 	float WrappingWidth = WrapTextAt.Get();
-	const FSlateFontInfo& FontInfo = Font.Get();
+	const FSlateFontInfo& FontInfo = GetFont();
 
 	// Text wrapping can either be used defined (WrapTextAt), automatic (AutoWrapText and CachedAutoWrapTextWidth), or a mixture of both
 	// Take whichever has the smallest value (>1)
@@ -261,7 +269,7 @@ void STextBlock::CacheDesiredSize()
 			TextMeasurement.X = FMath::Min(TextMeasurement.X, WrappingWidth);
 		}
 
-		this->Advanced_SetDesiredSize(TextMeasurement + ShadowOffset.Get());
+		this->Advanced_SetDesiredSize(TextMeasurement + GetShadowOffset());
 
 		// Update cached values
 		CachedWrapTextWidth = WrappingWidth;
@@ -270,7 +278,6 @@ void STextBlock::CacheDesiredSize()
 	}
 }
 
-
 FVector2D STextBlock::ComputeDesiredSize() const
 {
 	// Usually widgets just override ComputeDesiredSize(), but STextBlock
@@ -279,3 +286,40 @@ FVector2D STextBlock::ComputeDesiredSize() const
 	// Dummy implementation.
 	return FVector2D::ZeroVector;
 }
+
+void STextBlock::SetFont(const TAttribute< FSlateFontInfo >& InFont)
+{
+	Font = InFont;
+	bRequestCache = true;
+}
+
+void STextBlock::SetColorAndOpacity(const TAttribute<FSlateColor>& InColorAndOpacity)
+{
+	ColorAndOpacity = InColorAndOpacity;
+}
+
+void STextBlock::SetTextStyle(const FTextBlockStyle* InTextStyle)
+{
+	TextStyle = InTextStyle;
+	bRequestCache = true;
+}
+
+void STextBlock::SetWrapTextAt(const TAttribute<float>& InWrapTextAt)
+{
+	WrapTextAt = InWrapTextAt;
+}
+
+void STextBlock::SetAutoWrapText(const TAttribute<bool>& InAutoWrapText)
+{
+	AutoWrapText = InAutoWrapText;
+}
+
+void STextBlock::SetShadowOffset(const TAttribute<FVector2D>& InShadowOffset)
+{
+	ShadowOffset = InShadowOffset;
+}
+
+void STextBlock::SetShadowColorAndOpacity(const TAttribute<FLinearColor>& InShadowColorAndOpacity)
+{
+	ShadowColorAndOpacity = InShadowColorAndOpacity;
+}
diff --git a/Engine/Source/Runtime/Slate/Public/Widgets/Notifications/SProgressBar.h b/Engine/Source/Runtime/Slate/Public/Widgets/Notifications/SProgressBar.h
--- a/Engine/Source/Runtime/Slate/Public/Widgets/Notifications/SProgressBar.h
+++ b/Engine/Source/Runtime/Slate/Public/Widgets/Notifications/SProgressBar.h
@@ -62,9 +62,6 @@ public:
 
 	FVector2D ComputeDesiredSize() const;
 
-	/** See Slate attribute Percent */
-	TOptional<float> GetPercent() const;
-
 	/** See Slate attribute Percent */
 	void SetPercent(TAttribute< TOptional<float> > InPercent);
 
diff --git a/Engine/Source/Runtime/Slate/Public/Widgets/Text/STextBlock.h b/Engine/Source/Runtime/Slate/Public/Widgets/Text/STextBlock.h
--- a/Engine/Source/Runtime/Slate/Public/Widgets/Text/STextBlock.h
+++ b/Engine/Source/Runtime/Slate/Public/Widgets/Text/STextBlock.h
@@ -93,8 +93,7 @@ public:
 	{
 		return Text.Get();
 	}
-
-
+	
 	/**
 	 * Sets the text for this text block
 	 *
@@ -135,14 +134,25 @@ public:
 	 *
 	 * @param	InFont	The new font to use
 	 */
-	void SetFont( const TAttribute< FSlateFontInfo >& InFont )
-	{
-		Font = InFont;
-		bRequestCache = true;
-	}
+	void SetFont(const TAttribute< FSlateFontInfo >& InFont);
+
+	/** See ColorAndOpacity attribute */
+	void SetColorAndOpacity(const TAttribute<FSlateColor>& InColorAndOpacity);
+
+	/** See TextStyle argument */
+	void SetTextStyle(const FTextBlockStyle* InTextStyle);
+
+	/** See WrapTextAt attribute */
+	void SetWrapTextAt(const TAttribute<float>& InWrapTextAt);
+
+	/** See AutoWrapText attribute */
+	void SetAutoWrapText(const TAttribute<bool>& InAutoWrapText);
 
-	/** Set the color and opacity of the text */
-	void SetForegroundColor( const TAttribute<FSlateColor>& InSlateColor );
+	/** See ShadowOffset attribute */
+	void SetShadowOffset(const TAttribute<FVector2D>& InShadowOffset);
+
+	/** See ShadowColorAndOpacity attribute */
+	void SetShadowColorAndOpacity(const TAttribute<FLinearColor>& InShadowColorAndOpacity);
 
 	// SWidget interface
 	virtual int32 OnPaint( const FGeometry& AllottedGeometry, const FSlateRect& MyClippingRect, FSlateWindowElementList& OutDrawElements, int32 LayerId, const FWidgetStyle& InWidgetStyle, bool bParentEnabled ) const OVERRIDE;
@@ -155,15 +165,36 @@ private:
 	/** Caches string size to avoid remeasuring if possible */
 	void CacheStringSizeIfNeeded();
 
+	/** Gets the current foreground color */
+	FSlateColor GetColorAndOpacity() const;
+
+	/** Gets the current font */
+	FSlateFontInfo GetFont() const;
+
+	/** Gets the current shadow offset */
+	FVector2D GetShadowOffset() const;
+
+	/** Gets the current shadow color and opacity */
+	FLinearColor GetShadowColorAndOpacity() const;
+
+	/** Gets the current highlight color */
+	FLinearColor GetHighlightColor() const;
+
+	/** Gets the current highlight shape */
+	const FSlateBrush* GetHighlightShape() const;
+
 private:
 	/** The text displayed in this text block */
 	TAttribute< FString > Text;
 
+	/** The style for the text block */
+	const FTextBlockStyle* TextStyle;
+
 	/** Sets the font used to draw the text */
 	TAttribute< FSlateFontInfo > Font;
 
 	/** Text color and opacity */
-	TAttribute<FSlateColor> ForegroundColor;
+	TAttribute<FSlateColor> ColorAndOpacity;
 
 	/** Drop shadow offset in pixels */
 	TAttribute< FVector2D > ShadowOffset;
diff --git a/Engine/Source/Runtime/UMG/Private/Components/CanvasPanel.cpp b/Engine/Source/Runtime/UMG/Private/Components/CanvasPanel.cpp
--- a/Engine/Source/Runtime/UMG/Private/Components/CanvasPanel.cpp
+++ b/Engine/Source/Runtime/UMG/Private/Components/CanvasPanel.cpp
@@ -51,39 +51,87 @@ UWidget* UCanvasPanel::GetChildAt(int32 Index) const
 	return Slots[Index]->Content;
 }
 
-TSharedPtr<SFixedSizeCanvas> UCanvasPanel::GetCanvasWidget() const
+TSharedPtr<SConstraintCanvas> UCanvasPanel::GetCanvasWidget() const
 {
-	return MyCanvas.Pin();
+	return MyCanvas;
 }
 
-TSharedRef<SWidget> UCanvasPanel::RebuildWidget()
+int32 UCanvasPanel::GetChildIndex(UWidget* Content) const
 {
-	TSharedRef<SFixedSizeCanvas> NewCanvas = SNew(SFixedSizeCanvas, DesiredCanvasSize);
-	MyCanvas = NewCanvas;
-
-	TSharedPtr<SFixedSizeCanvas> Canvas = MyCanvas.Pin();
-	if ( Canvas.IsValid() )
+	for ( int32 SlotIndex = 0; SlotIndex < Slots.Num(); ++SlotIndex )
 	{
-		Canvas->ClearChildren();
+		UCanvasPanelSlot* Slot = Slots[SlotIndex];
 
-		for ( int32 SlotIndex = 0; SlotIndex < Slots.Num(); ++SlotIndex )
+		if ( Slot->Content == Content )
 		{
-			UCanvasPanelSlot* Slot = Slots[SlotIndex];
-			if ( Slot == NULL )
-			{
-				Slots[SlotIndex] = Slot = ConstructObject<UCanvasPanelSlot>(UCanvasPanelSlot::StaticClass(), this);
-			}
-
-			Slot->BuildSlot(Canvas.ToSharedRef());
+			return SlotIndex;
 		}
 	}
 
-	return NewCanvas;
+	return -1;
+}
+
+bool UCanvasPanel::AddChild(UWidget* Child, FVector2D Position)
+{
+	UCanvasPanelSlot* Slot = AddSlot(Child);
+	Slot->LayoutData.Offsets = FMargin(Position.X, Position.Y, 100, 25);
+
+	return true;
+}
+
+bool UCanvasPanel::RemoveChild(UWidget* Child)
+{
+	int32 SlotIndex = GetChildIndex(Child);
+	if ( SlotIndex != -1 )
+	{
+		Slots.RemoveAt(SlotIndex);
+		return true;
+	}
+
+	return false;
+}
+
+void UCanvasPanel::ReplaceChildAt(int32 Index, UWidget* Content)
+{
+	UCanvasPanelSlot* Slot = Slots[Index];
+	Slot->Content = Content;
+
+#if WITH_EDITOR
+	Content->Slot = Slot;
+#endif
+}
+
+void UCanvasPanel::InsertChildAt(int32 Index, UWidget* Content)
+{
+	UCanvasPanelSlot* Slot = ConstructObject<UCanvasPanelSlot>(UCanvasPanelSlot::StaticClass(), this);
+	Slot->SetFlags(RF_Transactional);
+	Slot->Content = Content;
+	Slot->Parent = this;
+
+#if WITH_EDITOR
+	Content->Slot = Slot;
+#endif
+
+	Slots.Insert(Slot, Index);
+}
+
+TSharedRef<SWidget> UCanvasPanel::RebuildWidget()
+{
+	MyCanvas = SNew(SFixedSizeCanvas, DesiredCanvasSize);
+
+	for ( auto Slot : Slots )
+	{
+		Slot->Parent = this;
+		Slot->BuildSlot(MyCanvas.ToSharedRef());
+	}
+
+	return MyCanvas.ToSharedRef();
 }
 
 UCanvasPanelSlot* UCanvasPanel::AddSlot(UWidget* Content)
 {
 	UCanvasPanelSlot* Slot = ConstructObject<UCanvasPanelSlot>(UCanvasPanelSlot::StaticClass(), this);
+	Slot->SetFlags(RF_Transactional);
 	Slot->Content = Content;
 	Slot->Parent = this;
 
@@ -96,13 +144,7 @@ UCanvasPanelSlot* UCanvasPanel::AddSlot(UWidget* Content)
 	return Slot;
 }
 
-bool UCanvasPanel::AddChild(UWidget* Child, FVector2D Position)
-{
-	UCanvasPanelSlot* Slot = AddSlot(Child);
-	Slot->LayoutData.Offsets = FMargin(Position.X, Position.Y, 100, 25);
 
-	return true;
-}
 
 bool UCanvasPanel::GetGeometryForSlot(UCanvasPanelSlot* Slot, FGeometry& ArrangedGeometry) const
 {
@@ -111,7 +153,7 @@ bool UCanvasPanel::GetGeometryForSlot(UCanvasPanelSlot* Slot, FGeometry& Arrange
 		return false;
 	}
 
-	TSharedPtr<SFixedSizeCanvas> Canvas = GetCanvasWidget();
+	TSharedPtr<SConstraintCanvas> Canvas = GetCanvasWidget();
 	if ( Canvas.IsValid() )
 	{
 		FArrangedChildren ArrangedChildren(EVisibility::All);
@@ -143,16 +185,5 @@ void UCanvasPanel::ConnectEditorData()
 
 void UCanvasPanel::PostEditChangeProperty(struct FPropertyChangedEvent& PropertyChangedEvent)
 {
-	// Ensure the slots have unique names
-	int32 SlotNumbering = 1;
-
-	TSet<FName> UniqueSlotNames;
-	for (int32 SlotIndex = 0; SlotIndex < Slots.Num(); ++SlotIndex)
-	{
-		if ( Slots[SlotIndex] == NULL )
-		{
-			Slots[SlotIndex] = ConstructObject<UCanvasPanelSlot>(UCanvasPanelSlot::StaticClass(), this);
-		}
-	}
 }
 #endif
\ No newline at end of file
diff --git a/Engine/Source/Runtime/UMG/Private/Components/HorizontalBox.cpp b/Engine/Source/Runtime/UMG/Private/Components/HorizontalBox.cpp
--- a/Engine/Source/Runtime/UMG/Private/Components/HorizontalBox.cpp
+++ b/Engine/Source/Runtime/UMG/Private/Components/HorizontalBox.cpp
@@ -67,6 +67,7 @@ void UHorizontalBox::ReplaceChildAt(int32 Index, UWidget* Content)
 void UHorizontalBox::InsertChildAt(int32 Index, UWidget* Content)
 {
 	UHorizontalBoxSlot* Slot = ConstructObject<UHorizontalBoxSlot>(UHorizontalBoxSlot::StaticClass(), this);
+	Slot->SetFlags(RF_Transactional);
 	Slot->Content = Content;
 	Slot->Parent = this;
 
@@ -82,14 +83,8 @@ TSharedRef<SWidget> UHorizontalBox::RebuildWidget()
 	TSharedRef<SHorizontalBox> NewWidget = SNew(SHorizontalBox);
 	MyHorizontalBox = NewWidget;
 
-	for ( int32 SlotIndex = 0; SlotIndex < Slots.Num(); ++SlotIndex )
+	for ( auto Slot : Slots )
 	{
-		UHorizontalBoxSlot* Slot = Slots[SlotIndex];
-		if ( Slot == NULL )
-		{
-			Slots[SlotIndex] = Slot = ConstructObject<UHorizontalBoxSlot>(UHorizontalBoxSlot::StaticClass(), this);
-		}
-
 		Slot->Parent = this;
 		Slot->BuildSlot(NewWidget);
 	}
@@ -100,6 +95,7 @@ TSharedRef<SWidget> UHorizontalBox::RebuildWidget()
 UHorizontalBoxSlot* UHorizontalBox::AddSlot(UWidget* Content)
 {
 	UHorizontalBoxSlot* Slot = ConstructObject<UHorizontalBoxSlot>(UHorizontalBoxSlot::StaticClass(), this);
+	Slot->SetFlags(RF_Transactional);
 	Slot->Content = Content;
 	Slot->Parent = this;
 
@@ -124,16 +120,5 @@ void UHorizontalBox::ConnectEditorData()
 
 void UHorizontalBox::PostEditChangeProperty(struct FPropertyChangedEvent& PropertyChangedEvent)
 {
-	// Ensure the slots have unique names
-	int32 SlotNumbering = 1;
-
-	TSet<FName> UniqueSlotNames;
-	for (int32 SlotIndex = 0; SlotIndex < Slots.Num(); ++SlotIndex)
-	{
-		if ( Slots[SlotIndex] == NULL )
-		{
-			Slots[SlotIndex] = ConstructObject<UHorizontalBoxSlot>(UHorizontalBoxSlot::StaticClass(), this);
-		}
-	}
 }
 #endif
diff --git a/Engine/Source/Runtime/UMG/Private/Components/Overlay.cpp b/Engine/Source/Runtime/UMG/Private/Components/Overlay.cpp
--- a/Engine/Source/Runtime/UMG/Private/Components/Overlay.cpp
+++ b/Engine/Source/Runtime/UMG/Private/Components/Overlay.cpp
@@ -58,14 +58,9 @@ TSharedRef<SWidget> UOverlay::RebuildWidget()
 	TSharedRef<SOverlay> NewWidget = SNew(SOverlay);
 	MyVerticalBox = NewWidget;
 
-	for ( int32 SlotIndex = 0; SlotIndex < Slots.Num(); ++SlotIndex )
+	for ( auto Slot : Slots )
 	{
-		UOverlaySlot* Slot = Slots[SlotIndex];
-		if ( Slot == NULL )
-		{
-			Slots[SlotIndex] = Slot = ConstructObject<UOverlaySlot>(UOverlaySlot::StaticClass(), this);
-		}
-
+		Slot->Parent = this;
 		Slot->BuildSlot(NewWidget);
 	}
 
@@ -75,7 +70,9 @@ TSharedRef<SWidget> UOverlay::RebuildWidget()
 UOverlaySlot* UOverlay::AddSlot(UWidget* Content)
 {
 	UOverlaySlot* Slot = ConstructObject<UOverlaySlot>(UOverlaySlot::StaticClass(), this);
+	Slot->SetFlags(RF_Transactional);
 	Slot->Content = Content;
+	Slot->Parent = this;
 
 #if WITH_EDITOR
 	Content->Slot = Slot;
@@ -101,16 +98,6 @@ void UOverlay::ConnectEditorData()
 
 void UOverlay::PostEditChangeProperty(struct FPropertyChangedEvent& PropertyChangedEvent)
 {
-	// Ensure the slots have unique names
-	int32 SlotNumbering = 1;
 
-	TSet<FName> UniqueSlotNames;
-	for (int32 SlotIndex = 0; SlotIndex < Slots.Num(); ++SlotIndex)
-	{
-		if ( Slots[SlotIndex] == NULL )
-		{
-			Slots[SlotIndex] = ConstructObject<UOverlaySlot>(UOverlaySlot::StaticClass(), this);
-		}
-	}
 }
 #endif
diff --git a/Engine/Source/Runtime/UMG/Private/Components/ProcessBar.cpp b/Engine/Source/Runtime/UMG/Private/Components/ProcessBar.cpp
--- a/Engine/Source/Runtime/UMG/Private/Components/ProcessBar.cpp
+++ b/Engine/Source/Runtime/UMG/Private/Components/ProcessBar.cpp
@@ -27,11 +27,6 @@ TSharedRef<SWidget> UProgressBar::RebuildWidget()
 	return MyProgressBar.ToSharedRef();
 }
 
-float UProgressBar::GetPercent()
-{
-	return MyProgressBar->GetPercent().Get(0.0f);
-}
-
 void UProgressBar::SetPercent(float InPercent)
 {
 	return MyProgressBar->SetPercent(InPercent);
diff --git a/Engine/Source/Runtime/UMG/Private/Components/ScrollBox.cpp b/Engine/Source/Runtime/UMG/Private/Components/ScrollBox.cpp
--- a/Engine/Source/Runtime/UMG/Private/Components/ScrollBox.cpp
+++ b/Engine/Source/Runtime/UMG/Private/Components/ScrollBox.cpp
@@ -69,6 +69,7 @@ void UScrollBox::ReplaceChildAt(int32 Index, UWidget* Content)
 void UScrollBox::InsertChildAt(int32 Index, UWidget* Content)
 {
 	UScrollBoxSlot* Slot = ConstructObject<UScrollBoxSlot>(UScrollBoxSlot::StaticClass(), this);
+	Slot->SetFlags(RF_Transactional);
 	Slot->Content = Content;
 	Slot->Parent = this;
 
@@ -89,14 +90,8 @@ TSharedRef<SWidget> UScrollBox::RebuildWidget()
 	{
 		Canvas->ClearChildren();
 
-		for ( int32 SlotIndex = 0; SlotIndex < Slots.Num(); ++SlotIndex )
+		for ( auto Slot : Slots )
 		{
-			UScrollBoxSlot* Slot = Slots[SlotIndex];
-			if ( Slot == NULL )
-			{
-				Slots[SlotIndex] = Slot = ConstructObject<UScrollBoxSlot>(UScrollBoxSlot::StaticClass(), this);
-			}
-
 			Slot->Parent = this;
 			Slot->BuildSlot(NewCanvas);
 		}
@@ -108,6 +103,7 @@ TSharedRef<SWidget> UScrollBox::RebuildWidget()
 UScrollBoxSlot* UScrollBox::AddSlot(UWidget* Content)
 {
 	UScrollBoxSlot* Slot = ConstructObject<UScrollBoxSlot>(UScrollBoxSlot::StaticClass(), this);
+	Slot->SetFlags(RF_Transactional);
 	Slot->Content = Content;
 	Slot->Parent = this;
 
@@ -138,17 +134,6 @@ void UScrollBox::ConnectEditorData()
 
 void UScrollBox::PostEditChangeProperty(struct FPropertyChangedEvent& PropertyChangedEvent)
 {
-	// Ensure the slots have unique names
-	int32 SlotNumbering = 1;
-
-	TSet<FName> UniqueSlotNames;
-	for (int32 SlotIndex = 0; SlotIndex < Slots.Num(); ++SlotIndex)
-	{
-		if ( Slots[SlotIndex] == NULL )
-		{
-			Slots[SlotIndex] = ConstructObject<UScrollBoxSlot>(UScrollBoxSlot::StaticClass(), this);
-		}
-	}
 }
 #endif
 
diff --git a/Engine/Source/Runtime/UMG/Private/Components/TextBlock.cpp b/Engine/Source/Runtime/UMG/Private/Components/TextBlock.cpp
--- a/Engine/Source/Runtime/UMG/Private/Components/TextBlock.cpp
+++ b/Engine/Source/Runtime/UMG/Private/Components/TextBlock.cpp
@@ -19,12 +19,20 @@ UTextBlock::UTextBlock(const FPostConstructInitializeProperties& PCIP)
 
 	Style = NULL;
 
-	// HACK Special font initialization hack since there are no font assets yet for slate.
+	// @TODO UMG HACK Special font initialization hack since there are no font assets yet for slate.
 	Font = FSlateFontInfo(TEXT("Slate/Fonts/Roboto-Bold.ttf"), 24);
 }
 
 TSharedRef<SWidget> UTextBlock::RebuildWidget()
 {
+	MyTextBlock = SNew(STextBlock);
+	return MyTextBlock.ToSharedRef();
+}
+
+void UTextBlock::SyncronizeProperties()
+{
+	Super::SyncronizeProperties();
+
 	FString FontPath = FPaths::EngineContentDir() / Font.FontName.ToString();
 
 	const FTextBlockStyle* StylePtr = ( Style != NULL ) ? Style->GetStyle<FTextBlockStyle>() : NULL;
@@ -37,14 +45,13 @@ TSharedRef<SWidget> UTextBlock::RebuildWidget()
 	TAttribute<FText> TextBinding = OPTIONAL_BINDING(FText, Text);
 	TAttribute<FSlateColor> ColorAndOpacityBinding = OPTIONAL_BINDING(FSlateColor, ColorAndOpacity);
 	TAttribute<FLinearColor> ShadowColorAndOpacityBinding = OPTIONAL_BINDING(FLinearColor, ShadowColorAndOpacity);
-	
-	return SNew(STextBlock)
-		.Text( TextBinding )
-		.TextStyle( StylePtr )
-		.Font( FSlateFontInfo(FontPath, Font.Size) )
-		.ColorAndOpacity( ColorAndOpacityBinding )
-		.ShadowOffset( ShadowOffset )
-		.ShadowColorAndOpacity( ShadowColorAndOpacityBinding );
+
+	MyTextBlock->SetText(TextBinding);
+	MyTextBlock->SetFont(FSlateFontInfo(FontPath, Font.Size));
+	MyTextBlock->SetColorAndOpacity(ColorAndOpacityBinding);
+	MyTextBlock->SetTextStyle(StylePtr);
+	MyTextBlock->SetShadowOffset(ShadowOffset);
+	MyTextBlock->SetShadowColorAndOpacity(ShadowColorAndOpacityBinding);
 }
 
 /////////////////////////////////////////////////////
diff --git a/Engine/Source/Runtime/UMG/Private/Components/UniformGridPanel.cpp b/Engine/Source/Runtime/UMG/Private/Components/UniformGridPanel.cpp
--- a/Engine/Source/Runtime/UMG/Private/Components/UniformGridPanel.cpp
+++ b/Engine/Source/Runtime/UMG/Private/Components/UniformGridPanel.cpp
@@ -63,14 +63,9 @@ TSharedRef<SWidget> UUniformGridPanel::RebuildWidget()
 	
 	MyUniformGridPanel = NewPanel;
 
-	for ( int32 SlotIndex = 0; SlotIndex < Slots.Num(); ++SlotIndex )
+	for ( auto Slot : Slots )
 	{
-		UUniformGridSlot* Slot = Slots[SlotIndex];
-		if ( Slot == NULL )
-		{
-			Slots[SlotIndex] = Slot = ConstructObject<UUniformGridSlot>(UUniformGridSlot::StaticClass(), this);
-		}
-
+		Slot->Parent = this;
 		Slot->BuildSlot(NewPanel);
 	}
 
@@ -80,6 +75,7 @@ TSharedRef<SWidget> UUniformGridPanel::RebuildWidget()
 UUniformGridSlot* UUniformGridPanel::AddSlot(UWidget* Content)
 {
 	UUniformGridSlot* Slot = ConstructObject<UUniformGridSlot>(UUniformGridSlot::StaticClass(), this);
+	Slot->SetFlags(RF_Transactional);
 	Slot->Content = Content;
 	Slot->Parent = this;
 
@@ -105,17 +101,6 @@ void UUniformGridPanel::ConnectEditorData()
 
 void UUniformGridPanel::PostEditChangeProperty(struct FPropertyChangedEvent& PropertyChangedEvent)
 {
-	// Ensure the slots have unique names
-	int32 SlotNumbering = 1;
-
-	TSet<FName> UniqueSlotNames;
-	for (int32 SlotIndex = 0; SlotIndex < Slots.Num(); ++SlotIndex)
-	{
-		if ( Slots[SlotIndex] == NULL )
-		{
-			Slots[SlotIndex] = ConstructObject<UUniformGridSlot>(UUniformGridSlot::StaticClass(), this);
-		}
-	}
 }
 
 #endif
diff --git a/Engine/Source/Runtime/UMG/Private/Components/VerticalBox.cpp b/Engine/Source/Runtime/UMG/Private/Components/VerticalBox.cpp
--- a/Engine/Source/Runtime/UMG/Private/Components/VerticalBox.cpp
+++ b/Engine/Source/Runtime/UMG/Private/Components/VerticalBox.cpp
@@ -67,6 +67,7 @@ void UVerticalBox::ReplaceChildAt(int32 Index, UWidget* Content)
 void UVerticalBox::InsertChildAt(int32 Index, UWidget* Content)
 {
 	UVerticalBoxSlot* Slot = ConstructObject<UVerticalBoxSlot>(UVerticalBoxSlot::StaticClass(), this);
+	Slot->SetFlags(RF_Transactional);
 	Slot->Content = Content;
 	Slot->Parent = this;
 
@@ -82,14 +83,8 @@ TSharedRef<SWidget> UVerticalBox::RebuildWidget()
 	TSharedRef<SVerticalBox> NewWidget = SNew(SVerticalBox);
 	MyVerticalBox = NewWidget;
 
-	for ( int32 SlotIndex = 0; SlotIndex < Slots.Num(); ++SlotIndex )
+	for ( auto Slot : Slots )
 	{
-		UVerticalBoxSlot* Slot = Slots[SlotIndex];
-		if ( Slot == NULL )
-		{
-			Slots[SlotIndex] = Slot = ConstructObject<UVerticalBoxSlot>(UVerticalBoxSlot::StaticClass(), this);
-		}
-
 		Slot->Parent = this;
 		Slot->BuildSlot(NewWidget);
 	}
@@ -100,6 +95,7 @@ TSharedRef<SWidget> UVerticalBox::RebuildWidget()
 UVerticalBoxSlot* UVerticalBox::AddSlot(UWidget* Content)
 {
 	UVerticalBoxSlot* Slot = ConstructObject<UVerticalBoxSlot>(UVerticalBoxSlot::StaticClass(), this);
+	Slot->SetFlags(RF_Transactional);
 	Slot->Content = Content;
 	Slot->Parent = this;
 
@@ -130,16 +126,5 @@ void UVerticalBox::ConnectEditorData()
 
 void UVerticalBox::PostEditChangeProperty(struct FPropertyChangedEvent& PropertyChangedEvent)
 {
-	// Ensure the slots have unique names
-	int32 SlotNumbering = 1;
-
-	TSet<FName> UniqueSlotNames;
-	for (int32 SlotIndex = 0; SlotIndex < Slots.Num(); ++SlotIndex)
-	{
-		if ( Slots[SlotIndex] == NULL )
-		{
-			Slots[SlotIndex] = ConstructObject<UVerticalBoxSlot>(UVerticalBoxSlot::StaticClass(), this);
-		}
-	}
 }
 #endif
diff --git a/Engine/Source/Runtime/UMG/Private/WidgetBlueprintLibrary.cpp b/Engine/Source/Runtime/UMG/Private/WidgetBlueprintLibrary.cpp
--- a/Engine/Source/Runtime/UMG/Private/WidgetBlueprintLibrary.cpp
+++ b/Engine/Source/Runtime/UMG/Private/WidgetBlueprintLibrary.cpp
@@ -26,6 +26,7 @@ UUserWidget* UWidgetBlueprintLibrary::Create(UObject* WorldContextObject, TSubcl
 
 	UWorld* World = GEngine->GetWorldFromContextObject(WorldContextObject);
 	UUserWidget* NewWidget = ConstructObject<UUserWidget>(WidgetType, World->GetCurrentLevel());
+	NewWidget->SetFlags(RF_Transactional);
 	//UUserWidget* NewWidget = World->SpawnActor<UUserWidget>(WidgetType);
 		
 	return NewWidget;
diff --git a/Engine/Source/Runtime/UMG/Public/Blueprint/WidgetTree.h b/Engine/Source/Runtime/UMG/Public/Blueprint/WidgetTree.h
--- a/Engine/Source/Runtime/UMG/Public/Blueprint/WidgetTree.h
+++ b/Engine/Source/Runtime/UMG/Public/Blueprint/WidgetTree.h
@@ -26,7 +26,8 @@ class UMG_API UWidgetTree : public UObject
 	template< class T >
 	T* ConstructWidget(TSubclassOf<class UWidget> WidgetType)
 	{
-		UWidget* Widget = (UWidget*)ConstructObject<class UWidget>(WidgetType, this);
+		UWidget* Widget = (UWidget*)ConstructObject<UWidget>(WidgetType, this);
+		Widget->SetFlags(RF_Transactional);
 		WidgetTemplates.Add(Widget);
 
 		return (T*)Widget;
diff --git a/Engine/Source/Runtime/UMG/Public/Components/CanvasPanel.h b/Engine/Source/Runtime/UMG/Public/Components/CanvasPanel.h
--- a/Engine/Source/Runtime/UMG/Public/Components/CanvasPanel.h
+++ b/Engine/Source/Runtime/UMG/Public/Components/CanvasPanel.h
@@ -18,15 +18,19 @@ class UMG_API UCanvasPanel : public UPanelWidget
 	FVector2D DesiredCanvasSize;
 
 	UCanvasPanelSlot* AddSlot(UWidget* Content);
-	virtual bool AddChild(UWidget* Child, FVector2D Position) OVERRIDE;
 
 	// UPanelWidget
 	virtual int32 GetChildrenCount() const OVERRIDE;
 	virtual UWidget* GetChildAt(int32 Index) const OVERRIDE;
+	virtual int32 GetChildIndex(UWidget* Content) const OVERRIDE;
+	virtual bool AddChild(UWidget* Child, FVector2D Position) OVERRIDE;
+	virtual bool RemoveChild(UWidget* Child) OVERRIDE;
+	virtual void ReplaceChildAt(int32 Index, UWidget* Child) OVERRIDE;
+	virtual void InsertChildAt(int32 Index, UWidget* Child) OVERRIDE;
 	// End UPanelWidget
 
 	/** Gets the underlying native canvas widget if it has been constructed */
-	TSharedPtr<class SFixedSizeCanvas> GetCanvasWidget() const;
+	TSharedPtr<class SConstraintCanvas> GetCanvasWidget() const;
 
 	/** Computes the geometry for a particular slot based on the current geometry of the canvas. */
 	bool GetGeometryForSlot(UCanvasPanelSlot* Slot, FGeometry& ArrangedGeometry) const;
@@ -43,7 +47,7 @@ class UMG_API UCanvasPanel : public UPanelWidget
 
 protected:
 
-	TWeakPtr<class SFixedSizeCanvas> MyCanvas;
+	TSharedPtr<class SFixedSizeCanvas> MyCanvas;
 
 protected:
 	// UWidget interface
diff --git a/Engine/Source/Runtime/UMG/Public/Components/ProgressBar.h b/Engine/Source/Runtime/UMG/Public/Components/ProgressBar.h
--- a/Engine/Source/Runtime/UMG/Public/Components/ProgressBar.h
+++ b/Engine/Source/Runtime/UMG/Public/Components/ProgressBar.h
@@ -31,10 +31,6 @@ class UMG_API UProgressBar : public UWidget
 	UPROPERTY(EditDefaultsOnly, Category=Appearance)
 	FVector2D BorderPadding;
 
-	/** Gets the current value of the ProgressBar. */
-	UFUNCTION(BlueprintPure, Category="Behavior")
-	float GetPercent();
-
 	/** Sets the current value of the ProgressBar. */
 	UFUNCTION(BlueprintCallable, Category="Behavior")
 	void SetPercent(float InPercent);
diff --git a/Engine/Source/Runtime/UMG/Public/Components/ScrollBox.h b/Engine/Source/Runtime/UMG/Public/Components/ScrollBox.h
--- a/Engine/Source/Runtime/UMG/Public/Components/ScrollBox.h
+++ b/Engine/Source/Runtime/UMG/Public/Components/ScrollBox.h
@@ -11,7 +11,7 @@ class UMG_API UScrollBox : public UPanelWidget
 	GENERATED_UCLASS_BODY()
 
 	/** The slots containing the widgets that are scrollable. */
-	UPROPERTY(EditDefaultsOnly, EditInline, Category=Slots)
+	UPROPERTY()
 	TArray<UScrollBoxSlot*> Slots;
 
 	//TODO UMG Add ways to make adding slots callable by blueprints.
diff --git a/Engine/Source/Runtime/UMG/Public/Components/TextBlock.h b/Engine/Source/Runtime/UMG/Public/Components/TextBlock.h
--- a/Engine/Source/Runtime/UMG/Public/Components/TextBlock.h
+++ b/Engine/Source/Runtime/UMG/Public/Components/TextBlock.h
@@ -43,21 +43,30 @@ public:
 	UPROPERTY()
 	FGetSlateColor ColorAndOpacityDelegate;
 
+	/** Whether text wraps onto a new line when it's length exceeds this width; if this value is zero or negative, no wrapping occurs. */
+	UPROPERTY(EditDefaultsOnly, Category=Appearance)
+	float WrapTextAt;
+
+	/** True if we're wrapping text automatically based on the computed horizontal space for this widget */
+	UPROPERTY(EditDefaultsOnly, Category=Appearance)
+	bool AutoWrapText;
+
 	/** The style to use to render the text */
 	UPROPERTY(EditDefaultsOnly, BlueprintReadOnly, Category=Style, meta=( DisplayThumbnail = "true" ))
 	USlateWidgetStyleAsset* Style;
 
-	//TODO UMG Set text via UFunction.
-	//TODO UMG set all these properties by UFunction as well.
+	///** Called when this text is double clicked */
+	//SLATE_EVENT(FOnClicked, OnDoubleClicked)
+
+	// UWidget interface
+	void SyncronizeProperties() OVERRIDE;
+	// End of UWidget interface
 
 protected:
 	// UWidget interface
 	virtual TSharedRef<SWidget> RebuildWidget() OVERRIDE;
 	// End of UWidget interface
 
-	/** Converts the bound FLinearColor to a slate color for the ColorAndOpacityDelegate */
-	FSlateColor GetColorAndOpacity() const;
-
-	/** Converts the bound FLinearColor to a slate color for the ShadowColorAndOpacityDelegate */
-	FLinearColor GetShadowColorAndOpacity() const;
+protected:
+	TSharedPtr<STextBlock> MyTextBlock;
 };
