Re-added some accidentally removed code

diff --git a/Engine/Source/Runtime/Slate/Private/Widgets/Text/SMultiLineEditableText.cpp b/Engine/Source/Runtime/Slate/Private/Widgets/Text/SMultiLineEditableText.cpp
--- a/Engine/Source/Runtime/Slate/Private/Widgets/Text/SMultiLineEditableText.cpp
+++ b/Engine/Source/Runtime/Slate/Private/Widgets/Text/SMultiLineEditableText.cpp
@@ -385,8 +385,7 @@ void SMultiLineEditableText::OnKeyboardFocusLost( const FKeyboardFocusEvent& InK
 		TextLayout->GetAsText(EditedText);
 
 		OnTextCommitted.ExecuteIfBound(FText::FromString(EditedText), TextAction);
-		// todo: jdale - testing
-		//RemoveCursorHighlight();
+		RemoveCursorHighlight();
 	}
 }
 
