SLATE - Adding SetForegroundColor back to STextBlock, marking it as Deprecated in 4.3.

diff --git a/Engine/Source/Runtime/Slate/Public/Widgets/Text/STextBlock.h b/Engine/Source/Runtime/Slate/Public/Widgets/Text/STextBlock.h
--- a/Engine/Source/Runtime/Slate/Public/Widgets/Text/STextBlock.h
+++ b/Engine/Source/Runtime/Slate/Public/Widgets/Text/STextBlock.h
@@ -136,6 +136,9 @@ public:
 	 */
 	void SetFont(const TAttribute< FSlateFontInfo >& InFont);
 
+	DEPRECATED(4.3, "SetForegroundColor is deprecated use SetColorAndOpacity instead.")
+	void SetForegroundColor(const TAttribute<FSlateColor>& InColorAndOpacity) { SetColorAndOpacity(InColorAndOpacity); }
+
 	/** See ColorAndOpacity attribute */
 	void SetColorAndOpacity(const TAttribute<FSlateColor>& InColorAndOpacity);
 
