renamed OVERRIDE to override

diff --git a/Engine/Plugins/Slate/SlateRemote/Source/SlateRemote/Private/SlateRemoteModule.cpp b/Engine/Plugins/Slate/SlateRemote/Source/SlateRemote/Private/SlateRemoteModule.cpp
--- a/Engine/Plugins/Slate/SlateRemote/Source/SlateRemote/Private/SlateRemoteModule.cpp
+++ b/Engine/Plugins/Slate/SlateRemote/Source/SlateRemote/Private/SlateRemoteModule.cpp
@@ -20,7 +20,7 @@ public:
 
 	// Begin IModuleInterface interface
 
-	virtual void StartupModule( ) OVERRIDE
+	virtual void StartupModule( ) override
 	{
 		if (!SupportsSlateRemote())
 		{
@@ -50,7 +50,7 @@ public:
 		RestartServices();
 	}
 
-	virtual void ShutdownModule( ) OVERRIDE
+	virtual void ShutdownModule( ) override
 	{
 		// unregister application events
 		FCoreDelegates::ApplicationHasReactivatedDelegate.RemoveAll(this);
@@ -68,7 +68,7 @@ public:
 		ShutdownRemoteServer();
 	}
 
-	virtual bool SupportsDynamicReloading( ) OVERRIDE
+	virtual bool SupportsDynamicReloading( ) override
 	{
 		return true;
 	}
diff --git a/Engine/Source/Developer/SlateReflector/Private/SlateReflectorModule.cpp b/Engine/Source/Developer/SlateReflector/Private/SlateReflectorModule.cpp
--- a/Engine/Source/Developer/SlateReflector/Private/SlateReflectorModule.cpp
+++ b/Engine/Source/Developer/SlateReflector/Private/SlateReflectorModule.cpp
@@ -21,7 +21,7 @@ public:
 
 	// Begin ISlateReflectorModule interface
 
-	virtual TSharedRef<SWidget> GetWidgetReflector( ) OVERRIDE
+	virtual TSharedRef<SWidget> GetWidgetReflector( ) override
 	{
 		TSharedPtr<SWidgetReflector> WidgetReflector = WidgetReflectorPtr.Pin();
 
@@ -34,7 +34,7 @@ public:
 		return WidgetReflector.ToSharedRef();
 	}
 
-	virtual void RegisterTabSpawner( const TSharedRef<FWorkspaceItem>& WorkspaceGroup ) OVERRIDE
+	virtual void RegisterTabSpawner( const TSharedRef<FWorkspaceItem>& WorkspaceGroup ) override
 	{
 		FGlobalTabmanager::Get()->RegisterNomadTabSpawner("WidgetReflector", FOnSpawnTab::CreateRaw(this, &FSlateReflectorModule::MakeWidgetReflectorTab) )
 			.SetDisplayName(LOCTEXT("WidgetReflectorTitle", "Widget Reflector"))
@@ -54,9 +54,9 @@ public:
 
 	// Begin IModuleInterface interface
 
-	virtual void StartupModule( ) OVERRIDE { }
+	virtual void StartupModule( ) override { }
 
-	virtual void ShutdownModule( ) OVERRIDE
+	virtual void ShutdownModule( ) override
 	{
 		UnregisterTabSpawner();
 	}
diff --git a/Engine/Source/Developer/SlateReflector/Private/Widgets/SWidgetReflector.h b/Engine/Source/Developer/SlateReflector/Private/Widgets/SWidgetReflector.h
--- a/Engine/Source/Developer/SlateReflector/Private/Widgets/SWidgetReflector.h
+++ b/Engine/Source/Developer/SlateReflector/Private/Widgets/SWidgetReflector.h
@@ -38,7 +38,7 @@ public:
 
 	// Begin SCompoundWidget overrides
 
-	virtual void Tick( const FGeometry& AllottedGeometry, const double InCurrentTime, const float InDeltaTime ) OVERRIDE;
+	virtual void Tick( const FGeometry& AllottedGeometry, const double InCurrentTime, const float InDeltaTime ) override;
 
 	// End SCompoundWidget overrides
 
@@ -46,36 +46,36 @@ public:
 
 	// Begin IWidgetReflector interface
 
-	virtual bool IsInPickingMode( ) const OVERRIDE
+	virtual bool IsInPickingMode( ) const override
 	{
 		return bIsPicking;
 	}
 
-	virtual bool IsShowingFocus( ) const OVERRIDE
+	virtual bool IsShowingFocus( ) const override
 	{
 		return bShowFocus;
 	}
 
-	virtual bool IsVisualizingLayoutUnderCursor( ) const OVERRIDE
+	virtual bool IsVisualizingLayoutUnderCursor( ) const override
 	{
 		return bIsPicking;
 	}
 
-	virtual void OnWidgetPicked( ) OVERRIDE
+	virtual void OnWidgetPicked( ) override
 	{
 		bIsPicking = false;
 	}
 
-	virtual bool ReflectorNeedsToDrawIn( TSharedRef<SWindow> ThisWindow ) const OVERRIDE;
+	virtual bool ReflectorNeedsToDrawIn( TSharedRef<SWindow> ThisWindow ) const override;
 
-	virtual void SetSourceAccessDelegate( FAccessSourceCode InDelegate ) OVERRIDE
+	virtual void SetSourceAccessDelegate( FAccessSourceCode InDelegate ) override
 	{
 		SourceAccessDelegate = InDelegate;
 	}
 
-	virtual void SetWidgetsToVisualize( const FWidgetPath& InWidgetsToVisualize ) OVERRIDE;
+	virtual void SetWidgetsToVisualize( const FWidgetPath& InWidgetsToVisualize ) override;
 
-	virtual int32 Visualize( const FWidgetPath& InWidgetsToVisualize, FSlateWindowElementList& OutDrawElements, int32 LayerId ) OVERRIDE;
+	virtual int32 Visualize( const FWidgetPath& InWidgetsToVisualize, FSlateWindowElementList& OutDrawElements, int32 LayerId ) override;
 
 	// End IWidgetReflector interface
 
diff --git a/Engine/Source/Developer/SlateReflector/Private/Widgets/SWidgetReflectorTreeWidgetItem.h b/Engine/Source/Developer/SlateReflector/Private/Widgets/SWidgetReflectorTreeWidgetItem.h
--- a/Engine/Source/Developer/SlateReflector/Private/Widgets/SWidgetReflectorTreeWidgetItem.h
+++ b/Engine/Source/Developer/SlateReflector/Private/Widgets/SWidgetReflectorTreeWidgetItem.h
@@ -41,7 +41,7 @@ public:
 	}
 
 	BEGIN_SLATE_FUNCTION_BUILD_OPTIMIZATION
-	virtual TSharedRef<SWidget> GenerateWidgetForColumn( const FName& ColumnName ) OVERRIDE
+	virtual TSharedRef<SWidget> GenerateWidgetForColumn( const FName& ColumnName ) override
 	{
 		if (ColumnName == TEXT("WidgetName"))
 		{
diff --git a/Engine/Source/Runtime/Slate/Private/Framework/Application/MenuStack.cpp b/Engine/Source/Runtime/Slate/Private/Framework/Application/MenuStack.cpp
--- a/Engine/Source/Runtime/Slate/Private/Framework/Application/MenuStack.cpp
+++ b/Engine/Source/Runtime/Slate/Private/Framework/Application/MenuStack.cpp
@@ -43,12 +43,12 @@ namespace
 	private:
 
 		/** This widget must support keyboard focus */
-		virtual bool SupportsKeyboardFocus() const OVERRIDE
+		virtual bool SupportsKeyboardFocus() const override
 		{
 			return true;
 		}
 
-		virtual FReply OnKeyDown( const FGeometry& MyGeometry, const FKeyboardEvent& InKeyboardEvent ) OVERRIDE
+		virtual FReply OnKeyDown( const FGeometry& MyGeometry, const FKeyboardEvent& InKeyboardEvent ) override
 		{
 			if (OnKeyDownDelegate.IsBound())
 			{
diff --git a/Engine/Source/Runtime/Slate/Private/Framework/Application/SWindowTitleBar.h b/Engine/Source/Runtime/Slate/Private/Framework/Application/SWindowTitleBar.h
--- a/Engine/Source/Runtime/Slate/Private/Framework/Application/SWindowTitleBar.h
+++ b/Engine/Source/Runtime/Slate/Private/Framework/Application/SWindowTitleBar.h
@@ -47,7 +47,7 @@ class SAppIconWidget
 			];
 	}
 
-	virtual EWindowZone::Type GetWindowZoneOverride() const OVERRIDE
+	virtual EWindowZone::Type GetWindowZoneOverride() const override
 	{
 		// Pretend we are a REAL system menu so the user can click to open a menu, or double-click to close the app on Windows
 		return EWindowZone::SysMenu;
diff --git a/Engine/Source/Runtime/Slate/Private/Framework/Docking/SDockingArea.h b/Engine/Source/Runtime/Slate/Private/Framework/Docking/SDockingArea.h
--- a/Engine/Source/Runtime/Slate/Private/Framework/Docking/SDockingArea.h
+++ b/Engine/Source/Runtime/Slate/Private/Framework/Docking/SDockingArea.h
@@ -43,18 +43,18 @@ public:
 	}
 	
 	/** Returns this dock area */
-	virtual TSharedPtr<SDockingArea> GetDockArea() OVERRIDE;
-	virtual TSharedPtr<const SDockingArea> GetDockArea() const OVERRIDE;
+	virtual TSharedPtr<SDockingArea> GetDockArea() override;
+	virtual TSharedPtr<const SDockingArea> GetDockArea() const override;
 
 	/** Returns the window that this dock area resides in directly and also manages */
 	TSharedPtr<SWindow> GetParentWindow() const;
 
-	virtual void OnDragEnter( const FGeometry& MyGeometry, const FDragDropEvent& DragDropEvent ) OVERRIDE;
-	virtual void OnDragLeave( const FDragDropEvent& DragDropEvent ) OVERRIDE;
+	virtual void OnDragEnter( const FGeometry& MyGeometry, const FDragDropEvent& DragDropEvent ) override;
+	virtual void OnDragLeave( const FDragDropEvent& DragDropEvent ) override;
 	
-	virtual FReply OnDrop( const FGeometry& MyGeometry, const FDragDropEvent& DragDropEvent ) OVERRIDE;
+	virtual FReply OnDrop( const FGeometry& MyGeometry, const FDragDropEvent& DragDropEvent ) override;
 
-	virtual FReply OnUserAttemptingDock( SDockingNode::RelativeDirection Direction, const FDragDropEvent& DragDropEvent ) OVERRIDE;
+	virtual FReply OnUserAttemptingDock( SDockingNode::RelativeDirection Direction, const FDragDropEvent& DragDropEvent ) override;
 
 	void OnTabFoundNewHome( const TSharedRef<SDockTab>& RelocatedTab, const TSharedRef<SWindow>& NewOwnerWindow );
 
@@ -71,12 +71,12 @@ public:
 
 	void SetParentWindow( const TSharedRef<SWindow>& NewParentWindow );
 
-	virtual TSharedPtr<FTabManager::FLayoutNode> GatherPersistentLayout() const OVERRIDE;
+	virtual TSharedPtr<FTabManager::FLayoutNode> GatherPersistentLayout() const override;
 
 	TSharedRef<FTabManager> GetTabManager() const;
 
 protected:
-	virtual SDockingNode::ECleanupRetVal CleanUpNodes() OVERRIDE;
+	virtual SDockingNode::ECleanupRetVal CleanUpNodes() override;
 
 private:
 	EVisibility TargetCrossVisibility() const;
@@ -87,7 +87,7 @@ private:
 	/** We were placed in a window, and it is being destroyed */
 	void OnOwningWindowBeingDestroyed(const TSharedRef<SWindow>& WindowBeingDestroyed);
 
-	virtual void OnLiveTabAdded() OVERRIDE;
+	virtual void OnLiveTabAdded() override;
 
 private:	
 	
diff --git a/Engine/Source/Runtime/Slate/Private/Framework/Docking/SDockingCross.h b/Engine/Source/Runtime/Slate/Private/Framework/Docking/SDockingCross.h
--- a/Engine/Source/Runtime/Slate/Private/Framework/Docking/SDockingCross.h
+++ b/Engine/Source/Runtime/Slate/Private/Framework/Docking/SDockingCross.h
@@ -17,11 +17,11 @@ public:
 	void Construct( const FArguments& InArgs, const TSharedPtr<class SDockingNode>& InOwnerNode );
 
 	// SWidget interface
-	virtual int32 OnPaint( const FGeometry& AllottedGeometry, const FSlateRect& MyClippingRect, FSlateWindowElementList& OutDrawElements, int32 LayerId, const FWidgetStyle& InWidgetStyle, bool bParentEnabled ) const OVERRIDE;
-	virtual FVector2D ComputeDesiredSize() const OVERRIDE;
-	virtual void OnDragLeave( const FDragDropEvent& DragDropEvent ) OVERRIDE;
-	virtual FReply OnDragOver( const FGeometry& MyGeometry, const FDragDropEvent& DragDropEvent ) OVERRIDE;
-	virtual FReply OnDrop( const FGeometry& MyGeometry, const FDragDropEvent& DragDropEvent ) OVERRIDE;
+	virtual int32 OnPaint( const FGeometry& AllottedGeometry, const FSlateRect& MyClippingRect, FSlateWindowElementList& OutDrawElements, int32 LayerId, const FWidgetStyle& InWidgetStyle, bool bParentEnabled ) const override;
+	virtual FVector2D ComputeDesiredSize() const override;
+	virtual void OnDragLeave( const FDragDropEvent& DragDropEvent ) override;
+	virtual FReply OnDragOver( const FGeometry& MyGeometry, const FDragDropEvent& DragDropEvent ) override;
+	virtual FReply OnDrop( const FGeometry& MyGeometry, const FDragDropEvent& DragDropEvent ) override;
 	// End of SWidget interface
 	
 private:
diff --git a/Engine/Source/Runtime/Slate/Private/Framework/Docking/SDockingSplitter.h b/Engine/Source/Runtime/Slate/Private/Framework/Docking/SDockingSplitter.h
--- a/Engine/Source/Runtime/Slate/Private/Framework/Docking/SDockingSplitter.h
+++ b/Engine/Source/Runtime/Slate/Private/Framework/Docking/SDockingSplitter.h
@@ -68,7 +68,7 @@ public:
 	TArray< TSharedRef<SDockingNode> > GetChildNodesRecursively() const;
 	
 	/** Recursively searches through all children looking for child tabs */
-	virtual TArray< TSharedRef<SDockTab> > GetAllChildTabs() const OVERRIDE;
+	virtual TArray< TSharedRef<SDockTab> > GetAllChildTabs() const override;
 
 	/** Gets the size coefficient of a given child dock node */
 	float GetSizeCoefficientForSlot(int32 Index) const;
@@ -76,7 +76,7 @@ public:
 	/** Gets the orientation of this splitter */
 	EOrientation GetOrientation() const;
 
-	virtual TSharedPtr<FTabManager::FLayoutNode> GatherPersistentLayout() const OVERRIDE;
+	virtual TSharedPtr<FTabManager::FLayoutNode> GatherPersistentLayout() const override;
 
 protected:
 
@@ -85,7 +85,7 @@ protected:
 
 	static SDockingNode::ECleanupRetVal MostResponsibility( SDockingNode::ECleanupRetVal A, SDockingNode::ECleanupRetVal B );
 
-	virtual SDockingNode::ECleanupRetVal CleanUpNodes() OVERRIDE;
+	virtual SDockingNode::ECleanupRetVal CleanUpNodes() override;
 
 	float ComputeChildCoefficientTotal() const;
 
diff --git a/Engine/Source/Runtime/Slate/Private/Framework/Docking/SDockingTabStack.h b/Engine/Source/Runtime/Slate/Private/Framework/Docking/SDockingTabStack.h
--- a/Engine/Source/Runtime/Slate/Private/Framework/Docking/SDockingTabStack.h
+++ b/Engine/Source/Runtime/Slate/Private/Framework/Docking/SDockingTabStack.h
@@ -55,37 +55,37 @@ public:
 	/** Set the content that the DockNode is presenting. */
 	void SetNodeContent( const TSharedRef<SWidget>& InContent, const TSharedRef<SWidget>& InContentLeft, const TSharedRef<SWidget>& InContentRight );
 
-	virtual FReply OnUserAttemptingDock( SDockingNode::RelativeDirection Direction, const FDragDropEvent& DragDropEvent ) OVERRIDE;
+	virtual FReply OnUserAttemptingDock( SDockingNode::RelativeDirection Direction, const FDragDropEvent& DragDropEvent ) override;
 	
 	/** Recursively searches through all children looking for child tabs */
-	virtual TArray< TSharedRef<SDockTab> > GetAllChildTabs() const OVERRIDE;
+	virtual TArray< TSharedRef<SDockTab> > GetAllChildTabs() const override;
 
-	virtual SSplitter::ESizeRule GetSizeRule() const OVERRIDE;
+	virtual SSplitter::ESizeRule GetSizeRule() const override;
 
 	void SetTabWellHidden( bool bShouldHideTabWell );
 	bool IsTabWellHidden() const;
 
-	virtual TSharedPtr<FTabManager::FLayoutNode> GatherPersistentLayout() const OVERRIDE;
+	virtual TSharedPtr<FTabManager::FLayoutNode> GatherPersistentLayout() const override;
 
 public:
 
-	virtual void Tick( const FGeometry& AllottedGeometry, const double InCurrentTime, const float InDeltaTime ) OVERRIDE;
+	virtual void Tick( const FGeometry& AllottedGeometry, const double InCurrentTime, const float InDeltaTime ) override;
 
-	virtual FReply OnDragOver( const FGeometry& MyGeometry, const FDragDropEvent& DragDropEvent ) OVERRIDE;
+	virtual FReply OnDragOver( const FGeometry& MyGeometry, const FDragDropEvent& DragDropEvent ) override;
 
-	virtual void OnDragLeave( const FDragDropEvent& DragDropEvent ) OVERRIDE;
+	virtual void OnDragLeave( const FDragDropEvent& DragDropEvent ) override;
 	
-	virtual FReply OnDrop( const FGeometry& MyGeometry, const FDragDropEvent& DragDropEvent ) OVERRIDE;
+	virtual FReply OnDrop( const FGeometry& MyGeometry, const FDragDropEvent& DragDropEvent ) override;
 
-	virtual void OnKeyboardFocusChanging( const FWeakWidgetPath& PreviousFocusPath, const FWidgetPath& NewWidgetPath ) OVERRIDE;
+	virtual void OnKeyboardFocusChanging( const FWeakWidgetPath& PreviousFocusPath, const FWidgetPath& NewWidgetPath ) override;
 
-	virtual FReply OnMouseButtonDown( const FGeometry& MyGeometry, const FPointerEvent& MouseEvent ) OVERRIDE;
+	virtual FReply OnMouseButtonDown( const FGeometry& MyGeometry, const FPointerEvent& MouseEvent ) override;
 
-	virtual bool SupportsKeyboardFocus() const OVERRIDE { return true; }
+	virtual bool SupportsKeyboardFocus() const override { return true; }
 
 protected:
 
-	virtual EWindowZone::Type GetWindowZoneOverride() const OVERRIDE
+	virtual EWindowZone::Type GetWindowZoneOverride() const override
 	{
 		// Pretend we are a title bar so the user can grab the area to move the window around
 		return EWindowZone::TitleBar;
@@ -107,7 +107,7 @@ protected:
 
 	FReply TabWellRightClicked( const FGeometry& TabWellGeometry, const FPointerEvent& MouseEvent );
 
-	virtual SDockingNode::ECleanupRetVal CleanUpNodes() OVERRIDE;
+	virtual SDockingNode::ECleanupRetVal CleanUpNodes() override;
 
 	int32 OpenPersistentTab( const FTabId& TabId, int32 OpenLocationAmongActiveTabs = INDEX_NONE );
 
@@ -116,7 +116,7 @@ protected:
 	void RemovePersistentTab( const FTabId& TabId );
 
 	/** Overridden from SDockingNode */
-	virtual void SetParentNode( TSharedRef<class SDockingSplitter> InParent ) OVERRIDE;
+	virtual void SetParentNode( TSharedRef<class SDockingSplitter> InParent ) override;
 
 private:
 	
diff --git a/Engine/Source/Runtime/Slate/Private/Framework/Docking/SDockingTabWell.h b/Engine/Source/Runtime/Slate/Private/Framework/Docking/SDockingTabWell.h
--- a/Engine/Source/Runtime/Slate/Private/Framework/Docking/SDockingTabWell.h
+++ b/Engine/Source/Runtime/Slate/Private/Framework/Docking/SDockingTabWell.h
@@ -84,15 +84,15 @@ public:
 	FVector2D ComputeChildSize( const FGeometry& AllottedGeometry ) const;
 
 	// SWidget interface
-	virtual void OnArrangeChildren( const FGeometry& AllottedGeometry, FArrangedChildren& ArrangedChildren ) const OVERRIDE;
-	virtual int32 OnPaint( const FGeometry& AllottedGeometry, const FSlateRect& MyClippingRect, FSlateWindowElementList& OutDrawElements, int32 LayerId, const FWidgetStyle& InWidgetStyle, bool bParentEnabled ) const OVERRIDE;
-	virtual FVector2D ComputeDesiredSize() const OVERRIDE;
-	virtual FChildren* GetChildren() OVERRIDE;
-	virtual void OnDragEnter( const FGeometry& MyGeometry, const FDragDropEvent& DragDropEvent ) OVERRIDE;
-	virtual void OnDragLeave( const FDragDropEvent& DragDropEvent ) OVERRIDE;
-	virtual FReply OnDragOver( const FGeometry& MyGeometry, const FDragDropEvent& DragDropEvent ) OVERRIDE;
-	virtual FReply OnDrop( const FGeometry& MyGeometry, const FDragDropEvent& DragDropEvent ) OVERRIDE;
-	virtual EWindowZone::Type GetWindowZoneOverride() const OVERRIDE;
+	virtual void OnArrangeChildren( const FGeometry& AllottedGeometry, FArrangedChildren& ArrangedChildren ) const override;
+	virtual int32 OnPaint( const FGeometry& AllottedGeometry, const FSlateRect& MyClippingRect, FSlateWindowElementList& OutDrawElements, int32 LayerId, const FWidgetStyle& InWidgetStyle, bool bParentEnabled ) const override;
+	virtual FVector2D ComputeDesiredSize() const override;
+	virtual FChildren* GetChildren() override;
+	virtual void OnDragEnter( const FGeometry& MyGeometry, const FDragDropEvent& DragDropEvent ) override;
+	virtual void OnDragLeave( const FDragDropEvent& DragDropEvent ) override;
+	virtual FReply OnDragOver( const FGeometry& MyGeometry, const FDragDropEvent& DragDropEvent ) override;
+	virtual FReply OnDrop( const FGeometry& MyGeometry, const FDragDropEvent& DragDropEvent ) override;
+	virtual EWindowZone::Type GetWindowZoneOverride() const override;
 	// End of SWidget interface
 
 private:
diff --git a/Engine/Source/Runtime/Slate/Private/Framework/Docking/SDockingTarget.h b/Engine/Source/Runtime/Slate/Private/Framework/Docking/SDockingTarget.h
--- a/Engine/Source/Runtime/Slate/Private/Framework/Docking/SDockingTarget.h
+++ b/Engine/Source/Runtime/Slate/Private/Framework/Docking/SDockingTarget.h
@@ -24,9 +24,9 @@ public:
 	void Construct( const FArguments& InArgs );
 
 	// SWidget interface
-	virtual void OnDragEnter( const FGeometry& MyGeometry, const FDragDropEvent& DragDropEvent ) OVERRIDE;
-	virtual void OnDragLeave( const FDragDropEvent& DragDropEvent ) OVERRIDE;
-	virtual FReply OnDrop( const FGeometry& MyGeometry, const FDragDropEvent& DragDropEvent ) OVERRIDE;
+	virtual void OnDragEnter( const FGeometry& MyGeometry, const FDragDropEvent& DragDropEvent ) override;
+	virtual void OnDragLeave( const FDragDropEvent& DragDropEvent ) override;
+	virtual FReply OnDrop( const FGeometry& MyGeometry, const FDragDropEvent& DragDropEvent ) override;
 	// End of SWidget interface
 
 	/** @return the DockNode which this target represent */
diff --git a/Engine/Source/Runtime/Slate/Private/Framework/MultiBox/MultiBoxCustomization.h b/Engine/Source/Runtime/Slate/Private/Framework/MultiBox/MultiBoxCustomization.h
--- a/Engine/Source/Runtime/Slate/Private/Framework/MultiBox/MultiBoxCustomization.h
+++ b/Engine/Source/Runtime/Slate/Private/Framework/MultiBox/MultiBoxCustomization.h
@@ -13,7 +13,7 @@ public:
 
 	void Construct( const FArguments& InArgs );
 
-	virtual void BuildMultiBlockWidget(const ISlateStyle* StyleSet, const FName& StyleName) OVERRIDE;
+	virtual void BuildMultiBlockWidget(const ISlateStyle* StyleSet, const FName& StyleName) override;
 private:
 	TSharedPtr<SWidget> Content;
 };
@@ -27,10 +27,10 @@ class SMultiBlockDragHandle : public SCompoundWidget
 
 	FReply OnMouseButtonDown( const FGeometry& MyGeometry, const FPointerEvent& MouseEvent );
 
-	virtual void OnDragEnter( const FGeometry& MyGeometry, const FDragDropEvent& DragDropEvent ) OVERRIDE;
-	virtual FReply OnDragOver( const FGeometry& MyGeometry, const FDragDropEvent& DragDropEvent ) OVERRIDE;
-	virtual FReply OnDrop( const FGeometry& MyGeometry, const FDragDropEvent& DragDropEvent ) OVERRIDE;
-	virtual FReply OnDragDetected( const FGeometry& MyGeometry, const FPointerEvent& MouseEvent ) OVERRIDE;
+	virtual void OnDragEnter( const FGeometry& MyGeometry, const FDragDropEvent& DragDropEvent ) override;
+	virtual FReply OnDragOver( const FGeometry& MyGeometry, const FDragDropEvent& DragDropEvent ) override;
+	virtual FReply OnDrop( const FGeometry& MyGeometry, const FDragDropEvent& DragDropEvent ) override;
+	virtual FReply OnDragDetected( const FGeometry& MyGeometry, const FPointerEvent& MouseEvent ) override;
 
 private:
 	TSharedPtr<const FMultiBlock> Block;
@@ -56,7 +56,7 @@ public:
 	}
 
 	/** FMultiBlock interface */
-	virtual TSharedRef< class IMultiBlockBaseWidget > ConstructWidget() const OVERRIDE;
+	virtual TSharedRef< class IMultiBlockBaseWidget > ConstructWidget() const override;
 	TSharedRef<const FMultiBlock> GetActualBlock() const { return ActualBlock.ToSharedRef(); }
 private:
 	TSharedPtr<const FMultiBlock> ActualBlock;
diff --git a/Engine/Source/Runtime/Slate/Private/Framework/MultiBox/MultiBoxExtender.cpp b/Engine/Source/Runtime/Slate/Private/Framework/MultiBox/MultiBoxExtender.cpp
--- a/Engine/Source/Runtime/Slate/Private/Framework/MultiBox/MultiBoxExtender.cpp
+++ b/Engine/Source/Runtime/Slate/Private/Framework/MultiBox/MultiBoxExtender.cpp
@@ -52,7 +52,7 @@ class FMenuBarExtension : public FExtensionBase
 public:
 
 	/** FExtensionBase implementation */
-	virtual EExtensionType::Type GetType() const OVERRIDE
+	virtual EExtensionType::Type GetType() const override
 	{
 		return EExtensionType::MenuBar;
 	}
@@ -69,7 +69,7 @@ class FMenuExtension : public FExtensionBase
 public:
 
 	/** FExtensionBase implementation */
-	virtual EExtensionType::Type GetType() const OVERRIDE
+	virtual EExtensionType::Type GetType() const override
 	{
 		return EExtensionType::Menu;
 	}
@@ -86,7 +86,7 @@ class FToolBarExtension : public FExtensionBase
 public:
 
 	/** FExtensionBase implementation */
-	virtual EExtensionType::Type GetType() const OVERRIDE
+	virtual EExtensionType::Type GetType() const override
 	{
 		return EExtensionType::ToolBar;
 	}
diff --git a/Engine/Source/Runtime/Slate/Private/Framework/MultiBox/SButtonRowBlock.h b/Engine/Source/Runtime/Slate/Private/Framework/MultiBox/SButtonRowBlock.h
--- a/Engine/Source/Runtime/Slate/Private/Framework/MultiBox/SButtonRowBlock.h
+++ b/Engine/Source/Runtime/Slate/Private/Framework/MultiBox/SButtonRowBlock.h
@@ -79,7 +79,7 @@ public:
 	/**
 	 * Builds this MultiBlock widget up from the MultiBlock associated with it
 	 */
-	virtual void BuildMultiBlockWidget(const ISlateStyle* StyleSet, const FName& StyleName) OVERRIDE;
+	virtual void BuildMultiBlockWidget(const ISlateStyle* StyleSet, const FName& StyleName) override;
 
 	/**
 	 * Construct this widget
diff --git a/Engine/Source/Runtime/Slate/Private/Framework/MultiBox/SClippingHorizontalBox.h b/Engine/Source/Runtime/Slate/Private/Framework/MultiBox/SClippingHorizontalBox.h
--- a/Engine/Source/Runtime/Slate/Private/Framework/MultiBox/SClippingHorizontalBox.h
+++ b/Engine/Source/Runtime/Slate/Private/Framework/MultiBox/SClippingHorizontalBox.h
@@ -21,12 +21,12 @@ public:
 	FORCENOINLINE SClippingHorizontalBox() { }
 
 	/** SWidget interface */
-	virtual void Tick( const FGeometry& AllottedGeometry, const double InCurrentTime, const float InDeltaTime ) OVERRIDE;
-	virtual int32 OnPaint( const FGeometry& AllottedGeometry, const FSlateRect& MyClippingRect, FSlateWindowElementList& OutDrawElements, int32 LayerId, const FWidgetStyle& InWidgetStyle, bool bParentEnabled ) const OVERRIDE;
-	virtual FVector2D ComputeDesiredSize() const OVERRIDE;
+	virtual void Tick( const FGeometry& AllottedGeometry, const double InCurrentTime, const float InDeltaTime ) override;
+	virtual int32 OnPaint( const FGeometry& AllottedGeometry, const FSlateRect& MyClippingRect, FSlateWindowElementList& OutDrawElements, int32 LayerId, const FWidgetStyle& InWidgetStyle, bool bParentEnabled ) const override;
+	virtual FVector2D ComputeDesiredSize() const override;
 
 	/** SPanel interface */
-	virtual void OnArrangeChildren( const FGeometry& AllottedGeometry, FArrangedChildren& ArrangedChildren ) const OVERRIDE;
+	virtual void OnArrangeChildren( const FGeometry& AllottedGeometry, FArrangedChildren& ArrangedChildren ) const override;
 
 	/** Construct this widget */
 	void Construct( const FArguments& InArgs );
diff --git a/Engine/Source/Runtime/Slate/Private/Framework/MultiBox/SEditableTextBlock.h b/Engine/Source/Runtime/Slate/Private/Framework/MultiBox/SEditableTextBlock.h
--- a/Engine/Source/Runtime/Slate/Private/Framework/MultiBox/SEditableTextBlock.h
+++ b/Engine/Source/Runtime/Slate/Private/Framework/MultiBox/SEditableTextBlock.h
@@ -81,7 +81,7 @@ public:
 	/**
 	 * Builds this MultiBlock widget up from the MultiBlock associated with it
 	 */
-	virtual void BuildMultiBlockWidget(const ISlateStyle* StyleSet, const FName& StyleName) OVERRIDE;
+	virtual void BuildMultiBlockWidget(const ISlateStyle* StyleSet, const FName& StyleName) override;
 
 	/**
 	 * Construct this widget
diff --git a/Engine/Source/Runtime/Slate/Private/Framework/MultiBox/SGroupMarkerBlock.h b/Engine/Source/Runtime/Slate/Private/Framework/MultiBox/SGroupMarkerBlock.h
--- a/Engine/Source/Runtime/Slate/Private/Framework/MultiBox/SGroupMarkerBlock.h
+++ b/Engine/Source/Runtime/Slate/Private/Framework/MultiBox/SGroupMarkerBlock.h
@@ -58,7 +58,7 @@ public:
 	/**
 	 * Builds this MultiBlock widget up from the MultiBlock associated with it
 	 */
-	virtual void BuildMultiBlockWidget(const ISlateStyle* StyleSet, const FName& StyleName) OVERRIDE;
+	virtual void BuildMultiBlockWidget(const ISlateStyle* StyleSet, const FName& StyleName) override;
 
 	/**
 	 * Construct this widget
diff --git a/Engine/Source/Runtime/Slate/Private/Framework/MultiBox/SHeadingBlock.h b/Engine/Source/Runtime/Slate/Private/Framework/MultiBox/SHeadingBlock.h
--- a/Engine/Source/Runtime/Slate/Private/Framework/MultiBox/SHeadingBlock.h
+++ b/Engine/Source/Runtime/Slate/Private/Framework/MultiBox/SHeadingBlock.h
@@ -58,7 +58,7 @@ public:
 	/**
 	 * Builds this MultiBlock widget up from the MultiBlock associated with it
 	 */
-	virtual void BuildMultiBlockWidget(const ISlateStyle* StyleSet, const FName& StyleName) OVERRIDE;
+	virtual void BuildMultiBlockWidget(const ISlateStyle* StyleSet, const FName& StyleName) override;
 
 
 	/**
diff --git a/Engine/Source/Runtime/Slate/Private/Framework/MultiBox/SMenuEntryBlock.h b/Engine/Source/Runtime/Slate/Private/Framework/MultiBox/SMenuEntryBlock.h
--- a/Engine/Source/Runtime/Slate/Private/Framework/MultiBox/SMenuEntryBlock.h
+++ b/Engine/Source/Runtime/Slate/Private/Framework/MultiBox/SMenuEntryBlock.h
@@ -87,7 +87,7 @@ public:
 	FMenuEntryBlock( const FName& InExtensionHook, const FUIAction& UIAction, const TSharedRef< SWidget > Contents, const FNewMenuDelegate& InEntryBuilder, TSharedPtr<FExtender> InExtender, bool bInSubMenu, TSharedPtr< const FUICommandList > InCommandList, bool bInCloseSelfOnly );
 
 	/** FMultiBlock interface */
-	virtual void CreateMenuEntry(class FMenuBuilder& MenuBuilder) const OVERRIDE;
+	virtual void CreateMenuEntry(class FMenuBuilder& MenuBuilder) const override;
 
 
 
@@ -161,7 +161,7 @@ public:
 	/**
 	 * Builds this MultiBlock widget up from the MultiBlock associated with it
 	 */
-	virtual void BuildMultiBlockWidget(const ISlateStyle* StyleSet, const FName& StyleName) OVERRIDE;
+	virtual void BuildMultiBlockWidget(const ISlateStyle* StyleSet, const FName& StyleName) override;
 
 	/**
 	 * Construct this widget
@@ -229,10 +229,10 @@ protected:
 	const FSlateBrush* OnGetCheckImage() const;
 
 	// SWidget interface
-	virtual void OnMouseEnter( const FGeometry& MyGeometry, const FPointerEvent& MouseEvent ) OVERRIDE;
-	virtual void OnMouseLeave( const FPointerEvent& MouseEvent ) OVERRIDE;
-	virtual FReply OnKeyDown( const FGeometry& MyGeometry, const FKeyboardEvent& KeyboardEvent ) OVERRIDE;
-	virtual void Tick( const FGeometry& AllottedGeometry, const double InCurrentTime, const float InDeltaTime ) OVERRIDE;
+	virtual void OnMouseEnter( const FGeometry& MyGeometry, const FPointerEvent& MouseEvent ) override;
+	virtual void OnMouseLeave( const FPointerEvent& MouseEvent ) override;
+	virtual FReply OnKeyDown( const FGeometry& MyGeometry, const FKeyboardEvent& KeyboardEvent ) override;
+	virtual void Tick( const FGeometry& AllottedGeometry, const double InCurrentTime, const float InDeltaTime ) override;
 	// End of SWidget interface
 
 	/**
diff --git a/Engine/Source/Runtime/Slate/Private/Framework/MultiBox/SMenuSeparatorBlock.h b/Engine/Source/Runtime/Slate/Private/Framework/MultiBox/SMenuSeparatorBlock.h
--- a/Engine/Source/Runtime/Slate/Private/Framework/MultiBox/SMenuSeparatorBlock.h
+++ b/Engine/Source/Runtime/Slate/Private/Framework/MultiBox/SMenuSeparatorBlock.h
@@ -53,7 +53,7 @@ public:
 	/**
 	 * Builds this MultiBlock widget up from the MultiBlock associated with it
 	 */
-	virtual void BuildMultiBlockWidget(const ISlateStyle* StyleSet, const FName& StyleName) OVERRIDE;
+	virtual void BuildMultiBlockWidget(const ISlateStyle* StyleSet, const FName& StyleName) override;
 
 
 	/**
diff --git a/Engine/Source/Runtime/Slate/Private/Framework/MultiBox/SToolBarSeparatorBlock.h b/Engine/Source/Runtime/Slate/Private/Framework/MultiBox/SToolBarSeparatorBlock.h
--- a/Engine/Source/Runtime/Slate/Private/Framework/MultiBox/SToolBarSeparatorBlock.h
+++ b/Engine/Source/Runtime/Slate/Private/Framework/MultiBox/SToolBarSeparatorBlock.h
@@ -17,7 +17,7 @@ public:
 	FToolBarSeparatorBlock(const FName& InExtensionHook);
 
 	/** FMultiBlock interface */
-	virtual void CreateMenuEntry(class FMenuBuilder& MenuBuilder) const OVERRIDE;
+	virtual void CreateMenuEntry(class FMenuBuilder& MenuBuilder) const override;
 
 
 private:
@@ -56,7 +56,7 @@ public:
 	/**
 	 * Builds this MultiBlock widget up from the MultiBlock associated with it
 	 */
-	virtual void BuildMultiBlockWidget(const ISlateStyle* StyleSet, const FName& StyleName) OVERRIDE;
+	virtual void BuildMultiBlockWidget(const ISlateStyle* StyleSet, const FName& StyleName) override;
 
 
 	/**
diff --git a/Engine/Source/Runtime/Slate/Private/Framework/MultiBox/SWidgetBlock.h b/Engine/Source/Runtime/Slate/Private/Framework/MultiBox/SWidgetBlock.h
--- a/Engine/Source/Runtime/Slate/Private/Framework/MultiBox/SWidgetBlock.h
+++ b/Engine/Source/Runtime/Slate/Private/Framework/MultiBox/SWidgetBlock.h
@@ -21,7 +21,7 @@ public:
 	FWidgetBlock( TSharedRef<SWidget> InContent, const FText& InLabel, bool bInNoIndent );
 
 	/** FMultiBlock interface */
-	virtual void CreateMenuEntry(class FMenuBuilder& MenuBuilder) const OVERRIDE;
+	virtual void CreateMenuEntry(class FMenuBuilder& MenuBuilder) const override;
 
 
 private:
@@ -68,7 +68,7 @@ public:
 	/**
 	 * Builds this MultiBlock widget up from the MultiBlock associated with it
 	 */
-	virtual void BuildMultiBlockWidget(const ISlateStyle* StyleSet, const FName& StyleName) OVERRIDE;
+	virtual void BuildMultiBlockWidget(const ISlateStyle* StyleSet, const FName& StyleName) override;
 
 	/**
 	 * Construct this widget
diff --git a/Engine/Source/Runtime/Slate/Private/Framework/Testing/SLayoutExample.cpp b/Engine/Source/Runtime/Slate/Private/Framework/Testing/SLayoutExample.cpp
--- a/Engine/Source/Runtime/Slate/Private/Framework/Testing/SLayoutExample.cpp
+++ b/Engine/Source/Runtime/Slate/Private/Framework/Testing/SLayoutExample.cpp
@@ -366,7 +366,7 @@ public:
 	END_SLATE_FUNCTION_BUILD_OPTIMIZATION
 
 	
-	virtual bool OnVisualizeTooltip( const TSharedPtr<SWidget>& TooltipContent ) OVERRIDE
+	virtual bool OnVisualizeTooltip( const TSharedPtr<SWidget>& TooltipContent ) override
 	{
 		// The layout example has its own way of visualizing tool tips.
 		// We show them below the items.
diff --git a/Engine/Source/Runtime/Slate/Private/Framework/Testing/STableViewTesting.cpp b/Engine/Source/Runtime/Slate/Private/Framework/Testing/STableViewTesting.cpp
--- a/Engine/Source/Runtime/Slate/Private/Framework/Testing/STableViewTesting.cpp
+++ b/Engine/Source/Runtime/Slate/Private/Framework/Testing/STableViewTesting.cpp
@@ -248,7 +248,7 @@ public:
 	SLATE_END_ARGS()
 	
 
-	virtual TSharedRef<SWidget> GenerateWidgetForColumn( const FName& ColumnName ) OVERRIDE
+	virtual TSharedRef<SWidget> GenerateWidgetForColumn( const FName& ColumnName ) override
 	{
 		FSlateFontInfo ItemEditorFont = FCoreStyle::Get().GetFontStyle(TEXT("NormalFont"));
 
diff --git a/Engine/Source/Runtime/Slate/Private/Framework/Testing/SUserWidgetTest.cpp b/Engine/Source/Runtime/Slate/Private/Framework/Testing/SUserWidgetTest.cpp
--- a/Engine/Source/Runtime/Slate/Private/Framework/Testing/SUserWidgetTest.cpp
+++ b/Engine/Source/Runtime/Slate/Private/Framework/Testing/SUserWidgetTest.cpp
@@ -21,7 +21,7 @@ public:
 	}
 
 
-	virtual void DoStuff() OVERRIDE
+	virtual void DoStuff() override
 	{
 
 	}
diff --git a/Engine/Source/Runtime/Slate/Private/Framework/Text/RichTextMarkupProcessing.h b/Engine/Source/Runtime/Slate/Private/Framework/Text/RichTextMarkupProcessing.h
--- a/Engine/Source/Runtime/Slate/Private/Framework/Text/RichTextMarkupProcessing.h
+++ b/Engine/Source/Runtime/Slate/Private/Framework/Text/RichTextMarkupProcessing.h
@@ -9,7 +9,7 @@ public:
 	static TSharedRef< FRichTextMarkupProcessing > Create();
 
 public:
-	virtual void Process(TArray<FTextLineParseResults>& Results, const FString& Input, FString& Output) OVERRIDE;
+	virtual void Process(TArray<FTextLineParseResults>& Results, const FString& Input, FString& Output) override;
 
 private:
 	FRichTextMarkupProcessing() {}
diff --git a/Engine/Source/Runtime/Slate/Private/Widgets/Colors/SColorPicker.cpp b/Engine/Source/Runtime/Slate/Private/Widgets/Colors/SColorPicker.cpp
--- a/Engine/Source/Runtime/Slate/Private/Widgets/Colors/SColorPicker.cpp
+++ b/Engine/Source/Runtime/Slate/Private/Widgets/Colors/SColorPicker.cpp
@@ -431,7 +431,7 @@ public:
 		);
 	}
 
-	virtual FReply OnMouseButtonDown( const FGeometry& MyGeometry, const FPointerEvent& MouseEvent ) OVERRIDE
+	virtual FReply OnMouseButtonDown( const FGeometry& MyGeometry, const FPointerEvent& MouseEvent ) override
 	{
 		// Clicking ANY mouse button when the dropper isn't active resets the active dropper states ready to activate
 		if (!HasMouseCapture())
@@ -444,7 +444,7 @@ public:
 		return SButton::OnMouseButtonDown(MyGeometry, MouseEvent);
 	}
 
-	virtual FReply OnMouseButtonUp( const FGeometry& MyGeometry, const FPointerEvent& MouseEvent ) OVERRIDE
+	virtual FReply OnMouseButtonUp( const FGeometry& MyGeometry, const FPointerEvent& MouseEvent ) override
 	{
 		// If a mouse click is completing and the dropper is active ALWAYS deactivate
 		bool bDeactivating = false;
@@ -488,7 +488,7 @@ public:
 		return Reply;
 	}
 
-	virtual FReply OnMouseMove( const FGeometry& MyGeometry, const FPointerEvent& MouseEvent ) OVERRIDE
+	virtual FReply OnMouseMove( const FGeometry& MyGeometry, const FPointerEvent& MouseEvent ) override
 	{
 		// If the mouse is captured and bWasClickActivated is set then we are in dropper mode
 		if (HasMouseCapture() && bWasClickActivated)
@@ -541,7 +541,7 @@ public:
 		return FReply::Unhandled();
 	}
 
-	virtual FCursorReply OnCursorQuery( const FGeometry& MyGeometry, const FPointerEvent& CursorEvent ) const OVERRIDE
+	virtual FCursorReply OnCursorQuery( const FGeometry& MyGeometry, const FPointerEvent& CursorEvent ) const override
 	{
 		// Cursor is changed to the dropper when dropper mode is active and the states are correct
 		if (HasMouseCapture() &&
diff --git a/Engine/Source/Runtime/Slate/Private/Widgets/Layout/SScrollBox.cpp b/Engine/Source/Runtime/Slate/Private/Widgets/Layout/SScrollBox.cpp
--- a/Engine/Source/Runtime/Slate/Private/Widgets/Layout/SScrollBox.cpp
+++ b/Engine/Source/Runtime/Slate/Private/Widgets/Layout/SScrollBox.cpp
@@ -36,7 +36,7 @@ public:
 
 public:
 
-	virtual void OnArrangeChildren( const FGeometry& AllottedGeometry, FArrangedChildren& ArrangedChildren ) const OVERRIDE
+	virtual void OnArrangeChildren( const FGeometry& AllottedGeometry, FArrangedChildren& ArrangedChildren ) const override
 	{
 		
 		float CurChildOffset = -PhysicalOffset;
@@ -66,7 +66,7 @@ public:
 
 
 
-	virtual FVector2D ComputeDesiredSize() const OVERRIDE
+	virtual FVector2D ComputeDesiredSize() const override
 	{
 		FVector2D ThisDesiredSize = FVector2D::ZeroVector;
 		for(int32 SlotIndex=0; SlotIndex < Children.Num(); ++SlotIndex )
@@ -82,7 +82,7 @@ public:
 
 
 
-	virtual FChildren* GetChildren() OVERRIDE
+	virtual FChildren* GetChildren() override
 	{
 		return &Children;
 	}
diff --git a/Engine/Source/Runtime/Slate/Private/Widgets/Notifications/SNotificationList.cpp b/Engine/Source/Runtime/Slate/Private/Widgets/Notifications/SNotificationList.cpp
--- a/Engine/Source/Runtime/Slate/Private/Widgets/Notifications/SNotificationList.cpp
+++ b/Engine/Source/Runtime/Slate/Private/Widgets/Notifications/SNotificationList.cpp
@@ -21,18 +21,18 @@ public:
 	
 	
 	/** Sets the text for message element */
-	virtual void SetText( const TAttribute< FText >& InText ) OVERRIDE
+	virtual void SetText( const TAttribute< FText >& InText ) override
 	{
 		Text = InText;
 		MyTextBlock->SetText( Text );
 	}
 
-	virtual ECompletionState GetCompletionState() const OVERRIDE
+	virtual ECompletionState GetCompletionState() const override
 	{
 		return CompletionState;
 	}
 
-	virtual void SetCompletionState(ECompletionState State) OVERRIDE
+	virtual void SetCompletionState(ECompletionState State) override
 	{
 		CompletionState = State;
 
@@ -44,7 +44,7 @@ public:
 		}
 	}
 
-	virtual void ExpireAndFadeout() OVERRIDE
+	virtual void ExpireAndFadeout() override
 	{
 		FadeAnimation = FCurveSequence();
 		// Add some space for the expire time
@@ -83,7 +83,7 @@ public:
 	}
 
 	/** Begins the fadeout of this message */
-	virtual void Fadeout() OVERRIDE
+	virtual void Fadeout() override
 	{
 		// Start fade animation
 		FadeAnimation = FCurveSequence();
@@ -91,7 +91,7 @@ public:
 		FadeAnimation.PlayReverse();
 	}
 
-	void Tick( const FGeometry& AllottedGeometry, const double InCurrentTime, const float InDeltaTime ) OVERRIDE
+	void Tick( const FGeometry& AllottedGeometry, const double InCurrentTime, const float InDeltaTime ) override
 	{
 		SCompoundWidget::Tick(AllottedGeometry, InCurrentTime, InDeltaTime);
 
@@ -490,7 +490,7 @@ public:
 	}
 
 	/** Sets the text and delegate for the hyperlink */
-	virtual void SetHyperlink( const FSimpleDelegate& InHyperlink, const TAttribute< FText >& InHyperlinkText = TAttribute< FText >() ) OVERRIDE
+	virtual void SetHyperlink( const FSimpleDelegate& InHyperlink, const TAttribute< FText >& InHyperlinkText = TAttribute< FText >() ) override
 	{
 		Hyperlink = InHyperlink;
 
diff --git a/Engine/Source/Runtime/Slate/Private/Widgets/Testing/STestSuite.cpp b/Engine/Source/Runtime/Slate/Private/Widgets/Testing/STestSuite.cpp
--- a/Engine/Source/Runtime/Slate/Private/Widgets/Testing/STestSuite.cpp
+++ b/Engine/Source/Runtime/Slate/Private/Widgets/Testing/STestSuite.cpp
@@ -56,7 +56,7 @@ public:
 		OnPaintHandler = InArgs._OnPaintHandler;
 	}
 
-	virtual FVector2D ComputeDesiredSize() const OVERRIDE
+	virtual FVector2D ComputeDesiredSize() const override
 	{
 		return FVector2D(128, 128);
 	}
diff --git a/Engine/Source/Runtime/Slate/Private/Widgets/Views/SHeaderRow.cpp b/Engine/Source/Runtime/Slate/Private/Widgets/Views/SHeaderRow.cpp
--- a/Engine/Source/Runtime/Slate/Private/Widgets/Views/SHeaderRow.cpp
+++ b/Engine/Source/Runtime/Slate/Private/Widgets/Views/SHeaderRow.cpp
@@ -182,7 +182,7 @@ public:
 		SortMode = NewMode;
 	}
 
-	virtual FReply OnMouseButtonUp( const FGeometry& MyGeometry, const FPointerEvent& MouseEvent ) OVERRIDE
+	virtual FReply OnMouseButtonUp( const FGeometry& MyGeometry, const FPointerEvent& MouseEvent ) override
 	{
 		if ( MouseEvent.GetEffectingButton() == EKeys::RightMouseButton && ContextMenuContent != SNullWidget::NullWidget )
 		{
diff --git a/Engine/Source/Runtime/Slate/Private/Widgets/Views/SListPanel.h b/Engine/Source/Runtime/Slate/Private/Widgets/Views/SListPanel.h
--- a/Engine/Source/Runtime/Slate/Private/Widgets/Views/SListPanel.h
+++ b/Engine/Source/Runtime/Slate/Private/Widgets/Views/SListPanel.h
@@ -59,10 +59,10 @@ public:
 public:
 
 	// SWidget interface
-	virtual void OnArrangeChildren( const FGeometry& AllottedGeometry, FArrangedChildren& ArrangedChildren ) const OVERRIDE;
-	virtual FVector2D ComputeDesiredSize() const OVERRIDE;
-	virtual FChildren* GetChildren() OVERRIDE;
-	virtual void Tick( const FGeometry& AllottedGeometry, const double InCurrentTime, const float InDeltaTime ) OVERRIDE;
+	virtual void OnArrangeChildren( const FGeometry& AllottedGeometry, FArrangedChildren& ArrangedChildren ) const override;
+	virtual FVector2D ComputeDesiredSize() const override;
+	virtual FChildren* GetChildren() override;
+	virtual void Tick( const FGeometry& AllottedGeometry, const double InCurrentTime, const float InDeltaTime ) override;
 	// End of SWidget interface
 
 	/** Fraction of the first item that we should offset by to simulate smooth scrolling. */
diff --git a/Engine/Source/Runtime/Slate/Public/Framework/Application/SlateApplication.h b/Engine/Source/Runtime/Slate/Public/Framework/Application/SlateApplication.h
--- a/Engine/Source/Runtime/Slate/Public/Framework/Application/SlateApplication.h
+++ b/Engine/Source/Runtime/Slate/Public/Framework/Application/SlateApplication.h
@@ -914,69 +914,69 @@ public:
 
 	// Begin FSlateApplicationBase interface
 
-	virtual TSharedRef<SWindow> AddWindow( TSharedRef<SWindow> InSlateWindow, const bool bShowImmediately = true ) OVERRIDE;
+	virtual TSharedRef<SWindow> AddWindow( TSharedRef<SWindow> InSlateWindow, const bool bShowImmediately = true ) override;
 
-	virtual void ArrangeWindowToFrontVirtual( TArray<TSharedRef<SWindow>>& Windows, const TSharedRef<SWindow>& WindowToBringToFront ) OVERRIDE
+	virtual void ArrangeWindowToFrontVirtual( TArray<TSharedRef<SWindow>>& Windows, const TSharedRef<SWindow>& WindowToBringToFront ) override
 	{
 		FSlateWindowHelper::ArrangeWindowToFront(Windows, WindowToBringToFront);
 	}
 
-	virtual bool FindPathToWidget( TSharedRef<const SWidget> InWidget, FWidgetPath& OutWidgetPath, EVisibility VisibilityFilter = EVisibility::Visible ) OVERRIDE
+	virtual bool FindPathToWidget( TSharedRef<const SWidget> InWidget, FWidgetPath& OutWidgetPath, EVisibility VisibilityFilter = EVisibility::Visible ) override
 	{
 		return FSlateWindowHelper::FindPathToWidget(GetInteractiveTopLevelWindows(), InWidget, OutWidgetPath, VisibilityFilter);
 	}
 
-	virtual const double GetCurrentTime( ) const OVERRIDE
+	virtual const double GetCurrentTime( ) const override
 	{
 		return CurrentTime;
 	}
 
-	virtual TSharedPtr<SWindow> GetActiveTopLevelWindow( ) const OVERRIDE;
+	virtual TSharedPtr<SWindow> GetActiveTopLevelWindow( ) const override;
 
-	virtual const FSlateBrush* GetAppIcon( ) const OVERRIDE;
+	virtual const FSlateBrush* GetAppIcon( ) const override;
 
-	virtual float GetApplicationScale( ) const OVERRIDE
+	virtual float GetApplicationScale( ) const override
 	{
 		return Scale;
 	}
 
-	virtual FVector2D GetCursorPos( ) const OVERRIDE;
+	virtual FVector2D GetCursorPos( ) const override;
 
-	virtual FVector2D GetCursorSize( ) const OVERRIDE;
+	virtual FVector2D GetCursorSize( ) const override;
 
-	virtual void GetDisplayMetrics( FDisplayMetrics& OutDisplayMetrics ) const OVERRIDE
+	virtual void GetDisplayMetrics( FDisplayMetrics& OutDisplayMetrics ) const override
 	{
 		PlatformApplication->GetDisplayMetrics( OutDisplayMetrics );
 	}
 
-	virtual TSharedPtr<SWidget> GetKeyboardFocusedWidget( ) const OVERRIDE;
+	virtual TSharedPtr<SWidget> GetKeyboardFocusedWidget( ) const override;
 
-	virtual TSharedPtr<SWidget> GetMouseCaptor( ) const OVERRIDE;
+	virtual TSharedPtr<SWidget> GetMouseCaptor( ) const override;
 
-	virtual FSlateRect GetPreferredWorkArea( ) const OVERRIDE;
+	virtual FSlateRect GetPreferredWorkArea( ) const override;
 
-	virtual bool HasFocusedDescendants( const TSharedRef<const SWidget>& Widget ) const OVERRIDE;
+	virtual bool HasFocusedDescendants( const TSharedRef<const SWidget>& Widget ) const override;
 
-	virtual bool IsExternalUIOpened() OVERRIDE
+	virtual bool IsExternalUIOpened() override
 	{
 		return bIsExternalUIOpened;
 	}
 
-	virtual FWidgetPath LocateWindowUnderMouse( FVector2D ScreenspaceMouseCoordinate, const TArray<TSharedRef<SWindow>>& Windows, bool bIgnoreEnabledStatus = false ) OVERRIDE;
+	virtual FWidgetPath LocateWindowUnderMouse( FVector2D ScreenspaceMouseCoordinate, const TArray<TSharedRef<SWindow>>& Windows, bool bIgnoreEnabledStatus = false ) override;
 
-	virtual TSharedRef<SWidget> MakeImage( const TAttribute<const FSlateBrush*>& Image, const TAttribute<FSlateColor>& Color, const TAttribute<EVisibility>& Visibility ) const OVERRIDE;
+	virtual TSharedRef<SWidget> MakeImage( const TAttribute<const FSlateBrush*>& Image, const TAttribute<FSlateColor>& Color, const TAttribute<EVisibility>& Visibility ) const override;
 
-	virtual TSharedRef<SWidget> MakeWindowTitleBar( const TSharedRef<SWindow>& Window, const TSharedPtr<SWidget>& CenterContent, EHorizontalAlignment CenterContentAlignment, TSharedPtr<IWindowTitleBar>& OutTitleBar ) const OVERRIDE;
+	virtual TSharedRef<SWidget> MakeWindowTitleBar( const TSharedRef<SWindow>& Window, const TSharedPtr<SWidget>& CenterContent, EHorizontalAlignment CenterContentAlignment, TSharedPtr<IWindowTitleBar>& OutTitleBar ) const override;
 
-	virtual TSharedRef<IToolTip> MakeToolTip( const TAttribute<FString>& ToolTipString ) OVERRIDE;
+	virtual TSharedRef<IToolTip> MakeToolTip( const TAttribute<FString>& ToolTipString ) override;
 
-	virtual TSharedRef<IToolTip> MakeToolTip( const TAttribute<FText>& ToolTipText ) OVERRIDE;
+	virtual TSharedRef<IToolTip> MakeToolTip( const TAttribute<FText>& ToolTipText ) override;
 
-	virtual TSharedRef<IToolTip> MakeToolTip( const FText& ToolTipText ) OVERRIDE;
+	virtual TSharedRef<IToolTip> MakeToolTip( const FText& ToolTipText ) override;
 
-	virtual void RequestDestroyWindow( TSharedRef<SWindow> WindowToDestroy ) OVERRIDE;
+	virtual void RequestDestroyWindow( TSharedRef<SWindow> WindowToDestroy ) override;
 
-	virtual bool SetKeyboardFocus( const FWidgetPath& InFocusPath, const EKeyboardFocusCause::Type InCause ) OVERRIDE;
+	virtual bool SetKeyboardFocus( const FWidgetPath& InFocusPath, const EKeyboardFocusCause::Type InCause ) override;
 
 	// End FSlateApplicationBase interface
 
@@ -984,83 +984,83 @@ public:
 
 	// Begin FGenericApplicationMessageHandler interface
 
-	virtual bool ShouldProcessUserInputMessages( const TSharedPtr< FGenericWindow >& PlatformWindow ) const OVERRIDE;
+	virtual bool ShouldProcessUserInputMessages( const TSharedPtr< FGenericWindow >& PlatformWindow ) const override;
 
-	virtual bool OnKeyChar( const TCHAR Character, const bool IsRepeat ) OVERRIDE;
+	virtual bool OnKeyChar( const TCHAR Character, const bool IsRepeat ) override;
 
-	virtual bool OnKeyDown( const int32 KeyCode, const uint32 CharacterCode, const bool IsRepeat ) OVERRIDE;
+	virtual bool OnKeyDown( const int32 KeyCode, const uint32 CharacterCode, const bool IsRepeat ) override;
 
-	virtual bool OnKeyUp( const int32 KeyCode, const uint32 CharacterCode, const bool IsRepeat ) OVERRIDE;
+	virtual bool OnKeyUp( const int32 KeyCode, const uint32 CharacterCode, const bool IsRepeat ) override;
 
-	virtual bool OnMouseDown( const TSharedPtr< FGenericWindow >& PlatformWindow, const EMouseButtons::Type Button ) OVERRIDE;
+	virtual bool OnMouseDown( const TSharedPtr< FGenericWindow >& PlatformWindow, const EMouseButtons::Type Button ) override;
 
-	virtual bool OnMouseUp( const EMouseButtons::Type Button ) OVERRIDE;
+	virtual bool OnMouseUp( const EMouseButtons::Type Button ) override;
 
-	virtual bool OnMouseDoubleClick( const TSharedPtr< FGenericWindow >& PlatformWindow, const EMouseButtons::Type Button ) OVERRIDE;
+	virtual bool OnMouseDoubleClick( const TSharedPtr< FGenericWindow >& PlatformWindow, const EMouseButtons::Type Button ) override;
 
-	virtual bool OnMouseWheel( const float Delta ) OVERRIDE;
+	virtual bool OnMouseWheel( const float Delta ) override;
 
-	virtual bool OnMouseMove() OVERRIDE;
+	virtual bool OnMouseMove() override;
 
-	virtual bool OnRawMouseMove( const int32 X, const int32 Y ) OVERRIDE;
+	virtual bool OnRawMouseMove( const int32 X, const int32 Y ) override;
 
-	virtual bool OnCursorSet() OVERRIDE;
+	virtual bool OnCursorSet() override;
 
 	virtual bool OnControllerAnalog( FKey Button, int32 ControllerId, float AnalogValue );
 
-	virtual bool OnControllerAnalog( EControllerButtons::Type Button, int32 ControllerId, float AnalogValue ) OVERRIDE;
+	virtual bool OnControllerAnalog( EControllerButtons::Type Button, int32 ControllerId, float AnalogValue ) override;
 
 	virtual bool OnControllerButtonPressed( FKey Button, int32 ControllerId, bool IsRepeat );
 
-	virtual bool OnControllerButtonPressed( EControllerButtons::Type Button, int32 ControllerId, bool IsRepeat ) OVERRIDE;
+	virtual bool OnControllerButtonPressed( EControllerButtons::Type Button, int32 ControllerId, bool IsRepeat ) override;
 
 	virtual bool OnControllerButtonReleased( FKey, int32 ControllerId, bool IsRepeat );
 
-	virtual bool OnControllerButtonReleased( EControllerButtons::Type Button, int32 ControllerId, bool IsRepeat ) OVERRIDE;
+	virtual bool OnControllerButtonReleased( EControllerButtons::Type Button, int32 ControllerId, bool IsRepeat ) override;
 
-	virtual bool OnTouchGesture( EGestureEvent::Type GestureType, const FVector2D& Delta, float WheelDelta ) OVERRIDE;
+	virtual bool OnTouchGesture( EGestureEvent::Type GestureType, const FVector2D& Delta, float WheelDelta ) override;
 	
-	virtual bool OnTouchStarted( const TSharedPtr< FGenericWindow >& PlatformWindow, const FVector2D& Location, int32 TouchIndex, int32 ControllerId ) OVERRIDE;
+	virtual bool OnTouchStarted( const TSharedPtr< FGenericWindow >& PlatformWindow, const FVector2D& Location, int32 TouchIndex, int32 ControllerId ) override;
 
-	virtual bool OnTouchMoved( const FVector2D& Location, int32 TouchIndex, int32 ControllerId ) OVERRIDE;
+	virtual bool OnTouchMoved( const FVector2D& Location, int32 TouchIndex, int32 ControllerId ) override;
 
-	virtual bool OnTouchEnded( const FVector2D& Location, int32 TouchIndex, int32 ControllerId ) OVERRIDE;
+	virtual bool OnTouchEnded( const FVector2D& Location, int32 TouchIndex, int32 ControllerId ) override;
 
-	virtual bool OnMotionDetected(const FVector& Tilt, const FVector& RotationRate, const FVector& Gravity, const FVector& Acceleration, int32 ControllerId) OVERRIDE;
+	virtual bool OnMotionDetected(const FVector& Tilt, const FVector& RotationRate, const FVector& Gravity, const FVector& Acceleration, int32 ControllerId) override;
 
-	virtual bool OnSizeChanged( const TSharedRef< FGenericWindow >& PlatformWindow, const int32 Width, const int32 Height, bool bWasMinimized = false ) OVERRIDE;
+	virtual bool OnSizeChanged( const TSharedRef< FGenericWindow >& PlatformWindow, const int32 Width, const int32 Height, bool bWasMinimized = false ) override;
 
-	virtual void OnOSPaint( const TSharedRef< FGenericWindow >& PlatformWindow ) OVERRIDE;
+	virtual void OnOSPaint( const TSharedRef< FGenericWindow >& PlatformWindow ) override;
 
-	virtual void OnResizingWindow( const TSharedRef< FGenericWindow >& PlatformWindow ) OVERRIDE;
+	virtual void OnResizingWindow( const TSharedRef< FGenericWindow >& PlatformWindow ) override;
 
-	virtual bool BeginReshapingWindow( const TSharedRef< FGenericWindow >& PlatformWindow ) OVERRIDE;
+	virtual bool BeginReshapingWindow( const TSharedRef< FGenericWindow >& PlatformWindow ) override;
 
-	virtual void FinishedReshapingWindow( const TSharedRef< FGenericWindow >& PlatformWindow ) OVERRIDE;
+	virtual void FinishedReshapingWindow( const TSharedRef< FGenericWindow >& PlatformWindow ) override;
 
-	virtual void OnMovedWindow( const TSharedRef< FGenericWindow >& PlatformWindow, const int32 X, const int32 Y ) OVERRIDE;
+	virtual void OnMovedWindow( const TSharedRef< FGenericWindow >& PlatformWindow, const int32 X, const int32 Y ) override;
 
-	virtual bool OnWindowActivationChanged( const TSharedRef< FGenericWindow >& PlatformWindow, const EWindowActivation::Type ActivationType ) OVERRIDE;
+	virtual bool OnWindowActivationChanged( const TSharedRef< FGenericWindow >& PlatformWindow, const EWindowActivation::Type ActivationType ) override;
 
-	virtual bool OnApplicationActivationChanged( const bool IsActive ) OVERRIDE;
+	virtual bool OnApplicationActivationChanged( const bool IsActive ) override;
 
-	virtual EWindowZone::Type GetWindowZoneForPoint( const TSharedRef< FGenericWindow >& PlatformWindow, const int32 X, const int32 Y ) OVERRIDE;
+	virtual EWindowZone::Type GetWindowZoneForPoint( const TSharedRef< FGenericWindow >& PlatformWindow, const int32 X, const int32 Y ) override;
 
-	virtual void OnWindowClose( const TSharedRef< FGenericWindow >& PlatformWindow ) OVERRIDE;
+	virtual void OnWindowClose( const TSharedRef< FGenericWindow >& PlatformWindow ) override;
 
-	virtual EDropEffect::Type OnDragEnterText( const TSharedRef< FGenericWindow >& Window, const FString& Text ) OVERRIDE;
+	virtual EDropEffect::Type OnDragEnterText( const TSharedRef< FGenericWindow >& Window, const FString& Text ) override;
 
-	virtual EDropEffect::Type OnDragEnterFiles( const TSharedRef< FGenericWindow >& Window, const TArray< FString >& Files ) OVERRIDE;
+	virtual EDropEffect::Type OnDragEnterFiles( const TSharedRef< FGenericWindow >& Window, const TArray< FString >& Files ) override;
 
 	EDropEffect::Type OnDragEnter( const TSharedRef< SWindow >& Window, const TSharedRef<FExternalDragOperation>& DragDropOperation );
 
-	virtual EDropEffect::Type OnDragOver( const TSharedPtr< FGenericWindow >& Window ) OVERRIDE;
+	virtual EDropEffect::Type OnDragOver( const TSharedPtr< FGenericWindow >& Window ) override;
 
-	virtual void OnDragLeave( const TSharedPtr< FGenericWindow >& Window ) OVERRIDE;
+	virtual void OnDragLeave( const TSharedPtr< FGenericWindow >& Window ) override;
 
-	virtual EDropEffect::Type OnDragDrop( const TSharedPtr< FGenericWindow >& Window ) OVERRIDE;
+	virtual EDropEffect::Type OnDragDrop( const TSharedPtr< FGenericWindow >& Window ) override;
 
-	virtual bool OnWindowAction( const TSharedRef< FGenericWindow >& PlatformWindow, const EWindowAction::Type InActionType ) OVERRIDE;
+	virtual bool OnWindowAction( const TSharedRef< FGenericWindow >& PlatformWindow, const EWindowAction::Type InActionType ) override;
 
 	// End FGenericApplicationMessageHandler interface
 
diff --git a/Engine/Source/Runtime/Slate/Public/Framework/Commands/GenericCommands.h b/Engine/Source/Runtime/Slate/Public/Framework/Commands/GenericCommands.h
--- a/Engine/Source/Runtime/Slate/Public/Framework/Commands/GenericCommands.h
+++ b/Engine/Source/Runtime/Slate/Public/Framework/Commands/GenericCommands.h
@@ -15,7 +15,7 @@ public:
 	{
 	}
 
-	virtual void RegisterCommands() OVERRIDE;	
+	virtual void RegisterCommands() override;	
 
 	TSharedPtr<FUICommandInfo> Cut;
 	TSharedPtr<FUICommandInfo> Copy;
diff --git a/Engine/Source/Runtime/Slate/Public/Framework/Commands/UICommandDragDropOp.h b/Engine/Source/Runtime/Slate/Public/Framework/Commands/UICommandDragDropOp.h
--- a/Engine/Source/Runtime/Slate/Public/Framework/Commands/UICommandDragDropOp.h
+++ b/Engine/Source/Runtime/Slate/Public/Framework/Commands/UICommandDragDropOp.h
@@ -32,9 +32,9 @@ public:
 	void SetOnDropNotification( FSimpleDelegate InOnDropNotification ) { OnDropNotification = InOnDropNotification; }
 
 	/** FDragDropOperation interface */
-	virtual TSharedPtr<SWidget> GetDefaultDecorator() const OVERRIDE;
-	virtual void OnDragged( const class FDragDropEvent& DragDropEvent ) OVERRIDE;
-	virtual void OnDrop( bool bDropWasHandled, const FPointerEvent& MouseEvent ) OVERRIDE;
+	virtual TSharedPtr<SWidget> GetDefaultDecorator() const override;
+	virtual void OnDragged( const class FDragDropEvent& DragDropEvent ) override;
+	virtual void OnDrop( bool bDropWasHandled, const FPointerEvent& MouseEvent ) override;
 
 public:
 
diff --git a/Engine/Source/Runtime/Slate/Public/Framework/Docking/TabManager.h b/Engine/Source/Runtime/Slate/Public/Framework/Docking/TabManager.h
--- a/Engine/Source/Runtime/Slate/Public/Framework/Docking/TabManager.h
+++ b/Engine/Source/Runtime/Slate/Public/Framework/Docking/TabManager.h
@@ -161,7 +161,7 @@ struct FTabSpawnerEntry : public FWorkspaceItem
 		return *this;
 	}
 
-	virtual TSharedPtr<FTabSpawnerEntry> AsSpawnerEntry() OVERRIDE
+	virtual TSharedPtr<FTabSpawnerEntry> AsSpawnerEntry() override
 	{
 		return SharedThis(this);
 	}
@@ -285,7 +285,7 @@ class SLATE_API FTabManager : public TSharedFromThis<FTabManager>
 					return SharedThis(this);
 				}
 
-				virtual TSharedPtr<FStack> AsStack() OVERRIDE
+				virtual TSharedPtr<FStack> AsStack() override
 				{
 					return SharedThis(this);
 				}
@@ -333,7 +333,7 @@ class SLATE_API FTabManager : public TSharedFromThis<FTabManager>
 					return SharedThis(this);
 				}
 
-				virtual TSharedPtr<FSplitter> AsSplitter() OVERRIDE
+				virtual TSharedPtr<FSplitter> AsSplitter() override
 				{
 					return SharedThis(this);
 				}
@@ -387,7 +387,7 @@ class SLATE_API FTabManager : public TSharedFromThis<FTabManager>
 					return SharedThis(this);
 				}
 
-				virtual TSharedPtr<FArea> AsArea() OVERRIDE
+				virtual TSharedPtr<FArea> AsArea() override
 				{
 					return SharedThis(this);
 				}
@@ -726,7 +726,7 @@ public:
 		return MakeShareable( new FGlobalTabmanager() );
 	}
 
-	virtual bool CanCloseManager( const TSet< TSharedRef<SDockTab> >& TabsToIgnore = TSet< TSharedRef<SDockTab> >()) OVERRIDE;
+	virtual bool CanCloseManager( const TSet< TSharedRef<SDockTab> >& TabsToIgnore = TSet< TSharedRef<SDockTab> >()) override;
 
 	/** Draw the user's attention to a child tab manager */
 	void DrawAttentionToTabManager( const TSharedRef<FTabManager>& ChildManager );
@@ -758,10 +758,10 @@ public:
 	int32 GetMaximumWindowCount() const { return AllAreasWindowMaxCount; }
 
 protected:
-	virtual void OnTabForegrounded( const TSharedPtr<SDockTab>& NewForegroundTab, const TSharedPtr<SDockTab>& BackgroundedTab ) OVERRIDE;
-	virtual void OnTabRelocated( const TSharedRef<SDockTab>& RelocatedTab, const TSharedPtr<SWindow>& NewOwnerWindow ) OVERRIDE;
-	virtual void OnTabClosing( const TSharedRef<SDockTab>& TabBeingClosed ) OVERRIDE;
-	virtual void UpdateStats() OVERRIDE;
+	virtual void OnTabForegrounded( const TSharedPtr<SDockTab>& NewForegroundTab, const TSharedPtr<SDockTab>& BackgroundedTab ) override;
+	virtual void OnTabRelocated( const TSharedRef<SDockTab>& RelocatedTab, const TSharedPtr<SWindow>& NewOwnerWindow ) override;
+	virtual void OnTabClosing( const TSharedRef<SDockTab>& TabBeingClosed ) override;
+	virtual void UpdateStats() override;
 
 public:
 	virtual void OnTabManagerClosing();
diff --git a/Engine/Source/Runtime/Slate/Public/Framework/MultiBox/MultiBox.h b/Engine/Source/Runtime/Slate/Public/Framework/MultiBox/MultiBox.h
--- a/Engine/Source/Runtime/Slate/Public/Framework/MultiBox/MultiBox.h
+++ b/Engine/Source/Runtime/Slate/Public/Framework/MultiBox/MultiBox.h
@@ -443,9 +443,9 @@ public:
 	virtual void BuildMultiBlockWidget(const ISlateStyle* StyleSet, const FName& StyleName) = 0;
 
 	/** SWidget Interface */
-	virtual void OnDragEnter( const FGeometry& MyGeometry, const FDragDropEvent& DragDropEvent ) OVERRIDE;
-	virtual FReply OnDragOver( const FGeometry& MyGeometry, const FDragDropEvent& DragDropEvent ) OVERRIDE;
-	virtual FReply OnDrop( const FGeometry& MyGeometry, const FDragDropEvent& DragDropEvent ) OVERRIDE;
+	virtual void OnDragEnter( const FGeometry& MyGeometry, const FDragDropEvent& DragDropEvent ) override;
+	virtual FReply OnDragOver( const FGeometry& MyGeometry, const FDragDropEvent& DragDropEvent ) override;
+	virtual FReply OnDrop( const FGeometry& MyGeometry, const FDragDropEvent& DragDropEvent ) override;
 protected:
 
 	/** Weak reference back to the MultiBox widget that owns us */
@@ -567,11 +567,11 @@ public:
 	static FReply FocusNextWidget( EFocusMoveDirection::Type MoveDirection );
 
 	/** SWidget interface */
-	virtual FReply OnDragOver( const FGeometry& MyGeometry, const FDragDropEvent& DragDropEvent ) OVERRIDE;
-	virtual FReply OnDrop( const FGeometry& MyGeometry, const FDragDropEvent& DragDropEvent ) OVERRIDE;
-	virtual bool SupportsKeyboardFocus() const OVERRIDE;
-	virtual FReply OnKeyboardFocusReceived( const FGeometry& MyGeometry, const FKeyboardFocusEvent& InKeyboardFocusEvent ) OVERRIDE;
-	virtual FReply OnKeyDown( const FGeometry& MyGeometry, const FKeyboardEvent& KeyboardEvent ) OVERRIDE;
+	virtual FReply OnDragOver( const FGeometry& MyGeometry, const FDragDropEvent& DragDropEvent ) override;
+	virtual FReply OnDrop( const FGeometry& MyGeometry, const FDragDropEvent& DragDropEvent ) override;
+	virtual bool SupportsKeyboardFocus() const override;
+	virtual FReply OnKeyboardFocusReceived( const FGeometry& MyGeometry, const FKeyboardFocusEvent& InKeyboardFocusEvent ) override;
+	virtual FReply OnKeyDown( const FGeometry& MyGeometry, const FKeyboardEvent& KeyboardEvent ) override;
 private:
 	/** Adds a block Widget to this widget */
 	void AddBlockWidget( const FMultiBlock& Block, TSharedPtr<SHorizontalBox> HorizontalBox, TSharedPtr<SVerticalBox> VerticalBox, EMultiBlockLocation::Type InLocation );
diff --git a/Engine/Source/Runtime/Slate/Public/Framework/MultiBox/MultiBoxBuilder.h b/Engine/Source/Runtime/Slate/Public/Framework/MultiBox/MultiBoxBuilder.h
--- a/Engine/Source/Runtime/Slate/Public/Framework/MultiBox/MultiBoxBuilder.h
+++ b/Engine/Source/Runtime/Slate/Public/Framework/MultiBox/MultiBoxBuilder.h
@@ -233,8 +233,8 @@ public:
 	
 protected:
 	/** FMultiBoxBuilder interface */
-	virtual void ApplyHook(FName InExtensionHook, EExtensionHook::Position HookPosition) OVERRIDE;
-	virtual void ApplySectionBeginning() OVERRIDE;
+	virtual void ApplyHook(FName InExtensionHook, EExtensionHook::Position HookPosition) override;
+	virtual void ApplySectionBeginning() override;
 
 public:
 	// These classes need access to the AddWrapperSubMenu() methods
@@ -302,7 +302,7 @@ public:
 	
 protected:
 	/** FMultiBoxBuilder interface */
-	virtual void ApplyHook(FName InExtensionHook, EExtensionHook::Position HookPosition) OVERRIDE;
+	virtual void ApplyHook(FName InExtensionHook, EExtensionHook::Position HookPosition) override;
 };
 
 
@@ -402,8 +402,8 @@ public:
 
 protected:
 	/** FMultiBoxBuilder interface */
-	virtual void ApplyHook(FName InExtensionHook, EExtensionHook::Position HookPosition) OVERRIDE;
-	virtual void ApplySectionBeginning() OVERRIDE;
+	virtual void ApplyHook(FName InExtensionHook, EExtensionHook::Position HookPosition) override;
+	virtual void ApplySectionBeginning() override;
 
 private:
 	/** Current extension hook name for sections to determine where sections begin and end */
@@ -462,5 +462,5 @@ public:
 	
 protected:
 	/** FMultiBoxBuilder interface */
-	virtual void ApplyHook(FName InExtensionHook, EExtensionHook::Position HookPosition) OVERRIDE {}
+	virtual void ApplyHook(FName InExtensionHook, EExtensionHook::Position HookPosition) override {}
 };
diff --git a/Engine/Source/Runtime/Slate/Public/Framework/MultiBox/SToolBarButtonBlock.h b/Engine/Source/Runtime/Slate/Public/Framework/MultiBox/SToolBarButtonBlock.h
--- a/Engine/Source/Runtime/Slate/Public/Framework/MultiBox/SToolBarButtonBlock.h
+++ b/Engine/Source/Runtime/Slate/Public/Framework/MultiBox/SToolBarButtonBlock.h
@@ -41,7 +41,7 @@ public:
 	void SetForceSmallIcons( const bool InForceSmallIcons ) { bForceSmallIcons = InForceSmallIcons; }
 
 	/** FMultiBlock interface */
-	virtual void CreateMenuEntry(class FMenuBuilder& MenuBuilder) const OVERRIDE;
+	virtual void CreateMenuEntry(class FMenuBuilder& MenuBuilder) const override;
 
 private:
 
@@ -106,7 +106,7 @@ public:
 	/**
 	 * Builds this MultiBlock widget up from the MultiBlock associated with it
 	 */
-	virtual void BuildMultiBlockWidget(const ISlateStyle* StyleSet, const FName& StyleName) OVERRIDE;
+	virtual void BuildMultiBlockWidget(const ISlateStyle* StyleSet, const FName& StyleName) override;
 
 
 	/**
diff --git a/Engine/Source/Runtime/Slate/Public/Framework/MultiBox/SToolBarComboButtonBlock.h b/Engine/Source/Runtime/Slate/Public/Framework/MultiBox/SToolBarComboButtonBlock.h
--- a/Engine/Source/Runtime/Slate/Public/Framework/MultiBox/SToolBarComboButtonBlock.h
+++ b/Engine/Source/Runtime/Slate/Public/Framework/MultiBox/SToolBarComboButtonBlock.h
@@ -24,7 +24,7 @@ public:
 	FToolBarComboButtonBlock( const FUIAction& InAction, const FOnGetContent& InMenuContentGenerator, const TAttribute<FText>& InLabel = TAttribute<FText>(), const TAttribute<FText>& InToolTip = TAttribute<FText>(), const TAttribute<FSlateIcon>& InIcon = TAttribute<FSlateIcon>(), bool bInSimpleComboBox = false );
 
 	/** FMultiBlock interface */
-	virtual void CreateMenuEntry(class FMenuBuilder& MenuBuilder) const OVERRIDE;
+	virtual void CreateMenuEntry(class FMenuBuilder& MenuBuilder) const override;
 
 	/** 
 	 * Sets the visibility of the blocks label
@@ -103,7 +103,7 @@ public:
 	/**
 	 * Builds this MultiBlock widget up from the MultiBlock associated with it
 	 */
-	virtual void BuildMultiBlockWidget(const ISlateStyle* StyleSet, const FName& StyleName) OVERRIDE;
+	virtual void BuildMultiBlockWidget(const ISlateStyle* StyleSet, const FName& StyleName) override;
 
 
 	/**
diff --git a/Engine/Source/Runtime/Slate/Public/Framework/Styling/ButtonWidgetStyle.h b/Engine/Source/Runtime/Slate/Public/Framework/Styling/ButtonWidgetStyle.h
--- a/Engine/Source/Runtime/Slate/Public/Framework/Styling/ButtonWidgetStyle.h
+++ b/Engine/Source/Runtime/Slate/Public/Framework/Styling/ButtonWidgetStyle.h
@@ -17,7 +17,7 @@ public:
 	UPROPERTY(Category=Appearance, EditAnywhere, meta=(ShowOnlyInnerProperties))
 	FButtonStyle ButtonStyle;
 
-	virtual const struct FSlateWidgetStyle* const GetStyle() const OVERRIDE
+	virtual const struct FSlateWidgetStyle* const GetStyle() const override
 	{
 		return static_cast< const struct FSlateWidgetStyle* >( &ButtonStyle );
 	}
diff --git a/Engine/Source/Runtime/Slate/Public/Framework/Styling/CheckBoxWidgetStyle.h b/Engine/Source/Runtime/Slate/Public/Framework/Styling/CheckBoxWidgetStyle.h
--- a/Engine/Source/Runtime/Slate/Public/Framework/Styling/CheckBoxWidgetStyle.h
+++ b/Engine/Source/Runtime/Slate/Public/Framework/Styling/CheckBoxWidgetStyle.h
@@ -16,7 +16,7 @@ public:
 	UPROPERTY(Category=Appearance, EditAnywhere, meta=(ShowOnlyInnerProperties))
 	FCheckBoxStyle CheckBoxStyle;
 
-	virtual const struct FSlateWidgetStyle* const GetStyle() const OVERRIDE
+	virtual const struct FSlateWidgetStyle* const GetStyle() const override
 	{
 		return static_cast< const struct FSlateWidgetStyle* >( &CheckBoxStyle );
 	}
diff --git a/Engine/Source/Runtime/Slate/Public/Framework/Styling/ComboBoxWidgetStyle.h b/Engine/Source/Runtime/Slate/Public/Framework/Styling/ComboBoxWidgetStyle.h
--- a/Engine/Source/Runtime/Slate/Public/Framework/Styling/ComboBoxWidgetStyle.h
+++ b/Engine/Source/Runtime/Slate/Public/Framework/Styling/ComboBoxWidgetStyle.h
@@ -17,7 +17,7 @@ public:
 	UPROPERTY(Category=Appearance, EditAnywhere, meta=(ShowOnlyInnerProperties))
 	FComboBoxStyle ComboBoxStyle;
 
-	virtual const struct FSlateWidgetStyle* const GetStyle() const OVERRIDE
+	virtual const struct FSlateWidgetStyle* const GetStyle() const override
 	{
 		return static_cast< const struct FSlateWidgetStyle* >( &ComboBoxStyle );
 	}
diff --git a/Engine/Source/Runtime/Slate/Public/Framework/Styling/ComboButtonWidgetStyle.h b/Engine/Source/Runtime/Slate/Public/Framework/Styling/ComboButtonWidgetStyle.h
--- a/Engine/Source/Runtime/Slate/Public/Framework/Styling/ComboButtonWidgetStyle.h
+++ b/Engine/Source/Runtime/Slate/Public/Framework/Styling/ComboButtonWidgetStyle.h
@@ -17,7 +17,7 @@ public:
 	UPROPERTY(Category=Appearance, EditAnywhere, meta=(ShowOnlyInnerProperties))
 	FComboButtonStyle ComboButtonStyle;
 
-	virtual const struct FSlateWidgetStyle* const GetStyle() const OVERRIDE
+	virtual const struct FSlateWidgetStyle* const GetStyle() const override
 	{
 		return static_cast< const struct FSlateWidgetStyle* >( &ComboButtonStyle );
 	}
diff --git a/Engine/Source/Runtime/Slate/Public/Framework/Styling/EditableTextWidgetStyle.h b/Engine/Source/Runtime/Slate/Public/Framework/Styling/EditableTextWidgetStyle.h
--- a/Engine/Source/Runtime/Slate/Public/Framework/Styling/EditableTextWidgetStyle.h
+++ b/Engine/Source/Runtime/Slate/Public/Framework/Styling/EditableTextWidgetStyle.h
@@ -17,7 +17,7 @@ public:
 	UPROPERTY(Category=Appearance, EditAnywhere, meta=(ShowOnlyInnerProperties))
 	FEditableTextStyle EditableTextStyle;
 
-	virtual const struct FSlateWidgetStyle* const GetStyle() const OVERRIDE
+	virtual const struct FSlateWidgetStyle* const GetStyle() const override
 	{
 		return static_cast< const struct FSlateWidgetStyle* >( &EditableTextStyle );
 	}
diff --git a/Engine/Source/Runtime/Slate/Public/Framework/Styling/TextBlockWidgetStyle.h b/Engine/Source/Runtime/Slate/Public/Framework/Styling/TextBlockWidgetStyle.h
--- a/Engine/Source/Runtime/Slate/Public/Framework/Styling/TextBlockWidgetStyle.h
+++ b/Engine/Source/Runtime/Slate/Public/Framework/Styling/TextBlockWidgetStyle.h
@@ -17,7 +17,7 @@ public:
 	UPROPERTY(Category=Appearance, EditAnywhere, meta=(ShowOnlyInnerProperties))
 	FTextBlockStyle TextBlockStyle;
 
-	virtual const struct FSlateWidgetStyle* const GetStyle() const OVERRIDE
+	virtual const struct FSlateWidgetStyle* const GetStyle() const override
 	{
 		return static_cast< const struct FSlateWidgetStyle* >( &TextBlockStyle );
 	}
diff --git a/Engine/Source/Runtime/Slate/Public/Framework/Text/DefaultLayoutBlock.h b/Engine/Source/Runtime/Slate/Public/Framework/Text/DefaultLayoutBlock.h
--- a/Engine/Source/Runtime/Slate/Public/Framework/Text/DefaultLayoutBlock.h
+++ b/Engine/Source/Runtime/Slate/Public/Framework/Text/DefaultLayoutBlock.h
@@ -12,13 +12,13 @@ public:
 
 	virtual ~FDefaultLayoutBlock() {}
 
-	virtual TSharedRef< IRun > GetRun() const OVERRIDE { return Run; }
-	virtual FTextRange GetTextRange() const OVERRIDE { return Range; }
-	virtual FVector2D GetSize() const OVERRIDE { return Size; }
-	virtual TSharedPtr< IRunHighlighter > GetHighlighter() const OVERRIDE { return Highlighter; }
+	virtual TSharedRef< IRun > GetRun() const override { return Run; }
+	virtual FTextRange GetTextRange() const override { return Range; }
+	virtual FVector2D GetSize() const override { return Size; }
+	virtual TSharedPtr< IRunHighlighter > GetHighlighter() const override { return Highlighter; }
 
-	virtual void SetLocationOffset( const FVector2D& InLocationOffset ) OVERRIDE { LocationOffset = InLocationOffset; }
-	virtual FVector2D GetLocationOffset() const OVERRIDE { return LocationOffset; }
+	virtual void SetLocationOffset( const FVector2D& InLocationOffset ) override { LocationOffset = InLocationOffset; }
+	virtual FVector2D GetLocationOffset() const override { return LocationOffset; }
 
 private:
 
diff --git a/Engine/Source/Runtime/Slate/Public/Framework/Text/GenericPlatformTextField.h b/Engine/Source/Runtime/Slate/Public/Framework/Text/GenericPlatformTextField.h
--- a/Engine/Source/Runtime/Slate/Public/Framework/Text/GenericPlatformTextField.h
+++ b/Engine/Source/Runtime/Slate/Public/Framework/Text/GenericPlatformTextField.h
@@ -7,7 +7,7 @@
 class FGenericPlatformTextField : public IPlatformTextField
 {
 public:
-	virtual void ShowKeyboard(bool bShow, TSharedPtr<SVirtualKeyboardEntry> TextEntryWidget) OVERRIDE {};
+	virtual void ShowKeyboard(bool bShow, TSharedPtr<SVirtualKeyboardEntry> TextEntryWidget) override {};
 
 };
 
diff --git a/Engine/Source/Runtime/Slate/Public/Framework/Text/IOS/IOSPlatformTextField.h b/Engine/Source/Runtime/Slate/Public/Framework/Text/IOS/IOSPlatformTextField.h
--- a/Engine/Source/Runtime/Slate/Public/Framework/Text/IOS/IOSPlatformTextField.h
+++ b/Engine/Source/Runtime/Slate/Public/Framework/Text/IOS/IOSPlatformTextField.h
@@ -14,7 +14,7 @@ public:
 	FIOSPlatformTextField();
 	virtual ~FIOSPlatformTextField();
 
-	virtual void ShowKeyboard(bool bShow, TSharedPtr<SVirtualKeyboardEntry> TextEntryWidget) OVERRIDE;
+	virtual void ShowKeyboard(bool bShow, TSharedPtr<SVirtualKeyboardEntry> TextEntryWidget) override;
 
 private:
 	SlateTextField* TextField;
diff --git a/Engine/Source/Runtime/Slate/Public/Framework/Text/SlateHyperlinkRun.h b/Engine/Source/Runtime/Slate/Public/Framework/Text/SlateHyperlinkRun.h
--- a/Engine/Source/Runtime/Slate/Public/Framework/Text/SlateHyperlinkRun.h
+++ b/Engine/Source/Runtime/Slate/Public/Framework/Text/SlateHyperlinkRun.h
@@ -36,38 +36,38 @@ public:
 
 	virtual ~FSlateHyperlinkRun() {}
 
-	virtual FTextRange GetTextRange() const OVERRIDE;
+	virtual FTextRange GetTextRange() const override;
 
-	virtual void SetTextRange( const FTextRange& Value ) OVERRIDE;
+	virtual void SetTextRange( const FTextRange& Value ) override;
 
-	virtual int16 GetBaseLine( float Scale ) const OVERRIDE;
+	virtual int16 GetBaseLine( float Scale ) const override;
 
-	virtual int16 GetMaxHeight( float Scale ) const OVERRIDE;
+	virtual int16 GetMaxHeight( float Scale ) const override;
 
-	virtual FVector2D Measure( int32 StartIndex, int32 EndIndex, float Scale ) const OVERRIDE;
+	virtual FVector2D Measure( int32 StartIndex, int32 EndIndex, float Scale ) const override;
 
-	virtual int8 GetKerning(int32 CurrentIndex, float Scale) const OVERRIDE;
+	virtual int8 GetKerning(int32 CurrentIndex, float Scale) const override;
 
-	virtual TSharedRef< ILayoutBlock > CreateBlock( int32 StartIndex, int32 EndIndex, FVector2D Size, const TSharedPtr< IRunHighlighter >& Highlighter ) OVERRIDE;
+	virtual TSharedRef< ILayoutBlock > CreateBlock( int32 StartIndex, int32 EndIndex, FVector2D Size, const TSharedPtr< IRunHighlighter >& Highlighter ) override;
 
-	virtual int32 OnPaint( const FTextLayout::FLineView& Line, const TSharedRef< ILayoutBlock >& Block, const FTextBlockStyle& DefaultStyle, const FGeometry& AllottedGeometry, const FSlateRect& MyClippingRect, FSlateWindowElementList& OutDrawElements, int32 LayerId, const FWidgetStyle& InWidgetStyle, bool bParentEnabled ) const OVERRIDE;
+	virtual int32 OnPaint( const FTextLayout::FLineView& Line, const TSharedRef< ILayoutBlock >& Block, const FTextBlockStyle& DefaultStyle, const FGeometry& AllottedGeometry, const FSlateRect& MyClippingRect, FSlateWindowElementList& OutDrawElements, int32 LayerId, const FWidgetStyle& InWidgetStyle, bool bParentEnabled ) const override;
 
-	virtual FChildren* GetChildren() OVERRIDE;
+	virtual FChildren* GetChildren() override;
 
-	virtual void ArrangeChildren( const TSharedRef< ILayoutBlock >& Block, const FGeometry& AllottedGeometry, FArrangedChildren& ArrangedChildren ) const OVERRIDE;
+	virtual void ArrangeChildren( const TSharedRef< ILayoutBlock >& Block, const FGeometry& AllottedGeometry, FArrangedChildren& ArrangedChildren ) const override;
 
-	virtual int32 GetTextIndexAt( const TSharedRef< ILayoutBlock >& Block, const FVector2D& Location, float Scale ) const OVERRIDE;
+	virtual int32 GetTextIndexAt( const TSharedRef< ILayoutBlock >& Block, const FVector2D& Location, float Scale ) const override;
 
-	virtual FVector2D GetLocationAt( const TSharedRef< ILayoutBlock >& Block, int32 Offset, float Scale ) const OVERRIDE;
+	virtual FVector2D GetLocationAt( const TSharedRef< ILayoutBlock >& Block, int32 Offset, float Scale ) const override;
 
-	virtual void BeginLayout() OVERRIDE { Children.Empty(); }
-	virtual void EndLayout() OVERRIDE {}
+	virtual void BeginLayout() override { Children.Empty(); }
+	virtual void EndLayout() override {}
 
-	virtual void Move(const TSharedRef<FString>& NewText, const FTextRange& NewRange) OVERRIDE;
-	virtual TSharedRef<IRun> Clone() const OVERRIDE;
+	virtual void Move(const TSharedRef<FString>& NewText, const FTextRange& NewRange) override;
+	virtual TSharedRef<IRun> Clone() const override;
 
-	virtual void AppendText(FString& AppendToText) const OVERRIDE;
-	virtual void AppendText(FString& AppendToText, const FTextRange& PartialRange) const OVERRIDE;
+	virtual void AppendText(FString& AppendToText) const override;
+	virtual void AppendText(FString& AppendToText, const FTextRange& PartialRange) const override;
 
 private:
 
diff --git a/Engine/Source/Runtime/Slate/Public/Framework/Text/SlateImageRun.h b/Engine/Source/Runtime/Slate/Public/Framework/Text/SlateImageRun.h
--- a/Engine/Source/Runtime/Slate/Public/Framework/Text/SlateImageRun.h
+++ b/Engine/Source/Runtime/Slate/Public/Framework/Text/SlateImageRun.h
@@ -12,37 +12,37 @@ public:
 
 	virtual ~FSlateImageRun() {}
 
-	virtual FTextRange GetTextRange() const OVERRIDE;
-	virtual void SetTextRange( const FTextRange& Value ) OVERRIDE;
+	virtual FTextRange GetTextRange() const override;
+	virtual void SetTextRange( const FTextRange& Value ) override;
 
-	virtual int16 GetBaseLine( float Scale ) const OVERRIDE;
-	virtual int16 GetMaxHeight( float Scale ) const OVERRIDE;
+	virtual int16 GetBaseLine( float Scale ) const override;
+	virtual int16 GetMaxHeight( float Scale ) const override;
 
-	virtual FVector2D Measure( int32 StartIndex, int32 EndIndex, float Scale ) const OVERRIDE;
+	virtual FVector2D Measure( int32 StartIndex, int32 EndIndex, float Scale ) const override;
 
-	virtual int8 GetKerning(int32 CurrentIndex, float Scale) const OVERRIDE;
+	virtual int8 GetKerning(int32 CurrentIndex, float Scale) const override;
 
 
-	virtual TSharedRef< ILayoutBlock > CreateBlock( int32 StartIndex, int32 EndIndex, FVector2D Size, const TSharedPtr< IRunHighlighter >& Highlighter ) OVERRIDE;
+	virtual TSharedRef< ILayoutBlock > CreateBlock( int32 StartIndex, int32 EndIndex, FVector2D Size, const TSharedPtr< IRunHighlighter >& Highlighter ) override;
 
-	virtual int32 OnPaint( const FTextLayout::FLineView& Line, const TSharedRef< ILayoutBlock >& Block, const FTextBlockStyle& DefaultStyle, const FGeometry& AllottedGeometry, const FSlateRect& MyClippingRect, FSlateWindowElementList& OutDrawElements, int32 LayerId, const FWidgetStyle& InWidgetStyle, bool bParentEnabled ) const OVERRIDE;
+	virtual int32 OnPaint( const FTextLayout::FLineView& Line, const TSharedRef< ILayoutBlock >& Block, const FTextBlockStyle& DefaultStyle, const FGeometry& AllottedGeometry, const FSlateRect& MyClippingRect, FSlateWindowElementList& OutDrawElements, int32 LayerId, const FWidgetStyle& InWidgetStyle, bool bParentEnabled ) const override;
 
-	virtual FChildren* GetChildren() OVERRIDE;
+	virtual FChildren* GetChildren() override;
 
-	virtual void ArrangeChildren( const TSharedRef< ILayoutBlock >& Block, const FGeometry& AllottedGeometry, FArrangedChildren& ArrangedChildren ) const OVERRIDE;
+	virtual void ArrangeChildren( const TSharedRef< ILayoutBlock >& Block, const FGeometry& AllottedGeometry, FArrangedChildren& ArrangedChildren ) const override;
 
-	virtual int32 GetTextIndexAt( const TSharedRef< ILayoutBlock >& Block, const FVector2D& Location, float Scale ) const OVERRIDE;
+	virtual int32 GetTextIndexAt( const TSharedRef< ILayoutBlock >& Block, const FVector2D& Location, float Scale ) const override;
 
-	virtual FVector2D GetLocationAt( const TSharedRef< ILayoutBlock >& Block, int32 Offset, float Scale ) const OVERRIDE;
+	virtual FVector2D GetLocationAt( const TSharedRef< ILayoutBlock >& Block, int32 Offset, float Scale ) const override;
 
-	virtual void BeginLayout() OVERRIDE {}
-	virtual void EndLayout() OVERRIDE {}
+	virtual void BeginLayout() override {}
+	virtual void EndLayout() override {}
 
-	virtual void Move(const TSharedRef<FString>& NewText, const FTextRange& NewRange) OVERRIDE;
-	virtual TSharedRef<IRun> Clone() const OVERRIDE;
+	virtual void Move(const TSharedRef<FString>& NewText, const FTextRange& NewRange) override;
+	virtual TSharedRef<IRun> Clone() const override;
 
-	virtual void AppendText(FString& Text) const OVERRIDE;
-	virtual void AppendText(FString& AppendToText, const FTextRange& PartialRange) const OVERRIDE;
+	virtual void AppendText(FString& Text) const override;
+	virtual void AppendText(FString& AppendToText, const FTextRange& PartialRange) const override;
 
 private:
 
diff --git a/Engine/Source/Runtime/Slate/Public/Framework/Text/SlateSimpleRunHighlighter.h b/Engine/Source/Runtime/Slate/Public/Framework/Text/SlateSimpleRunHighlighter.h
--- a/Engine/Source/Runtime/Slate/Public/Framework/Text/SlateSimpleRunHighlighter.h
+++ b/Engine/Source/Runtime/Slate/Public/Framework/Text/SlateSimpleRunHighlighter.h
@@ -9,11 +9,11 @@ public:
 
 	virtual ~FSlateSimpleRunHighlighter() {}
 
-	virtual int32 OnPaint( const FTextLayout::FLineView& Line, const TSharedRef< ISlateRun >& Run, const TSharedRef< ILayoutBlock >& Block, const FTextBlockStyle& DefaultStyle, const FGeometry& AllottedGeometry, const FSlateRect& MyClippingRect, FSlateWindowElementList& OutDrawElements, int32 LayerId, const FWidgetStyle& InWidgetStyle, bool bParentEnabled ) const OVERRIDE;
+	virtual int32 OnPaint( const FTextLayout::FLineView& Line, const TSharedRef< ISlateRun >& Run, const TSharedRef< ILayoutBlock >& Block, const FTextBlockStyle& DefaultStyle, const FGeometry& AllottedGeometry, const FSlateRect& MyClippingRect, FSlateWindowElementList& OutDrawElements, int32 LayerId, const FWidgetStyle& InWidgetStyle, bool bParentEnabled ) const override;
 
-	virtual FChildren* GetChildren() OVERRIDE;
+	virtual FChildren* GetChildren() override;
 
-	virtual void OnArrangeChildren( const TSharedRef< ILayoutBlock >& Block, const FGeometry& AllottedGeometry, FArrangedChildren& ArrangedChildren ) const OVERRIDE;
+	virtual void OnArrangeChildren( const TSharedRef< ILayoutBlock >& Block, const FGeometry& AllottedGeometry, FArrangedChildren& ArrangedChildren ) const override;
 
 private:
 
diff --git a/Engine/Source/Runtime/Slate/Public/Framework/Text/SlateTextLayout.h b/Engine/Source/Runtime/Slate/Public/Framework/Text/SlateTextLayout.h
--- a/Engine/Source/Runtime/Slate/Public/Framework/Text/SlateTextLayout.h
+++ b/Engine/Source/Runtime/Slate/Public/Framework/Text/SlateTextLayout.h
@@ -20,7 +20,7 @@ public:
 	//	LineModels.Add( LineModel );
 	//}
 
-	virtual void EndLayout() OVERRIDE;
+	virtual void EndLayout() override;
 
 
 protected:
diff --git a/Engine/Source/Runtime/Slate/Public/Framework/Text/SlateTextRun.h b/Engine/Source/Runtime/Slate/Public/Framework/Text/SlateTextRun.h
--- a/Engine/Source/Runtime/Slate/Public/Framework/Text/SlateTextRun.h
+++ b/Engine/Source/Runtime/Slate/Public/Framework/Text/SlateTextRun.h
@@ -13,34 +13,34 @@ public:
 
 	virtual ~FSlateTextRun() {}
 
-	virtual FTextRange GetTextRange() const OVERRIDE;
-	virtual void SetTextRange( const FTextRange& Value ) OVERRIDE;
+	virtual FTextRange GetTextRange() const override;
+	virtual void SetTextRange( const FTextRange& Value ) override;
 
-	virtual int16 GetBaseLine( float Scale ) const OVERRIDE;
-	virtual int16 GetMaxHeight( float Scale ) const OVERRIDE;
-	virtual FVector2D Measure( int32 StartIndex, int32 EndIndex, float Scale ) const OVERRIDE;
-	virtual int8 GetKerning(int32 CurrentIndex, float Scale) const OVERRIDE;
+	virtual int16 GetBaseLine( float Scale ) const override;
+	virtual int16 GetMaxHeight( float Scale ) const override;
+	virtual FVector2D Measure( int32 StartIndex, int32 EndIndex, float Scale ) const override;
+	virtual int8 GetKerning(int32 CurrentIndex, float Scale) const override;
 
-	virtual TSharedRef< ILayoutBlock > CreateBlock( int32 StartIndex, int32 EndIndex, FVector2D Size, const TSharedPtr< IRunHighlighter >& Highlighter ) OVERRIDE;
+	virtual TSharedRef< ILayoutBlock > CreateBlock( int32 StartIndex, int32 EndIndex, FVector2D Size, const TSharedPtr< IRunHighlighter >& Highlighter ) override;
 
-	virtual int32 OnPaint( const FTextLayout::FLineView& Line, const TSharedRef< ILayoutBlock >& Block, const FTextBlockStyle& DefaultStyle, const FGeometry& AllottedGeometry, const FSlateRect& MyClippingRect, FSlateWindowElementList& OutDrawElements, int32 LayerId, const FWidgetStyle& InWidgetStyle, bool bParentEnabled ) const OVERRIDE;
+	virtual int32 OnPaint( const FTextLayout::FLineView& Line, const TSharedRef< ILayoutBlock >& Block, const FTextBlockStyle& DefaultStyle, const FGeometry& AllottedGeometry, const FSlateRect& MyClippingRect, FSlateWindowElementList& OutDrawElements, int32 LayerId, const FWidgetStyle& InWidgetStyle, bool bParentEnabled ) const override;
 
-	virtual FChildren* GetChildren() OVERRIDE;
+	virtual FChildren* GetChildren() override;
 
-	virtual void ArrangeChildren( const TSharedRef< ILayoutBlock >& Block, const FGeometry& AllottedGeometry, FArrangedChildren& ArrangedChildren ) const OVERRIDE;
+	virtual void ArrangeChildren( const TSharedRef< ILayoutBlock >& Block, const FGeometry& AllottedGeometry, FArrangedChildren& ArrangedChildren ) const override;
 
-	virtual void BeginLayout() OVERRIDE {}
-	virtual void EndLayout() OVERRIDE {}
+	virtual void BeginLayout() override {}
+	virtual void EndLayout() override {}
 
-	virtual int32 GetTextIndexAt( const TSharedRef< ILayoutBlock >& Block, const FVector2D& Location, float Scale ) const OVERRIDE;
+	virtual int32 GetTextIndexAt( const TSharedRef< ILayoutBlock >& Block, const FVector2D& Location, float Scale ) const override;
 
-	virtual FVector2D GetLocationAt(const TSharedRef< ILayoutBlock >& Block, int32 Offset, float Scale) const OVERRIDE;
+	virtual FVector2D GetLocationAt(const TSharedRef< ILayoutBlock >& Block, int32 Offset, float Scale) const override;
 
-	virtual void Move(const TSharedRef<FString>& NewText, const FTextRange& NewRange) OVERRIDE;
-	virtual TSharedRef<IRun> Clone() const OVERRIDE;
+	virtual void Move(const TSharedRef<FString>& NewText, const FTextRange& NewRange) override;
+	virtual TSharedRef<IRun> Clone() const override;
 
-	virtual void AppendText(FString& Text) const OVERRIDE;
-	virtual void AppendText(FString& AppendToText, const FTextRange& PartialRange) const OVERRIDE;
+	virtual void AppendText(FString& Text) const override;
+	virtual void AppendText(FString& AppendToText, const FTextRange& PartialRange) const override;
 
 private:
 
diff --git a/Engine/Source/Runtime/Slate/Public/Framework/Text/SlateWidgetRun.h b/Engine/Source/Runtime/Slate/Public/Framework/Text/SlateWidgetRun.h
--- a/Engine/Source/Runtime/Slate/Public/Framework/Text/SlateWidgetRun.h
+++ b/Engine/Source/Runtime/Slate/Public/Framework/Text/SlateWidgetRun.h
@@ -34,35 +34,35 @@ public:
 
 	virtual ~FSlateWidgetRun() {}
 
-	virtual FTextRange GetTextRange() const OVERRIDE;
-	virtual void SetTextRange( const FTextRange& Value ) OVERRIDE;
+	virtual FTextRange GetTextRange() const override;
+	virtual void SetTextRange( const FTextRange& Value ) override;
 
-	virtual int16 GetBaseLine( float Scale ) const OVERRIDE;
-	virtual int16 GetMaxHeight( float Scale ) const OVERRIDE;
+	virtual int16 GetBaseLine( float Scale ) const override;
+	virtual int16 GetMaxHeight( float Scale ) const override;
 
-	virtual FVector2D Measure( int32 StartIndex, int32 EndIndex, float Scale ) const OVERRIDE;
-	virtual int8 GetKerning(int32 CurrentIndex, float Scale) const OVERRIDE;
+	virtual FVector2D Measure( int32 StartIndex, int32 EndIndex, float Scale ) const override;
+	virtual int8 GetKerning(int32 CurrentIndex, float Scale) const override;
 
-	virtual TSharedRef< ILayoutBlock > CreateBlock( int32 StartIndex, int32 EndIndex, FVector2D Size, const TSharedPtr< IRunHighlighter >& Highlighter ) OVERRIDE;
+	virtual TSharedRef< ILayoutBlock > CreateBlock( int32 StartIndex, int32 EndIndex, FVector2D Size, const TSharedPtr< IRunHighlighter >& Highlighter ) override;
 
-	virtual int32 OnPaint( const FTextLayout::FLineView& Line, const TSharedRef< ILayoutBlock >& Block, const FTextBlockStyle& DefaultStyle, const FGeometry& AllottedGeometry, const FSlateRect& MyClippingRect, FSlateWindowElementList& OutDrawElements, int32 LayerId, const FWidgetStyle& InWidgetStyle, bool bParentEnabled ) const OVERRIDE;
+	virtual int32 OnPaint( const FTextLayout::FLineView& Line, const TSharedRef< ILayoutBlock >& Block, const FTextBlockStyle& DefaultStyle, const FGeometry& AllottedGeometry, const FSlateRect& MyClippingRect, FSlateWindowElementList& OutDrawElements, int32 LayerId, const FWidgetStyle& InWidgetStyle, bool bParentEnabled ) const override;
 
-	virtual FChildren* GetChildren() OVERRIDE;
+	virtual FChildren* GetChildren() override;
 
 	virtual void ArrangeChildren( const TSharedRef< ILayoutBlock >& Block, const FGeometry& AllottedGeometry, FArrangedChildren& ArrangedChildren ) const;
 
-	virtual int32 GetTextIndexAt( const TSharedRef< ILayoutBlock >& Block, const FVector2D& Location, float Scale ) const OVERRIDE;
+	virtual int32 GetTextIndexAt( const TSharedRef< ILayoutBlock >& Block, const FVector2D& Location, float Scale ) const override;
 
-	virtual FVector2D GetLocationAt( const TSharedRef< ILayoutBlock >& Block, int32 Offset, float Scale ) const OVERRIDE;
+	virtual FVector2D GetLocationAt( const TSharedRef< ILayoutBlock >& Block, int32 Offset, float Scale ) const override;
 
-	virtual void BeginLayout() OVERRIDE {}
-	virtual void EndLayout() OVERRIDE {}
+	virtual void BeginLayout() override {}
+	virtual void EndLayout() override {}
 
-	virtual void Move(const TSharedRef<FString>& NewText, const FTextRange& NewRange) OVERRIDE;
-	virtual TSharedRef<IRun> Clone() const OVERRIDE;
+	virtual void Move(const TSharedRef<FString>& NewText, const FTextRange& NewRange) override;
+	virtual TSharedRef<IRun> Clone() const override;
 
-	virtual void AppendText(FString& Text) const OVERRIDE;
-	virtual void AppendText(FString& AppendToText, const FTextRange& PartialRange) const OVERRIDE;
+	virtual void AppendText(FString& Text) const override;
+	virtual void AppendText(FString& AppendToText, const FTextRange& PartialRange) const override;
 
 private:
 
diff --git a/Engine/Source/Runtime/Slate/Public/Framework/Text/TextDecorators.h b/Engine/Source/Runtime/Slate/Public/Framework/Text/TextDecorators.h
--- a/Engine/Source/Runtime/Slate/Public/Framework/Text/TextDecorators.h
+++ b/Engine/Source/Runtime/Slate/Public/Framework/Text/TextDecorators.h
@@ -16,9 +16,9 @@ public:
 
 public:
 
-	virtual bool Supports( const FTextRunParseResults& RunInfo, const FString& Text ) const OVERRIDE;
+	virtual bool Supports( const FTextRunParseResults& RunInfo, const FString& Text ) const override;
 
-	virtual TSharedRef< ISlateRun > Create( const FTextRunParseResults& RunInfo, const FString& OriginalText, const TSharedRef< FString >& InOutModelText, const ISlateStyle* Style ) OVERRIDE;
+	virtual TSharedRef< ISlateRun > Create( const FTextRunParseResults& RunInfo, const FString& OriginalText, const TSharedRef< FString >& InOutModelText, const ISlateStyle* Style ) override;
 
 private:
 
@@ -39,9 +39,9 @@ public:
 
 public:
 
-	virtual bool Supports( const FTextRunParseResults& RunInfo, const FString& Text ) const OVERRIDE;
+	virtual bool Supports( const FTextRunParseResults& RunInfo, const FString& Text ) const override;
 
-	virtual TSharedRef< ISlateRun > Create( const FTextRunParseResults& RunInfo, const FString& OriginalText, const TSharedRef< FString >& InOutModelText, const ISlateStyle* Style ) OVERRIDE;
+	virtual TSharedRef< ISlateRun > Create( const FTextRunParseResults& RunInfo, const FString& OriginalText, const TSharedRef< FString >& InOutModelText, const ISlateStyle* Style ) override;
 
 private:
 
@@ -62,9 +62,9 @@ public:
 
 public:
 
-	virtual bool Supports( const FTextRunParseResults& RunInfo, const FString& Text ) const OVERRIDE;
+	virtual bool Supports( const FTextRunParseResults& RunInfo, const FString& Text ) const override;
 
-	virtual TSharedRef< ISlateRun > Create( const FTextRunParseResults& RunInfo, const FString& OriginalText, const TSharedRef< FString >& InOutModelText, const ISlateStyle* Style ) OVERRIDE;
+	virtual TSharedRef< ISlateRun > Create( const FTextRunParseResults& RunInfo, const FString& OriginalText, const TSharedRef< FString >& InOutModelText, const ISlateStyle* Style ) override;
 
 private:
 
diff --git a/Engine/Source/Runtime/Slate/Public/Framework/Text/WidgetLayoutBlock.h b/Engine/Source/Runtime/Slate/Public/Framework/Text/WidgetLayoutBlock.h
--- a/Engine/Source/Runtime/Slate/Public/Framework/Text/WidgetLayoutBlock.h
+++ b/Engine/Source/Runtime/Slate/Public/Framework/Text/WidgetLayoutBlock.h
@@ -12,18 +12,18 @@ public:
 
 	virtual ~FWidgetLayoutBlock() {}
 
-	virtual TSharedRef< IRun > GetRun() const OVERRIDE { return Run; }
+	virtual TSharedRef< IRun > GetRun() const override { return Run; }
 
-	virtual FTextRange GetTextRange() const OVERRIDE { return Range; }
+	virtual FTextRange GetTextRange() const override { return Range; }
 
-	virtual FVector2D GetSize() const OVERRIDE { return Size; }
+	virtual FVector2D GetSize() const override { return Size; }
 
-	virtual TSharedPtr< IRunHighlighter > GetHighlighter() const OVERRIDE { return Highlighter; }
+	virtual TSharedPtr< IRunHighlighter > GetHighlighter() const override { return Highlighter; }
 
 	TSharedRef< SWidget > GetWidget() const { return Widget; }
 
-	virtual void SetLocationOffset( const FVector2D& InLocationOffset ) OVERRIDE { LocationOffset = InLocationOffset; }
-	virtual FVector2D GetLocationOffset() const OVERRIDE { return LocationOffset; }
+	virtual void SetLocationOffset( const FVector2D& InLocationOffset ) override { LocationOffset = InLocationOffset; }
+	virtual FVector2D GetLocationOffset() const override { return LocationOffset; }
 
 private:
 
diff --git a/Engine/Source/Runtime/Slate/Public/Widgets/Colors/SColorBlock.h b/Engine/Source/Runtime/Slate/Public/Widgets/Colors/SColorBlock.h
--- a/Engine/Source/Runtime/Slate/Public/Widgets/Colors/SColorBlock.h
+++ b/Engine/Source/Runtime/Slate/Public/Widgets/Colors/SColorBlock.h
@@ -48,9 +48,9 @@ public:
 
 private:
 	// SWidget interface
-	virtual int32 OnPaint( const FGeometry& AllottedGeometry, const FSlateRect& MyClippingRect, FSlateWindowElementList& OutDrawElements, int32 LayerId, const FWidgetStyle& InWidgetStyle, bool bParentEnabled ) const OVERRIDE;
-	virtual FReply OnMouseButtonDown( const FGeometry& MyGeometry, const FPointerEvent& MouseEvent ) OVERRIDE;
-	virtual FVector2D ComputeDesiredSize() const OVERRIDE;
+	virtual int32 OnPaint( const FGeometry& AllottedGeometry, const FSlateRect& MyClippingRect, FSlateWindowElementList& OutDrawElements, int32 LayerId, const FWidgetStyle& InWidgetStyle, bool bParentEnabled ) const override;
+	virtual FReply OnMouseButtonDown( const FGeometry& MyGeometry, const FPointerEvent& MouseEvent ) override;
+	virtual FVector2D ComputeDesiredSize() const override;
 	// End of SWidget interface
 
 private:
diff --git a/Engine/Source/Runtime/Slate/Public/Widgets/Colors/SColorPicker.h b/Engine/Source/Runtime/Slate/Public/Widgets/Colors/SColorPicker.h
--- a/Engine/Source/Runtime/Slate/Public/Widgets/Colors/SColorPicker.h
+++ b/Engine/Source/Runtime/Slate/Public/Widgets/Colors/SColorPicker.h
@@ -159,7 +159,7 @@ public:
 protected:	
 
 	// SWidget interface
-	virtual void Tick( const FGeometry& AllottedGeometry, const double InCurrentTime, const float InDeltaTime ) OVERRIDE;
+	virtual void Tick( const FGeometry& AllottedGeometry, const double InCurrentTime, const float InDeltaTime ) override;
 	// End of SWidget interface
 
 	/** Backup all the colors that are being modified */
diff --git a/Engine/Source/Runtime/Slate/Public/Widgets/Colors/SColorSpectrum.h b/Engine/Source/Runtime/Slate/Public/Widgets/Colors/SColorSpectrum.h
--- a/Engine/Source/Runtime/Slate/Public/Widgets/Colors/SColorSpectrum.h
+++ b/Engine/Source/Runtime/Slate/Public/Widgets/Colors/SColorSpectrum.h
@@ -62,7 +62,7 @@ public:
 
 	virtual FReply OnMouseMove( const FGeometry& MyGeometry, const FPointerEvent& MouseEvent );
 
-	virtual int32 OnPaint( const FGeometry& AllottedGeometry, const FSlateRect& MyClippingRect, FSlateWindowElementList& OutDrawElements, int32 LayerId, const FWidgetStyle& InWidgetStyle, bool bParentEnabled ) const OVERRIDE;
+	virtual int32 OnPaint( const FGeometry& AllottedGeometry, const FSlateRect& MyClippingRect, FSlateWindowElementList& OutDrawElements, int32 LayerId, const FWidgetStyle& InWidgetStyle, bool bParentEnabled ) const override;
 
 	// End overrides
 	
diff --git a/Engine/Source/Runtime/Slate/Public/Widgets/Colors/SColorThemes.h b/Engine/Source/Runtime/Slate/Public/Widgets/Colors/SColorThemes.h
--- a/Engine/Source/Runtime/Slate/Public/Widgets/Colors/SColorThemes.h
+++ b/Engine/Source/Runtime/Slate/Public/Widgets/Colors/SColorThemes.h
@@ -69,7 +69,7 @@ protected:
 	 *
 	 * @return A reply that indicated whether this event was handled.
 	 */
-	virtual void OnDragEnter( const FGeometry& MyGeometry, const FDragDropEvent& DragDropEvent ) OVERRIDE;
+	virtual void OnDragEnter( const FGeometry& MyGeometry, const FDragDropEvent& DragDropEvent ) override;
 
 	/**
 	 * Called during drag and drop when the drag leaves a widget.
@@ -78,7 +78,7 @@ protected:
 	 *
 	 * @return A reply that indicated whether this event was handled.
 	 */
-	virtual void OnDragLeave( const FDragDropEvent& DragDropEvent ) OVERRIDE;
+	virtual void OnDragLeave( const FDragDropEvent& DragDropEvent ) override;
 
 	/**
 	 * Called when the user is dropping something onto a widget; terminates drag and drop.
@@ -88,7 +88,7 @@ protected:
 	 *
 	 * @return A reply that indicated whether this event was handled.
 	 */
-	virtual FReply OnDrop( const FGeometry& MyGeometry, const FDragDropEvent& DragDropEvent ) OVERRIDE;
+	virtual FReply OnDrop( const FGeometry& MyGeometry, const FDragDropEvent& DragDropEvent ) override;
 	
 	const FSlateBrush* GetBorderStyle() const;
 
@@ -149,9 +149,9 @@ public:
 	void Construct(const FArguments& InArgs );
 	
 private:
-	virtual FReply OnMouseButtonDown( const FGeometry& MyGeometry, const FPointerEvent& MouseEvent ) OVERRIDE;
-	virtual FReply OnMouseButtonUp( const FGeometry& MyGeometry, const FPointerEvent& MouseEvent ) OVERRIDE;
-	virtual FReply OnDragDetected( const FGeometry& MyGeometry, const FPointerEvent& MouseEvent ) OVERRIDE;
+	virtual FReply OnMouseButtonDown( const FGeometry& MyGeometry, const FPointerEvent& MouseEvent ) override;
+	virtual FReply OnMouseButtonUp( const FGeometry& MyGeometry, const FPointerEvent& MouseEvent ) override;
+	virtual FReply OnDragDetected( const FGeometry& MyGeometry, const FPointerEvent& MouseEvent ) override;
 
 	
 	FLinearColor GetColor() const;
@@ -251,7 +251,7 @@ public:
 	 * @param AllottedGeometry    The geometry allotted for this widget by its parent.
 	 * @param ArrangedChildren    The array to which to add the WidgetGeometries that represent the arranged children.
 	 */
-	virtual void OnArrangeChildren( const FGeometry& AllottedGeometry, FArrangedChildren& ArrangedChildren ) const OVERRIDE;
+	virtual void OnArrangeChildren( const FGeometry& AllottedGeometry, FArrangedChildren& ArrangedChildren ) const override;
 
 	/**
 	 * A Panel's desired size in the space required to arrange of its children on the screen while respecting all of
@@ -259,7 +259,7 @@ public:
 	 *
 	 * @return The desired size.
 	 */
-	virtual FVector2D ComputeDesiredSize() const OVERRIDE;
+	virtual FVector2D ComputeDesiredSize() const override;
 
 	/**
 	 * All widgets must provide a way to access their children in a layout-agnostic way.
@@ -268,7 +268,7 @@ public:
 	 * provides layout information about the child it stores. In that case
 	 * GetChildren should simply return the TPanelChildren<Slot>.
 	 */
-	virtual FChildren* GetChildren() OVERRIDE;
+	virtual FChildren* GetChildren() override;
 	
 	/**
 	 * Called during drag and drop when the drag enters a widget.
@@ -278,7 +278,7 @@ public:
 	 *
 	 * @return A reply that indicated whether this event was handled.
 	 */
-	virtual void OnDragEnter( const FGeometry& MyGeometry, const FDragDropEvent& DragDropEvent ) OVERRIDE;
+	virtual void OnDragEnter( const FGeometry& MyGeometry, const FDragDropEvent& DragDropEvent ) override;
 	
 	/**
 	 * Called during drag and drop when the drag leaves a widget.
@@ -287,7 +287,7 @@ public:
 	 *
 	 * @return A reply that indicated whether this event was handled.
 	 */
-	virtual void OnDragLeave( const FDragDropEvent& DragDropEvent ) OVERRIDE;
+	virtual void OnDragLeave( const FDragDropEvent& DragDropEvent ) override;
 
 	/**
 	 * Called during drag and drop when the the mouse is being dragged over a widget.
@@ -297,7 +297,7 @@ public:
 	 *
 	 * @return A reply that indicated whether this event was handled.
 	 */
-	virtual FReply OnDragOver( const FGeometry& MyGeometry, const FDragDropEvent& DragDropEvent ) OVERRIDE;
+	virtual FReply OnDragOver( const FGeometry& MyGeometry, const FDragDropEvent& DragDropEvent ) override;
 
 	/**
 	 * Called when the user is dropping something onto a widget; terminates drag and drop.
@@ -307,7 +307,7 @@ public:
 	 *
 	 * @return A reply that indicated whether this event was handled.
 	 */
-	virtual FReply OnDrop( const FGeometry& MyGeometry, const FDragDropEvent& DragDropEvent ) OVERRIDE;
+	virtual FReply OnDrop( const FGeometry& MyGeometry, const FDragDropEvent& DragDropEvent ) override;
 
 	/**
 	 * Adds a new color block to the Bar
@@ -589,17 +589,17 @@ public:
 	 * @param bDropWasHandled   true when the drop was handled by some widget; false otherwise
 	 * @param MouseEvent        The mouse event which caused the on drop to be called.
 	 */
-	virtual void OnDrop( bool bDropWasHandled, const FPointerEvent& MouseEvent ) OVERRIDE;
+	virtual void OnDrop( bool bDropWasHandled, const FPointerEvent& MouseEvent ) override;
 
 	/** 
 	 * Called when the mouse was moved during a drag and drop operation
 	 *
 	 * @param DragDropEvent    The event that describes this drag drop operation.
 	 */
-	virtual void OnDragged( const class FDragDropEvent& DragDropEvent ) OVERRIDE;
+	virtual void OnDragged( const class FDragDropEvent& DragDropEvent ) override;
 	
 	/** Gets the widget that will serve as the decorator unless overriden. If you do not override, you will have no decorator */
-	virtual TSharedPtr<SWidget> GetDefaultDecorator() const OVERRIDE;
+	virtual TSharedPtr<SWidget> GetDefaultDecorator() const override;
 	
 	/**
 	 * Makes a new FColorDragDrop to hold on to
diff --git a/Engine/Source/Runtime/Slate/Public/Widgets/Colors/SColorWheel.h b/Engine/Source/Runtime/Slate/Public/Widgets/Colors/SColorWheel.h
--- a/Engine/Source/Runtime/Slate/Public/Widgets/Colors/SColorWheel.h
+++ b/Engine/Source/Runtime/Slate/Public/Widgets/Colors/SColorWheel.h
@@ -62,7 +62,7 @@ public:
 
 	virtual FReply OnMouseMove( const FGeometry& MyGeometry, const FPointerEvent& MouseEvent );
 
-	virtual int32 OnPaint( const FGeometry& AllottedGeometry, const FSlateRect& MyClippingRect, FSlateWindowElementList& OutDrawElements, int32 LayerId, const FWidgetStyle& InWidgetStyle, bool bParentEnabled ) const OVERRIDE;
+	virtual int32 OnPaint( const FGeometry& AllottedGeometry, const FSlateRect& MyClippingRect, FSlateWindowElementList& OutDrawElements, int32 LayerId, const FWidgetStyle& InWidgetStyle, bool bParentEnabled ) const override;
 
 	// End SWidget overrides
 	
diff --git a/Engine/Source/Runtime/Slate/Public/Widgets/Docking/SDockTab.h b/Engine/Source/Runtime/Slate/Public/Widgets/Docking/SDockTab.h
--- a/Engine/Source/Runtime/Slate/Public/Widgets/Docking/SDockTab.h
+++ b/Engine/Source/Runtime/Slate/Public/Widgets/Docking/SDockTab.h
@@ -84,25 +84,25 @@ public:
 	/** Construct the widget from the declaration. */
 	void Construct( const FArguments& InArgs );
 
-	virtual void Tick( const FGeometry& AllottedGeometry, const double InCurrentTime, const float InDeltaTime ) OVERRIDE;
+	virtual void Tick( const FGeometry& AllottedGeometry, const double InCurrentTime, const float InDeltaTime ) override;
 
-	virtual FReply OnMouseButtonDown( const FGeometry& MyGeometry, const FPointerEvent& MouseEvent ) OVERRIDE;
+	virtual FReply OnMouseButtonDown( const FGeometry& MyGeometry, const FPointerEvent& MouseEvent ) override;
 
-	virtual FReply OnMouseButtonDoubleClick( const FGeometry& MyGeometry, const FPointerEvent& MouseEvent ) OVERRIDE;
+	virtual FReply OnMouseButtonDoubleClick( const FGeometry& MyGeometry, const FPointerEvent& MouseEvent ) override;
 
-	virtual FReply OnDragDetected( const FGeometry& MyGeometry, const FPointerEvent& MouseEvent  ) OVERRIDE;
+	virtual FReply OnDragDetected( const FGeometry& MyGeometry, const FPointerEvent& MouseEvent  ) override;
 	
-	virtual FReply OnMouseButtonUp( const FGeometry& MyGeometry, const FPointerEvent& MouseEvent ) OVERRIDE;
+	virtual FReply OnMouseButtonUp( const FGeometry& MyGeometry, const FPointerEvent& MouseEvent ) override;
 
-	virtual void OnDragEnter( const FGeometry& MyGeometry, const FDragDropEvent& DragDropEvent ) OVERRIDE;
+	virtual void OnDragEnter( const FGeometry& MyGeometry, const FDragDropEvent& DragDropEvent ) override;
 
-	virtual void OnDragLeave( const FDragDropEvent& DragDropEvent ) OVERRIDE;
+	virtual void OnDragLeave( const FDragDropEvent& DragDropEvent ) override;
 
-	virtual FReply OnDrop( const FGeometry& MyGeometry, const FDragDropEvent& DragDropEvent ) OVERRIDE;
+	virtual FReply OnDrop( const FGeometry& MyGeometry, const FDragDropEvent& DragDropEvent ) override;
 
-	virtual FReply OnTouchStarted( const FGeometry& MyGeometry, const FPointerEvent& InTouchEvent ) OVERRIDE;
+	virtual FReply OnTouchStarted( const FGeometry& MyGeometry, const FPointerEvent& InTouchEvent ) override;
 
-	virtual FReply OnTouchEnded( const FGeometry& MyGeometry, const FPointerEvent& InTouchEvent ) OVERRIDE;
+	virtual FReply OnTouchEnded( const FGeometry& MyGeometry, const FPointerEvent& InTouchEvent ) override;
 
 	void SetContent( TSharedRef<SWidget> InContent );
 	void SetLeftContent( TSharedRef<SWidget> InContent );
diff --git a/Engine/Source/Runtime/Slate/Public/Widgets/Images/SImage.h b/Engine/Source/Runtime/Slate/Public/Widgets/Images/SImage.h
--- a/Engine/Source/Runtime/Slate/Public/Widgets/Images/SImage.h
+++ b/Engine/Source/Runtime/Slate/Public/Widgets/Images/SImage.h
@@ -32,9 +32,9 @@ public:
 	void Construct( const FArguments& InArgs );
 
 	// SWidget interface
-	virtual int32 OnPaint( const FGeometry& AllottedGeometry, const FSlateRect& MyClippingRect, FSlateWindowElementList& OutDrawElements, int32 LayerId, const FWidgetStyle& InWidgetStyle, bool bParentEnabled ) const OVERRIDE;
-	virtual FReply OnMouseButtonDown( const FGeometry& MyGeometry, const FPointerEvent& MouseEvent ) OVERRIDE;
-	virtual FVector2D ComputeDesiredSize() const OVERRIDE;
+	virtual int32 OnPaint( const FGeometry& AllottedGeometry, const FSlateRect& MyClippingRect, FSlateWindowElementList& OutDrawElements, int32 LayerId, const FWidgetStyle& InWidgetStyle, bool bParentEnabled ) const override;
+	virtual FReply OnMouseButtonDown( const FGeometry& MyGeometry, const FPointerEvent& MouseEvent ) override;
+	virtual FVector2D ComputeDesiredSize() const override;
 	// End of SWidget interface
 
 public:
diff --git a/Engine/Source/Runtime/Slate/Public/Widgets/Images/SSpinningImage.h b/Engine/Source/Runtime/Slate/Public/Widgets/Images/SSpinningImage.h
--- a/Engine/Source/Runtime/Slate/Public/Widgets/Images/SSpinningImage.h
+++ b/Engine/Source/Runtime/Slate/Public/Widgets/Images/SSpinningImage.h
@@ -32,7 +32,7 @@ public:
 	void Construct( const FArguments& InArgs );
 
 	/* Override SImage's OnPaint to draw rotated */
-	virtual int32 OnPaint( const FGeometry& AllottedGeometry, const FSlateRect& MyClippingRect, FSlateWindowElementList& OutDrawElements, int32 LayerId, const FWidgetStyle& InWidgetStyle, bool bParentEnabled ) const OVERRIDE;
+	virtual int32 OnPaint( const FGeometry& AllottedGeometry, const FSlateRect& MyClippingRect, FSlateWindowElementList& OutDrawElements, int32 LayerId, const FWidgetStyle& InWidgetStyle, bool bParentEnabled ) const override;
 
 private:
 	/** The sequence to drive the spinning animation */
diff --git a/Engine/Source/Runtime/Slate/Public/Widgets/Images/SThrobber.h b/Engine/Source/Runtime/Slate/Public/Widgets/Images/SThrobber.h
--- a/Engine/Source/Runtime/Slate/Public/Widgets/Images/SThrobber.h
+++ b/Engine/Source/Runtime/Slate/Public/Widgets/Images/SThrobber.h
@@ -71,8 +71,8 @@ public:
 	/** Constructs the widget */
 	void Construct(const FArguments& InArgs);
 
-	virtual int32 OnPaint( const FGeometry& AllottedGeometry, const FSlateRect& MyClippingRect, FSlateWindowElementList& OutDrawElements, int32 LayerId, const FWidgetStyle& InWidgetStyle, bool bParentEnabled ) const OVERRIDE;
-	virtual FVector2D ComputeDesiredSize() const OVERRIDE;
+	virtual int32 OnPaint( const FGeometry& AllottedGeometry, const FSlateRect& MyClippingRect, FSlateWindowElementList& OutDrawElements, int32 LayerId, const FWidgetStyle& InWidgetStyle, bool bParentEnabled ) const override;
+	virtual FVector2D ComputeDesiredSize() const override;
 
 private:
 	/** The sequence to drive the spinning animation */
diff --git a/Engine/Source/Runtime/Slate/Public/Widgets/Input/SButton.h b/Engine/Source/Runtime/Slate/Public/Widgets/Input/SButton.h
--- a/Engine/Source/Runtime/Slate/Public/Widgets/Input/SButton.h
+++ b/Engine/Source/Runtime/Slate/Public/Widgets/Input/SButton.h
@@ -81,21 +81,21 @@ public:
 	 *
 	 * @return  True if this widget can take keyboard focus
 	 */
-	virtual bool SupportsKeyboardFocus() const OVERRIDE;
+	virtual bool SupportsKeyboardFocus() const override;
 
-	virtual FReply OnKeyDown( const FGeometry& MyGeometry, const FKeyboardEvent& InKeyboardEvent ) OVERRIDE;
+	virtual FReply OnKeyDown( const FGeometry& MyGeometry, const FKeyboardEvent& InKeyboardEvent ) override;
 
-	virtual FReply OnMouseButtonDown( const FGeometry& MyGeometry, const FPointerEvent& MouseEvent ) OVERRIDE;
+	virtual FReply OnMouseButtonDown( const FGeometry& MyGeometry, const FPointerEvent& MouseEvent ) override;
 
-	virtual FReply OnMouseButtonDoubleClick( const FGeometry& InMyGeometry, const FPointerEvent& InMouseEvent ) OVERRIDE;
+	virtual FReply OnMouseButtonDoubleClick( const FGeometry& InMyGeometry, const FPointerEvent& InMouseEvent ) override;
 	
-	virtual FReply OnMouseButtonUp( const FGeometry& MyGeometry, const FPointerEvent& MouseEvent ) OVERRIDE;
+	virtual FReply OnMouseButtonUp( const FGeometry& MyGeometry, const FPointerEvent& MouseEvent ) override;
 	
-	virtual FReply OnMouseMove( const FGeometry& MyGeometry, const FPointerEvent& MouseEvent ) OVERRIDE;
+	virtual FReply OnMouseMove( const FGeometry& MyGeometry, const FPointerEvent& MouseEvent ) override;
 	
-	virtual void OnMouseEnter( const FGeometry& MyGeometry, const FPointerEvent& MouseEvent ) OVERRIDE;
+	virtual void OnMouseEnter( const FGeometry& MyGeometry, const FPointerEvent& MouseEvent ) override;
 
-	virtual void OnMouseLeave( const FPointerEvent& MouseEvent ) OVERRIDE;
+	virtual void OnMouseLeave( const FPointerEvent& MouseEvent ) override;
 
 	/**
 	 * Returns true if this button is currently pressed
@@ -177,5 +177,5 @@ protected:
 	FSlateSound PressedSound;
 
 private:
-	virtual FVector2D ComputeDesiredSize() const OVERRIDE;
+	virtual FVector2D ComputeDesiredSize() const override;
 };
diff --git a/Engine/Source/Runtime/Slate/Public/Widgets/Input/SCheckBox.h b/Engine/Source/Runtime/Slate/Public/Widgets/Input/SCheckBox.h
--- a/Engine/Source/Runtime/Slate/Public/Widgets/Input/SCheckBox.h
+++ b/Engine/Source/Runtime/Slate/Public/Widgets/Input/SCheckBox.h
@@ -99,13 +99,13 @@ public:
 	void Construct( const FArguments& InArgs );
 
 	// SWidget interface
-	virtual bool SupportsKeyboardFocus() const OVERRIDE;
-	virtual FReply OnKeyUp( const FGeometry& MyGeometry, const FKeyboardEvent& InKeyboardEvent ) OVERRIDE;
-	virtual FReply OnMouseButtonDown( const FGeometry& MyGeometry, const FPointerEvent& MouseEvent ) OVERRIDE;
-	virtual FReply OnMouseButtonDoubleClick( const FGeometry& InMyGeometry, const FPointerEvent& InMouseEvent ) OVERRIDE;
-	virtual FReply OnMouseButtonUp( const FGeometry& MyGeometry, const FPointerEvent& MouseEvent ) OVERRIDE;
-	virtual void OnMouseEnter( const FGeometry& MyGeometry, const FPointerEvent& MouseEvent ) OVERRIDE;
-	virtual void OnMouseLeave( const FPointerEvent& MouseEvent ) OVERRIDE;
+	virtual bool SupportsKeyboardFocus() const override;
+	virtual FReply OnKeyUp( const FGeometry& MyGeometry, const FKeyboardEvent& InKeyboardEvent ) override;
+	virtual FReply OnMouseButtonDown( const FGeometry& MyGeometry, const FPointerEvent& MouseEvent ) override;
+	virtual FReply OnMouseButtonDoubleClick( const FGeometry& InMyGeometry, const FPointerEvent& InMouseEvent ) override;
+	virtual FReply OnMouseButtonUp( const FGeometry& MyGeometry, const FPointerEvent& MouseEvent ) override;
+	virtual void OnMouseEnter( const FGeometry& MyGeometry, const FPointerEvent& MouseEvent ) override;
+	virtual void OnMouseLeave( const FPointerEvent& MouseEvent ) override;
 	// End of SWidget interface
 
 	/**
diff --git a/Engine/Source/Runtime/Slate/Public/Widgets/Input/SComboBox.h b/Engine/Source/Runtime/Slate/Public/Widgets/Input/SComboBox.h
--- a/Engine/Source/Runtime/Slate/Public/Widgets/Input/SComboBox.h
+++ b/Engine/Source/Runtime/Slate/Public/Widgets/Input/SComboBox.h
@@ -244,7 +244,7 @@ protected:
 		return FReply::Unhandled();
 	}
 
-	FReply OnKeyDown( const FGeometry& MyGeometry, const FKeyboardEvent& InKeyboardEvent ) OVERRIDE
+	FReply OnKeyDown( const FGeometry& MyGeometry, const FKeyboardEvent& InKeyboardEvent ) override
 	{
 		const FKey Key = InKeyboardEvent.GetKey();
 
@@ -274,7 +274,7 @@ protected:
 		return SComboButton::OnKeyDown( MyGeometry, InKeyboardEvent );
 	}
 
-	virtual bool SupportsKeyboardFocus() const OVERRIDE
+	virtual bool SupportsKeyboardFocus() const override
 	{
 		return true;
 	}
@@ -320,7 +320,7 @@ private:
 	}
 
 	/** Handle clicking on the content menu */
-	virtual FReply OnButtonClicked() OVERRIDE
+	virtual FReply OnButtonClicked() override
 	{
 		// if user clicked to close the combo menu
 		if (this->IsOpen())
diff --git a/Engine/Source/Runtime/Slate/Public/Widgets/Input/SEditableText.h b/Engine/Source/Runtime/Slate/Public/Widgets/Input/SEditableText.h
--- a/Engine/Source/Runtime/Slate/Public/Widgets/Input/SEditableText.h
+++ b/Engine/Source/Runtime/Slate/Public/Widgets/Input/SEditableText.h
@@ -165,19 +165,19 @@ protected:
 		virtual ~FTextInputMethodContext() {}
 
 	private:
-		virtual bool IsReadOnly() OVERRIDE;
-		virtual uint32 GetTextLength() OVERRIDE;
-		virtual void GetSelectionRange(uint32& BeginIndex, uint32& Length, ECaretPosition& CaretPosition) OVERRIDE;
-		virtual void SetSelectionRange(const uint32 BeginIndex, const uint32 Length, const ECaretPosition CaretPosition) OVERRIDE;
-		virtual void GetTextInRange(const uint32 BeginIndex, const uint32 Length, FString& OutString) OVERRIDE;
-		virtual void SetTextInRange(const uint32 BeginIndex, const uint32 Length, const FString& InString) OVERRIDE;
-		virtual int32 GetCharacterIndexFromPoint(const FVector2D& Point) OVERRIDE;
-		virtual bool GetTextBounds(const uint32 BeginIndex, const uint32 Length, FVector2D& Position, FVector2D& Size) OVERRIDE;
-		virtual void GetScreenBounds(FVector2D& Position, FVector2D& Size) OVERRIDE;
-		virtual TSharedPtr<FGenericWindow> GetWindow() OVERRIDE;
-		virtual void BeginComposition() OVERRIDE;
-		virtual void UpdateCompositionRange(const int32 InBeginIndex, const uint32 InLength) OVERRIDE;
-		virtual void EndComposition() OVERRIDE;
+		virtual bool IsReadOnly() override;
+		virtual uint32 GetTextLength() override;
+		virtual void GetSelectionRange(uint32& BeginIndex, uint32& Length, ECaretPosition& CaretPosition) override;
+		virtual void SetSelectionRange(const uint32 BeginIndex, const uint32 Length, const ECaretPosition CaretPosition) override;
+		virtual void GetTextInRange(const uint32 BeginIndex, const uint32 Length, FString& OutString) override;
+		virtual void SetTextInRange(const uint32 BeginIndex, const uint32 Length, const FString& InString) override;
+		virtual int32 GetCharacterIndexFromPoint(const FVector2D& Point) override;
+		virtual bool GetTextBounds(const uint32 BeginIndex, const uint32 Length, FVector2D& Position, FVector2D& Size) override;
+		virtual void GetScreenBounds(FVector2D& Position, FVector2D& Size) override;
+		virtual TSharedPtr<FGenericWindow> GetWindow() override;
+		virtual void BeginComposition() override;
+		virtual void UpdateCompositionRange(const int32 InBeginIndex, const uint32 InLength) override;
+		virtual void EndComposition() override;
 
 	private:
 		TWeakPtr<SEditableText> OwningWidget;
@@ -271,62 +271,62 @@ protected:
 
 public:
 	// BEGIN ITextEditorWidget interface
-	virtual void StartChangingText() OVERRIDE;
-	virtual void FinishChangingText() OVERRIDE;
-	virtual bool GetIsReadOnly() const OVERRIDE;
-	virtual void BackspaceChar() OVERRIDE;
-	virtual void DeleteChar() OVERRIDE;
-	virtual bool CanTypeCharacter(const TCHAR CharInQuestion) const OVERRIDE;
-	virtual void TypeChar( const int32 Character ) OVERRIDE;
-	virtual FReply MoveCursor( ECursorMoveMethod::Type Method, const FVector2D& Direction, ECursorAction::Type Action ) OVERRIDE;
-	virtual void JumpTo(ETextLocation::Type JumpLocation, ECursorAction::Type Action) OVERRIDE;
-	virtual void ClearSelection() OVERRIDE;
-	virtual void SelectAllText() OVERRIDE;
-	virtual bool SelectAllTextWhenFocused() OVERRIDE;
-	virtual void SelectWordAt(const FVector2D& LocalPosition) OVERRIDE;
-	virtual void BeginDragSelection() OVERRIDE;
-	virtual bool IsDragSelecting() const OVERRIDE;
-	virtual void EndDragSelection() OVERRIDE;
-	virtual bool AnyTextSelected() const OVERRIDE;
-	virtual bool IsTextSelectedAt(const FVector2D& LocalPosition) const OVERRIDE;
-	virtual void SetWasFocusedByLastMouseDown( bool Value ) OVERRIDE;
-	virtual bool WasFocusedByLastMouseDown() const OVERRIDE;
-	virtual void SetHasDragSelectedSinceFocused( bool Value ) OVERRIDE;
-	virtual bool HasDragSelectedSinceFocused() const OVERRIDE;
-	virtual FReply OnEscape() OVERRIDE;
-	virtual void OnEnter() OVERRIDE;
-	virtual bool CanExecuteCut() const OVERRIDE;
-	virtual void CutSelectedTextToClipboard() OVERRIDE;
-	virtual bool CanExecuteCopy() const OVERRIDE;
-	virtual void CopySelectedTextToClipboard() OVERRIDE;
-	virtual bool CanExecutePaste() const OVERRIDE;
-	virtual void PasteTextFromClipboard() OVERRIDE;
-	virtual bool CanExecuteUndo() const OVERRIDE;
-	virtual void Undo() OVERRIDE;
-	virtual void Redo() OVERRIDE;
-	virtual TSharedRef< SWidget > GetWidget() OVERRIDE;
-	virtual void SummonContextMenu( const FVector2D& InLocation ) OVERRIDE;
-	virtual void LoadText() OVERRIDE;
+	virtual void StartChangingText() override;
+	virtual void FinishChangingText() override;
+	virtual bool GetIsReadOnly() const override;
+	virtual void BackspaceChar() override;
+	virtual void DeleteChar() override;
+	virtual bool CanTypeCharacter(const TCHAR CharInQuestion) const override;
+	virtual void TypeChar( const int32 Character ) override;
+	virtual FReply MoveCursor( ECursorMoveMethod::Type Method, const FVector2D& Direction, ECursorAction::Type Action ) override;
+	virtual void JumpTo(ETextLocation::Type JumpLocation, ECursorAction::Type Action) override;
+	virtual void ClearSelection() override;
+	virtual void SelectAllText() override;
+	virtual bool SelectAllTextWhenFocused() override;
+	virtual void SelectWordAt(const FVector2D& LocalPosition) override;
+	virtual void BeginDragSelection() override;
+	virtual bool IsDragSelecting() const override;
+	virtual void EndDragSelection() override;
+	virtual bool AnyTextSelected() const override;
+	virtual bool IsTextSelectedAt(const FVector2D& LocalPosition) const override;
+	virtual void SetWasFocusedByLastMouseDown( bool Value ) override;
+	virtual bool WasFocusedByLastMouseDown() const override;
+	virtual void SetHasDragSelectedSinceFocused( bool Value ) override;
+	virtual bool HasDragSelectedSinceFocused() const override;
+	virtual FReply OnEscape() override;
+	virtual void OnEnter() override;
+	virtual bool CanExecuteCut() const override;
+	virtual void CutSelectedTextToClipboard() override;
+	virtual bool CanExecuteCopy() const override;
+	virtual void CopySelectedTextToClipboard() override;
+	virtual bool CanExecutePaste() const override;
+	virtual void PasteTextFromClipboard() override;
+	virtual bool CanExecuteUndo() const override;
+	virtual void Undo() override;
+	virtual void Redo() override;
+	virtual TSharedRef< SWidget > GetWidget() override;
+	virtual void SummonContextMenu( const FVector2D& InLocation ) override;
+	virtual void LoadText() override;
 	// END ITextEditorWidget interface
 
 protected:
 	// BEGIN SWidget interface
-	virtual void Tick( const FGeometry& AllottedGeometry, const double InCurrentTime, const float InDeltaTime ) OVERRIDE;
-	virtual FVector2D ComputeDesiredSize() const OVERRIDE;
-	virtual int32 OnPaint( const FGeometry& AllottedGeometry, const FSlateRect& MyClippingRect, FSlateWindowElementList& OutDrawElements, int32 LayerId, const FWidgetStyle& InWidgetStyle, bool bParentEnabled ) const OVERRIDE;
-	virtual FReply OnDragOver( const FGeometry& MyGeometry, const FDragDropEvent& DragDropEvent ) OVERRIDE;
-	virtual FReply OnDrop( const FGeometry& MyGeometry, const FDragDropEvent& DragDropEvent ) OVERRIDE;
-	virtual bool SupportsKeyboardFocus() const OVERRIDE;
-	virtual FReply OnKeyboardFocusReceived( const FGeometry& MyGeometry, const FKeyboardFocusEvent& InKeyboardFocusEvent ) OVERRIDE;
-	virtual void OnKeyboardFocusLost( const FKeyboardFocusEvent& InKeyboardFocusEvent ) OVERRIDE;
-	virtual FReply OnKeyChar( const FGeometry& MyGeometry,  const FCharacterEvent& InCharacterEvent ) OVERRIDE;
-	virtual FReply OnKeyDown( const FGeometry& MyGeometry, const FKeyboardEvent& InKeyboardEvent ) OVERRIDE;
-	virtual FReply OnKeyUp( const FGeometry& MyGeometry, const FKeyboardEvent& InKeyboardEvent ) OVERRIDE;
-	virtual FReply OnMouseButtonDown( const FGeometry& InMyGeometry, const FPointerEvent& InMouseEvent ) OVERRIDE;
-	virtual FReply OnMouseButtonUp( const FGeometry& InMyGeometry, const FPointerEvent& InMouseEvent ) OVERRIDE;
-	virtual FReply OnMouseMove( const FGeometry& InMyGeometry, const FPointerEvent& InMouseEvent ) OVERRIDE;
-	virtual FReply OnMouseButtonDoubleClick( const FGeometry& InMyGeometry, const FPointerEvent& InMouseEvent ) OVERRIDE;
-	virtual FCursorReply OnCursorQuery( const FGeometry& MyGeometry, const FPointerEvent& CursorEvent ) const OVERRIDE;
+	virtual void Tick( const FGeometry& AllottedGeometry, const double InCurrentTime, const float InDeltaTime ) override;
+	virtual FVector2D ComputeDesiredSize() const override;
+	virtual int32 OnPaint( const FGeometry& AllottedGeometry, const FSlateRect& MyClippingRect, FSlateWindowElementList& OutDrawElements, int32 LayerId, const FWidgetStyle& InWidgetStyle, bool bParentEnabled ) const override;
+	virtual FReply OnDragOver( const FGeometry& MyGeometry, const FDragDropEvent& DragDropEvent ) override;
+	virtual FReply OnDrop( const FGeometry& MyGeometry, const FDragDropEvent& DragDropEvent ) override;
+	virtual bool SupportsKeyboardFocus() const override;
+	virtual FReply OnKeyboardFocusReceived( const FGeometry& MyGeometry, const FKeyboardFocusEvent& InKeyboardFocusEvent ) override;
+	virtual void OnKeyboardFocusLost( const FKeyboardFocusEvent& InKeyboardFocusEvent ) override;
+	virtual FReply OnKeyChar( const FGeometry& MyGeometry,  const FCharacterEvent& InCharacterEvent ) override;
+	virtual FReply OnKeyDown( const FGeometry& MyGeometry, const FKeyboardEvent& InKeyboardEvent ) override;
+	virtual FReply OnKeyUp( const FGeometry& MyGeometry, const FKeyboardEvent& InKeyboardEvent ) override;
+	virtual FReply OnMouseButtonDown( const FGeometry& InMyGeometry, const FPointerEvent& InMouseEvent ) override;
+	virtual FReply OnMouseButtonUp( const FGeometry& InMyGeometry, const FPointerEvent& InMouseEvent ) override;
+	virtual FReply OnMouseMove( const FGeometry& InMyGeometry, const FPointerEvent& InMouseEvent ) override;
+	virtual FReply OnMouseButtonDoubleClick( const FGeometry& InMyGeometry, const FPointerEvent& InMouseEvent ) override;
+	virtual FCursorReply OnCursorQuery( const FGeometry& MyGeometry, const FPointerEvent& CursorEvent ) const override;
 	// END SWidget interface
 
 protected:
diff --git a/Engine/Source/Runtime/Slate/Public/Widgets/Input/SEditableTextBox.h b/Engine/Source/Runtime/Slate/Public/Widgets/Input/SEditableTextBox.h
--- a/Engine/Source/Runtime/Slate/Public/Widgets/Input/SEditableTextBox.h
+++ b/Engine/Source/Runtime/Slate/Public/Widgets/Input/SEditableTextBox.h
@@ -121,10 +121,10 @@ public:
 	void SetError( const FString& InError );
 
 	// SWidget overrides
-	virtual bool SupportsKeyboardFocus() const OVERRIDE;
-	virtual bool HasKeyboardFocus() const OVERRIDE;
-	virtual FReply OnKeyboardFocusReceived( const FGeometry& MyGeometry, const FKeyboardFocusEvent& InKeyboardFocusEvent ) OVERRIDE;
-	virtual FReply OnKeyDown( const FGeometry& MyGeometry, const FKeyboardEvent& InKeyboardEvent ) OVERRIDE;
+	virtual bool SupportsKeyboardFocus() const override;
+	virtual bool HasKeyboardFocus() const override;
+	virtual FReply OnKeyboardFocusReceived( const FGeometry& MyGeometry, const FKeyboardFocusEvent& InKeyboardFocusEvent ) override;
+	virtual FReply OnKeyDown( const FGeometry& MyGeometry, const FKeyboardEvent& InKeyboardEvent ) override;
 
 protected:
 
diff --git a/Engine/Source/Runtime/Slate/Public/Widgets/Input/SHyperlink.h b/Engine/Source/Runtime/Slate/Public/Widgets/Input/SHyperlink.h
--- a/Engine/Source/Runtime/Slate/Public/Widgets/Input/SHyperlink.h
+++ b/Engine/Source/Runtime/Slate/Public/Widgets/Input/SHyperlink.h
@@ -53,7 +53,7 @@ public:
 	}
 
 	// SWidget implementation
-	virtual FCursorReply OnCursorQuery( const FGeometry& MyGeometry, const FPointerEvent& CursorEvent ) const OVERRIDE
+	virtual FCursorReply OnCursorQuery( const FGeometry& MyGeometry, const FPointerEvent& CursorEvent ) const override
 	{
 		return FCursorReply::Cursor( EMouseCursor::Hand );
 	}
diff --git a/Engine/Source/Runtime/Slate/Public/Widgets/Input/SMenuAnchor.h b/Engine/Source/Runtime/Slate/Public/Widgets/Input/SMenuAnchor.h
--- a/Engine/Source/Runtime/Slate/Public/Widgets/Input/SMenuAnchor.h
+++ b/Engine/Source/Runtime/Slate/Public/Widgets/Input/SMenuAnchor.h
@@ -49,7 +49,7 @@ public:
 	SMenuAnchor();
 
 	// SWidget interface
-	virtual void Tick( const FGeometry& AllottedGeometry, const double InCurrentTime, const float InDeltaTime ) OVERRIDE;
+	virtual void Tick( const FGeometry& AllottedGeometry, const double InCurrentTime, const float InDeltaTime ) override;
 	// End of SWidget interface
 
 	
diff --git a/Engine/Source/Runtime/Slate/Public/Widgets/Input/SMultiLineEditableTextBox.h b/Engine/Source/Runtime/Slate/Public/Widgets/Input/SMultiLineEditableTextBox.h
--- a/Engine/Source/Runtime/Slate/Public/Widgets/Input/SMultiLineEditableTextBox.h
+++ b/Engine/Source/Runtime/Slate/Public/Widgets/Input/SMultiLineEditableTextBox.h
@@ -135,9 +135,9 @@ public:
 	void SetError( const FString& InError );
 
 	// SWidget overrides
-	virtual bool SupportsKeyboardFocus() const OVERRIDE;
-	virtual bool HasKeyboardFocus() const OVERRIDE;
-	virtual FReply OnKeyboardFocusReceived( const FGeometry& MyGeometry, const FKeyboardFocusEvent& InKeyboardFocusEvent ) OVERRIDE;
+	virtual bool SupportsKeyboardFocus() const override;
+	virtual bool HasKeyboardFocus() const override;
+	virtual FReply OnKeyboardFocusReceived( const FGeometry& MyGeometry, const FKeyboardFocusEvent& InKeyboardFocusEvent ) override;
 
 protected:
 
diff --git a/Engine/Source/Runtime/Slate/Public/Widgets/Input/SNumericEntryBox.h b/Engine/Source/Runtime/Slate/Public/Widgets/Input/SNumericEntryBox.h
--- a/Engine/Source/Runtime/Slate/Public/Widgets/Input/SNumericEntryBox.h
+++ b/Engine/Source/Runtime/Slate/Public/Widgets/Input/SNumericEntryBox.h
@@ -205,12 +205,12 @@ public:
 private:
 
 	// Begin SWidget Interface
-	virtual bool SupportsKeyboardFocus() const OVERRIDE
+	virtual bool SupportsKeyboardFocus() const override
 	{
 		return StaticCastSharedPtr<SWidget>(EditableText)->SupportsKeyboardFocus();
 	}
 
-	virtual FReply OnKeyboardFocusReceived( const FGeometry& MyGeometry, const FKeyboardFocusEvent& InKeyboardFocusEvent ) OVERRIDE
+	virtual FReply OnKeyboardFocusReceived( const FGeometry& MyGeometry, const FKeyboardFocusEvent& InKeyboardFocusEvent ) override
 	{
 		FReply Reply = FReply::Handled();
 
@@ -234,7 +234,7 @@ private:
 		return Reply;
 	}
 
-	virtual FReply OnKeyDown( const FGeometry& MyGeometry, const FKeyboardEvent& InKeyboardEvent ) OVERRIDE
+	virtual FReply OnKeyDown( const FGeometry& MyGeometry, const FKeyboardEvent& InKeyboardEvent ) override
 	{
 		FKey Key = InKeyboardEvent.GetKey();
 
@@ -405,7 +405,7 @@ private:
 	/**
 	 * Caches the value and performs widget visibility maintenance
 	 */
-	virtual void Tick( const FGeometry& AllottedGeometry, const double InCurrentTime, const float InDeltaTime ) OVERRIDE
+	virtual void Tick( const FGeometry& AllottedGeometry, const double InCurrentTime, const float InDeltaTime ) override
 	{
 		// Visibility toggle only matters if the spinbox is used
 		if( SpinBox.IsValid() )
diff --git a/Engine/Source/Runtime/Slate/Public/Widgets/Input/SRichTextHyperlink.h b/Engine/Source/Runtime/Slate/Public/Widgets/Input/SRichTextHyperlink.h
--- a/Engine/Source/Runtime/Slate/Public/Widgets/Input/SRichTextHyperlink.h
+++ b/Engine/Source/Runtime/Slate/Public/Widgets/Input/SRichTextHyperlink.h
@@ -30,19 +30,19 @@ public:
 		);
 	}
 
-	virtual void OnMouseEnter( const FGeometry& MyGeometry, const FPointerEvent& MouseEvent ) OVERRIDE
+	virtual void OnMouseEnter( const FGeometry& MyGeometry, const FPointerEvent& MouseEvent ) override
 	{
 		SHyperlink::OnMouseEnter( MyGeometry, MouseEvent );
 		ViewModel->SetIsHovered( true );
 	}
 
-	virtual void OnMouseLeave( const FPointerEvent& MouseEvent ) OVERRIDE
+	virtual void OnMouseLeave( const FPointerEvent& MouseEvent ) override
 	{
 		SHyperlink::OnMouseLeave( MouseEvent );
 		ViewModel->SetIsHovered( false );
 	}
 
-	virtual FReply OnMouseButtonDown( const FGeometry& MyGeometry, const FPointerEvent& MouseEvent ) OVERRIDE
+	virtual FReply OnMouseButtonDown( const FGeometry& MyGeometry, const FPointerEvent& MouseEvent ) override
 	{
 		FReply Reply = SHyperlink::OnMouseButtonDown( MyGeometry, MouseEvent );
 		ViewModel->SetIsPressed( bIsPressed );
@@ -50,7 +50,7 @@ public:
 		return Reply;
 	}
 
-	virtual FReply OnMouseButtonUp( const FGeometry& MyGeometry, const FPointerEvent& MouseEvent ) OVERRIDE
+	virtual FReply OnMouseButtonUp( const FGeometry& MyGeometry, const FPointerEvent& MouseEvent ) override
 	{
 		FReply Reply = SHyperlink::OnMouseButtonUp( MyGeometry, MouseEvent );
 		ViewModel->SetIsPressed( bIsPressed );
@@ -58,12 +58,12 @@ public:
 		return Reply;
 	}
 
-	virtual bool IsHovered() const OVERRIDE
+	virtual bool IsHovered() const override
 	{
 		return ViewModel->IsHovered();
 	}
 
-	virtual bool IsPressed() const OVERRIDE
+	virtual bool IsPressed() const override
 	{
 		return ViewModel->IsPressed();
 	}
diff --git a/Engine/Source/Runtime/Slate/Public/Widgets/Input/SSearchBox.h b/Engine/Source/Runtime/Slate/Public/Widgets/Input/SSearchBox.h
--- a/Engine/Source/Runtime/Slate/Public/Widgets/Input/SSearchBox.h
+++ b/Engine/Source/Runtime/Slate/Public/Widgets/Input/SSearchBox.h
@@ -59,7 +59,7 @@ public:
 	SLATE_END_ARGS()
 
 	void Construct( const FArguments& InArgs );
-	virtual void Tick( const FGeometry& AllottedGeometry, const double InCurrentTime, const float InDeltaTime ) OVERRIDE;
+	virtual void Tick( const FGeometry& AllottedGeometry, const double InCurrentTime, const float InDeltaTime ) override;
 
 private:
 
diff --git a/Engine/Source/Runtime/Slate/Public/Widgets/Input/SSlider.h b/Engine/Source/Runtime/Slate/Public/Widgets/Input/SSlider.h
--- a/Engine/Source/Runtime/Slate/Public/Widgets/Input/SSlider.h
+++ b/Engine/Source/Runtime/Slate/Public/Widgets/Input/SSlider.h
@@ -71,15 +71,15 @@ public:
 
 public:
 
-	virtual int32 OnPaint( const FGeometry& AllottedGeometry, const FSlateRect& MyClippingRect, FSlateWindowElementList& OutDrawElements, int32 LayerId, const FWidgetStyle& InWidgetStyle, bool bParentEnabled ) const OVERRIDE;
+	virtual int32 OnPaint( const FGeometry& AllottedGeometry, const FSlateRect& MyClippingRect, FSlateWindowElementList& OutDrawElements, int32 LayerId, const FWidgetStyle& InWidgetStyle, bool bParentEnabled ) const override;
 
 	virtual FVector2D ComputeDesiredSize() const;
 
-	virtual FReply OnMouseButtonDown( const FGeometry& MyGeometry, const FPointerEvent& MouseEvent ) OVERRIDE;
+	virtual FReply OnMouseButtonDown( const FGeometry& MyGeometry, const FPointerEvent& MouseEvent ) override;
 	
-	virtual FReply OnMouseButtonUp( const FGeometry& MyGeometry, const FPointerEvent& MouseEvent ) OVERRIDE;
+	virtual FReply OnMouseButtonUp( const FGeometry& MyGeometry, const FPointerEvent& MouseEvent ) override;
 	
-	virtual FReply OnMouseMove( const FGeometry& MyGeometry, const FPointerEvent& MouseEvent ) OVERRIDE;
+	virtual FReply OnMouseMove( const FGeometry& MyGeometry, const FPointerEvent& MouseEvent ) override;
 
 protected:
 
diff --git a/Engine/Source/Runtime/Slate/Public/Widgets/Input/SSpinBox.h b/Engine/Source/Runtime/Slate/Public/Widgets/Input/SSpinBox.h
--- a/Engine/Source/Runtime/Slate/Public/Widgets/Input/SSpinBox.h
+++ b/Engine/Source/Runtime/Slate/Public/Widgets/Input/SSpinBox.h
@@ -171,7 +171,7 @@ public:
 	*
 	* @return The maximum layer ID attained by this widget or any of its children.
 	*/
-	virtual int32 OnPaint( const FGeometry& AllottedGeometry, const FSlateRect& MyClippingRect, FSlateWindowElementList& OutDrawElements, int32 LayerId, const FWidgetStyle& InWidgetStyle, bool bParentEnabled ) const OVERRIDE
+	virtual int32 OnPaint( const FGeometry& AllottedGeometry, const FSlateRect& MyClippingRect, FSlateWindowElementList& OutDrawElements, int32 LayerId, const FWidgetStyle& InWidgetStyle, bool bParentEnabled ) const override
 	{
 		const bool bActiveFeedback = IsHovered() || bDragging;
 
@@ -242,7 +242,7 @@ public:
 	 *
 	 * @return Whether the event was handled along with possible requests for the system to take action.
 	 */
-	virtual FReply OnMouseButtonDown( const FGeometry& MyGeometry, const FPointerEvent& MouseEvent ) OVERRIDE
+	virtual FReply OnMouseButtonDown( const FGeometry& MyGeometry, const FPointerEvent& MouseEvent ) override
 	{
 		if ( MouseEvent.GetEffectingButton() == EKeys::LeftMouseButton )
 		{
@@ -265,7 +265,7 @@ public:
 	 *
 	 * @return Whether the event was handled along with possible requests for the system to take action.
 	 */
-	virtual FReply OnMouseButtonUp( const FGeometry& MyGeometry, const FPointerEvent& MouseEvent ) OVERRIDE
+	virtual FReply OnMouseButtonUp( const FGeometry& MyGeometry, const FPointerEvent& MouseEvent ) override
 	{
 		if ( MouseEvent.GetEffectingButton() == EKeys::LeftMouseButton && this->HasMouseCapture() )
 		{
@@ -298,7 +298,7 @@ public:
 	 *
 	 * @return Whether the event was handled along with possible requests for the system to take action.
 	 */
-	virtual FReply OnMouseMove( const FGeometry& MyGeometry, const FPointerEvent& MouseEvent ) OVERRIDE
+	virtual FReply OnMouseMove( const FGeometry& MyGeometry, const FPointerEvent& MouseEvent ) override
 	{
 		if ( this->HasMouseCapture() )
 		{
@@ -372,7 +372,7 @@ public:
 		return FReply::Unhandled();
 	}
 
-	virtual FCursorReply OnCursorQuery( const FGeometry& MyGeometry, const FPointerEvent& CursorEvent ) const OVERRIDE
+	virtual FCursorReply OnCursorQuery( const FGeometry& MyGeometry, const FPointerEvent& CursorEvent ) const override
 	{
 		return bDragging ? 
 			FCursorReply::Cursor( EMouseCursor::None ) :
@@ -386,7 +386,7 @@ public:
 	}
 
 
-	virtual FReply OnKeyboardFocusReceived( const FGeometry& MyGeometry, const FKeyboardFocusEvent& InKeyboardFocusEvent ) OVERRIDE
+	virtual FReply OnKeyboardFocusReceived( const FGeometry& MyGeometry, const FKeyboardFocusEvent& InKeyboardFocusEvent ) override
 	{
 		if ( InKeyboardFocusEvent.GetCause() == EKeyboardFocusCause::Keyboard || InKeyboardFocusEvent.GetCause() == EKeyboardFocusCause::SetDirectly )
 		{
@@ -399,7 +399,7 @@ public:
 		}
 	}
 	
-	virtual FReply OnKeyDown( const FGeometry& MyGeometry, const FKeyboardEvent& InKeyboardEvent ) OVERRIDE
+	virtual FReply OnKeyDown( const FGeometry& MyGeometry, const FKeyboardEvent& InKeyboardEvent ) override
 	{
 		const FKey Key = InKeyboardEvent.GetKey();
 		if ( Key == EKeys::Escape && HasMouseCapture())
diff --git a/Engine/Source/Runtime/Slate/Public/Widgets/Input/SSuggestionTextBox.h b/Engine/Source/Runtime/Slate/Public/Widgets/Input/SSuggestionTextBox.h
--- a/Engine/Source/Runtime/Slate/Public/Widgets/Input/SSuggestionTextBox.h
+++ b/Engine/Source/Runtime/Slate/Public/Widgets/Input/SSuggestionTextBox.h
@@ -237,11 +237,11 @@ protected:
 
 protected:
 
-	virtual void OnKeyboardFocusLost( const FKeyboardFocusEvent& InKeyboardFocusEvent ) OVERRIDE;
+	virtual void OnKeyboardFocusLost( const FKeyboardFocusEvent& InKeyboardFocusEvent ) override;
 
-	virtual FReply OnKeyDown( const FGeometry& MyGeometry, const FKeyboardEvent& KeyboardEvent ) OVERRIDE;
+	virtual FReply OnKeyDown( const FGeometry& MyGeometry, const FKeyboardEvent& KeyboardEvent ) override;
 
-	virtual bool SupportsKeyboardFocus( ) const OVERRIDE;
+	virtual bool SupportsKeyboardFocus( ) const override;
 
 
 private:
diff --git a/Engine/Source/Runtime/Slate/Public/Widgets/Input/STextComboPopup.h b/Engine/Source/Runtime/Slate/Public/Widgets/Input/STextComboPopup.h
--- a/Engine/Source/Runtime/Slate/Public/Widgets/Input/STextComboPopup.h
+++ b/Engine/Source/Runtime/Slate/Public/Widgets/Input/STextComboPopup.h
@@ -35,7 +35,7 @@ public:
 	void Construct( const FArguments& InArgs );
 
 	/** SWidget interface */
-	virtual FReply OnKeyDown( const FGeometry& MyGeometry, const FKeyboardEvent& InKeyboardEvent ) OVERRIDE;
+	virtual FReply OnKeyDown( const FGeometry& MyGeometry, const FKeyboardEvent& InKeyboardEvent ) override;
 
 	/** Called to create a widget for each string */
 	TSharedRef<SWidget> MakeItemWidget( TSharedPtr<FString> StringItem );
diff --git a/Engine/Source/Runtime/Slate/Public/Widgets/Input/SVirtualJoystick.h b/Engine/Source/Runtime/Slate/Public/Widgets/Input/SVirtualJoystick.h
--- a/Engine/Source/Runtime/Slate/Public/Widgets/Input/SVirtualJoystick.h
+++ b/Engine/Source/Runtime/Slate/Public/Widgets/Input/SVirtualJoystick.h
@@ -125,17 +125,17 @@ public:
 	 */
 	void SetGlobalParameters(float InActiveOpacity, float InInactiveOpacity, float InTimeUntilDeactive, float InTimeUntilReset, float InActivationDelay, bool InbPreventReCenter, float InStartupDelay);
 
-	virtual int32 OnPaint( const FGeometry& AllottedGeometry, const FSlateRect& MyClippingRect, FSlateWindowElementList& OutDrawElements, int32 LayerId, const FWidgetStyle& InWidgetStyle, bool bParentEnabled ) const OVERRIDE;
+	virtual int32 OnPaint( const FGeometry& AllottedGeometry, const FSlateRect& MyClippingRect, FSlateWindowElementList& OutDrawElements, int32 LayerId, const FWidgetStyle& InWidgetStyle, bool bParentEnabled ) const override;
 	
-	virtual FVector2D ComputeDesiredSize() const OVERRIDE;
+	virtual FVector2D ComputeDesiredSize() const override;
 
-	virtual FReply OnTouchStarted(const FGeometry& MyGeometry, const FPointerEvent& Event) OVERRIDE;
-	virtual FReply OnTouchMoved(const FGeometry& MyGeometry, const FPointerEvent& Event) OVERRIDE;
-	virtual FReply OnTouchEnded(const FGeometry& MyGeometry, const FPointerEvent& Event) OVERRIDE;
+	virtual FReply OnTouchStarted(const FGeometry& MyGeometry, const FPointerEvent& Event) override;
+	virtual FReply OnTouchMoved(const FGeometry& MyGeometry, const FPointerEvent& Event) override;
+	virtual FReply OnTouchEnded(const FGeometry& MyGeometry, const FPointerEvent& Event) override;
 
-	virtual void Tick( const FGeometry& AllottedGeometry, const double InCurrentTime, const float InDeltaTime ) OVERRIDE;
+	virtual void Tick( const FGeometry& AllottedGeometry, const double InCurrentTime, const float InDeltaTime ) override;
 
-	virtual bool SupportsKeyboardFocus() const OVERRIDE
+	virtual bool SupportsKeyboardFocus() const override
 	{
 		return true;
 	}
diff --git a/Engine/Source/Runtime/Slate/Public/Widgets/Input/SVirtualKeyboardEntry.h b/Engine/Source/Runtime/Slate/Public/Widgets/Input/SVirtualKeyboardEntry.h
--- a/Engine/Source/Runtime/Slate/Public/Widgets/Input/SVirtualKeyboardEntry.h
+++ b/Engine/Source/Runtime/Slate/Public/Widgets/Input/SVirtualKeyboardEntry.h
@@ -134,12 +134,12 @@ public:
 protected:
 
 	// SWidget interface
-	virtual void Tick( const FGeometry& AllottedGeometry, const double InCurrentTime, const float InDeltaTime ) OVERRIDE;
-	virtual FVector2D ComputeDesiredSize() const OVERRIDE;
-	virtual int32 OnPaint( const FGeometry& AllottedGeometry, const FSlateRect& MyClippingRect, FSlateWindowElementList& OutDrawElements, int32 LayerId, const FWidgetStyle& InWidgetStyle, bool bParentEnabled ) const OVERRIDE;
-	virtual bool SupportsKeyboardFocus() const OVERRIDE;
-	virtual FReply OnKeyboardFocusReceived( const FGeometry& MyGeometry, const FKeyboardFocusEvent& InKeyboardFocusEvent ) OVERRIDE;
-	virtual void OnKeyboardFocusLost( const FKeyboardFocusEvent& InKeyboardFocusEvent ) OVERRIDE;
+	virtual void Tick( const FGeometry& AllottedGeometry, const double InCurrentTime, const float InDeltaTime ) override;
+	virtual FVector2D ComputeDesiredSize() const override;
+	virtual int32 OnPaint( const FGeometry& AllottedGeometry, const FSlateRect& MyClippingRect, FSlateWindowElementList& OutDrawElements, int32 LayerId, const FWidgetStyle& InWidgetStyle, bool bParentEnabled ) const override;
+	virtual bool SupportsKeyboardFocus() const override;
+	virtual FReply OnKeyboardFocusReceived( const FGeometry& MyGeometry, const FKeyboardFocusEvent& InKeyboardFocusEvent ) override;
+	virtual void OnKeyboardFocusLost( const FKeyboardFocusEvent& InKeyboardFocusEvent ) override;
 	// End of SWidget interface
 
 private:
diff --git a/Engine/Source/Runtime/Slate/Public/Widgets/Layout/SBorder.h b/Engine/Source/Runtime/Slate/Public/Widgets/Layout/SBorder.h
--- a/Engine/Source/Runtime/Slate/Public/Widgets/Layout/SBorder.h
+++ b/Engine/Source/Runtime/Slate/Public/Widgets/Layout/SBorder.h
@@ -109,12 +109,12 @@ public:
 
 public:
 	// SWidget interface
-	virtual int32 OnPaint( const FGeometry& AllottedGeometry, const FSlateRect& MyClippingRect, FSlateWindowElementList& OutDrawElements, int32 LayerId, const FWidgetStyle& InWidgetStyle, bool bParentEnabled ) const OVERRIDE;
-	virtual FReply OnMouseButtonDown( const FGeometry& MyGeometry, const FPointerEvent& MouseEvent ) OVERRIDE;
-	virtual FReply OnMouseButtonUp( const FGeometry& MyGeometry, const FPointerEvent& MouseEvent ) OVERRIDE;
-	virtual FReply OnMouseMove( const FGeometry& MyGeometry, const FPointerEvent& MouseEvent ) OVERRIDE;
-	virtual FReply OnMouseButtonDoubleClick( const FGeometry& MyGeometry, const FPointerEvent& MouseEvent ) OVERRIDE;
-	virtual FVector2D ComputeDesiredSize() const OVERRIDE;
+	virtual int32 OnPaint( const FGeometry& AllottedGeometry, const FSlateRect& MyClippingRect, FSlateWindowElementList& OutDrawElements, int32 LayerId, const FWidgetStyle& InWidgetStyle, bool bParentEnabled ) const override;
+	virtual FReply OnMouseButtonDown( const FGeometry& MyGeometry, const FPointerEvent& MouseEvent ) override;
+	virtual FReply OnMouseButtonUp( const FGeometry& MyGeometry, const FPointerEvent& MouseEvent ) override;
+	virtual FReply OnMouseMove( const FGeometry& MyGeometry, const FPointerEvent& MouseEvent ) override;
+	virtual FReply OnMouseButtonDoubleClick( const FGeometry& MyGeometry, const FPointerEvent& MouseEvent ) override;
+	virtual FVector2D ComputeDesiredSize() const override;
 	// End of SWidget interface
 
  protected:
diff --git a/Engine/Source/Runtime/Slate/Public/Widgets/Layout/SBox.h b/Engine/Source/Runtime/Slate/Public/Widgets/Layout/SBox.h
--- a/Engine/Source/Runtime/Slate/Public/Widgets/Layout/SBox.h
+++ b/Engine/Source/Runtime/Slate/Public/Widgets/Layout/SBox.h
@@ -52,10 +52,10 @@ class SLATE_API SBox : public SPanel
 		void Construct( const FArguments& InArgs );
 
 
-		virtual FVector2D ComputeDesiredSize() const OVERRIDE;
-		virtual void OnArrangeChildren( const FGeometry& AllottedGeometry, FArrangedChildren& ArrangedChildren ) const OVERRIDE;
-		virtual FChildren* GetChildren() OVERRIDE;
-		virtual int32 OnPaint( const FGeometry& AllottedGeometry, const FSlateRect& MyClippingRect, FSlateWindowElementList& OutDrawElements, int32 LayerId, const FWidgetStyle& InWidgetStyle, bool bParentEnabled ) const OVERRIDE;
+		virtual FVector2D ComputeDesiredSize() const override;
+		virtual void OnArrangeChildren( const FGeometry& AllottedGeometry, FArrangedChildren& ArrangedChildren ) const override;
+		virtual FChildren* GetChildren() override;
+		virtual int32 OnPaint( const FGeometry& AllottedGeometry, const FSlateRect& MyClippingRect, FSlateWindowElementList& OutDrawElements, int32 LayerId, const FWidgetStyle& InWidgetStyle, bool bParentEnabled ) const override;
 
 private:
 		FBoxSlot ChildSlot;
diff --git a/Engine/Source/Runtime/Slate/Public/Widgets/Layout/SDPIScaler.h b/Engine/Source/Runtime/Slate/Public/Widgets/Layout/SDPIScaler.h
--- a/Engine/Source/Runtime/Slate/Public/Widgets/Layout/SDPIScaler.h
+++ b/Engine/Source/Runtime/Slate/Public/Widgets/Layout/SDPIScaler.h
@@ -28,11 +28,11 @@ public:
 
 	void Construct( const FArguments& InArgs );
 
-	virtual void OnArrangeChildren( const FGeometry& AllottedGeometry, FArrangedChildren& ArrangedChildren ) const OVERRIDE;
+	virtual void OnArrangeChildren( const FGeometry& AllottedGeometry, FArrangedChildren& ArrangedChildren ) const override;
 	
-	virtual FVector2D ComputeDesiredSize() const OVERRIDE;
+	virtual FVector2D ComputeDesiredSize() const override;
 
-	virtual FChildren* GetChildren() OVERRIDE;
+	virtual FChildren* GetChildren() override;
 
 protected:
 
diff --git a/Engine/Source/Runtime/Slate/Public/Widgets/Layout/SEnableBox.h b/Engine/Source/Runtime/Slate/Public/Widgets/Layout/SEnableBox.h
--- a/Engine/Source/Runtime/Slate/Public/Widgets/Layout/SEnableBox.h
+++ b/Engine/Source/Runtime/Slate/Public/Widgets/Layout/SEnableBox.h
@@ -19,7 +19,7 @@ public:
 	}
 
 	// SWidget interface
-	virtual int32 OnPaint(const FGeometry& AllottedGeometry, const FSlateRect& MyClippingRect, FSlateWindowElementList& OutDrawElements, int32 LayerId, const FWidgetStyle& InWidgetStyle, bool bParentEnabled) const OVERRIDE
+	virtual int32 OnPaint(const FGeometry& AllottedGeometry, const FSlateRect& MyClippingRect, FSlateWindowElementList& OutDrawElements, int32 LayerId, const FWidgetStyle& InWidgetStyle, bool bParentEnabled) const override
 	{
 		return SBox::OnPaint(AllottedGeometry, MyClippingRect, OutDrawElements, LayerId, InWidgetStyle, /*bParentEnabled=*/ true);
 	}
diff --git a/Engine/Source/Runtime/Slate/Public/Widgets/Layout/SExpandableArea.h b/Engine/Source/Runtime/Slate/Public/Widgets/Layout/SExpandableArea.h
--- a/Engine/Source/Runtime/Slate/Public/Widgets/Layout/SExpandableArea.h
+++ b/Engine/Source/Runtime/Slate/Public/Widgets/Layout/SExpandableArea.h
@@ -97,7 +97,7 @@ protected:
 	/**
 	 * Computes the desired size of this area. Optionally clamping to Max height
 	 */
-	virtual FVector2D ComputeDesiredSize() const OVERRIDE;
+	virtual FVector2D ComputeDesiredSize() const override;
 
 
 protected:
diff --git a/Engine/Source/Runtime/Slate/Public/Widgets/Layout/SFxWidget.h b/Engine/Source/Runtime/Slate/Public/Widgets/Layout/SFxWidget.h
--- a/Engine/Source/Runtime/Slate/Public/Widgets/Layout/SFxWidget.h
+++ b/Engine/Source/Runtime/Slate/Public/Widgets/Layout/SFxWidget.h
@@ -65,9 +65,9 @@ public:
 	
 
 protected:
-	virtual int32 OnPaint( const FGeometry& AllottedGeometry, const FSlateRect& MyClippingRect, FSlateWindowElementList& OutDrawElements, int32 LayerId, const FWidgetStyle& InWidgetStyle, bool bParentEnabled ) const OVERRIDE;
-	virtual FVector2D ComputeDesiredSize() const OVERRIDE;
-	virtual void OnArrangeChildren( const FGeometry& AllottedGeometry, FArrangedChildren& ArrangedChildren ) const OVERRIDE;
+	virtual int32 OnPaint( const FGeometry& AllottedGeometry, const FSlateRect& MyClippingRect, FSlateWindowElementList& OutDrawElements, int32 LayerId, const FWidgetStyle& InWidgetStyle, bool bParentEnabled ) const override;
+	virtual FVector2D ComputeDesiredSize() const override;
+	virtual void OnArrangeChildren( const FGeometry& AllottedGeometry, FArrangedChildren& ArrangedChildren ) const override;
 
 	TAttribute<float> RenderScale;
 	TAttribute<FVector2D> RenderScaleOrigin;
diff --git a/Engine/Source/Runtime/Slate/Public/Widgets/Layout/SGridPanel.h b/Engine/Source/Runtime/Slate/Public/Widgets/Layout/SGridPanel.h
--- a/Engine/Source/Runtime/Slate/Public/Widgets/Layout/SGridPanel.h
+++ b/Engine/Source/Runtime/Slate/Public/Widgets/Layout/SGridPanel.h
@@ -147,11 +147,11 @@ public:
 	void Construct( const FArguments& InArgs );
 
 	// SWidget interface
-	virtual int32 OnPaint( const FGeometry& AllottedGeometry, const FSlateRect& MyClippingRect, FSlateWindowElementList& OutDrawElements, int32 LayerId, const FWidgetStyle& InWidgetStyle, bool bParentEnabled ) const OVERRIDE;
-	virtual void OnArrangeChildren( const FGeometry& AllottedGeometry, FArrangedChildren& ArrangedChildren ) const OVERRIDE;
-	virtual void CacheDesiredSize() OVERRIDE;
-	virtual FVector2D ComputeDesiredSize() const OVERRIDE;
-	virtual FChildren* GetChildren() OVERRIDE;
+	virtual int32 OnPaint( const FGeometry& AllottedGeometry, const FSlateRect& MyClippingRect, FSlateWindowElementList& OutDrawElements, int32 LayerId, const FWidgetStyle& InWidgetStyle, bool bParentEnabled ) const override;
+	virtual void OnArrangeChildren( const FGeometry& AllottedGeometry, FArrangedChildren& ArrangedChildren ) const override;
+	virtual void CacheDesiredSize() override;
+	virtual FVector2D ComputeDesiredSize() const override;
+	virtual FChildren* GetChildren() override;
 	// End of SWidget interface
 
 	
diff --git a/Engine/Source/Runtime/Slate/Public/Widgets/Layout/SScrollBox.h b/Engine/Source/Runtime/Slate/Public/Widgets/Layout/SScrollBox.h
--- a/Engine/Source/Runtime/Slate/Public/Widgets/Layout/SScrollBox.h
+++ b/Engine/Source/Runtime/Slate/Public/Widgets/Layout/SScrollBox.h
@@ -63,15 +63,15 @@ public:
 public:
 
 	// SWidget interface
-	virtual void Tick( const FGeometry& AllottedGeometry, const double InCurrentTime, const float InDeltaTime ) OVERRIDE;
-	virtual FReply OnMouseButtonDown( const FGeometry& MyGeometry, const FPointerEvent& MouseEvent ) OVERRIDE;
-	virtual FReply OnMouseButtonUp( const FGeometry& MyGeometry, const FPointerEvent& MouseEvent ) OVERRIDE;
-	virtual FReply OnMouseMove( const FGeometry& MyGeometry, const FPointerEvent& MouseEvent ) OVERRIDE;
-	virtual void OnMouseLeave( const FPointerEvent& MouseEvent ) OVERRIDE;
-	virtual FReply OnMouseWheel( const FGeometry& MyGeometry, const FPointerEvent& MouseEvent ) OVERRIDE;
-	virtual FReply OnDragDetected( const FGeometry& MyGeometry, const FPointerEvent& MouseEvent ) OVERRIDE;
-	virtual FCursorReply OnCursorQuery( const FGeometry& MyGeometry, const FPointerEvent& CursorEvent ) const OVERRIDE;
-	virtual int32 OnPaint( const FGeometry& AllottedGeometry, const FSlateRect& MyClippingRect, FSlateWindowElementList& OutDrawElements, int32 LayerId, const FWidgetStyle& InWidgetStyle, bool bParentEnabled ) const OVERRIDE;
+	virtual void Tick( const FGeometry& AllottedGeometry, const double InCurrentTime, const float InDeltaTime ) override;
+	virtual FReply OnMouseButtonDown( const FGeometry& MyGeometry, const FPointerEvent& MouseEvent ) override;
+	virtual FReply OnMouseButtonUp( const FGeometry& MyGeometry, const FPointerEvent& MouseEvent ) override;
+	virtual FReply OnMouseMove( const FGeometry& MyGeometry, const FPointerEvent& MouseEvent ) override;
+	virtual void OnMouseLeave( const FPointerEvent& MouseEvent ) override;
+	virtual FReply OnMouseWheel( const FGeometry& MyGeometry, const FPointerEvent& MouseEvent ) override;
+	virtual FReply OnDragDetected( const FGeometry& MyGeometry, const FPointerEvent& MouseEvent ) override;
+	virtual FCursorReply OnCursorQuery( const FGeometry& MyGeometry, const FPointerEvent& CursorEvent ) const override;
+	virtual int32 OnPaint( const FGeometry& AllottedGeometry, const FSlateRect& MyClippingRect, FSlateWindowElementList& OutDrawElements, int32 LayerId, const FWidgetStyle& InWidgetStyle, bool bParentEnabled ) const override;
 	// End of SWidget interface
 
 private:
diff --git a/Engine/Source/Runtime/Slate/Public/Widgets/Layout/SSeparator.h b/Engine/Source/Runtime/Slate/Public/Widgets/Layout/SSeparator.h
--- a/Engine/Source/Runtime/Slate/Public/Widgets/Layout/SSeparator.h
+++ b/Engine/Source/Runtime/Slate/Public/Widgets/Layout/SSeparator.h
@@ -27,7 +27,7 @@ public:
 	void Construct( const FArguments& InArgs );
 
 	// SWidget interface
-	virtual FVector2D ComputeDesiredSize() const OVERRIDE;
+	virtual FVector2D ComputeDesiredSize() const override;
 	// End of SWidget interface
 private:
 	EOrientation Orientation;
diff --git a/Engine/Source/Runtime/Slate/Public/Widgets/Layout/SSpacer.h b/Engine/Source/Runtime/Slate/Public/Widgets/Layout/SSpacer.h
--- a/Engine/Source/Runtime/Slate/Public/Widgets/Layout/SSpacer.h
+++ b/Engine/Source/Runtime/Slate/Public/Widgets/Layout/SSpacer.h
@@ -24,8 +24,8 @@ public:
 	void Construct( const FArguments& InArgs );
 
 	// SWidget interface
-	virtual int32 OnPaint( const FGeometry& AllottedGeometry, const FSlateRect& MyClippingRect, FSlateWindowElementList& OutDrawElements, int32 LayerId, const FWidgetStyle& InWidgetStyle, bool bParentEnabled ) const OVERRIDE;
-	virtual FVector2D ComputeDesiredSize() const OVERRIDE;
+	virtual int32 OnPaint( const FGeometry& AllottedGeometry, const FSlateRect& MyClippingRect, FSlateWindowElementList& OutDrawElements, int32 LayerId, const FWidgetStyle& InWidgetStyle, bool bParentEnabled ) const override;
+	virtual FVector2D ComputeDesiredSize() const override;
 	// End of SWidget interface
 
 	FVector2D GetSize() const
diff --git a/Engine/Source/Runtime/Slate/Public/Widgets/Layout/SSplitter.h b/Engine/Source/Runtime/Slate/Public/Widgets/Layout/SSplitter.h
--- a/Engine/Source/Runtime/Slate/Public/Widgets/Layout/SSplitter.h
+++ b/Engine/Source/Runtime/Slate/Public/Widgets/Layout/SSplitter.h
@@ -152,17 +152,17 @@ public:
 
 public:
 
-	virtual void OnArrangeChildren( const FGeometry& AllottedGeometry, FArrangedChildren& ArrangedChildren ) const OVERRIDE;
+	virtual void OnArrangeChildren( const FGeometry& AllottedGeometry, FArrangedChildren& ArrangedChildren ) const override;
 
 
-	virtual int32 OnPaint( const FGeometry& AllottedGeometry, const FSlateRect& MyClippingRect, FSlateWindowElementList& OutDrawElements, int32 LayerId, const FWidgetStyle& InWidgetStyle, bool bParentEnabled ) const OVERRIDE;
+	virtual int32 OnPaint( const FGeometry& AllottedGeometry, const FSlateRect& MyClippingRect, FSlateWindowElementList& OutDrawElements, int32 LayerId, const FWidgetStyle& InWidgetStyle, bool bParentEnabled ) const override;
 
 
 	/**
 	 * A Panel's desired size in the space required to arrange of its children on the screen while respecting all of
 	 * the children's desired sizes and any layout-related options specified by the user. See StackPanel for an example.
 	 */
-	virtual FVector2D ComputeDesiredSize() const OVERRIDE;
+	virtual FVector2D ComputeDesiredSize() const override;
 
 	/**
 	 * All widgets must provide a way to access their children in a layout-agnostic way.
@@ -171,7 +171,7 @@ public:
 	 * provides layout information about the child it stores. In that case
 	 * GetChildren should simply return the TPanelChildren<Slot>. See StackPanel for an example.
 	 */
-	virtual FChildren* GetChildren() OVERRIDE;
+	virtual FChildren* GetChildren() override;
 
 	/**
 	 * The system calls this method to notify the widget that a mouse button was pressed within it. This event is bubbled.
@@ -181,7 +181,7 @@ public:
 	 *
 	 * @return Whether the event was handled along with possible requests for the system to take action.
 	 */
-	virtual FReply OnMouseButtonDown( const FGeometry& MyGeometry, const FPointerEvent& MouseEvent ) OVERRIDE;
+	virtual FReply OnMouseButtonDown( const FGeometry& MyGeometry, const FPointerEvent& MouseEvent ) override;
 	
 	/**
 	 * The system calls this method to notify the widget that a mouse button was release within it. This event is bubbled.
@@ -191,18 +191,18 @@ public:
 	 *
 	 * @return Whether the event was handled along with possible requests for the system to take action.
 	 */
-	virtual FReply OnMouseButtonUp( const FGeometry& MyGeometry, const FPointerEvent& MouseEvent ) OVERRIDE;
+	virtual FReply OnMouseButtonUp( const FGeometry& MyGeometry, const FPointerEvent& MouseEvent ) override;
 	
-	virtual FReply OnMouseMove( const FGeometry& MyGeometry, const FPointerEvent& MouseEvent ) OVERRIDE;
+	virtual FReply OnMouseMove( const FGeometry& MyGeometry, const FPointerEvent& MouseEvent ) override;
 
-	virtual void OnMouseLeave( const FPointerEvent& MouseEvent ) OVERRIDE;
+	virtual void OnMouseLeave( const FPointerEvent& MouseEvent ) override;
 
 	/**
 	 * The system asks each widget under the mouse to provide a cursor. This event is bubbled.
 	 * 
 	 * @return FCursorReply::Unhandled() if the event is not handled; return FCursorReply::Cursor() otherwise.
 	 */
-	virtual FCursorReply OnCursorQuery( const FGeometry& MyGeometry, const FPointerEvent& CursorEvent ) const OVERRIDE;
+	virtual FCursorReply OnCursorQuery( const FGeometry& MyGeometry, const FPointerEvent& CursorEvent ) const override;
 
 	/**
 	 * Change the orientation of the splitter
@@ -406,13 +406,13 @@ public:
 
 private:
 
-	virtual void OnArrangeChildren( const FGeometry& AllottedGeometry, FArrangedChildren& ArrangedChildren ) const OVERRIDE;
+	virtual void OnArrangeChildren( const FGeometry& AllottedGeometry, FArrangedChildren& ArrangedChildren ) const override;
 
 	/**
 	 * A Panel's desired size in the space required to arrange of its children on the screen while respecting all of
 	 * the children's desired sizes and any layout-related options specified by the user. See StackPanel for an example.
 	 */
-	virtual FVector2D ComputeDesiredSize() const OVERRIDE;
+	virtual FVector2D ComputeDesiredSize() const override;
 
 	/**
 	 * All widgets must provide a way to access their children in a layout-agnostic way.
@@ -421,7 +421,7 @@ private:
 	 * provides layout information about the child it stores. In that case
 	 * GetChildren should simply return the TPanelChildren<Slot>. See StackPanel for an example.
 	 */
-	virtual FChildren* GetChildren() OVERRIDE;
+	virtual FChildren* GetChildren() override;
 
 	/**
 	 * The system calls this method to notify the widget that a mouse button was pressed within it. This event is bubbled.
@@ -431,7 +431,7 @@ private:
 	 *
 	 * @return Whether the event was handled along with possible requests for the system to take action.
 	 */
- 	virtual FReply OnMouseButtonDown( const FGeometry& MyGeometry, const FPointerEvent& MouseEvent ) OVERRIDE;
+ 	virtual FReply OnMouseButtonDown( const FGeometry& MyGeometry, const FPointerEvent& MouseEvent ) override;
 	
 	/**
 	 * The system calls this method to notify the widget that a mouse button was release within it. This event is bubbled.
@@ -441,7 +441,7 @@ private:
 	 *
 	 * @return Whether the event was handled along with possible requests for the system to take action.
 	 */
- 	virtual FReply OnMouseButtonUp( const FGeometry& MyGeometry, const FPointerEvent& MouseEvent ) OVERRIDE;
+ 	virtual FReply OnMouseButtonUp( const FGeometry& MyGeometry, const FPointerEvent& MouseEvent ) override;
 
 	/**
 	 * The system calls this method to notify the widget that a mouse moved within it. This event is bubbled.
@@ -451,12 +451,12 @@ private:
 	 *
 	 * @return Whether the event was handled along with possible requests for the system to take action.
 	 */
-	virtual FReply OnMouseMove( const FGeometry& MyGeometry, const FPointerEvent& MouseEvent ) OVERRIDE;
+	virtual FReply OnMouseMove( const FGeometry& MyGeometry, const FPointerEvent& MouseEvent ) override;
 
 	/**
 	 * @return The cursor that should be visible
 	 */
-	virtual FCursorReply OnCursorQuery( const FGeometry& MyGeometry, const FPointerEvent& CursorEvent ) const OVERRIDE;
+	virtual FCursorReply OnCursorQuery( const FGeometry& MyGeometry, const FPointerEvent& CursorEvent ) const override;
 
 	/**
 	 * Calculates the axis being resized
diff --git a/Engine/Source/Runtime/Slate/Public/Widgets/Layout/SUniformGridPanel.h b/Engine/Source/Runtime/Slate/Public/Widgets/Layout/SUniformGridPanel.h
--- a/Engine/Source/Runtime/Slate/Public/Widgets/Layout/SUniformGridPanel.h
+++ b/Engine/Source/Runtime/Slate/Public/Widgets/Layout/SUniformGridPanel.h
@@ -53,9 +53,9 @@ public:
 	void Construct( const FArguments& InArgs );
 
 	// BEGIN SPanel INTERFACE	
-	virtual void OnArrangeChildren( const FGeometry& AllottedGeometry, FArrangedChildren& ArrangedChildren ) const OVERRIDE;
-	virtual FVector2D ComputeDesiredSize() const OVERRIDE;
-	virtual FChildren* GetChildren() OVERRIDE;
+	virtual void OnArrangeChildren( const FGeometry& AllottedGeometry, FArrangedChildren& ArrangedChildren ) const override;
+	virtual FVector2D ComputeDesiredSize() const override;
+	virtual FChildren* GetChildren() override;
 	// END SPanel INTERFACE
 
 	/**
diff --git a/Engine/Source/Runtime/Slate/Public/Widgets/Layout/SWidgetSwitcher.h b/Engine/Source/Runtime/Slate/Public/Widgets/Layout/SWidgetSwitcher.h
--- a/Engine/Source/Runtime/Slate/Public/Widgets/Layout/SWidgetSwitcher.h
+++ b/Engine/Source/Runtime/Slate/Public/Widgets/Layout/SWidgetSwitcher.h
@@ -120,11 +120,11 @@ protected:
 
 	// Begin SCompoundWidget interface
 
-	virtual void OnArrangeChildren( const FGeometry& AllottedGeometry, FArrangedChildren& ArrangedChildren ) const OVERRIDE;
+	virtual void OnArrangeChildren( const FGeometry& AllottedGeometry, FArrangedChildren& ArrangedChildren ) const override;
 	
-	virtual FVector2D ComputeDesiredSize( ) const OVERRIDE;
+	virtual FVector2D ComputeDesiredSize( ) const override;
 
-	virtual FChildren* GetChildren( ) OVERRIDE;
+	virtual FChildren* GetChildren( ) override;
 
 	// End SCompoundWidget interface
 
@@ -141,11 +141,11 @@ private:
 			, WidgetIndex( InWidgetIndex )
 		{ }
 		
-		virtual int32 Num() const OVERRIDE { return AllChildren->Num() > 0 ? 1 : 0; }
+		virtual int32 Num() const override { return AllChildren->Num() > 0 ? 1 : 0; }
 		
-		virtual TSharedRef<SWidget> GetChildAt( int32 Index ) OVERRIDE { check(Index == 0); return (*AllChildren)[WidgetIndex->Get()].Widget; }
+		virtual TSharedRef<SWidget> GetChildAt( int32 Index ) override { check(Index == 0); return (*AllChildren)[WidgetIndex->Get()].Widget; }
 		
-		virtual TSharedRef<const SWidget> GetChildAt( int32 Index ) const OVERRIDE { check(Index == 0); return (*AllChildren)[WidgetIndex->Get()].Widget; }
+		virtual TSharedRef<const SWidget> GetChildAt( int32 Index ) const override { check(Index == 0); return (*AllChildren)[WidgetIndex->Get()].Widget; }
 		
 	private:
 
diff --git a/Engine/Source/Runtime/Slate/Public/Widgets/Layout/SWrapBox.h b/Engine/Source/Runtime/Slate/Public/Widgets/Layout/SWrapBox.h
--- a/Engine/Source/Runtime/Slate/Public/Widgets/Layout/SWrapBox.h
+++ b/Engine/Source/Runtime/Slate/Public/Widgets/Layout/SWrapBox.h
@@ -82,15 +82,15 @@ public:
 
 	void Construct( const FArguments & InArgs );
 
-	virtual void Tick( const FGeometry& AllottedGeometry, const double InCurrentTime, const float InDeltaTime ) OVERRIDE;
+	virtual void Tick( const FGeometry& AllottedGeometry, const double InCurrentTime, const float InDeltaTime ) override;
 
-	virtual void OnArrangeChildren( const FGeometry& AllottedGeometry, FArrangedChildren& ArrangedChildren ) const OVERRIDE;
+	virtual void OnArrangeChildren( const FGeometry& AllottedGeometry, FArrangedChildren& ArrangedChildren ) const override;
 
 	void ClearChildren();
 
-	virtual FVector2D ComputeDesiredSize() const OVERRIDE;
+	virtual FVector2D ComputeDesiredSize() const override;
 
-	virtual FChildren* GetChildren() OVERRIDE;
+	virtual FChildren* GetChildren() override;
 
 private:
 
diff --git a/Engine/Source/Runtime/Slate/Public/Widgets/Navigation/SWizard.h b/Engine/Source/Runtime/Slate/Public/Widgets/Navigation/SWizard.h
--- a/Engine/Source/Runtime/Slate/Public/Widgets/Navigation/SWizard.h
+++ b/Engine/Source/Runtime/Slate/Public/Widgets/Navigation/SWizard.h
@@ -247,7 +247,7 @@ public:
 
 	// Begin SCompoundWidget interface
 
-	virtual FVector2D ComputeDesiredSize( ) const OVERRIDE;
+	virtual FVector2D ComputeDesiredSize( ) const override;
 
 	// End SCompoundWidget interface
 
diff --git a/Engine/Source/Runtime/Slate/Public/Widgets/Notifications/SErrorHint.h b/Engine/Source/Runtime/Slate/Public/Widgets/Notifications/SErrorHint.h
--- a/Engine/Source/Runtime/Slate/Public/Widgets/Notifications/SErrorHint.h
+++ b/Engine/Source/Runtime/Slate/Public/Widgets/Notifications/SErrorHint.h
@@ -15,12 +15,12 @@ public:
 
 	// IErrorReportingWidget interface
 
-	virtual void SetError( const FText& InErrorText ) OVERRIDE;
-	virtual void SetError( const FString& InErrorText ) OVERRIDE;
+	virtual void SetError( const FText& InErrorText ) override;
+	virtual void SetError( const FString& InErrorText ) override;
 
-	virtual bool HasError() const OVERRIDE;
+	virtual bool HasError() const override;
 
-	virtual TSharedRef<SWidget> AsWidget() OVERRIDE;
+	virtual TSharedRef<SWidget> AsWidget() override;
 
 	// IErrorReportingWidget interface
 
diff --git a/Engine/Source/Runtime/Slate/Public/Widgets/Notifications/SErrorText.h b/Engine/Source/Runtime/Slate/Public/Widgets/Notifications/SErrorText.h
--- a/Engine/Source/Runtime/Slate/Public/Widgets/Notifications/SErrorText.h
+++ b/Engine/Source/Runtime/Slate/Public/Widgets/Notifications/SErrorText.h
@@ -28,12 +28,12 @@ public:
 
 	// IErrorReportingWidget interface
 
-	virtual void SetError( const FText& InErrorText ) OVERRIDE;
-	virtual void SetError( const FString& InErrorText ) OVERRIDE;
+	virtual void SetError( const FText& InErrorText ) override;
+	virtual void SetError( const FString& InErrorText ) override;
 
-	virtual bool HasError() const OVERRIDE;
+	virtual bool HasError() const override;
 
-	virtual TSharedRef<SWidget> AsWidget() OVERRIDE;
+	virtual TSharedRef<SWidget> AsWidget() override;
 
 	// IErrorReportingWidget interface
 
diff --git a/Engine/Source/Runtime/Slate/Public/Widgets/Notifications/SPopUpErrorText.h b/Engine/Source/Runtime/Slate/Public/Widgets/Notifications/SPopUpErrorText.h
--- a/Engine/Source/Runtime/Slate/Public/Widgets/Notifications/SPopUpErrorText.h
+++ b/Engine/Source/Runtime/Slate/Public/Widgets/Notifications/SPopUpErrorText.h
@@ -18,12 +18,12 @@ public:
 
 	// IErrorReportingWidget interface
 
-	virtual void SetError( const FText& InErrorText ) OVERRIDE;
-	virtual void SetError( const FString& InErrorText ) OVERRIDE;
+	virtual void SetError( const FText& InErrorText ) override;
+	virtual void SetError( const FString& InErrorText ) override;
 
-	virtual bool HasError() const OVERRIDE;
+	virtual bool HasError() const override;
 
-	virtual TSharedRef<SWidget> AsWidget() OVERRIDE;
+	virtual TSharedRef<SWidget> AsWidget() override;
 
 	// IErrorReportingWidget interface
 
diff --git a/Engine/Source/Runtime/Slate/Public/Widgets/Notifications/SProgressBar.h b/Engine/Source/Runtime/Slate/Public/Widgets/Notifications/SProgressBar.h
--- a/Engine/Source/Runtime/Slate/Public/Widgets/Notifications/SProgressBar.h
+++ b/Engine/Source/Runtime/Slate/Public/Widgets/Notifications/SProgressBar.h
@@ -58,7 +58,7 @@ public:
 	 */
 	void Construct(const FArguments& InArgs);
 
-	virtual int32 OnPaint( const FGeometry& AllottedGeometry, const FSlateRect& MyClippingRect, FSlateWindowElementList& OutDrawElements, int32 LayerId, const FWidgetStyle& InWidgetStyle, bool bParentEnabled ) const OVERRIDE;
+	virtual int32 OnPaint( const FGeometry& AllottedGeometry, const FSlateRect& MyClippingRect, FSlateWindowElementList& OutDrawElements, int32 LayerId, const FWidgetStyle& InWidgetStyle, bool bParentEnabled ) const override;
 
 	FVector2D ComputeDesiredSize() const;
 
diff --git a/Engine/Source/Runtime/Slate/Public/Widgets/SCanvas.h b/Engine/Source/Runtime/Slate/Public/Widgets/SCanvas.h
--- a/Engine/Source/Runtime/Slate/Public/Widgets/SCanvas.h
+++ b/Engine/Source/Runtime/Slate/Public/Widgets/SCanvas.h
@@ -153,13 +153,13 @@ public:
 
 	// Begin SWidget overrides
 
-	virtual void OnArrangeChildren( const FGeometry& AllottedGeometry, FArrangedChildren& ArrangedChildren ) const OVERRIDE;
+	virtual void OnArrangeChildren( const FGeometry& AllottedGeometry, FArrangedChildren& ArrangedChildren ) const override;
 	
-	virtual int32 OnPaint( const FGeometry& AllottedGeometry, const FSlateRect& MyClippingRect, FSlateWindowElementList& OutDrawElements, int32 LayerId, const FWidgetStyle& InWidgetStyle, bool bParentEnabled ) const OVERRIDE;
+	virtual int32 OnPaint( const FGeometry& AllottedGeometry, const FSlateRect& MyClippingRect, FSlateWindowElementList& OutDrawElements, int32 LayerId, const FWidgetStyle& InWidgetStyle, bool bParentEnabled ) const override;
 
-	virtual FVector2D ComputeDesiredSize() const OVERRIDE;
+	virtual FVector2D ComputeDesiredSize() const override;
 
-	virtual FChildren* GetChildren() OVERRIDE;
+	virtual FChildren* GetChildren() override;
 
 	// End SWidget overrides
 
diff --git a/Engine/Source/Runtime/Slate/Public/Widgets/SToolTip.h b/Engine/Source/Runtime/Slate/Public/Widgets/SToolTip.h
--- a/Engine/Source/Runtime/Slate/Public/Widgets/SToolTip.h
+++ b/Engine/Source/Runtime/Slate/Public/Widgets/SToolTip.h
@@ -60,19 +60,19 @@ public:
 
 	// Begin IToolTip interface
 
-	virtual TSharedRef<class SWidget> AsWidget( ) OVERRIDE
+	virtual TSharedRef<class SWidget> AsWidget( ) override
 	{
 		return AsShared();
 	}
 
-	virtual const TSharedRef<SWidget>& GetContentWidget( ) OVERRIDE
+	virtual const TSharedRef<SWidget>& GetContentWidget( ) override
 	{
 		return GetContent();
 	}
 
-	virtual bool IsEmpty( ) const OVERRIDE;
+	virtual bool IsEmpty( ) const override;
 
-	virtual bool IsInteractive( ) const OVERRIDE;
+	virtual bool IsInteractive( ) const override;
 
 	// End IToolTip interface
 
diff --git a/Engine/Source/Runtime/Slate/Public/Widgets/SViewport.h b/Engine/Source/Runtime/Slate/Public/Widgets/SViewport.h
--- a/Engine/Source/Runtime/Slate/Public/Widgets/SViewport.h
+++ b/Engine/Source/Runtime/Slate/Public/Widgets/SViewport.h
@@ -62,12 +62,12 @@ public:
 public:
 
 	/** SViewport wants keyboard focus */
-	virtual bool SupportsKeyboardFocus() const OVERRIDE { return true; }
+	virtual bool SupportsKeyboardFocus() const override { return true; }
 
 	/** 
 	 * Computes the ideal size necessary to display this widget.
 	 */
-	virtual FVector2D ComputeDesiredSize() const OVERRIDE
+	virtual FVector2D ComputeDesiredSize() const override
 	{
 		return FVector2D( 320, 240 );
 	}
@@ -129,29 +129,29 @@ public:
 
 	// SWidget interface
 
-	virtual int32 OnPaint( const FGeometry& AllottedGeometry, const FSlateRect& MyClippingRect, FSlateWindowElementList& OutDrawElements, int32 LayerId, const FWidgetStyle& InWidgetStyle, bool bParentEnabled ) const OVERRIDE;
-	virtual void Tick( const FGeometry& AllottedGeometry, const double InCurrentTime, const float InDeltaTime ) OVERRIDE;
-	virtual FReply OnTouchStarted( const FGeometry& MyGeometry, const FPointerEvent& TouchEvent ) OVERRIDE;
-	virtual FReply OnTouchMoved( const FGeometry& MyGeometry, const FPointerEvent& TouchEvent ) OVERRIDE;
-	virtual FReply OnTouchEnded( const FGeometry& MyGeometry, const FPointerEvent& TouchEvent ) OVERRIDE;
-	virtual FReply OnTouchGesture( const FGeometry& MyGeometry, const FPointerEvent& GestureEvent ) OVERRIDE;
-	virtual FCursorReply OnCursorQuery( const FGeometry& MyGeometry, const FPointerEvent& CursorEvent ) const OVERRIDE;
-	virtual FReply OnMouseButtonDown( const FGeometry& MyGeometry, const FPointerEvent& MouseEvent ) OVERRIDE;
-	virtual FReply OnMouseButtonUp( const FGeometry& MyGeometry, const FPointerEvent& MouseEvent ) OVERRIDE;
-	virtual void OnMouseEnter( const FGeometry& MyGeometry, const FPointerEvent& MouseEvent ) OVERRIDE;
-	virtual void OnMouseLeave( const FPointerEvent& MouseEvent ) OVERRIDE;
-	virtual FReply OnMouseMove( const FGeometry& MyGeometry, const FPointerEvent& MouseEvent ) OVERRIDE;
-	virtual FReply OnMouseWheel( const FGeometry& MyGeometry, const FPointerEvent& MouseEvent ) OVERRIDE;
-	virtual FReply OnMouseButtonDoubleClick( const FGeometry& MyGeometry, const FPointerEvent& MouseEvent ) OVERRIDE;
-	virtual FReply OnKeyDown( const FGeometry& MyGeometry, const FKeyboardEvent& KeyboardEvent ) OVERRIDE;
-	virtual FReply OnKeyUp( const FGeometry& MyGeometry, const FKeyboardEvent& KeyboardEvent ) OVERRIDE;
-	virtual FReply OnKeyChar( const FGeometry& MyGeometry, const FCharacterEvent& CharacterEvent ) OVERRIDE;
-	virtual FReply OnKeyboardFocusReceived( const FGeometry& MyGeometry, const FKeyboardFocusEvent& InKeyboardFocusEvent ) OVERRIDE;
-	virtual void OnKeyboardFocusLost( const FKeyboardFocusEvent& InKeyboardFocusEvent ) OVERRIDE;
-	virtual FReply OnControllerButtonPressed( const FGeometry& MyGeometry, const FControllerEvent& ControllerEvent ) OVERRIDE;
-	virtual FReply OnControllerButtonReleased( const FGeometry& MyGeometry, const FControllerEvent& ControllerEvent ) OVERRIDE;
-	virtual FReply OnControllerAnalogValueChanged( const FGeometry& MyGeometry, const FControllerEvent& ControllerEvent ) OVERRIDE;
-	virtual FReply OnMotionDetected( const FGeometry& MyGeometry, const FMotionEvent& MotionEvent ) OVERRIDE;
+	virtual int32 OnPaint( const FGeometry& AllottedGeometry, const FSlateRect& MyClippingRect, FSlateWindowElementList& OutDrawElements, int32 LayerId, const FWidgetStyle& InWidgetStyle, bool bParentEnabled ) const override;
+	virtual void Tick( const FGeometry& AllottedGeometry, const double InCurrentTime, const float InDeltaTime ) override;
+	virtual FReply OnTouchStarted( const FGeometry& MyGeometry, const FPointerEvent& TouchEvent ) override;
+	virtual FReply OnTouchMoved( const FGeometry& MyGeometry, const FPointerEvent& TouchEvent ) override;
+	virtual FReply OnTouchEnded( const FGeometry& MyGeometry, const FPointerEvent& TouchEvent ) override;
+	virtual FReply OnTouchGesture( const FGeometry& MyGeometry, const FPointerEvent& GestureEvent ) override;
+	virtual FCursorReply OnCursorQuery( const FGeometry& MyGeometry, const FPointerEvent& CursorEvent ) const override;
+	virtual FReply OnMouseButtonDown( const FGeometry& MyGeometry, const FPointerEvent& MouseEvent ) override;
+	virtual FReply OnMouseButtonUp( const FGeometry& MyGeometry, const FPointerEvent& MouseEvent ) override;
+	virtual void OnMouseEnter( const FGeometry& MyGeometry, const FPointerEvent& MouseEvent ) override;
+	virtual void OnMouseLeave( const FPointerEvent& MouseEvent ) override;
+	virtual FReply OnMouseMove( const FGeometry& MyGeometry, const FPointerEvent& MouseEvent ) override;
+	virtual FReply OnMouseWheel( const FGeometry& MyGeometry, const FPointerEvent& MouseEvent ) override;
+	virtual FReply OnMouseButtonDoubleClick( const FGeometry& MyGeometry, const FPointerEvent& MouseEvent ) override;
+	virtual FReply OnKeyDown( const FGeometry& MyGeometry, const FKeyboardEvent& KeyboardEvent ) override;
+	virtual FReply OnKeyUp( const FGeometry& MyGeometry, const FKeyboardEvent& KeyboardEvent ) override;
+	virtual FReply OnKeyChar( const FGeometry& MyGeometry, const FCharacterEvent& CharacterEvent ) override;
+	virtual FReply OnKeyboardFocusReceived( const FGeometry& MyGeometry, const FKeyboardFocusEvent& InKeyboardFocusEvent ) override;
+	virtual void OnKeyboardFocusLost( const FKeyboardFocusEvent& InKeyboardFocusEvent ) override;
+	virtual FReply OnControllerButtonPressed( const FGeometry& MyGeometry, const FControllerEvent& ControllerEvent ) override;
+	virtual FReply OnControllerButtonReleased( const FGeometry& MyGeometry, const FControllerEvent& ControllerEvent ) override;
+	virtual FReply OnControllerAnalogValueChanged( const FGeometry& MyGeometry, const FControllerEvent& ControllerEvent ) override;
+	virtual FReply OnMotionDetected( const FGeometry& MyGeometry, const FMotionEvent& MotionEvent ) override;
 
 	// End of SWidget interface
 
diff --git a/Engine/Source/Runtime/Slate/Public/Widgets/SWeakWidget.h b/Engine/Source/Runtime/Slate/Public/Widgets/SWeakWidget.h
--- a/Engine/Source/Runtime/Slate/Public/Widgets/SWeakWidget.h
+++ b/Engine/Source/Runtime/Slate/Public/Widgets/SWeakWidget.h
@@ -41,13 +41,13 @@ public:
 
 	// Begin SWidget interface
 
-	virtual FVector2D ComputeDesiredSize() const OVERRIDE;
+	virtual FVector2D ComputeDesiredSize() const override;
 
-	virtual void OnArrangeChildren( const FGeometry& AllottedGeometry, FArrangedChildren& ArrangedChildren ) const OVERRIDE;
+	virtual void OnArrangeChildren( const FGeometry& AllottedGeometry, FArrangedChildren& ArrangedChildren ) const override;
 		
-	virtual FChildren* GetChildren() OVERRIDE;
+	virtual FChildren* GetChildren() override;
 
-	virtual int32 OnPaint( const FGeometry& AllottedGeometry, const FSlateRect& MyClippingRect, FSlateWindowElementList& OutDrawElements, int32 LayerId, const FWidgetStyle& InWidgetStyle, bool bParentEnabled ) const OVERRIDE;
+	virtual int32 OnPaint( const FGeometry& AllottedGeometry, const FSlateRect& MyClippingRect, FSlateWindowElementList& OutDrawElements, int32 LayerId, const FWidgetStyle& InWidgetStyle, bool bParentEnabled ) const override;
 
 	// End SWidget interface
 
diff --git a/Engine/Source/Runtime/Slate/Public/Widgets/Text/SInlineEditableTextBlock.h b/Engine/Source/Runtime/Slate/Public/Widgets/Text/SInlineEditableTextBlock.h
--- a/Engine/Source/Runtime/Slate/Public/Widgets/Text/SInlineEditableTextBlock.h
+++ b/Engine/Source/Runtime/Slate/Public/Widgets/Text/SInlineEditableTextBlock.h
@@ -78,17 +78,17 @@ class SLATE_API SInlineEditableTextBlock: public SCompoundWidget
 	 *
 	 * @return  This widget can gain focus if IsSelected is not bound.
 	 */
-	virtual bool SupportsKeyboardFocus() const OVERRIDE;
+	virtual bool SupportsKeyboardFocus() const override;
 
-	virtual FReply OnMouseButtonDown( const FGeometry& MyGeometry, const FPointerEvent& MouseEvent ) OVERRIDE;
+	virtual FReply OnMouseButtonDown( const FGeometry& MyGeometry, const FPointerEvent& MouseEvent ) override;
 
-	virtual FReply OnDragOver( const FGeometry& MyGeometry, const FDragDropEvent& DragDropEvent ) OVERRIDE;
+	virtual FReply OnDragOver( const FGeometry& MyGeometry, const FDragDropEvent& DragDropEvent ) override;
 
-	virtual FReply OnMouseButtonDoubleClick( const FGeometry& InMyGeometry, const FPointerEvent& InMouseEvent ) OVERRIDE;
+	virtual FReply OnMouseButtonDoubleClick( const FGeometry& InMyGeometry, const FPointerEvent& InMouseEvent ) override;
 
-	virtual FReply OnKeyDown( const FGeometry& MyGeometry, const FKeyboardEvent& InKeyboardEvent ) OVERRIDE;
+	virtual FReply OnKeyDown( const FGeometry& MyGeometry, const FKeyboardEvent& InKeyboardEvent ) override;
 
-	virtual void Tick( const FGeometry& AllottedGeometry, const double InCurrentTime, const float InDeltaTime ) OVERRIDE;
+	virtual void Tick( const FGeometry& AllottedGeometry, const double InCurrentTime, const float InDeltaTime ) override;
 	 
 	/** Switches the widget to editing mode */
 	void EnterEditingMode();
diff --git a/Engine/Source/Runtime/Slate/Public/Widgets/Text/SMultiLineEditableText.h b/Engine/Source/Runtime/Slate/Public/Widgets/Text/SMultiLineEditableText.h
--- a/Engine/Source/Runtime/Slate/Public/Widgets/Text/SMultiLineEditableText.h
+++ b/Engine/Source/Runtime/Slate/Public/Widgets/Text/SMultiLineEditableText.h
@@ -80,8 +80,8 @@ public:
 	 */
 	void SetText(const TAttribute< FText >& InText);
 
-	virtual bool GetIsReadOnly() const OVERRIDE;
-	virtual void ClearSelection() OVERRIDE;
+	virtual bool GetIsReadOnly() const override;
+	virtual void ClearSelection() override;
 
 private:
 	
@@ -206,11 +206,11 @@ private:
 
 		virtual ~FSlateCursorRunHighlighter() {}
 
-		virtual int32 OnPaint( const FTextLayout::FLineView& Line, const TSharedRef< ISlateRun >& Run, const TSharedRef< ILayoutBlock >& Block, const FTextBlockStyle& DefaultStyle, const FGeometry& AllottedGeometry, const FSlateRect& MyClippingRect, FSlateWindowElementList& OutDrawElements, int32 LayerId, const FWidgetStyle& InWidgetStyle, bool bParentEnabled ) const OVERRIDE;
+		virtual int32 OnPaint( const FTextLayout::FLineView& Line, const TSharedRef< ISlateRun >& Run, const TSharedRef< ILayoutBlock >& Block, const FTextBlockStyle& DefaultStyle, const FGeometry& AllottedGeometry, const FSlateRect& MyClippingRect, FSlateWindowElementList& OutDrawElements, int32 LayerId, const FWidgetStyle& InWidgetStyle, bool bParentEnabled ) const override;
 
-		virtual FChildren* GetChildren() OVERRIDE;
+		virtual FChildren* GetChildren() override;
 
-		virtual void OnArrangeChildren( const TSharedRef< ILayoutBlock >& Block, const FGeometry& AllottedGeometry, FArrangedChildren& ArrangedChildren ) const OVERRIDE;
+		virtual void OnArrangeChildren( const TSharedRef< ILayoutBlock >& Block, const FGeometry& AllottedGeometry, FArrangedChildren& ArrangedChildren ) const override;
 
 		void SetMode( EMode Mode );
 
@@ -231,62 +231,62 @@ private:
 private:
 
 	// BEGIN ITextEditorWidget interface
-	virtual void StartChangingText() OVERRIDE;
-	virtual void FinishChangingText() OVERRIDE;
-	virtual void BackspaceChar() OVERRIDE;
-	virtual void DeleteChar() OVERRIDE;
-	virtual bool CanTypeCharacter(const TCHAR CharInQuestion) const OVERRIDE;
-	virtual void TypeChar( const int32 Character ) OVERRIDE;
-	virtual FReply MoveCursor( ECursorMoveMethod::Type Method, const FVector2D& Direction, ECursorAction::Type Action ) OVERRIDE;
-	virtual void JumpTo(ETextLocation::Type JumpLocation, ECursorAction::Type Action) OVERRIDE;
-	virtual void SelectAllText() OVERRIDE;
-	virtual bool SelectAllTextWhenFocused() OVERRIDE;
-	virtual void SelectWordAt(const FVector2D& LocalPosition) OVERRIDE;
-	virtual void BeginDragSelection() OVERRIDE;
-	virtual bool IsDragSelecting() const OVERRIDE;
-	virtual void EndDragSelection() OVERRIDE;
-	virtual bool AnyTextSelected() const OVERRIDE;
-	virtual bool IsTextSelectedAt(const FVector2D& LocalPosition) const OVERRIDE;
-	virtual void SetWasFocusedByLastMouseDown( bool Value ) OVERRIDE;
-	virtual bool WasFocusedByLastMouseDown() const OVERRIDE;
-	virtual void SetHasDragSelectedSinceFocused( bool Value ) OVERRIDE;
-	virtual bool HasDragSelectedSinceFocused() const OVERRIDE;
-	virtual FReply OnEscape() OVERRIDE;
-	virtual void OnEnter() OVERRIDE;
-	virtual bool CanExecuteCut() const OVERRIDE;
-	virtual void CutSelectedTextToClipboard() OVERRIDE;
-	virtual bool CanExecuteCopy() const OVERRIDE;
-	virtual void CopySelectedTextToClipboard() OVERRIDE;
-	virtual bool CanExecutePaste() const OVERRIDE;
-	virtual void PasteTextFromClipboard() OVERRIDE;
-	virtual bool CanExecuteUndo() const OVERRIDE;
-	virtual void Undo() OVERRIDE;
-	virtual void Redo() OVERRIDE;
-	virtual TSharedRef< SWidget > GetWidget() OVERRIDE;
-	virtual void SummonContextMenu( const FVector2D& InLocation ) OVERRIDE;
-	virtual void LoadText() OVERRIDE;
+	virtual void StartChangingText() override;
+	virtual void FinishChangingText() override;
+	virtual void BackspaceChar() override;
+	virtual void DeleteChar() override;
+	virtual bool CanTypeCharacter(const TCHAR CharInQuestion) const override;
+	virtual void TypeChar( const int32 Character ) override;
+	virtual FReply MoveCursor( ECursorMoveMethod::Type Method, const FVector2D& Direction, ECursorAction::Type Action ) override;
+	virtual void JumpTo(ETextLocation::Type JumpLocation, ECursorAction::Type Action) override;
+	virtual void SelectAllText() override;
+	virtual bool SelectAllTextWhenFocused() override;
+	virtual void SelectWordAt(const FVector2D& LocalPosition) override;
+	virtual void BeginDragSelection() override;
+	virtual bool IsDragSelecting() const override;
+	virtual void EndDragSelection() override;
+	virtual bool AnyTextSelected() const override;
+	virtual bool IsTextSelectedAt(const FVector2D& LocalPosition) const override;
+	virtual void SetWasFocusedByLastMouseDown( bool Value ) override;
+	virtual bool WasFocusedByLastMouseDown() const override;
+	virtual void SetHasDragSelectedSinceFocused( bool Value ) override;
+	virtual bool HasDragSelectedSinceFocused() const override;
+	virtual FReply OnEscape() override;
+	virtual void OnEnter() override;
+	virtual bool CanExecuteCut() const override;
+	virtual void CutSelectedTextToClipboard() override;
+	virtual bool CanExecuteCopy() const override;
+	virtual void CopySelectedTextToClipboard() override;
+	virtual bool CanExecutePaste() const override;
+	virtual void PasteTextFromClipboard() override;
+	virtual bool CanExecuteUndo() const override;
+	virtual void Undo() override;
+	virtual void Redo() override;
+	virtual TSharedRef< SWidget > GetWidget() override;
+	virtual void SummonContextMenu( const FVector2D& InLocation ) override;
+	virtual void LoadText() override;
 	// END ITextEditorWidget interface
 
 
 private:
 	// BEGIN SWidget interface
-	virtual void Tick( const FGeometry& AllottedGeometry, const double InCurrentTime, const float InDeltaTime ) OVERRIDE;
-	virtual int32 OnPaint( const FGeometry& AllottedGeometry, const FSlateRect& MyClippingRect, FSlateWindowElementList& OutDrawElements, int32 LayerId, const FWidgetStyle& InWidgetStyle, bool bParentEnabled ) const OVERRIDE;
-	virtual void CacheDesiredSize() OVERRIDE;
-	virtual FVector2D ComputeDesiredSize() const OVERRIDE;
-	virtual FChildren* GetChildren() OVERRIDE;
-	virtual void OnArrangeChildren( const FGeometry& AllottedGeometry, FArrangedChildren& ArrangedChildren ) const OVERRIDE;
-	virtual bool SupportsKeyboardFocus() const OVERRIDE;
-	virtual FReply OnKeyChar( const FGeometry& MyGeometry,const FCharacterEvent& InCharacterEvent ) OVERRIDE;
-	virtual FReply OnKeyDown( const FGeometry& MyGeometry, const FKeyboardEvent& InKeyboardEvent ) OVERRIDE;
-	virtual FReply OnKeyUp( const FGeometry& MyGeometry, const FKeyboardEvent& InKeyboardEvent ) OVERRIDE;
-	virtual FReply OnKeyboardFocusReceived( const FGeometry& MyGeometry, const FKeyboardFocusEvent& InKeyboardFocusEvent ) OVERRIDE;
-	virtual void OnKeyboardFocusLost( const FKeyboardFocusEvent& InKeyboardFocusEvent ) OVERRIDE;
-	virtual FReply OnMouseButtonDown( const FGeometry& MyGeometry, const FPointerEvent& MouseEvent ) OVERRIDE;
-	virtual FReply OnMouseButtonUp( const FGeometry& MyGeometry, const FPointerEvent& MouseEvent ) OVERRIDE;
-	virtual FReply OnMouseMove( const FGeometry& MyGeometry, const FPointerEvent& MouseEvent ) OVERRIDE;
-	virtual FReply OnMouseButtonDoubleClick(const FGeometry& MyGeometry, const FPointerEvent& MouseEvent) OVERRIDE;
-	virtual FCursorReply OnCursorQuery( const FGeometry& MyGeometry, const FPointerEvent& CursorEvent ) const OVERRIDE;
+	virtual void Tick( const FGeometry& AllottedGeometry, const double InCurrentTime, const float InDeltaTime ) override;
+	virtual int32 OnPaint( const FGeometry& AllottedGeometry, const FSlateRect& MyClippingRect, FSlateWindowElementList& OutDrawElements, int32 LayerId, const FWidgetStyle& InWidgetStyle, bool bParentEnabled ) const override;
+	virtual void CacheDesiredSize() override;
+	virtual FVector2D ComputeDesiredSize() const override;
+	virtual FChildren* GetChildren() override;
+	virtual void OnArrangeChildren( const FGeometry& AllottedGeometry, FArrangedChildren& ArrangedChildren ) const override;
+	virtual bool SupportsKeyboardFocus() const override;
+	virtual FReply OnKeyChar( const FGeometry& MyGeometry,const FCharacterEvent& InCharacterEvent ) override;
+	virtual FReply OnKeyDown( const FGeometry& MyGeometry, const FKeyboardEvent& InKeyboardEvent ) override;
+	virtual FReply OnKeyUp( const FGeometry& MyGeometry, const FKeyboardEvent& InKeyboardEvent ) override;
+	virtual FReply OnKeyboardFocusReceived( const FGeometry& MyGeometry, const FKeyboardFocusEvent& InKeyboardFocusEvent ) override;
+	virtual void OnKeyboardFocusLost( const FKeyboardFocusEvent& InKeyboardFocusEvent ) override;
+	virtual FReply OnMouseButtonDown( const FGeometry& MyGeometry, const FPointerEvent& MouseEvent ) override;
+	virtual FReply OnMouseButtonUp( const FGeometry& MyGeometry, const FPointerEvent& MouseEvent ) override;
+	virtual FReply OnMouseMove( const FGeometry& MyGeometry, const FPointerEvent& MouseEvent ) override;
+	virtual FReply OnMouseButtonDoubleClick(const FGeometry& MyGeometry, const FPointerEvent& MouseEvent) override;
+	virtual FCursorReply OnCursorQuery( const FGeometry& MyGeometry, const FPointerEvent& CursorEvent ) const override;
 	// END SWidget interface
 
 private:
diff --git a/Engine/Source/Runtime/Slate/Public/Widgets/Text/SRichTextBlock.h b/Engine/Source/Runtime/Slate/Public/Widgets/Text/SRichTextBlock.h
--- a/Engine/Source/Runtime/Slate/Public/Widgets/Text/SRichTextBlock.h
+++ b/Engine/Source/Runtime/Slate/Public/Widgets/Text/SRichTextBlock.h
@@ -101,14 +101,14 @@ public:
 	void Construct( const FArguments& InArgs );
 
 	virtual int32 OnPaint( const FGeometry& AllottedGeometry, const FSlateRect& MyClippingRect, FSlateWindowElementList& OutDrawElements, int32 LayerId, const FWidgetStyle& InWidgetStyle, bool bParentEnabled ) const;
-	virtual void Tick( const FGeometry& AllottedGeometry, const double InCurrentTime, const float InDeltaTime ) OVERRIDE;
+	virtual void Tick( const FGeometry& AllottedGeometry, const double InCurrentTime, const float InDeltaTime ) override;
 
-	virtual void CacheDesiredSize() OVERRIDE;
-	virtual FVector2D ComputeDesiredSize() const OVERRIDE;
+	virtual void CacheDesiredSize() override;
+	virtual FVector2D ComputeDesiredSize() const override;
 
-	virtual FChildren* GetChildren() OVERRIDE;
+	virtual FChildren* GetChildren() override;
 
-	virtual void OnArrangeChildren( const FGeometry& AllottedGeometry, FArrangedChildren& ArrangedChildren ) const OVERRIDE;
+	virtual void OnArrangeChildren( const FGeometry& AllottedGeometry, FArrangedChildren& ArrangedChildren ) const override;
 
 	/**
 	 * Gets the text assigned to this text block
diff --git a/Engine/Source/Runtime/Slate/Public/Widgets/Text/STextBlock.h b/Engine/Source/Runtime/Slate/Public/Widgets/Text/STextBlock.h
--- a/Engine/Source/Runtime/Slate/Public/Widgets/Text/STextBlock.h
+++ b/Engine/Source/Runtime/Slate/Public/Widgets/Text/STextBlock.h
@@ -158,10 +158,10 @@ public:
 	void SetShadowColorAndOpacity(const TAttribute<FLinearColor>& InShadowColorAndOpacity);
 
 	// SWidget interface
-	virtual int32 OnPaint( const FGeometry& AllottedGeometry, const FSlateRect& MyClippingRect, FSlateWindowElementList& OutDrawElements, int32 LayerId, const FWidgetStyle& InWidgetStyle, bool bParentEnabled ) const OVERRIDE;
-	virtual FReply OnMouseButtonDoubleClick( const FGeometry& InMyGeometry, const FPointerEvent& InMouseEvent ) OVERRIDE;
-	virtual void CacheDesiredSize() OVERRIDE;
-	virtual FVector2D ComputeDesiredSize() const OVERRIDE;
+	virtual int32 OnPaint( const FGeometry& AllottedGeometry, const FSlateRect& MyClippingRect, FSlateWindowElementList& OutDrawElements, int32 LayerId, const FWidgetStyle& InWidgetStyle, bool bParentEnabled ) const override;
+	virtual FReply OnMouseButtonDoubleClick( const FGeometry& InMyGeometry, const FPointerEvent& InMouseEvent ) override;
+	virtual void CacheDesiredSize() override;
+	virtual FVector2D ComputeDesiredSize() const override;
 	// End of SWidget interface
 
 private:
diff --git a/Engine/Source/Runtime/Slate/Public/Widgets/Tutorials/STutorialWrapper.h b/Engine/Source/Runtime/Slate/Public/Widgets/Tutorials/STutorialWrapper.h
--- a/Engine/Source/Runtime/Slate/Public/Widgets/Tutorials/STutorialWrapper.h
+++ b/Engine/Source/Runtime/Slate/Public/Widgets/Tutorials/STutorialWrapper.h
@@ -22,8 +22,8 @@ public:
 	SLATE_API void Construct(const FArguments& InArgs, const FName& Name);
 
 	/** Begin SWidget interface */
-	SLATE_API virtual void Tick( const FGeometry& AllottedGeometry, const double InCurrentTime, const float InDeltaTime ) OVERRIDE;
-	SLATE_API virtual int32 OnPaint(const FGeometry& AllottedGeometry, const FSlateRect& MyClippingRect, FSlateWindowElementList& OutDrawElements, int32 LayerId, const FWidgetStyle& InWidgetStyle, bool bParentEnabled ) const OVERRIDE;
+	SLATE_API virtual void Tick( const FGeometry& AllottedGeometry, const double InCurrentTime, const float InDeltaTime ) override;
+	SLATE_API virtual int32 OnPaint(const FGeometry& AllottedGeometry, const FSlateRect& MyClippingRect, FSlateWindowElementList& OutDrawElements, int32 LayerId, const FWidgetStyle& InWidgetStyle, bool bParentEnabled ) const override;
 	/** End SWidget interface */
 
 	/** 
diff --git a/Engine/Source/Runtime/Slate/Public/Widgets/Views/SListView.h b/Engine/Source/Runtime/Slate/Public/Widgets/Views/SListView.h
--- a/Engine/Source/Runtime/Slate/Public/Widgets/Views/SListView.h
+++ b/Engine/Source/Runtime/Slate/Public/Widgets/Views/SListView.h
@@ -157,7 +157,7 @@ public:
 
 public:
 	// Inherited from SWidget
-	virtual FReply OnKeyDown( const FGeometry& MyGeometry, const FKeyboardEvent& InKeyboardEvent ) OVERRIDE
+	virtual FReply OnKeyDown( const FGeometry& MyGeometry, const FKeyboardEvent& InKeyboardEvent ) override
 	{
 		const TArray<ItemType>& ItemsSourceRef = (*this->ItemsSource);
 
@@ -332,7 +332,7 @@ public:
 		return STableViewBase::OnKeyDown(MyGeometry, InKeyboardEvent);
 	}
 
-	virtual FReply OnMouseButtonDown( const FGeometry& MyGeometry, const FPointerEvent& MouseEvent ) OVERRIDE
+	virtual FReply OnMouseButtonDown( const FGeometry& MyGeometry, const FPointerEvent& MouseEvent ) override
 	{
 		if ( bClearSelectionOnClick
 			&& SelectionMode != ESelectionMode::None
@@ -355,7 +355,7 @@ public:
 		return STableViewBase::OnMouseButtonDown(MyGeometry, MouseEvent);
 	}
 
-	virtual FReply OnMouseButtonUp( const FGeometry& MyGeometry, const FPointerEvent& MouseEvent ) OVERRIDE
+	virtual FReply OnMouseButtonUp( const FGeometry& MyGeometry, const FPointerEvent& MouseEvent ) override
 	{
 		if ( bClearSelectionOnClick
 			&& SelectionMode != ESelectionMode::None
@@ -510,7 +510,7 @@ public:
 	// These handle selection, expansion, and other such properties common to ItemsWidgets.
 	//
 
-	virtual void Private_SetItemSelection( ItemType TheItem, bool bShouldBeSelected, bool bWasUserDirected = false ) OVERRIDE
+	virtual void Private_SetItemSelection( ItemType TheItem, bool bShouldBeSelected, bool bWasUserDirected = false ) override
 	{
 		if ( SelectionMode == ESelectionMode::None )
 		{
@@ -536,14 +536,14 @@ public:
 		this->InertialScrollManager.ClearScrollVelocity();
 	}
 
-	virtual void Private_ClearSelection() OVERRIDE
+	virtual void Private_ClearSelection() override
 	{
 		SelectedItems.Empty();
 
 		this->InertialScrollManager.ClearScrollVelocity();
 	}
 
-	virtual void Private_SelectRangeFromCurrentTo( ItemType InRangeSelectionEnd ) OVERRIDE
+	virtual void Private_SelectRangeFromCurrentTo( ItemType InRangeSelectionEnd ) override
 	{
 		if ( SelectionMode == ESelectionMode::None )
 		{
@@ -576,7 +576,7 @@ public:
 		this->InertialScrollManager.ClearScrollVelocity();
 	}
 
-	virtual void Private_SignalSelectionChanged(ESelectInfo::Type SelectInfo) OVERRIDE
+	virtual void Private_SignalSelectionChanged(ESelectInfo::Type SelectInfo) override
 	{
 		if ( SelectionMode == ESelectionMode::None )
 		{
@@ -593,71 +593,71 @@ public:
 		}
 	}
 
-	virtual const ItemType* Private_ItemFromWidget( const ITableRow* TheWidget ) const OVERRIDE
+	virtual const ItemType* Private_ItemFromWidget( const ITableRow* TheWidget ) const override
 	{
 		ItemType const * LookupResult = WidgetGenerator.WidgetMapToItem.Find( TheWidget );
 		return LookupResult == NULL ? NULL : LookupResult;
 	}
 
-	virtual bool Private_UsesSelectorFocus() const OVERRIDE
+	virtual bool Private_UsesSelectorFocus() const override
 	{
 		return true;
 	}
 
-	virtual bool Private_HasSelectorFocus( const ItemType& TheItem ) const OVERRIDE
+	virtual bool Private_HasSelectorFocus( const ItemType& TheItem ) const override
 	{
 		return SelectorItem == TheItem;
 	}
 
-	virtual bool Private_IsItemSelected( const ItemType& TheItem ) const OVERRIDE
+	virtual bool Private_IsItemSelected( const ItemType& TheItem ) const override
 	{
 		return NULL != SelectedItems.Find(TheItem);
 	}
 
-	virtual bool Private_IsItemExpanded( const ItemType& TheItem ) const OVERRIDE
+	virtual bool Private_IsItemExpanded( const ItemType& TheItem ) const override
 	{
 		// List View does not support item expansion.
 		return false;	
 	}
 
-	virtual void Private_SetItemExpansion( ItemType TheItem, bool bShouldBeExpanded ) OVERRIDE
+	virtual void Private_SetItemExpansion( ItemType TheItem, bool bShouldBeExpanded ) override
 	{
 		// Do nothing; you cannot expand an item in a list!
 	}
 
-	virtual void Private_OnExpanderArrowShiftClicked( ItemType TheItem, bool bShouldBeExpanded ) OVERRIDE
+	virtual void Private_OnExpanderArrowShiftClicked( ItemType TheItem, bool bShouldBeExpanded ) override
 	{
 		// Do nothing; you cannot expand an item in a list!
 	}
 
-	virtual bool Private_DoesItemHaveChildren( int32 ItemIndexInList ) const OVERRIDE
+	virtual bool Private_DoesItemHaveChildren( int32 ItemIndexInList ) const override
 	{
 		// List View items cannot have children
 		return false;
 	}
 
-	virtual int32 Private_GetNumSelectedItems() const OVERRIDE
+	virtual int32 Private_GetNumSelectedItems() const override
 	{
 		return SelectedItems.Num();
 	}
 
-	virtual int32 Private_GetNestingDepth( int32 ItemIndexInList ) const OVERRIDE
+	virtual int32 Private_GetNestingDepth( int32 ItemIndexInList ) const override
 	{
 		// List View items are not indented
 		return 0;
 	}
 
-	virtual ESelectionMode::Type Private_GetSelectionMode() const OVERRIDE
+	virtual ESelectionMode::Type Private_GetSelectionMode() const override
 	{
 		return SelectionMode.Get();
 	}
 
-	virtual void Private_OnItemRightClicked( ItemType TheItem, const FPointerEvent& MouseEvent ) OVERRIDE
+	virtual void Private_OnItemRightClicked( ItemType TheItem, const FPointerEvent& MouseEvent ) override
 	{
 		this->OnRightMouseButtonUp( MouseEvent.GetScreenSpacePosition() );
 	}
 
-	virtual bool Private_OnItemDoubleClicked( ItemType TheItem ) OVERRIDE
+	virtual bool Private_OnItemDoubleClicked( ItemType TheItem ) override
 	{
 		if( OnDoubleClick.ExecuteIfBound( TheItem ) )
 		{
@@ -667,12 +667,12 @@ public:
 		return false;	// Not handled
 	}
 
-	virtual ETableViewMode::Type GetTableViewMode() const OVERRIDE
+	virtual ETableViewMode::Type GetTableViewMode() const override
 	{
 		return TableViewMode;
 	}
 
-	virtual TSharedRef<SWidget> AsWidget() OVERRIDE
+	virtual TSharedRef<SWidget> AsWidget() override
 	{
 		return SharedThis(this);
 	}
@@ -730,7 +730,7 @@ public:
 	 * Update generate Widgets for Items as needed and clean up any Widgets that are no longer needed.
 	 * Re-arrange the visible widget order as necessary.
 	 */
-	virtual FReGenerateResults ReGenerateItems( const FGeometry& MyGeometry ) OVERRIDE
+	virtual FReGenerateResults ReGenerateItems( const FGeometry& MyGeometry ) override
 	{
 		// Clear all the items from our panel. We will re-add them in the correct order momentarily.
 		this->ClearWidgets();
@@ -1066,7 +1066,7 @@ protected:
 	 * 
 	 * @param ListViewGeometry  The geometry of the listView; can be useful for centering the item.
 	 */
-	virtual void ScrollIntoView( const FGeometry& ListViewGeometry ) OVERRIDE
+	virtual void ScrollIntoView( const FGeometry& ListViewGeometry ) override
 	{
 		if ( TListTypeTraits<ItemType>::IsPtrValid(ItemToScrollIntoView) && ItemsSource != NULL )
 		{
@@ -1103,7 +1103,7 @@ protected:
 		}
 	}
 
-	virtual void NotifyItemScrolledIntoView() OVERRIDE
+	virtual void NotifyItemScrolledIntoView() override
 	{
 		// Notify that an item that came into view
 		if ( TListTypeTraits<ItemType>::IsPtrValid( ItemToNotifyWhenInView ) )
@@ -1116,7 +1116,7 @@ protected:
 		}
 	}
 
-	virtual float ScrollBy( const FGeometry& MyGeometry, float ScrollByAmountInSlateUnits, EAllowOverscroll AllowOverscroll ) OVERRIDE
+	virtual float ScrollBy( const FGeometry& MyGeometry, float ScrollByAmountInSlateUnits, EAllowOverscroll AllowOverscroll ) override
 	{
 		float AbsScrollByAmount = FMath::Abs( ScrollByAmountInSlateUnits );
 		int32 StartingItemIndex = (int32)ScrollOffset;
diff --git a/Engine/Source/Runtime/Slate/Public/Widgets/Views/STableRow.h b/Engine/Source/Runtime/Slate/Public/Widgets/Views/STableRow.h
--- a/Engine/Source/Runtime/Slate/Public/Widgets/Views/STableRow.h
+++ b/Engine/Source/Runtime/Slate/Public/Widgets/Views/STableRow.h
@@ -134,7 +134,7 @@ public:
 		}
 	}
 
-	virtual int32 OnPaint( const FGeometry& AllottedGeometry, const FSlateRect& MyClippingRect, FSlateWindowElementList& OutDrawElements, int32 LayerId, const FWidgetStyle& InWidgetStyle, bool bParentEnabled ) const OVERRIDE
+	virtual int32 OnPaint( const FGeometry& AllottedGeometry, const FSlateRect& MyClippingRect, FSlateWindowElementList& OutDrawElements, int32 LayerId, const FWidgetStyle& InWidgetStyle, bool bParentEnabled ) const override
 	{
 		int32 Result = SBorder::OnPaint(AllottedGeometry, MyClippingRect, OutDrawElements, LayerId, InWidgetStyle, bParentEnabled);
 
@@ -395,7 +395,7 @@ public:
 		return FReply::Unhandled();
 	}
 
-	virtual FReply OnTouchStarted( const FGeometry& MyGeometry, const FPointerEvent& InTouchEvent ) OVERRIDE
+	virtual FReply OnTouchStarted( const FGeometry& MyGeometry, const FPointerEvent& InTouchEvent ) override
 	{
 		bProcessingSelectionTouch = true;
 
@@ -406,7 +406,7 @@ public:
 			.DetectDrag( SharedThis(this), EKeys::LeftMouseButton );
 	}
 
-	virtual FReply OnTouchEnded( const FGeometry& MyGeometry, const FPointerEvent& InTouchEvent ) OVERRIDE
+	virtual FReply OnTouchEnded( const FGeometry& MyGeometry, const FPointerEvent& InTouchEvent ) override
 	{
 		if ( bProcessingSelectionTouch )
 		{
@@ -456,7 +456,7 @@ public:
 		}
 	}
 
-	virtual FReply OnDragDetected( const FGeometry& MyGeometry, const FPointerEvent& MouseEvent ) OVERRIDE
+	virtual FReply OnDragDetected( const FGeometry& MyGeometry, const FPointerEvent& MouseEvent ) override
 	{
 		if (bProcessingSelectionTouch)
 		{
@@ -481,7 +481,7 @@ public:
 		}
 	}
 
-	virtual void OnDragEnter(FGeometry const& MyGeometry, FDragDropEvent const& DragDropEvent) OVERRIDE
+	virtual void OnDragEnter(FGeometry const& MyGeometry, FDragDropEvent const& DragDropEvent) override
 	{
 		if (OnDragEnter_Handler.IsBound())
 		{
@@ -493,7 +493,7 @@ public:
 		}
 	}
 
-	virtual void OnDragLeave(FDragDropEvent const& DragDropEvent) OVERRIDE
+	virtual void OnDragLeave(FDragDropEvent const& DragDropEvent) override
 	{
 		if (OnDragLeave_Handler.IsBound())
 		{
@@ -505,7 +505,7 @@ public:
 		}
 	}
 
-	virtual FReply OnDrop(FGeometry const& MyGeometry, FDragDropEvent const& DragDropEvent) OVERRIDE
+	virtual FReply OnDrop(FGeometry const& MyGeometry, FDragDropEvent const& DragDropEvent) override
 	{
 		if (OnDrop_Handler.IsBound())
 		{
@@ -517,12 +517,12 @@ public:
 		}
 	}
 
-	virtual void SetIndexInList( int32 InIndexInList ) OVERRIDE
+	virtual void SetIndexInList( int32 InIndexInList ) override
 	{
 		IndexInList = InIndexInList;
 	}
 
-	virtual bool IsItemExpanded() const OVERRIDE
+	virtual bool IsItemExpanded() const override
 	{
 		TSharedPtr< ITypedTableView<ItemType> > OwnerWidget = OwnerTablePtr.Pin();
 		const ItemType* MyItem = OwnerWidget->Private_ItemFromWidget( this );
@@ -530,7 +530,7 @@ public:
 		return bIsItemExpanded;
 	}
 
-	virtual void ToggleExpansion() OVERRIDE
+	virtual void ToggleExpansion() override
 	{
 		TSharedPtr< ITypedTableView<ItemType> > OwnerWidget = OwnerTablePtr.Pin();
 
@@ -544,22 +544,22 @@ public:
 		}
 	}
 
-	virtual int32 GetIndentLevel() const OVERRIDE
+	virtual int32 GetIndentLevel() const override
 	{
 		return OwnerTablePtr.Pin()->Private_GetNestingDepth( IndexInList );
 	}
 
-	virtual int32 DoesItemHaveChildren() const OVERRIDE
+	virtual int32 DoesItemHaveChildren() const override
 	{
 		return OwnerTablePtr.Pin()->Private_DoesItemHaveChildren( IndexInList );
 	}
 
-	virtual TSharedRef<SWidget> AsWidget() OVERRIDE
+	virtual TSharedRef<SWidget> AsWidget() override
 	{
 		return SharedThis(this);
 	}
 
-	virtual TSharedPtr<SWidget> GetContent() OVERRIDE
+	virtual TSharedPtr<SWidget> GetContent() override
 	{
 		if ( this->Content.IsValid() )
 		{
@@ -571,7 +571,7 @@ public:
 		}
 	}
 
-	virtual void Private_OnExpanderArrowShiftClicked() OVERRIDE
+	virtual void Private_OnExpanderArrowShiftClicked() override
 	{
 		TSharedPtr< ITypedTableView<ItemType> > OwnerWidget = OwnerTablePtr.Pin();
 
@@ -804,7 +804,7 @@ protected:
 		this->GenerateColumns( HeaderRow.ToSharedRef() );
 	}
 
-	virtual void ConstructChildren( ETableViewMode::Type InOwnerTableMode, const TAttribute<FMargin>& InPadding, const TSharedRef<SWidget>& InContent ) OVERRIDE
+	virtual void ConstructChildren( ETableViewMode::Type InOwnerTableMode, const TAttribute<FMargin>& InPadding, const TSharedRef<SWidget>& InContent ) override
 	{
 		STableRow<ItemType>::Content = InContent;
 
diff --git a/Engine/Source/Runtime/Slate/Public/Widgets/Views/STableViewBase.h b/Engine/Source/Runtime/Slate/Public/Widgets/Views/STableViewBase.h
--- a/Engine/Source/Runtime/Slate/Public/Widgets/Views/STableViewBase.h
+++ b/Engine/Source/Runtime/Slate/Public/Widgets/Views/STableViewBase.h
@@ -31,20 +31,20 @@ public:
 	void ScrollBar_OnUserScrolled( float InScrollOffsetFraction );
 
 	// SWidget interface
-	virtual void OnKeyboardFocusLost( const FKeyboardFocusEvent& InKeyboardFocusEvent ) OVERRIDE;
-	virtual void OnMouseCaptureLost() OVERRIDE;
-	virtual bool SupportsKeyboardFocus() const OVERRIDE;
-	virtual void Tick( const FGeometry& AllottedGeometry, const double InCurrentTime, const float InDeltaTime ) OVERRIDE;
-	virtual FReply OnMouseButtonDown( const FGeometry& MyGeometry, const FPointerEvent& MouseEvent ) OVERRIDE;
-	virtual FReply OnMouseButtonUp( const FGeometry& MyGeometry, const FPointerEvent& MouseEvent ) OVERRIDE;
-	virtual void OnMouseLeave( const FPointerEvent& MouseEvent ) OVERRIDE;
-	virtual FReply OnMouseMove( const FGeometry& MyGeometry, const FPointerEvent& MouseEvent ) OVERRIDE;
-	virtual FReply OnMouseWheel( const FGeometry& MyGeometry, const FPointerEvent& MouseEvent ) OVERRIDE;
-	virtual FReply OnKeyDown( const FGeometry& MyGeometry, const FKeyboardEvent& InKeyboardEvent ) OVERRIDE;
-	virtual FCursorReply OnCursorQuery( const FGeometry& MyGeometry, const FPointerEvent& CursorEvent ) const OVERRIDE;
-	virtual FReply OnTouchStarted( const FGeometry& MyGeometry, const FPointerEvent& InTouchEvent ) OVERRIDE;
-	virtual FReply OnTouchMoved( const FGeometry& MyGeometry, const FPointerEvent& InTouchEvent ) OVERRIDE;
-	virtual FReply OnTouchEnded( const FGeometry& MyGeometry, const FPointerEvent& InTouchEvent ) OVERRIDE;
+	virtual void OnKeyboardFocusLost( const FKeyboardFocusEvent& InKeyboardFocusEvent ) override;
+	virtual void OnMouseCaptureLost() override;
+	virtual bool SupportsKeyboardFocus() const override;
+	virtual void Tick( const FGeometry& AllottedGeometry, const double InCurrentTime, const float InDeltaTime ) override;
+	virtual FReply OnMouseButtonDown( const FGeometry& MyGeometry, const FPointerEvent& MouseEvent ) override;
+	virtual FReply OnMouseButtonUp( const FGeometry& MyGeometry, const FPointerEvent& MouseEvent ) override;
+	virtual void OnMouseLeave( const FPointerEvent& MouseEvent ) override;
+	virtual FReply OnMouseMove( const FGeometry& MyGeometry, const FPointerEvent& MouseEvent ) override;
+	virtual FReply OnMouseWheel( const FGeometry& MyGeometry, const FPointerEvent& MouseEvent ) override;
+	virtual FReply OnKeyDown( const FGeometry& MyGeometry, const FKeyboardEvent& InKeyboardEvent ) override;
+	virtual FCursorReply OnCursorQuery( const FGeometry& MyGeometry, const FPointerEvent& CursorEvent ) const override;
+	virtual FReply OnTouchStarted( const FGeometry& MyGeometry, const FPointerEvent& InTouchEvent ) override;
+	virtual FReply OnTouchMoved( const FGeometry& MyGeometry, const FPointerEvent& InTouchEvent ) override;
+	virtual FReply OnTouchEnded( const FGeometry& MyGeometry, const FPointerEvent& InTouchEvent ) override;
 	// End of SWidget interface
 
 	/** @return The number of Widgets we currently have generated. */
@@ -69,9 +69,9 @@ public:
 	bool IsPendingRefresh() const;
 
 	// IScrollableWidget interface
-	virtual FVector2D GetScrollDistance() OVERRIDE;
-	virtual FVector2D GetScrollDistanceRemaining() OVERRIDE;
-	virtual TSharedRef<class SWidget> GetScrollWidget() OVERRIDE;
+	virtual FVector2D GetScrollDistance() override;
+	virtual FVector2D GetScrollDistanceRemaining() override;
+	virtual TSharedRef<class SWidget> GetScrollWidget() override;
 	// End of IScrollableWidget interface
 
 	virtual int32 OnPaint(
@@ -80,7 +80,7 @@ public:
 		FSlateWindowElementList& OutDrawElements,
 		int32 LayerId,
 		const FWidgetStyle& InWidgetStyle,
-		bool bParentEnabled) const OVERRIDE;
+		bool bParentEnabled) const override;
 
 	/** Is this list backing a tree or just a standalone list */
 	const ETableViewMode::Type TableViewMode;
diff --git a/Engine/Source/Runtime/Slate/Public/Widgets/Views/STileView.h b/Engine/Source/Runtime/Slate/Public/Widgets/Views/STileView.h
--- a/Engine/Source/Runtime/Slate/Public/Widgets/Views/STileView.h
+++ b/Engine/Source/Runtime/Slate/Public/Widgets/Views/STileView.h
@@ -113,7 +113,7 @@ public:
 
 public:
 	// Inherited from SWidget
-	virtual FReply OnKeyDown( const FGeometry& MyGeometry, const FKeyboardEvent& InKeyboardEvent ) OVERRIDE
+	virtual FReply OnKeyDown( const FGeometry& MyGeometry, const FKeyboardEvent& InKeyboardEvent ) override
 	{
 		const TArray<ItemType>& ItemsSourceRef = (*this->ItemsSource);
 
@@ -152,7 +152,7 @@ public:
 
 public:	
 
-	virtual STableViewBase::FReGenerateResults ReGenerateItems( const FGeometry& MyGeometry ) OVERRIDE
+	virtual STableViewBase::FReGenerateResults ReGenerateItems( const FGeometry& MyGeometry ) override
 	{
 		// Clear all the items from our panel. We will re-add them in the correct order momentarily.
 		this->ClearWidgets();
@@ -220,7 +220,7 @@ public:
 
 	}
 
-	virtual int32 GetNumItemsBeingObserved() const OVERRIDE
+	virtual int32 GetNumItemsBeingObserved() const override
 	{
 		const int32 NumItemsBeingObserved = this->ItemsSource == NULL ? 0 : this->ItemsSource->Num();
 		const int32 NumItemsWide = GetNumItemsWide();
@@ -240,7 +240,7 @@ public:
 
 protected:
 
-	virtual float ScrollBy( const FGeometry& MyGeometry, float ScrollByAmountInSlateUnits, EAllowOverscroll AllowOverscroll ) OVERRIDE
+	virtual float ScrollBy( const FGeometry& MyGeometry, float ScrollByAmountInSlateUnits, EAllowOverscroll AllowOverscroll ) override
 	{
 		// Working around a CLANG bug, where all the base class
 		// members require an explicit namespace resolution.
@@ -270,7 +270,7 @@ protected:
 		}
 	}
 
-	virtual int32 GetNumItemsWide() const OVERRIDE
+	virtual int32 GetNumItemsWide() const override
 	{
 		const float ItemWidth = this->GetItemWidth();
 		const int32 NumItemsWide = ItemWidth > 0 ? FMath::FloorToInt(this->PanelGeometryLastTick.Size.X / ItemWidth) : 1;
diff --git a/Engine/Source/Runtime/Slate/Public/Widgets/Views/STreeView.h b/Engine/Source/Runtime/Slate/Public/Widgets/Views/STreeView.h
--- a/Engine/Source/Runtime/Slate/Public/Widgets/Views/STreeView.h
+++ b/Engine/Source/Runtime/Slate/Public/Widgets/Views/STreeView.h
@@ -189,7 +189,7 @@ class STreeView : public SListView< ItemType >
 
 public:
 	// Inherited from SWidget
-	virtual FReply OnKeyDown( const FGeometry& MyGeometry, const FKeyboardEvent& InKeyboardEvent ) OVERRIDE
+	virtual FReply OnKeyDown( const FGeometry& MyGeometry, const FKeyboardEvent& InKeyboardEvent ) override
 	{
 		// Check for selection/expansion toggling keys (Left, Right)
 		// SelectorItem represents the keyboard selection. If it isn't valid then we don't know what to expand.
@@ -267,13 +267,13 @@ private:
 	// Tree View adds the ability to expand/collapse items.
 	// All the selection functionality is inherited from ListView.
 
-	virtual bool Private_IsItemExpanded( const ItemType& TheItem ) const OVERRIDE
+	virtual bool Private_IsItemExpanded( const ItemType& TheItem ) const override
 	{
 		const FSparseItemInfo* ItemInfo = SparseItemInfos.Find(TheItem);
 		return ItemInfo != NULL && ItemInfo->bIsExpanded;
 	}
 
-	virtual void Private_SetItemExpansion( ItemType TheItem, bool bShouldBeExpanded ) OVERRIDE
+	virtual void Private_SetItemExpansion( ItemType TheItem, bool bShouldBeExpanded ) override
 	{
 		const FSparseItemInfo* const SparseItemInfo = SparseItemInfos.Find(TheItem);
 		bool bWasExpanded = false;
@@ -298,7 +298,7 @@ private:
 		}
 	}
 
-	virtual void Private_OnExpanderArrowShiftClicked( ItemType TheItem, bool bShouldBeExpanded ) OVERRIDE
+	virtual void Private_OnExpanderArrowShiftClicked( ItemType TheItem, bool bShouldBeExpanded ) override
 	{
 		if(OnSetExpansionRecursive.IsBound())
 		{
@@ -310,7 +310,7 @@ private:
 		}
 	}
 
-	virtual bool Private_DoesItemHaveChildren( int32 ItemIndexInList ) const OVERRIDE
+	virtual bool Private_DoesItemHaveChildren( int32 ItemIndexInList ) const override
 	{
 		bool bHasChildren = false;
 		if (ItemIndexInList < DenseItemInfos.Num())
@@ -320,7 +320,7 @@ private:
 		return bHasChildren;
 	}
 
-	virtual int32 Private_GetNestingDepth( int32 ItemIndexInList ) const OVERRIDE
+	virtual int32 Private_GetNestingDepth( int32 ItemIndexInList ) const override
 	{
 		int32 NextingLevel = 0;
 		if (ItemIndexInList < DenseItemInfos.Num())
diff --git a/Engine/Source/Runtime/SlateCore/Private/Widgets/SNullWidget.cpp b/Engine/Source/Runtime/SlateCore/Private/Widgets/SNullWidget.cpp
--- a/Engine/Source/Runtime/SlateCore/Private/Widgets/SNullWidget.cpp
+++ b/Engine/Source/Runtime/SlateCore/Private/Widgets/SNullWidget.cpp
@@ -34,12 +34,12 @@ public:
 
 	// SWidget interface
 
-	virtual int32 OnPaint( const FGeometry& AllottedGeometry, const FSlateRect& MyClippingRect, FSlateWindowElementList& OutDrawElements, int32 LayerId, const FWidgetStyle& InWidgetStyle, bool bParentEnabled ) const OVERRIDE
+	virtual int32 OnPaint( const FGeometry& AllottedGeometry, const FSlateRect& MyClippingRect, FSlateWindowElementList& OutDrawElements, int32 LayerId, const FWidgetStyle& InWidgetStyle, bool bParentEnabled ) const override
 	{
 		return LayerId;
 	}
 
-	virtual FVector2D ComputeDesiredSize( ) const OVERRIDE
+	virtual FVector2D ComputeDesiredSize( ) const override
 	{
 		return FVector2D(0.0f, 0.0f);
 	}
diff --git a/Engine/Source/Runtime/SlateCore/Private/Widgets/SWindow.cpp b/Engine/Source/Runtime/SlateCore/Private/Widgets/SWindow.cpp
--- a/Engine/Source/Runtime/SlateCore/Private/Widgets/SWindow.cpp
+++ b/Engine/Source/Runtime/SlateCore/Private/Widgets/SWindow.cpp
@@ -88,7 +88,7 @@ public:
 
 private:
 
-	virtual void OnArrangeChildren( const FGeometry& AllottedGeometry, FArrangedChildren& ArrangedChildren ) const OVERRIDE
+	virtual void OnArrangeChildren( const FGeometry& AllottedGeometry, FArrangedChildren& ArrangedChildren ) const override
 	{
 		const FVector2D WindowDesktopPosition = (ensure(OwnerWindow.IsValid()))
 			? OwnerWindow.Pin()->GetPositionInScreen()
@@ -135,7 +135,7 @@ private:
 		}
 	}
 
-	virtual FVector2D ComputeDesiredSize() const OVERRIDE
+	virtual FVector2D ComputeDesiredSize() const override
 	{
 		return FVector2D(100,100);
 	}
@@ -147,7 +147,7 @@ private:
 	 * provides layout information about the child it stores. In that case
 	 * GetChildren should simply return the TPanelChildren<Slot>. See StackPanel for an example.
 	 */
-	virtual FChildren* GetChildren() OVERRIDE
+	virtual FChildren* GetChildren() override
 	{
 		return &Children;
 	}
diff --git a/Engine/Source/Runtime/SlateCore/Public/Input/DragAndDrop.h b/Engine/Source/Runtime/SlateCore/Public/Input/DragAndDrop.h
--- a/Engine/Source/Runtime/SlateCore/Public/Input/DragAndDrop.h
+++ b/Engine/Source/Runtime/SlateCore/Public/Input/DragAndDrop.h
@@ -171,7 +171,7 @@ DECLARE_DELEGATE_OneParam( FOnDragDropUpdate,
 
 #define DRAG_DROP_OPERATOR_TYPE(TYPE, BASE) \
 	static const FString& GetTypeId() { static FString Type = TEXT(#TYPE); return Type; } \
-	virtual bool IsOfTypeImpl(const FString& Type) const OVERRIDE { return GetTypeId() == Type || BASE::IsOfTypeImpl(Type); }
+	virtual bool IsOfTypeImpl(const FString& Type) const override { return GetTypeId() == Type || BASE::IsOfTypeImpl(Type); }
 
 
 /**
@@ -184,7 +184,7 @@ private:
 	/** A private constructor to ensure that the appropriate "New" factory method is used below */
 	FExternalDragOperation(){}
 
-	virtual bool IsExternalOperation() const OVERRIDE { return true; }
+	virtual bool IsExternalOperation() const override { return true; }
 
 public:
 	DRAG_DROP_OPERATOR_TYPE(FExternalDragOperation, FDragDropOperation)
diff --git a/Engine/Source/Runtime/SlateCore/Public/Logging/EventLogger.h b/Engine/Source/Runtime/SlateCore/Public/Logging/EventLogger.h
--- a/Engine/Source/Runtime/SlateCore/Public/Logging/EventLogger.h
+++ b/Engine/Source/Runtime/SlateCore/Public/Logging/EventLogger.h
@@ -31,18 +31,18 @@ public:
 
 	// Begin IEventLogger interface
 
-	virtual FString GetLog( ) const OVERRIDE
+	virtual FString GetLog( ) const override
 	{
 		return LoggerA->GetLog() + TEXT("\n") + LoggerB->GetLog();
 	}
 
-	virtual void Log(EEventLog::Type Event, const FString& AdditionalContent, TSharedPtr<SWidget> Widget) OVERRIDE
+	virtual void Log(EEventLog::Type Event, const FString& AdditionalContent, TSharedPtr<SWidget> Widget) override
 	{
 		LoggerA->Log(Event, AdditionalContent, Widget);
 		LoggerB->Log(Event, AdditionalContent, Widget);
 	}
 
-	virtual void SaveToFile( ) OVERRIDE
+	virtual void SaveToFile( ) override
 	{
 		LoggerA->SaveToFile();
 		LoggerB->SaveToFile();
@@ -67,11 +67,11 @@ public:
 
 	// Begin IEventLogger interface
 
-	virtual FString GetLog( ) const OVERRIDE;
+	virtual FString GetLog( ) const override;
 
-	virtual void Log( EEventLog::Type Event, const FString& AdditionalContent, TSharedPtr<SWidget> Widget ) OVERRIDE;
+	virtual void Log( EEventLog::Type Event, const FString& AdditionalContent, TSharedPtr<SWidget> Widget ) override;
 	
-	virtual void SaveToFile( ) OVERRIDE;
+	virtual void SaveToFile( ) override;
 
 	// EndIEventLogger interface
 
@@ -89,11 +89,11 @@ public:
 
 	// Begin IEventLogger interface
 
-	virtual FString GetLog( ) const OVERRIDE;
+	virtual FString GetLog( ) const override;
 
-	virtual void Log( EEventLog::Type Event, const FString& AdditionalContent, TSharedPtr<SWidget> Widget ) OVERRIDE;
+	virtual void Log( EEventLog::Type Event, const FString& AdditionalContent, TSharedPtr<SWidget> Widget ) override;
 	
-	virtual void SaveToFile( ) OVERRIDE { }
+	virtual void SaveToFile( ) override { }
 	
 	// End IEventLogger interface
 
@@ -111,11 +111,11 @@ public:
 
 	// Begin IEventLogger interface
 	
-	virtual FString GetLog( ) const OVERRIDE;
+	virtual FString GetLog( ) const override;
 
-	virtual void Log( EEventLog::Type Event, const FString& AdditionalContent, TSharedPtr<SWidget> Widget ) OVERRIDE;
+	virtual void Log( EEventLog::Type Event, const FString& AdditionalContent, TSharedPtr<SWidget> Widget ) override;
 	
-	virtual void SaveToFile( ) OVERRIDE { }
+	virtual void SaveToFile( ) override { }
 	
 	// End IEventLogger interface
 
diff --git a/Engine/Source/Runtime/SlateCore/Public/Sound/ISlateSoundDevice.h b/Engine/Source/Runtime/SlateCore/Public/Sound/ISlateSoundDevice.h
--- a/Engine/Source/Runtime/SlateCore/Public/Sound/ISlateSoundDevice.h
+++ b/Engine/Source/Runtime/SlateCore/Public/Sound/ISlateSoundDevice.h
@@ -20,11 +20,11 @@ public:
 class SLATECORE_API FNullSlateSoundDevice : public ISlateSoundDevice
 {
 public:
-	virtual void PlaySound(const FSlateSound&, int32) const OVERRIDE
+	virtual void PlaySound(const FSlateSound&, int32) const override
 	{
 	}
 
-	virtual float GetSoundDuration(const FSlateSound& Sound) const OVERRIDE
+	virtual float GetSoundDuration(const FSlateSound& Sound) const override
 	{
 		return 0.0f;
 	}
diff --git a/Engine/Source/Runtime/SlateCore/Public/Styling/SlateStyle.h b/Engine/Source/Runtime/SlateCore/Public/Styling/SlateStyle.h
--- a/Engine/Source/Runtime/SlateCore/Public/Styling/SlateStyle.h
+++ b/Engine/Source/Runtime/SlateCore/Public/Styling/SlateStyle.h
@@ -19,8 +19,8 @@ public:
 	/** Destructor. */
 	virtual ~FSlateStyleSet();
 
-	virtual const FName& GetStyleSetName() const OVERRIDE;
-	virtual void GetResources(TArray< const FSlateBrush* >& OutResources) const OVERRIDE;
+	virtual const FName& GetStyleSetName() const override;
+	virtual void GetResources(TArray< const FSlateBrush* >& OutResources) const override;
 	virtual void SetContentRoot(const FString& InContentRootDir);
 	virtual FString RootToContentDir(const ANSICHAR* RelativePath, const TCHAR* Extension);
 	virtual FString RootToContentDir(const WIDECHAR* RelativePath, const TCHAR* Extension);
@@ -38,23 +38,23 @@ public:
 	virtual FString RootToCoreContentDir(const WIDECHAR* RelativePath);
 	virtual FString RootToCoreContentDir(const FString& RelativePath);
 
-	virtual float GetFloat(const FName PropertyName, const ANSICHAR* Specifier = nullptr) const OVERRIDE;
-	virtual FVector2D GetVector(const FName PropertyName, const ANSICHAR* Specifier = nullptr) const OVERRIDE;
-	virtual const FLinearColor& GetColor(const FName PropertyName, const ANSICHAR* Specifier = nullptr) const OVERRIDE;
-	virtual const FSlateColor GetSlateColor(const FName PropertyName, const ANSICHAR* Specifier = nullptr) const OVERRIDE;
-	virtual const FMargin& GetMargin(const FName PropertyName, const ANSICHAR* Specifier = nullptr) const OVERRIDE;
+	virtual float GetFloat(const FName PropertyName, const ANSICHAR* Specifier = nullptr) const override;
+	virtual FVector2D GetVector(const FName PropertyName, const ANSICHAR* Specifier = nullptr) const override;
+	virtual const FLinearColor& GetColor(const FName PropertyName, const ANSICHAR* Specifier = nullptr) const override;
+	virtual const FSlateColor GetSlateColor(const FName PropertyName, const ANSICHAR* Specifier = nullptr) const override;
+	virtual const FMargin& GetMargin(const FName PropertyName, const ANSICHAR* Specifier = nullptr) const override;
 
-	virtual const FSlateBrush* GetBrush(const FName PropertyName, const ANSICHAR* Specifier = nullptr) const OVERRIDE;
-	virtual const FSlateBrush* GetOptionalBrush(const FName PropertyName, const ANSICHAR* Specifier = nullptr, const FSlateBrush* const DefaultBrush = FStyleDefaults::GetNoBrush()) const OVERRIDE;
+	virtual const FSlateBrush* GetBrush(const FName PropertyName, const ANSICHAR* Specifier = nullptr) const override;
+	virtual const FSlateBrush* GetOptionalBrush(const FName PropertyName, const ANSICHAR* Specifier = nullptr, const FSlateBrush* const DefaultBrush = FStyleDefaults::GetNoBrush()) const override;
 
-	virtual const TSharedPtr< FSlateDynamicImageBrush > GetDynamicImageBrush(const FName BrushTemplate, const FName TextureName, const ANSICHAR* Specifier = nullptr) OVERRIDE;
-	virtual const TSharedPtr< FSlateDynamicImageBrush > GetDynamicImageBrush(const FName BrushTemplate, const ANSICHAR* Specifier, UTexture2D* TextureResource, const FName TextureName) OVERRIDE;
-	virtual const TSharedPtr< FSlateDynamicImageBrush > GetDynamicImageBrush(const FName BrushTemplate, UTexture2D* TextureResource, const FName TextureName) OVERRIDE;
+	virtual const TSharedPtr< FSlateDynamicImageBrush > GetDynamicImageBrush(const FName BrushTemplate, const FName TextureName, const ANSICHAR* Specifier = nullptr) override;
+	virtual const TSharedPtr< FSlateDynamicImageBrush > GetDynamicImageBrush(const FName BrushTemplate, const ANSICHAR* Specifier, UTexture2D* TextureResource, const FName TextureName) override;
+	virtual const TSharedPtr< FSlateDynamicImageBrush > GetDynamicImageBrush(const FName BrushTemplate, UTexture2D* TextureResource, const FName TextureName) override;
 
-	virtual FSlateBrush* GetDefaultBrush() const OVERRIDE;
+	virtual FSlateBrush* GetDefaultBrush() const override;
 
-	virtual const FSlateSound& GetSound(const FName PropertyName, const ANSICHAR* Specifier = nullptr) const OVERRIDE;
-	virtual FSlateFontInfo GetFontStyle(const FName PropertyName, const ANSICHAR* Specifier = nullptr) const OVERRIDE;
+	virtual const FSlateSound& GetSound(const FName PropertyName, const ANSICHAR* Specifier = nullptr) const override;
+	virtual FSlateFontInfo GetFontStyle(const FName PropertyName, const ANSICHAR* Specifier = nullptr) const override;
 
 public:
 
@@ -192,9 +192,9 @@ public:
 
 protected:
 
-	virtual const FSlateWidgetStyle* GetWidgetStyleInternal(const FName DesiredTypeName, const FName StyleName) const OVERRIDE;
+	virtual const FSlateWidgetStyle* GetWidgetStyleInternal(const FName DesiredTypeName, const FName StyleName) const override;
 
-	virtual void Log(ISlateStyle::EStyleMessageSeverity Severity, const FText& Message) const OVERRIDE;
+	virtual void Log(ISlateStyle::EStyleMessageSeverity Severity, const FText& Message) const override;
 
 	virtual void LogUnusedBrushResources();
 
diff --git a/Engine/Source/Runtime/SlateCore/Public/Styling/SlateTypes.h b/Engine/Source/Runtime/SlateCore/Public/Styling/SlateTypes.h
--- a/Engine/Source/Runtime/SlateCore/Public/Styling/SlateTypes.h
+++ b/Engine/Source/Runtime/SlateCore/Public/Styling/SlateTypes.h
@@ -35,10 +35,10 @@ struct SLATECORE_API FCheckBoxStyle : public FSlateWidgetStyle
 
 	virtual ~FCheckBoxStyle() {}
 
-	virtual void GetResources( TArray< const FSlateBrush* > & OutBrushes ) const OVERRIDE;
+	virtual void GetResources( TArray< const FSlateBrush* > & OutBrushes ) const override;
 
 	static const FName TypeName;
-	virtual const FName GetTypeName() const OVERRIDE { return TypeName; };
+	virtual const FName GetTypeName() const override { return TypeName; };
 
 	static const FCheckBoxStyle& GetDefault();
 
@@ -163,7 +163,7 @@ struct SLATECORE_API FTextBlockStyle : public FSlateWidgetStyle
 	virtual ~FTextBlockStyle() {}
 
 	static const FName TypeName;
-	virtual const FName GetTypeName() const OVERRIDE { return TypeName; };
+	virtual const FName GetTypeName() const override { return TypeName; };
 
 	static const FTextBlockStyle& GetDefault();
 
@@ -219,10 +219,10 @@ struct SLATECORE_API FButtonStyle : public FSlateWidgetStyle
 
 	virtual ~FButtonStyle() {}
 
-	virtual void GetResources( TArray< const FSlateBrush* >& OutBrushes ) const OVERRIDE;
+	virtual void GetResources( TArray< const FSlateBrush* >& OutBrushes ) const override;
 
 	static const FName TypeName;
-	virtual const FName GetTypeName() const OVERRIDE { return TypeName; };
+	virtual const FName GetTypeName() const override { return TypeName; };
 
 	static const FButtonStyle& GetDefault();
 
@@ -309,10 +309,10 @@ struct SLATECORE_API FComboButtonStyle : public FSlateWidgetStyle
 
 	virtual ~FComboButtonStyle() {}
 
-	virtual void GetResources( TArray< const FSlateBrush* >& OutBrushes ) const OVERRIDE;
+	virtual void GetResources( TArray< const FSlateBrush* >& OutBrushes ) const override;
 
 	static const FName TypeName;
-	virtual const FName GetTypeName() const OVERRIDE { return TypeName; };
+	virtual const FName GetTypeName() const override { return TypeName; };
 
 	static const FComboButtonStyle& GetDefault();
 
@@ -358,10 +358,10 @@ struct SLATECORE_API FComboBoxStyle : public FSlateWidgetStyle
 
 	virtual ~FComboBoxStyle() {}
 
-	virtual void GetResources( TArray< const FSlateBrush* >& OutBrushes ) const OVERRIDE;
+	virtual void GetResources( TArray< const FSlateBrush* >& OutBrushes ) const override;
 
 	static const FName TypeName;
-	virtual const FName GetTypeName() const OVERRIDE { return TypeName; };
+	virtual const FName GetTypeName() const override { return TypeName; };
 
 	static const FComboBoxStyle& GetDefault();
 
@@ -418,10 +418,10 @@ struct SLATECORE_API FHyperlinkStyle : public FSlateWidgetStyle
 
 	virtual ~FHyperlinkStyle() {}
 
-	virtual void GetResources( TArray< const FSlateBrush* >& OutBrushes ) const OVERRIDE;
+	virtual void GetResources( TArray< const FSlateBrush* >& OutBrushes ) const override;
 
 	static const FName TypeName;
-	virtual const FName GetTypeName() const OVERRIDE { return TypeName; };
+	virtual const FName GetTypeName() const override { return TypeName; };
 
 	static const FHyperlinkStyle& GetDefault();
 
@@ -453,10 +453,10 @@ struct SLATECORE_API FEditableTextStyle : public FSlateWidgetStyle
 
 	virtual ~FEditableTextStyle() {}
 
-	virtual void GetResources( TArray< const FSlateBrush* >& OutBrushes ) const OVERRIDE;
+	virtual void GetResources( TArray< const FSlateBrush* >& OutBrushes ) const override;
 
 	static const FName TypeName;
-	virtual const FName GetTypeName() const OVERRIDE { return TypeName; }
+	virtual const FName GetTypeName() const override { return TypeName; }
 
 	static const FEditableTextStyle& GetDefault();
 
@@ -504,10 +504,10 @@ struct SLATECORE_API FEditableTextBoxStyle : public FSlateWidgetStyle
 
 	virtual ~FEditableTextBoxStyle() {}
 
-	virtual void GetResources( TArray< const FSlateBrush* >& OutBrushes ) const OVERRIDE;
+	virtual void GetResources( TArray< const FSlateBrush* >& OutBrushes ) const override;
 
 	static const FName TypeName;
-	virtual const FName GetTypeName() const OVERRIDE { return TypeName; };
+	virtual const FName GetTypeName() const override { return TypeName; };
 
 	static const FEditableTextBoxStyle& GetDefault();
 
@@ -571,10 +571,10 @@ struct SLATECORE_API FInlineEditableTextBlockStyle : public FSlateWidgetStyle
 
 	virtual ~FInlineEditableTextBlockStyle() {}
 
-	virtual void GetResources( TArray< const FSlateBrush* >& OutBrushes ) const OVERRIDE;
+	virtual void GetResources( TArray< const FSlateBrush* >& OutBrushes ) const override;
 
 	static const FName TypeName;
-	virtual const FName GetTypeName() const OVERRIDE { return TypeName; };
+	virtual const FName GetTypeName() const override { return TypeName; };
 
 	static const FInlineEditableTextBlockStyle& GetDefault();
 
@@ -602,10 +602,10 @@ struct SLATECORE_API FProgressBarStyle : public FSlateWidgetStyle
 
 	virtual ~FProgressBarStyle() {}
 
-	virtual void GetResources( TArray< const FSlateBrush* >& OutBrushes ) const OVERRIDE;
+	virtual void GetResources( TArray< const FSlateBrush* >& OutBrushes ) const override;
 
 	static const FName TypeName;
-	virtual const FName GetTypeName() const OVERRIDE { return TypeName; };
+	virtual const FName GetTypeName() const override { return TypeName; };
 
 	static const FProgressBarStyle& GetDefault();
 
@@ -638,10 +638,10 @@ struct SLATECORE_API FScrollBarStyle : public FSlateWidgetStyle
 
 	virtual ~FScrollBarStyle() {}
 
-	virtual void GetResources( TArray< const FSlateBrush* >& OutBrushes ) const OVERRIDE;
+	virtual void GetResources( TArray< const FSlateBrush* >& OutBrushes ) const override;
 
 	static const FName TypeName;
-	virtual const FName GetTypeName() const OVERRIDE { return TypeName; };
+	virtual const FName GetTypeName() const override { return TypeName; };
 
 	static const FScrollBarStyle& GetDefault();
 
@@ -684,10 +684,10 @@ struct SLATECORE_API FExpandableAreaStyle : public FSlateWidgetStyle
 
 	virtual ~FExpandableAreaStyle() {}
 
-	virtual void GetResources( TArray< const FSlateBrush* >& OutBrushes ) const OVERRIDE;
+	virtual void GetResources( TArray< const FSlateBrush* >& OutBrushes ) const override;
 
 	static const FName TypeName;
-	virtual const FName GetTypeName() const OVERRIDE { return TypeName; };
+	virtual const FName GetTypeName() const override { return TypeName; };
 
 	static const FExpandableAreaStyle& GetDefault();
 
@@ -715,10 +715,10 @@ struct SLATECORE_API FSearchBoxStyle : public FSlateWidgetStyle
 
 	virtual ~FSearchBoxStyle() {}
 
-	virtual void GetResources( TArray< const FSlateBrush* >& OutBrushes ) const OVERRIDE;
+	virtual void GetResources( TArray< const FSlateBrush* >& OutBrushes ) const override;
 
 	static const FName TypeName;
-	virtual const FName GetTypeName() const OVERRIDE { return TypeName; };
+	virtual const FName GetTypeName() const override { return TypeName; };
 
 	static const FSearchBoxStyle& GetDefault();
 
@@ -766,10 +766,10 @@ struct SLATECORE_API FSliderStyle : public FSlateWidgetStyle
 
 	virtual ~FSliderStyle() {}
 
-	virtual void GetResources( TArray< const FSlateBrush* >& OutBrushes ) const OVERRIDE;
+	virtual void GetResources( TArray< const FSlateBrush* >& OutBrushes ) const override;
 
 	static const FName TypeName;
-	virtual const FName GetTypeName() const OVERRIDE { return TypeName; };
+	virtual const FName GetTypeName() const override { return TypeName; };
 
 	static const FSliderStyle& GetDefault();
 
@@ -797,10 +797,10 @@ struct SLATECORE_API FVolumeControlStyle : public FSlateWidgetStyle
 
 	virtual ~FVolumeControlStyle() {}
 
-	virtual void GetResources( TArray< const FSlateBrush* >& OutBrushes ) const OVERRIDE;
+	virtual void GetResources( TArray< const FSlateBrush* >& OutBrushes ) const override;
 
 	static const FName TypeName;
-	virtual const FName GetTypeName() const OVERRIDE { return TypeName; };
+	virtual const FName GetTypeName() const override { return TypeName; };
 
 	static const FVolumeControlStyle& GetDefault();
 
@@ -847,10 +847,10 @@ struct SLATECORE_API FInlineTextImageStyle : public FSlateWidgetStyle
 
 	virtual ~FInlineTextImageStyle() {}
 
-	virtual void GetResources( TArray< const FSlateBrush* >& OutBrushes ) const OVERRIDE;
+	virtual void GetResources( TArray< const FSlateBrush* >& OutBrushes ) const override;
 
 	static const FName TypeName;
-	virtual const FName GetTypeName() const OVERRIDE { return TypeName; };
+	virtual const FName GetTypeName() const override { return TypeName; };
 
 	static const FInlineTextImageStyle& GetDefault();
 
@@ -877,10 +877,10 @@ struct SLATECORE_API FSpinBoxStyle : public FSlateWidgetStyle
 
 	virtual ~FSpinBoxStyle() {}
 
-	virtual void GetResources( TArray< const FSlateBrush* >& OutBrushes ) const OVERRIDE;
+	virtual void GetResources( TArray< const FSlateBrush* >& OutBrushes ) const override;
 
 	static const FName TypeName;
-	virtual const FName GetTypeName() const OVERRIDE { return TypeName; };
+	virtual const FName GetTypeName() const override { return TypeName; };
 
 	static const FSpinBoxStyle& GetDefault();
 
@@ -933,10 +933,10 @@ struct SLATECORE_API FSplitterStyle : public FSlateWidgetStyle
 
 	virtual ~FSplitterStyle() {}
 
-	virtual void GetResources( TArray< const FSlateBrush* >& OutBrushes ) const OVERRIDE;
+	virtual void GetResources( TArray< const FSlateBrush* >& OutBrushes ) const override;
 
 	static const FName TypeName;
-	virtual const FName GetTypeName() const OVERRIDE { return TypeName; };
+	virtual const FName GetTypeName() const override { return TypeName; };
 
 	static const FSplitterStyle& GetDefault();
 
@@ -964,10 +964,10 @@ struct SLATECORE_API FTableRowStyle : public FSlateWidgetStyle
 
 	virtual ~FTableRowStyle() {}
 
-	virtual void GetResources( TArray< const FSlateBrush* >& OutBrushes ) const OVERRIDE;
+	virtual void GetResources( TArray< const FSlateBrush* >& OutBrushes ) const override;
 
 	static const FName TypeName;
-	virtual const FName GetTypeName() const OVERRIDE { return TypeName; };
+	virtual const FName GetTypeName() const override { return TypeName; };
 
 	static const FTableRowStyle& GetDefault();
 
@@ -1040,10 +1040,10 @@ struct SLATECORE_API FTableColumnHeaderStyle : public FSlateWidgetStyle
 
 	virtual ~FTableColumnHeaderStyle() {}
 
-	virtual void GetResources( TArray< const FSlateBrush* >& OutBrushes ) const OVERRIDE;
+	virtual void GetResources( TArray< const FSlateBrush* >& OutBrushes ) const override;
 
 	static const FName TypeName;
-	virtual const FName GetTypeName() const OVERRIDE { return TypeName; };
+	virtual const FName GetTypeName() const override { return TypeName; };
 
 	static const FTableColumnHeaderStyle& GetDefault();
 
@@ -1096,10 +1096,10 @@ struct SLATECORE_API FHeaderRowStyle : public FSlateWidgetStyle
 
 	virtual ~FHeaderRowStyle() {}
 
-	virtual void GetResources( TArray< const FSlateBrush* >& OutBrushes ) const OVERRIDE;
+	virtual void GetResources( TArray< const FSlateBrush* >& OutBrushes ) const override;
 
 	static const FName TypeName;
-	virtual const FName GetTypeName() const OVERRIDE { return TypeName; };
+	virtual const FName GetTypeName() const override { return TypeName; };
 
 	static const FHeaderRowStyle& GetDefault();
 
@@ -1142,10 +1142,10 @@ struct SLATECORE_API FDockTabStyle : public FSlateWidgetStyle
 
 	virtual ~FDockTabStyle() {}
 
-	virtual void GetResources( TArray< const FSlateBrush* >& OutBrushes ) const OVERRIDE;
+	virtual void GetResources( TArray< const FSlateBrush* >& OutBrushes ) const override;
 
 	static const FName TypeName;
-	virtual const FName GetTypeName() const OVERRIDE { return TypeName; };
+	virtual const FName GetTypeName() const override { return TypeName; };
 
 	static const FDockTabStyle& GetDefault();
 
@@ -1218,10 +1218,10 @@ struct SLATECORE_API FScrollBoxStyle : public FSlateWidgetStyle
 
 	virtual ~FScrollBoxStyle() {}
 
-	virtual void GetResources( TArray< const FSlateBrush* >& OutBrushes ) const OVERRIDE;
+	virtual void GetResources( TArray< const FSlateBrush* >& OutBrushes ) const override;
 
 	static const FName TypeName;
-	virtual const FName GetTypeName() const OVERRIDE { return TypeName; };
+	virtual const FName GetTypeName() const override { return TypeName; };
 
 	static const FScrollBoxStyle& GetDefault();
 
@@ -1249,10 +1249,10 @@ struct SLATECORE_API FScrollBorderStyle : public FSlateWidgetStyle
 
 	virtual ~FScrollBorderStyle() {}
 
-	virtual void GetResources( TArray< const FSlateBrush* >& OutBrushes ) const OVERRIDE;
+	virtual void GetResources( TArray< const FSlateBrush* >& OutBrushes ) const override;
 
 	static const FName TypeName;
-	virtual const FName GetTypeName() const OVERRIDE { return TypeName; };
+	virtual const FName GetTypeName() const override { return TypeName; };
 
 	static const FScrollBorderStyle& GetDefault();
 
@@ -1280,10 +1280,10 @@ struct SLATECORE_API FWindowStyle : public FSlateWidgetStyle
 
 	virtual ~FWindowStyle() {}
 
-	virtual void GetResources( TArray< const FSlateBrush* >& OutBrushes ) const OVERRIDE;
+	virtual void GetResources( TArray< const FSlateBrush* >& OutBrushes ) const override;
 
 	static const FName TypeName;
-	virtual const FName GetTypeName() const OVERRIDE { return TypeName; };
+	virtual const FName GetTypeName() const override { return TypeName; };
 
 	static const FWindowStyle& GetDefault();
 
diff --git a/Engine/Source/Runtime/SlateCore/Public/Styling/SlateWidgetStyleContainerBase.h b/Engine/Source/Runtime/SlateCore/Public/Styling/SlateWidgetStyleContainerBase.h
--- a/Engine/Source/Runtime/SlateCore/Public/Styling/SlateWidgetStyleContainerBase.h
+++ b/Engine/Source/Runtime/SlateCore/Public/Styling/SlateWidgetStyleContainerBase.h
@@ -17,5 +17,5 @@ class SLATECORE_API USlateWidgetStyleContainerBase : public UObject, public ISla
 
 public:
 
-	virtual const struct FSlateWidgetStyle* const GetStyle() const OVERRIDE;
+	virtual const struct FSlateWidgetStyle* const GetStyle() const override;
 };
diff --git a/Engine/Source/Runtime/SlateCore/Public/Textures/SlateShaderResource.h b/Engine/Source/Runtime/SlateCore/Public/Textures/SlateShaderResource.h
--- a/Engine/Source/Runtime/SlateCore/Public/Textures/SlateShaderResource.h
+++ b/Engine/Source/Runtime/SlateCore/Public/Textures/SlateShaderResource.h
@@ -135,7 +135,7 @@ public:
 
 	// Begin FSlateShaderResource interface
 
-	virtual ESlateShaderResource::Type GetType() const OVERRIDE
+	virtual ESlateShaderResource::Type GetType() const override
 	{
 		return ESlateShaderResource::Texture;
 	}
diff --git a/Engine/Source/Runtime/SlateCore/Public/Widgets/DeclarativeSyntaxSupport.h b/Engine/Source/Runtime/SlateCore/Public/Widgets/DeclarativeSyntaxSupport.h
--- a/Engine/Source/Runtime/SlateCore/Public/Widgets/DeclarativeSyntaxSupport.h
+++ b/Engine/Source/Runtime/SlateCore/Public/Widgets/DeclarativeSyntaxSupport.h
@@ -73,7 +73,7 @@
  *         ] );
  *     }
  *
- *     virtual void DoStuff() OVERRIDE {}
+ *     virtual void DoStuff() override {}
  *     
  *     // Truly private. All handlers can be inlined (less boilerplate).
  *     // All private members are truly private.
diff --git a/Engine/Source/Runtime/SlateCore/Public/Widgets/SCompoundWidget.h b/Engine/Source/Runtime/SlateCore/Public/Widgets/SCompoundWidget.h
--- a/Engine/Source/Runtime/SlateCore/Public/Widgets/SCompoundWidget.h
+++ b/Engine/Source/Runtime/SlateCore/Public/Widgets/SCompoundWidget.h
@@ -70,15 +70,15 @@ public:
 
 	// Begin SWidgetOverrides
 
-	virtual int32 OnPaint( const FGeometry& AllottedGeometry, const FSlateRect& MyClippingRect, FSlateWindowElementList& OutDrawElements, int32 LayerId, const FWidgetStyle& InWidgetStyle, bool bParentEnabled ) const OVERRIDE;
+	virtual int32 OnPaint( const FGeometry& AllottedGeometry, const FSlateRect& MyClippingRect, FSlateWindowElementList& OutDrawElements, int32 LayerId, const FWidgetStyle& InWidgetStyle, bool bParentEnabled ) const override;
 	
-	virtual FChildren* GetChildren() OVERRIDE;
+	virtual FChildren* GetChildren() override;
 
-	virtual FVector2D ComputeDesiredSize() const OVERRIDE;
+	virtual FVector2D ComputeDesiredSize() const override;
 
-	virtual void OnArrangeChildren( const FGeometry& AllottedGeometry, FArrangedChildren& ArrangedChildren ) const OVERRIDE;
+	virtual void OnArrangeChildren( const FGeometry& AllottedGeometry, FArrangedChildren& ArrangedChildren ) const override;
 
-	virtual FSlateColor GetForegroundColor() const OVERRIDE;
+	virtual FSlateColor GetForegroundColor() const override;
 
 	// Begin SWidgetOverrides
 
diff --git a/Engine/Source/Runtime/SlateCore/Public/Widgets/SOverlay.h b/Engine/Source/Runtime/SlateCore/Public/Widgets/SOverlay.h
--- a/Engine/Source/Runtime/SlateCore/Public/Widgets/SOverlay.h
+++ b/Engine/Source/Runtime/SlateCore/Public/Widgets/SOverlay.h
@@ -119,10 +119,10 @@ public:
 	}
 
 	// SWidget interface
-	virtual void OnArrangeChildren( const FGeometry& AllottedGeometry, FArrangedChildren& ArrangedChildren ) const OVERRIDE;
-	virtual FVector2D ComputeDesiredSize() const OVERRIDE;
-	virtual FChildren* GetChildren() OVERRIDE;
-	virtual int32 OnPaint( const FGeometry& AllottedGeometry, const FSlateRect& MyClippingRect, FSlateWindowElementList& OutDrawElements, int32 LayerId, const FWidgetStyle& InWidgetStyle, bool bParentEnabled ) const OVERRIDE;
+	virtual void OnArrangeChildren( const FGeometry& AllottedGeometry, FArrangedChildren& ArrangedChildren ) const override;
+	virtual FVector2D ComputeDesiredSize() const override;
+	virtual FChildren* GetChildren() override;
+	virtual int32 OnPaint( const FGeometry& AllottedGeometry, const FSlateRect& MyClippingRect, FSlateWindowElementList& OutDrawElements, int32 LayerId, const FWidgetStyle& InWidgetStyle, bool bParentEnabled ) const override;
 	// End of SWidget interface
 
 protected:
diff --git a/Engine/Source/Runtime/SlateCore/Public/Widgets/SPanel.h b/Engine/Source/Runtime/SlateCore/Public/Widgets/SPanel.h
--- a/Engine/Source/Runtime/SlateCore/Public/Widgets/SPanel.h
+++ b/Engine/Source/Runtime/SlateCore/Public/Widgets/SPanel.h
@@ -56,7 +56,7 @@ public:
 
 	// Begin SWidget overrides
 
-	virtual int32 OnPaint( const FGeometry& AllottedGeometry, const FSlateRect& MyClippingRect, FSlateWindowElementList& OutDrawElements, int32 LayerId, const FWidgetStyle& InWidgetStyle, bool bParentEnabled ) const OVERRIDE;
+	virtual int32 OnPaint( const FGeometry& AllottedGeometry, const FSlateRect& MyClippingRect, FSlateWindowElementList& OutDrawElements, int32 LayerId, const FWidgetStyle& InWidgetStyle, bool bParentEnabled ) const override;
 
 	// End SWidget overrides
 
diff --git a/Engine/Source/Runtime/SlateCore/Public/Widgets/SUserWidget.h b/Engine/Source/Runtime/SlateCore/Public/Widgets/SUserWidget.h
--- a/Engine/Source/Runtime/SlateCore/Public/Widgets/SUserWidget.h
+++ b/Engine/Source/Runtime/SlateCore/Public/Widgets/SUserWidget.h
@@ -31,7 +31,7 @@
  *   class SMyWidget : public ::SMyWidget
  *   {
  *     public:
- *     virtual void Construct( const FArguments& InArgs ) OVERRIDE
+ *     virtual void Construct( const FArguments& InArgs ) override
  *     {
  *        SUserWidget::Construct( SUserWidget::FArguments()
  *        [
diff --git a/Engine/Source/Runtime/SlateCore/Public/Widgets/SWindow.h b/Engine/Source/Runtime/SlateCore/Public/Widgets/SWindow.h
--- a/Engine/Source/Runtime/SlateCore/Public/Widgets/SWindow.h
+++ b/Engine/Source/Runtime/SlateCore/Public/Widgets/SWindow.h
@@ -555,16 +555,16 @@ public:
 	/** @return Gets the radius of the corner rounding of the window. */
 	int32 GetCornerRadius();
 
-	virtual bool SupportsKeyboardFocus() const OVERRIDE;
+	virtual bool SupportsKeyboardFocus() const override;
 
 	virtual FReply OnKeyboardFocusReceived( const FGeometry& MyGeometry, const FKeyboardFocusEvent& InKeyboardFocusEvent );
 
-	virtual FReply OnMouseButtonDown( const FGeometry& MyGeometry, const FPointerEvent& MouseEvent ) OVERRIDE;
-	virtual FReply OnMouseButtonUp( const FGeometry& MyGeometry, const FPointerEvent& MouseEvent ) OVERRIDE;
-	virtual FReply OnMouseMove( const FGeometry& MyGeometry, const FPointerEvent& MouseEvent ) OVERRIDE;
+	virtual FReply OnMouseButtonDown( const FGeometry& MyGeometry, const FPointerEvent& MouseEvent ) override;
+	virtual FReply OnMouseButtonUp( const FGeometry& MyGeometry, const FPointerEvent& MouseEvent ) override;
+	virtual FReply OnMouseMove( const FGeometry& MyGeometry, const FPointerEvent& MouseEvent ) override;
 
 	/** The window's desired size takes into account the ratio between the slate units and the pixel size */
-	virtual FVector2D ComputeDesiredSize() const OVERRIDE;
+	virtual FVector2D ComputeDesiredSize() const override;
 
 	/** Marks this window as having been drawn at least once */
 	inline void MarkAsDrawn()
@@ -638,7 +638,7 @@ public:
 
 	// Begin SWidget overrides
 
-	virtual void Tick( const FGeometry& AllottedGeometry, const double InCurrentTime, const float InDeltaTime ) OVERRIDE;
+	virtual void Tick( const FGeometry& AllottedGeometry, const double InCurrentTime, const float InDeltaTime ) override;
 
 	// End SWidget overrides
 
