Fixed Clang compile error n Slate/Public/Framework/Commands/Command.h

diff --git a/Engine/Source/Runtime/Slate/Public/Framework/Commands/Commands.h b/Engine/Source/Runtime/Slate/Public/Framework/Commands/Commands.h
--- a/Engine/Source/Runtime/Slate/Public/Framework/Commands/Commands.h
+++ b/Engine/Source/Runtime/Slate/Public/Framework/Commands/Commands.h
@@ -3,6 +3,7 @@
 #pragma once
 
 #include "Framework/Commands/UICommandInfo.h"
+#include "InputBindingManager.h"
 
 //
 // COMMAND DECLARATION
