Added a comment explaining why CachedAutoWrapTextWidth was cached in OnPaint() rather than Tick()

This is because Tick() doesn't get called while resizing windows, but OnPaint() does.

diff --git a/Engine/Source/Runtime/Slate/Private/Widgets/Text/SMultiLineEditableText.cpp b/Engine/Source/Runtime/Slate/Private/Widgets/Text/SMultiLineEditableText.cpp
--- a/Engine/Source/Runtime/Slate/Private/Widgets/Text/SMultiLineEditableText.cpp
+++ b/Engine/Source/Runtime/Slate/Private/Widgets/Text/SMultiLineEditableText.cpp
@@ -1502,6 +1502,7 @@ void SMultiLineEditableText::Tick( const FGeometry& AllottedGeometry, const doub
 int32 SMultiLineEditableText::OnPaint( const FGeometry& AllottedGeometry, const FSlateRect& MyClippingRect, FSlateWindowElementList& OutDrawElements, int32 LayerId, const FWidgetStyle& InWidgetStyle, bool bParentEnabled ) const
 {
 	// Update the auto-wrap size now that we have computed paint geometry; won't take affect until text frame
+	// Note: This is done here rather than in Tick(), because Tick() doesn't get called while resizing windows, but OnPaint() does
 	if ( AutoWrapText.Get() )
 	{
 		CachedAutoWrapTextWidth = AllottedGeometry.Size.X;
diff --git a/Engine/Source/Runtime/Slate/Private/Widgets/Text/STextBlock.cpp b/Engine/Source/Runtime/Slate/Private/Widgets/Text/STextBlock.cpp
--- a/Engine/Source/Runtime/Slate/Private/Widgets/Text/STextBlock.cpp
+++ b/Engine/Source/Runtime/Slate/Private/Widgets/Text/STextBlock.cpp
@@ -98,6 +98,7 @@ int32 STextBlock::OnPaint( const FGeometry& AllottedGeometry, const FSlateRect&
 	FSlateFontInfo FontInfo = GetFont();
 
 	// Perform text auto-wrapping if that was enabled.
+	// Note: This is done here rather than in Tick(), because Tick() doesn't get called while resizing windows, but OnPaint() does
 	if( AutoWrapText.Get() )
 	{
 		const float OldWrapTextAt = CachedAutoWrapTextWidth;
