Slate: STileView: STableViewBase: Fixed an issue where fully visible STileViews would inappropriately consume mouse wheel messages. The fix was done in two parts. First there was a floating point precision issue with determining if we are at the end of the list. Added IsNearlyZero() to fix this. The second issue was that the stats returned from RegenerateWidgets in the STileView could not return expected results. This was because the logic was assuming that ExactNumWidgetsOnScreen returns fractions based on the amount of widgets that are visible. This can't be done if you have 4 widgets in a row all being partially clipped.

I've changed the ExactNumWidgetsOnScreen to ExactNumRowsOnScreen and adjusted the behaviors accordingly. This fixed an issue where if there wasn't enough items to fill a row it was assumed that we weren't at the start of the list.

So the end result was that if we aren't at the start and the end of the list then scrolling was necessary.

diff --git a/Engine/Source/Runtime/Slate/Private/Framework/Testing/STableViewTesting.cpp b/Engine/Source/Runtime/Slate/Private/Framework/Testing/STableViewTesting.cpp
--- a/Engine/Source/Runtime/Slate/Private/Framework/Testing/STableViewTesting.cpp
+++ b/Engine/Source/Runtime/Slate/Private/Framework/Testing/STableViewTesting.cpp
@@ -733,7 +733,7 @@ public:
 							// Tile view items are this wide
 							.ItemWidth(128)
 							// Tile view items are this tall
-							.ItemHeight(64)
+							.ItemHeight(75)
 							// Tell the tile view where to get its source data
 							.ListItemsSource( &Items )
 							// When the list view needs to generate a widget for some data item, use this method
diff --git a/Engine/Source/Runtime/Slate/Private/Widgets/Views/STableViewBase.cpp b/Engine/Source/Runtime/Slate/Private/Widgets/Views/STableViewBase.cpp
--- a/Engine/Source/Runtime/Slate/Private/Widgets/Views/STableViewBase.cpp
+++ b/Engine/Source/Runtime/Slate/Private/Widgets/Views/STableViewBase.cpp
@@ -195,7 +195,8 @@ void STableViewBase::Tick( const FGeometry& AllottedGeometry, const double InCur
 
 			const int32 NumItemsBeingObserved = GetNumItemsBeingObserved();
 
-			const bool bEnoughRoomForAllItems = ReGenerateResults.ExactNumWidgetsOnScreen >= NumItemsBeingObserved;
+			const int32 NumItemsWide = GetNumItemsWide();
+			const bool bEnoughRoomForAllItems = ReGenerateResults.ExactNumRowsOnScreen >= (NumItemsBeingObserved / NumItemsWide);
 			if (bEnoughRoomForAllItems)
 			{
 				// We can show all the items, so make sure there is no scrolling.
@@ -218,14 +219,14 @@ void STableViewBase::Tick( const FGeometry& AllottedGeometry, const double InCur
 			{
 				// The thumb size is whatever fraction of the items we are currently seeing (including partially seen items).
 				// e.g. if we are seeing 0.5 of the first generated widget and 0.75 of the last widget, that's 1.25 widgets.
-				const double ThumbSizeFraction = ReGenerateResults.ExactNumWidgetsOnScreen / NumItemsBeingObserved;
+				const double ThumbSizeFraction = ReGenerateResults.ExactNumRowsOnScreen / (NumItemsBeingObserved / NumItemsWide);
 				const double OffsetFraction = ScrollOffset / NumItemsBeingObserved;
 				ScrollBar->SetState( OffsetFraction, ThumbSizeFraction );
 			}
 
 			NotifyItemScrolledIntoView();
 
-			bWasAtEndOfList = ScrollBar->DistanceFromBottom() <= 0.f ? true : false;
+			bWasAtEndOfList = FMath::IsNearlyZero(ScrollBar->DistanceFromBottom()) ? true : false;
 
 			bItemsNeedRefresh = false;
 			ItemsPanel->SetRefreshPending(false);
@@ -646,9 +647,9 @@ void STableViewBase::OnRightMouseButtonUp(const FVector2D& SummonLocation)
 
 float STableViewBase::GetScrollRateInItems() const
 {
-	return ( LastGenerateResults.HeightOfGeneratedItems != 0 && LastGenerateResults.ExactNumWidgetsOnScreen != 0)
+	return (LastGenerateResults.HeightOfGeneratedItems != 0 && LastGenerateResults.ExactNumRowsOnScreen != 0)
 		// Approximate a consistent scrolling rate based on the average item height.
-		? LastGenerateResults.ExactNumWidgetsOnScreen / LastGenerateResults.HeightOfGeneratedItems
+		? LastGenerateResults.ExactNumRowsOnScreen / LastGenerateResults.HeightOfGeneratedItems
 		// Scroll 1/2 an item at a time as a default.
 		: 0.5f;
 }
diff --git a/Engine/Source/Runtime/Slate/Public/Widgets/Views/SListView.h b/Engine/Source/Runtime/Slate/Public/Widgets/Views/SListView.h
--- a/Engine/Source/Runtime/Slate/Public/Widgets/Views/SListView.h
+++ b/Engine/Source/Runtime/Slate/Public/Widgets/Views/SListView.h
@@ -1077,7 +1077,7 @@ protected:
 				if (NumLiveWidgets == 0 && IsPendingRefresh())
 				{
 					// Use the last number of widgets on screen to estimate if we actually need to scroll.
-					NumLiveWidgets = LastGenerateResults.ExactNumWidgetsOnScreen;
+					NumLiveWidgets = LastGenerateResults.ExactNumRowsOnScreen;
 				}
 
 				// Only scroll the item into view if it's not already in the visible range
diff --git a/Engine/Source/Runtime/Slate/Public/Widgets/Views/STableViewBase.h b/Engine/Source/Runtime/Slate/Public/Widgets/Views/STableViewBase.h
--- a/Engine/Source/Runtime/Slate/Public/Widgets/Views/STableViewBase.h
+++ b/Engine/Source/Runtime/Slate/Public/Widgets/Views/STableViewBase.h
@@ -161,7 +161,7 @@ protected:
 		FReGenerateResults( double InNewScrollOffset, double InHeightGenerated, double InItemsOnScreen, bool AtEndOfList )
 		: NewScrollOffset( InNewScrollOffset )
 		, HeightOfGeneratedItems( InHeightGenerated )
-		, ExactNumWidgetsOnScreen( InItemsOnScreen )
+		, ExactNumRowsOnScreen( InItemsOnScreen )
 		, bGeneratedPastLastItem( AtEndOfList )
 		{
 
@@ -173,8 +173,8 @@ protected:
 		/** The total height of the widgets that we have generated to represent the visible subset of the items*/
 		double HeightOfGeneratedItems;
 
-		/** How many items are fitting on the screen, including fractions */
-		double ExactNumWidgetsOnScreen;
+		/** How many rows are fitting on the screen, including fractions */
+		double ExactNumRowsOnScreen;
 
 		/** True when we have generated  */
 		bool bGeneratedPastLastItem;
diff --git a/Engine/Source/Runtime/Slate/Public/Widgets/Views/STileView.h b/Engine/Source/Runtime/Slate/Public/Widgets/Views/STileView.h
--- a/Engine/Source/Runtime/Slate/Public/Widgets/Views/STileView.h
+++ b/Engine/Source/Runtime/Slate/Public/Widgets/Views/STileView.h
@@ -180,15 +180,32 @@ public:
 			// Index of the item at which we start generating based on how far scrolled down we are
 			int32 StartIndex = FMath::Max( 0, FMath::FloorToInt(ClampedScrollOffset / NumItemsWide) * NumItemsWide );
 
-			// Let the WidgetGenerator that we are starting a pass so that it can keep track of data items and widgets.
+			// Let the WidgetGenerator know that we are starting a pass so that it can keep track of data items and widgets.
 			this->WidgetGenerator.OnBeginGenerationPass();
 
 			// Actually generate the widgets.
 			bool bKeepGenerating = true;
+			bool bNewRow = true;
+			double NumRowsShownOnScreen = 0;
 			for( int32 ItemIndex = StartIndex; bKeepGenerating && ItemIndex < SourceItems->Num(); ++ItemIndex )
 			{
 				const ItemType& CurItem = (*SourceItems)[ItemIndex];
 
+				if (bNewRow)
+				{
+					bNewRow = false;
+					HeightUsedSoFar += ItemHeight;
+
+					if (HeightUsedSoFar > MyGeometry.Size.Y)
+					{
+						NumRowsShownOnScreen += 1.0f - ((HeightUsedSoFar - MyGeometry.Size.Y) / ItemHeight);
+					}
+					else
+					{
+						++NumRowsShownOnScreen;
+					}
+				}
+
 				const float GeneratedItemHeight = SListView<ItemType>::GenerateWidgetForItem( CurItem, ItemIndex, StartIndex );
 
 				// The widget used up some of the available horizontal space.
@@ -199,7 +216,8 @@ public:
 				{
 					// A new row of widgets was added.
 					WidthUsedSoFar = 0;
-					HeightUsedSoFar += ItemHeight;	
+					bNewRow = true;
+
 					// Stop when we've generated a widget that's partially clipped by the bottom of the list.
 					if ( HeightUsedSoFar > MyGeometry.Size.Y + ItemHeight )
 					{
@@ -211,9 +229,7 @@ public:
 			// We have completed the generation pass. The WidgetGenerator will clean up unused Widgets.
 			this->WidgetGenerator.OnEndGenerationPass();
 
-			const double ExactNumItemsOnScreen = MyGeometry.Size.Y / ItemHeight * NumItemsWide;
-
-			return STableViewBase::FReGenerateResults(ClampedScrollOffset, ItemHeight * ExactNumItemsOnScreen, ExactNumItemsOnScreen, bAtEndOfList);
+			return STableViewBase::FReGenerateResults(ClampedScrollOffset, HeightUsedSoFar, NumRowsShownOnScreen, bAtEndOfList);
 		}
 
 		return STableViewBase::FReGenerateResults(0, 0, 0, false);
