SLATE/UMG - Extending the SProgressBar to be capable of having the style images be overriden.  Also providing methods for setting padding...etc after initialization.

diff --git a/Engine/Source/Runtime/Slate/Private/Widgets/Notifications/SProgressBar.cpp b/Engine/Source/Runtime/Slate/Private/Widgets/Notifications/SProgressBar.cpp
--- a/Engine/Source/Runtime/Slate/Private/Widgets/Notifications/SProgressBar.cpp
+++ b/Engine/Source/Runtime/Slate/Private/Widgets/Notifications/SProgressBar.cpp
@@ -6,14 +6,17 @@
 void SProgressBar::Construct( const FArguments& InArgs )
 {
 	check(InArgs._Style);
+	
+	Style = InArgs._Style;
 
 	Percent = InArgs._Percent;
 
 	BarFillType = InArgs._BarFillType;
-
-	BackgroundImage = &InArgs._Style->BackgroundImage;
-	FillImage = &InArgs._Style->FillImage;
-	MarqueeImage = &InArgs._Style->MarqueeImage;
+	
+	BackgroundImage = InArgs._BackgroundImage;
+	FillImage = InArgs._FillImage;
+	MarqueeImage = InArgs._MarqueeImage;
+	
 	FillColorAndOpacity = InArgs._FillColorAndOpacity;
 	BorderPadding = InArgs._BorderPadding;
 
@@ -26,6 +29,56 @@ void SProgressBar::SetPercent(TAttribute< TOptional<float> > InPercent)
 	Percent = InPercent;
 }
 
+void SProgressBar::SetStyle(const FProgressBarStyle* InStyle)
+{
+	Style = InStyle;
+}
+
+void SProgressBar::SetBarFillType(EProgressBarFillType::Type InBarFillType)
+{
+	BarFillType = InBarFillType;
+}
+
+void SProgressBar::SetFillColorAndOpacity(TAttribute< FSlateColor > InFillColorAndOpacity)
+{
+	FillColorAndOpacity = InFillColorAndOpacity;
+}
+
+void SProgressBar::SetBorderPadding(TAttribute< FVector2D > InBorderPadding)
+{
+	BorderPadding = InBorderPadding;
+}
+
+void SProgressBar::SetBackgroundImage(const FSlateBrush* InBackgroundImage)
+{
+	BackgroundImage = InBackgroundImage;
+}
+
+void SProgressBar::SetFillImage(const FSlateBrush* InFillImage)
+{
+	FillImage = InFillImage;
+}
+
+void SProgressBar::SetMarqueeImage(const FSlateBrush* InMarqueeImage)
+{
+	MarqueeImage = InMarqueeImage;
+}
+
+const FSlateBrush* SProgressBar::GetBackgroundImage() const
+{
+	return BackgroundImage ? BackgroundImage : &Style->BackgroundImage;
+}
+
+const FSlateBrush* SProgressBar::GetFillImage() const
+{
+	return FillImage ? FillImage : &Style->FillImage;
+}
+
+const FSlateBrush* SProgressBar::GetMarqueeImage() const
+{
+	return MarqueeImage ? MarqueeImage : &Style->MarqueeImage;
+}
+
 int32 SProgressBar::OnPaint( const FGeometry& AllottedGeometry, const FSlateRect& MyClippingRect, FSlateWindowElementList& OutDrawElements, int32 LayerId, const FWidgetStyle& InWidgetStyle, bool bParentEnabled ) const
 {
 	// Used to track the layer ID we will return.
@@ -34,7 +87,9 @@ int32 SProgressBar::OnPaint( const FGeometry& AllottedGeometry, const FSlateRect
 	bool bEnabled = ShouldBeEnabled( bParentEnabled );
 	const ESlateDrawEffect::Type DrawEffects = bEnabled ? ESlateDrawEffect::None : ESlateDrawEffect::DisabledEffect;
 	
-	const FColor FillColorAndOpacitySRGB( InWidgetStyle.GetColorAndOpacityTint() * FillColorAndOpacity.Get().GetColor(InWidgetStyle) * FillImage->GetTint( InWidgetStyle ) );
+	const FSlateBrush* CurrentFillImage = GetFillImage();
+	
+	const FColor FillColorAndOpacitySRGB( InWidgetStyle.GetColorAndOpacityTint() * FillColorAndOpacity.Get().GetColor(InWidgetStyle) * CurrentFillImage->GetTint( InWidgetStyle ) );
 	const FColor ColorAndOpacitySRGB = InWidgetStyle.GetColorAndOpacityTint();
 
 	TOptional<float> ProgressFraction = Percent.Get();	
@@ -42,15 +97,17 @@ int32 SProgressBar::OnPaint( const FGeometry& AllottedGeometry, const FSlateRect
 	// Paint inside the border only. 
 	FPaintGeometry ForegroundPaintGeometry = AllottedGeometry.ToInflatedPaintGeometry( -BorderPadding.Get() );
 	const FSlateRect ForegroundClippingRect = ForegroundPaintGeometry.ToSlateRect().IntersectionWith( MyClippingRect );
+	
+	const FSlateBrush* CurrentBackgroundImage = GetBackgroundImage();
 
 	FSlateDrawElement::MakeBox(
 		OutDrawElements,
 		RetLayerId++,
 		AllottedGeometry.ToPaintGeometry(),
-		BackgroundImage,
+		CurrentBackgroundImage,
 		MyClippingRect,
 		DrawEffects,
-		InWidgetStyle.GetColorAndOpacityTint() * BackgroundImage->GetTint( InWidgetStyle )
+		InWidgetStyle.GetColorAndOpacityTint() * CurrentBackgroundImage->GetTint( InWidgetStyle )
 	);	
 	
 	if( ProgressFraction.IsSet() )
@@ -67,7 +124,7 @@ int32 SProgressBar::OnPaint( const FGeometry& AllottedGeometry, const FSlateRect
 					AllottedGeometry.ToPaintGeometry(
 						FVector2D( AllottedGeometry.Size.X - (AllottedGeometry.Size.X * ( ClampedFraction )) , 0.0f),
 						FVector2D( AllottedGeometry.Size.X * ( ClampedFraction ) , AllottedGeometry.Size.Y )),
-					FillImage,
+					CurrentFillImage,
 					ForegroundClippingRect,
 					DrawEffects,
 					FillColorAndOpacitySRGB
@@ -84,7 +141,7 @@ int32 SProgressBar::OnPaint( const FGeometry& AllottedGeometry, const FSlateRect
 					AllottedGeometry.ToPaintGeometry(
 						FVector2D( (AllottedGeometry.Size.X * 0.5f) - ((AllottedGeometry.Size.X * ( ClampedFraction ))*0.5), 0.0f),
 						FVector2D( AllottedGeometry.Size.X * ( ClampedFraction ) , AllottedGeometry.Size.Y )),
-					FillImage,
+					CurrentFillImage,
 					ForegroundClippingRect,
 					DrawEffects,
 					FillColorAndOpacitySRGB
@@ -102,7 +159,7 @@ int32 SProgressBar::OnPaint( const FGeometry& AllottedGeometry, const FSlateRect
 					AllottedGeometry.ToPaintGeometry(
 						FVector2D::ZeroVector,
 						FVector2D( AllottedGeometry.Size.X * ( ClampedFraction ) , AllottedGeometry.Size.Y )),
-					FillImage,
+					CurrentFillImage,
 					ForegroundClippingRect,
 					DrawEffects,
 					FillColorAndOpacitySRGB
@@ -114,9 +171,11 @@ int32 SProgressBar::OnPaint( const FGeometry& AllottedGeometry, const FSlateRect
 	}
 	else
 	{
+		const FSlateBrush* CurrentMarqueeImage = GetMarqueeImage();
+		
 		// Draw Marquee
-		const float MarqueeAnimOffset = MarqueeImage->ImageSize.X * CurveSequence.GetLerpLooping();
-		const float MarqueeImageSize = MarqueeImage->ImageSize.X;
+		const float MarqueeAnimOffset = CurrentMarqueeImage->ImageSize.X * CurveSequence.GetLerpLooping();
+		const float MarqueeImageSize = CurrentMarqueeImage->ImageSize.X;
 
 		FSlateDrawElement::MakeBox(
 			OutDrawElements,
@@ -124,7 +183,7 @@ int32 SProgressBar::OnPaint( const FGeometry& AllottedGeometry, const FSlateRect
 			AllottedGeometry.ToPaintGeometry(
 				FVector2D( MarqueeAnimOffset - MarqueeImageSize, 0.0f ),
 				FVector2D( AllottedGeometry.Size.X + MarqueeImageSize, AllottedGeometry.Size.Y )),
-			MarqueeImage,
+			CurrentMarqueeImage,
 			ForegroundClippingRect,
 			DrawEffects,
 			ColorAndOpacitySRGB
@@ -142,7 +201,5 @@ int32 SProgressBar::OnPaint( const FGeometry& AllottedGeometry, const FSlateRect
  */
 FVector2D SProgressBar::ComputeDesiredSize() const
 {
-	return MarqueeImage->ImageSize;
+	return GetMarqueeImage()->ImageSize;
 }
-
-
diff --git a/Engine/Source/Runtime/Slate/Public/Widgets/Notifications/SProgressBar.h b/Engine/Source/Runtime/Slate/Public/Widgets/Notifications/SProgressBar.h
--- a/Engine/Source/Runtime/Slate/Public/Widgets/Notifications/SProgressBar.h
+++ b/Engine/Source/Runtime/Slate/Public/Widgets/Notifications/SProgressBar.h
@@ -32,6 +32,9 @@ public:
 		, _Percent( TOptional<float>() )
 		, _FillColorAndOpacity( FLinearColor::White )
 		, _BorderPadding( FVector2D(1,0) )
+		, _BackgroundImage(nullptr)
+		, _FillImage(nullptr)
+		, _MarqueeImage(nullptr)
 		{}
 
 		/** Style used for the progress bar */
@@ -48,6 +51,15 @@ public:
 
 		/** Border Padding around fill bar */
 		SLATE_ATTRIBUTE( FVector2D, BorderPadding )
+	
+		/** The brush to use as the background of the progress bar */
+		SLATE_ARGUMENT(const FSlateBrush*, BackgroundImage)
+	
+		/** The brush to use as the fill image */
+		SLATE_ARGUMENT(const FSlateBrush*, FillImage)
+	
+		/** The brush to use as the marquee image */
+		SLATE_ARGUMENT(const FSlateBrush*, MarqueeImage)
 
 	SLATE_END_ARGS()
 
@@ -62,10 +74,42 @@ public:
 
 	FVector2D ComputeDesiredSize() const;
 
-	/** See Slate attribute Percent */
+	/** See attribute Percent */
 	void SetPercent(TAttribute< TOptional<float> > InPercent);
+	
+	/** See attribute Style */
+	void SetStyle(const FProgressBarStyle* InStyle);
+	
+	/** See attribute BarFillType */
+	void SetBarFillType(EProgressBarFillType::Type InBarFillType);
+	
+	/** See attribute SetFillColorAndOpacity */
+	void SetFillColorAndOpacity(TAttribute< FSlateColor > InFillColorAndOpacity);
+	
+	/** See attribute BorderPadding */
+	void SetBorderPadding(TAttribute< FVector2D > InBorderPadding);
+	
+	/** See attribute BackgroundImage */
+	void SetBackgroundImage(const FSlateBrush* InBackgroundImage);
+	
+	/** See attribute FillImage */
+	void SetFillImage(const FSlateBrush* InFillImage);
+	
+	/** See attribute MarqueeImage */
+	void SetMarqueeImage(const FSlateBrush* InMarqueeImage);
+	
+private:
+	/** Gets the current background image. */
+	const FSlateBrush* GetBackgroundImage() const;
+	/** Gets the current fill image */
+	const FSlateBrush* GetFillImage() const;
+	/** Gets the current marquee image */
+	const FSlateBrush* GetMarqueeImage() const;
 
 private:
+	
+	/** The style of the progress bar */
+	const FProgressBarStyle* Style;
 
 	/** The text displayed over the progress bar */
 	TAttribute< TOptional<float> > Percent;
@@ -89,6 +133,5 @@ private:
 
 	/** Curve sequence to drive progress bar animation */
 	FCurveSequence CurveSequence;
-
 };
 
