SLATE/UMG - The SUniformGridPanel now allows slots to be removed at runtime like other layout panels.

diff --git a/Engine/Source/Runtime/Slate/Private/Widgets/Layout/SUniformGridPanel.cpp b/Engine/Source/Runtime/Slate/Private/Widgets/Layout/SUniformGridPanel.cpp
--- a/Engine/Source/Runtime/Slate/Private/Widgets/Layout/SUniformGridPanel.cpp
+++ b/Engine/Source/Runtime/Slate/Private/Widgets/Layout/SUniformGridPanel.cpp
@@ -96,6 +96,19 @@ SUniformGridPanel::FSlot& SUniformGridPanel::AddSlot( int32 Column, int32 Row )
 	return NewSlot;
 }
 
+bool SUniformGridPanel::RemoveSlot( const TSharedRef<SWidget>& SlotWidget )
+{
+	for (int32 SlotIdx = 0; SlotIdx < Children.Num(); ++SlotIdx)
+	{
+		if ( SlotWidget == Children[SlotIdx].Widget )
+		{
+			Children.RemoveAt(SlotIdx);
+			return true;
+		}
+	}
+	
+	return false;
+}
 
 void SUniformGridPanel::ClearChildren()
 {
diff --git a/Engine/Source/Runtime/Slate/Public/Widgets/Layout/SUniformGridPanel.h b/Engine/Source/Runtime/Slate/Public/Widgets/Layout/SUniformGridPanel.h
--- a/Engine/Source/Runtime/Slate/Public/Widgets/Layout/SUniformGridPanel.h
+++ b/Engine/Source/Runtime/Slate/Public/Widgets/Layout/SUniformGridPanel.h
@@ -64,6 +64,14 @@ public:
 	 * @return A reference to the newly-added slot
 	 */
 	FSlot& AddSlot( int32 Column, int32 Row );
+	
+	/**
+	 * Removes a slot from this panel which contains the specified SWidget
+	 *
+	 * @param SlotWidget The widget to match when searching through the slots
+	 * @returns The true if the slot was removed and false if no slot was found matching the widget
+	 */
+	bool RemoveSlot( const TSharedRef<SWidget>& SlotWidget );
 
 	/** Removes all slots from the panel */
 	void ClearChildren();
