SLATE/UMG - Adding more methods to the SSlider to allow style related properties to be modified after initialization.

diff --git a/Engine/Source/Runtime/Slate/Private/Widgets/Input/SSlider.cpp b/Engine/Source/Runtime/Slate/Private/Widgets/Input/SSlider.cpp
--- a/Engine/Source/Runtime/Slate/Private/Widgets/Input/SSlider.cpp
+++ b/Engine/Source/Runtime/Slate/Private/Widgets/Input/SSlider.cpp
@@ -199,3 +199,28 @@ void SSlider::SetValue(const TAttribute<float>& InValueAttribute)
 {
 	ValueAttribute = InValueAttribute;
 }
+
+void SSlider::SetIndentHandle(const TAttribute<bool>& InIndentHandle)
+{
+    IndentHandle = InIndentHandle;
+}
+
+void SSlider::SetLocked(const TAttribute<bool>& InLocked)
+{
+    LockedAttribute = InLocked;
+}
+
+void SSlider::SetOrientation(EOrientation InOrientation)
+{
+    Orientation = InOrientation;
+}
+
+void SSlider::SetSliderBarColor(FSlateColor InSliderBarColor)
+{
+    SliderBarColor = InSliderBarColor;
+}
+
+void SSlider::SetSliderHandleColor(FSlateColor InSliderHandleColor)
+{
+    SliderHandleColor = InSliderHandleColor;
+}
diff --git a/Engine/Source/Runtime/Slate/Public/Widgets/Input/SSlider.h b/Engine/Source/Runtime/Slate/Public/Widgets/Input/SSlider.h
--- a/Engine/Source/Runtime/Slate/Public/Widgets/Input/SSlider.h
+++ b/Engine/Source/Runtime/Slate/Public/Widgets/Input/SSlider.h
@@ -68,6 +68,21 @@ public:
 
 	/** See the Value attribute */
 	void SetValue(const TAttribute<float>& InValueAttribute);
+    
+    /** See the IndentHandle attribute */
+	void SetIndentHandle(const TAttribute<bool>& InIndentHandle);
+    
+    /** See the Locked attribute */
+	void SetLocked(const TAttribute<bool>& InLocked);
+    
+    /** See the Orientation attribute */
+	void SetOrientation(EOrientation InOrientation);
+    
+    /** See the SliderBarColor attribute */
+	void SetSliderBarColor(FSlateColor InSliderBarColor);
+    
+    /** See the SliderHandleColor attribute */
+	void SetSliderHandleColor(FSlateColor InSliderHandleColor);
 
 public:
 
