SLATE/UMG - Adding more methods to SEditableText to allow more attributes to be set after initialization.

diff --git a/Engine/Source/Runtime/Slate/Private/Widgets/Input/SEditableText.cpp b/Engine/Source/Runtime/Slate/Private/Widgets/Input/SEditableText.cpp
--- a/Engine/Source/Runtime/Slate/Private/Widgets/Input/SEditableText.cpp
+++ b/Engine/Source/Runtime/Slate/Private/Widgets/Input/SEditableText.cpp
@@ -1966,3 +1966,23 @@ void SEditableText::OnWindowClosed(const TSharedRef<SWindow>&)
 	// Give this focus when the context menu has been dismissed
 	FSlateApplication::Get().SetKeyboardFocus( AsShared(), EKeyboardFocusCause::OtherWidgetLostFocus );
 }
+
+void SEditableText::SetHintText( const TAttribute< FText >& InHintText )
+{
+    HintText = InHintText;
+}
+
+void SEditableText::SetIsReadOnly( TAttribute< bool > InIsReadOnly )
+{
+    IsReadOnly = InIsReadOnly;
+}
+
+void SEditableText::SetIsPassword( TAttribute< bool > InIsPassword )
+{
+    IsPassword = InIsPassword;
+}
+
+void SEditableText::SetColorAndOpacity(TAttribute<FSlateColor> Color)
+{
+    ColorAndOpacity = Color;
+}
diff --git a/Engine/Source/Runtime/Slate/Public/Widgets/Input/SEditableText.h b/Engine/Source/Runtime/Slate/Public/Widgets/Input/SEditableText.h
--- a/Engine/Source/Runtime/Slate/Public/Widgets/Input/SEditableText.h
+++ b/Engine/Source/Runtime/Slate/Public/Widgets/Input/SEditableText.h
@@ -131,6 +131,18 @@ public:
 	 * @param  InNewText  The new text
 	 */
 	void SetText( const TAttribute< FText >& InNewText );
+    
+    /** See the HintText attribute */
+    void SetHintText( const TAttribute< FText >& InHintText );
+    
+    /** See the IsReadOnly attribute */
+    void SetIsReadOnly( TAttribute< bool > InIsReadOnly );
+    
+    /** See the IsPassword attribute */
+    void SetIsPassword( TAttribute< bool > InIsPassword );
+    
+    /** See the ColorAndOpacity attribute */
+    void SetColorAndOpacity(TAttribute<FSlateColor> Color);
 
 	/**
 	 * Restores the text to the original state
@@ -442,6 +454,7 @@ private:
 	/** The text content for this editable text widget */
 	TAttribute< FText > Text;
 
+    /** The text content fot watermark/hinting what text to type here */
 	TAttribute< FText > HintText;
 
 	/** The font used to draw the text */
