SLATE/UMG - Adding individual overrides for images of a checkbox, and the ability to change the style after initializaing the checkbox.

diff --git a/Engine/Source/Runtime/Slate/Private/Widgets/Input/SCheckBox.cpp b/Engine/Source/Runtime/Slate/Private/Widgets/Input/SCheckBox.cpp
--- a/Engine/Source/Runtime/Slate/Private/Widgets/Input/SCheckBox.cpp
+++ b/Engine/Source/Runtime/Slate/Private/Widgets/Input/SCheckBox.cpp
@@ -11,6 +11,20 @@
 void SCheckBox::Construct( const SCheckBox::FArguments& InArgs )
 {
 	check(InArgs._Style != NULL);
+	Style = InArgs._Style;
+
+	UncheckedImage = InArgs._UncheckedImage;
+	UncheckedHoveredImage = InArgs._UncheckedHoveredImage;
+	UncheckedPressedImage = InArgs._UncheckedPressedImage;
+
+	CheckedImage = InArgs._CheckedImage;
+	CheckedHoveredImage = InArgs._CheckedHoveredImage;
+	CheckedPressedImage = InArgs._CheckedPressedImage;
+
+	UndeterminedImage = InArgs._UndeterminedImage;
+	UndeterminedHoveredImage = InArgs._UndeterminedHoveredImage;
+	UndeterminedPressedImage = InArgs._UndeterminedPressedImage;
+	
 	TAttribute<FMargin> Padding = InArgs._Padding.IsSet() ? InArgs._Padding : InArgs._Style->Padding;
 	TAttribute<FSlateColor> ForegroundColor = InArgs._ForegroundColor.IsSet() ? InArgs._ForegroundColor : InArgs._Style->ForegroundColor;
 	TAttribute<FSlateColor> BorderBackgroundColor = InArgs._BorderBackgroundColor.IsSet() ? InArgs._BorderBackgroundColor : InArgs._Style->BorderBackgroundColor;
@@ -27,7 +41,7 @@ void SCheckBox::Construct( const SCheckBox::FArguments& InArgs )
 			SNew(SHorizontalBox)
 			// Make sure we aren't trying to compute the desired size when the check box is collapsed
 			.Visibility(InArgs._Visibility)
-			+SHorizontalBox::Slot()
+			+ SHorizontalBox::Slot()
 			.AutoWidth()
 			.VAlign( VAlign_Center )
 			.HAlign( HAlign_Center )
@@ -70,17 +84,6 @@ void SCheckBox::Construct( const SCheckBox::FArguments& InArgs )
 		];
 	}
 
-	/* Set images to use in various SCheckBox states */
-	UncheckedImage = &InArgs._Style->UncheckedImage;
-	UncheckedHoveredImage = &InArgs._Style->UncheckedHoveredImage;
-	UncheckedPressedImage = &InArgs._Style->UncheckedPressedImage;
-	CheckedImage = &InArgs._Style->CheckedImage;
-	CheckedHoveredImage = &InArgs._Style->CheckedHoveredImage;
-	CheckedPressedImage = &InArgs._Style->CheckedPressedImage;
-	UndeterminedImage = &InArgs._Style->UndeterminedImage;
-	UndeterminedHoveredImage = &InArgs._Style->UndeterminedHoveredImage;
-	UndeterminedPressedImage = &InArgs._Style->UndeterminedPressedImage;
-
 	IsCheckboxChecked = InArgs._IsChecked;
 	OnCheckStateChanged = InArgs._OnCheckStateChanged;
 
@@ -273,16 +276,16 @@ const FSlateBrush* SCheckBox::OnGetCheckImage() const
 	switch( State )
 	{
 		case ESlateCheckBoxState::Unchecked:
-			ImageToUse = IsPressed() ? UncheckedPressedImage : ( IsHovered() ? UncheckedHoveredImage : UncheckedImage );
+			ImageToUse = IsPressed() ? GetUncheckedPressedImage() : ( IsHovered() ? GetUncheckedHoveredImage() : GetUncheckedImage() );
 			break;
 	
 		case ESlateCheckBoxState::Checked:
-			ImageToUse = IsPressed() ? CheckedPressedImage : ( IsHovered() ? CheckedHoveredImage : CheckedImage );
+			ImageToUse = IsPressed() ? GetCheckedPressedImage() : ( IsHovered() ? GetCheckedHoveredImage() : GetCheckedImage() );
 			break;
 	
 		default:
 		case ESlateCheckBoxState::Undetermined:
-			ImageToUse = IsPressed() ? UndeterminedPressedImage : ( IsHovered() ? UndeterminedHoveredImage : UndeterminedImage );
+			ImageToUse = IsPressed() ? GetUndeterminedPressedImage() : ( IsHovered() ? GetUndeterminedHoveredImage() : GetUndeterminedImage() );
 			break;
 	}
 
@@ -337,3 +340,98 @@ void SCheckBox::PlayHoverSound() const
 {
 	FSlateApplication::Get().PlaySound( HoveredSound );
 }
+
+void SCheckBox::SetStyle(const FCheckBoxStyle* InStyle)
+{
+	Style = InStyle;
+}
+
+void SCheckBox::SetUncheckedImage(const FSlateBrush* Brush)
+{
+	UncheckedImage = Brush;
+}
+
+void SCheckBox::SetUncheckedHoveredImage(const FSlateBrush* Brush)
+{
+	UncheckedHoveredImage = Brush;
+}
+
+void SCheckBox::SetUncheckedPressedImage(const FSlateBrush* Brush)
+{
+	UncheckedPressedImage = Brush;
+}
+
+void SCheckBox::SetCheckedImage(const FSlateBrush* Brush)
+{
+	CheckedImage = Brush;
+}
+
+void SCheckBox::SetCheckedHoveredImage(const FSlateBrush* Brush)
+{
+	CheckedHoveredImage = Brush;
+}
+
+void SCheckBox::SetCheckedPressedImage(const FSlateBrush* Brush)
+{
+	CheckedPressedImage = Brush;
+}
+
+void SCheckBox::SetUndeterminedImage(const FSlateBrush* Brush)
+{
+	UndeterminedImage = Brush;
+}
+
+void SCheckBox::SetUndeterminedHoveredImage(const FSlateBrush* Brush)
+{
+	UndeterminedHoveredImage = Brush;
+}
+
+void SCheckBox::SetUndeterminedPressedImage(const FSlateBrush* Brush)
+{
+	UndeterminedPressedImage = Brush;
+}
+
+const FSlateBrush* SCheckBox::GetUncheckedImage() const
+{
+	return UncheckedImage ? UncheckedImage : &Style->UncheckedImage;
+}
+
+const FSlateBrush* SCheckBox::GetUncheckedHoveredImage() const
+{
+	return UncheckedHoveredImage ? UncheckedHoveredImage : &Style->UncheckedHoveredImage;
+}
+
+const FSlateBrush* SCheckBox::GetUncheckedPressedImage() const
+{
+	return UncheckedPressedImage ? UncheckedPressedImage : &Style->UncheckedPressedImage;
+}
+
+const FSlateBrush* SCheckBox::GetCheckedImage() const
+{
+	return CheckedImage ? CheckedImage : &Style->CheckedImage;
+}
+
+const FSlateBrush* SCheckBox::GetCheckedHoveredImage() const
+{
+	return CheckedHoveredImage ? CheckedHoveredImage : &Style->CheckedHoveredImage;
+}
+
+const FSlateBrush* SCheckBox::GetCheckedPressedImage() const
+{
+	return CheckedPressedImage ? CheckedPressedImage : &Style->CheckedPressedImage;
+}
+
+const FSlateBrush* SCheckBox::GetUndeterminedImage() const
+{
+	return UndeterminedImage ? UndeterminedImage : &Style->UndeterminedImage;
+}
+
+const FSlateBrush* SCheckBox::GetUndeterminedHoveredImage() const
+{
+	return UndeterminedHoveredImage ? UndeterminedHoveredImage : &Style->UndeterminedHoveredImage;
+}
+
+const FSlateBrush* SCheckBox::GetUndeterminedPressedImage() const
+{
+	return UndeterminedPressedImage ? UndeterminedPressedImage : &Style->UndeterminedPressedImage;
+}
diff --git a/Engine/Source/Runtime/Slate/Public/Widgets/Input/SCheckBox.h b/Engine/Source/Runtime/Slate/Public/Widgets/Input/SCheckBox.h
--- a/Engine/Source/Runtime/Slate/Public/Widgets/Input/SCheckBox.h
+++ b/Engine/Source/Runtime/Slate/Public/Widgets/Input/SCheckBox.h
@@ -16,7 +16,6 @@ namespace ESlateCheckBoxState
 
 		/** Neither checked nor unchecked */
 	const Type Undetermined = 2;
-	
 };
 
 
@@ -43,6 +42,15 @@ public:
 		, _ForegroundColor()
 		, _BorderBackgroundColor ()
 		, _IsFocusable( true )
+		, _UncheckedImage( NULL )
+		, _UncheckedHoveredImage( NULL )
+		, _UncheckedPressedImage( NULL )
+		, _CheckedImage( NULL )
+		, _CheckedHoveredImage( NULL )
+		, _CheckedPressedImage( NULL )
+		, _UndeterminedImage( NULL )
+		, _UndeterminedHoveredImage( NULL )
+		, _UndeterminedPressedImage( NULL )
 		{}
 
 		/** Content to be placed next to the check box, or for a toggle button, the content to be placed inside the button */
@@ -88,6 +96,33 @@ public:
 		/** The sound to play when the check box is hovered */
 		SLATE_ARGUMENT( TOptional<FSlateSound>, HoveredSoundOverride )
 
+		/** The unchecked image for the checkbox - overrides the style's */
+		SLATE_ARGUMENT(const FSlateBrush*, UncheckedImage)
+
+		/** The unchecked hovered image for the checkbox - overrides the style's */
+		SLATE_ARGUMENT(const FSlateBrush*, UncheckedHoveredImage)
+
+		/** The unchecked pressed image for the checkbox - overrides the style's */
+		SLATE_ARGUMENT(const FSlateBrush*, UncheckedPressedImage)
+
+		/** The checked image for the checkbox - overrides the style's */
+		SLATE_ARGUMENT(const FSlateBrush*, CheckedImage)
+
+		/** The checked hovered image for the checkbox - overrides the style's */
+		SLATE_ARGUMENT(const FSlateBrush*, CheckedHoveredImage)
+
+		/** The checked pressed image for the checkbox - overrides the style's */
+		SLATE_ARGUMENT(const FSlateBrush*, CheckedPressedImage)
+
+		/** The undetermined image for the checkbox - overrides the style's */
+		SLATE_ARGUMENT(const FSlateBrush*, UndeterminedImage)
+
+		/** The undetermined hovered image for the checkbox - overrides the style's */
+		SLATE_ARGUMENT(const FSlateBrush*, UndeterminedHoveredImage)
+
+		/** The undetermined pressed image for the checkbox - overrides the style's */
+		SLATE_ARGUMENT(const FSlateBrush*, UndeterminedPressedImage)
+
 	SLATE_END_ARGS()
 
 
@@ -132,6 +167,30 @@ public:
 	 * Toggles the checked state for this check box, fire events as needed
 	 */
 	void ToggleCheckedState();
+	
+	/** See the Style attribute */
+	void SetStyle(const FCheckBoxStyle* InStyle);
+	
+	/** See the UncheckedImage attribute */
+	void SetUncheckedImage(const FSlateBrush* Brush);
+	/** See the UncheckedHoveredImage attribute */
+	void SetUncheckedHoveredImage(const FSlateBrush* Brush);
+	/** See the UncheckedPressedImage attribute */
+	void SetUncheckedPressedImage(const FSlateBrush* Brush);
+	
+	/** See the CheckedImage attribute */
+	void SetCheckedImage(const FSlateBrush* Brush);
+	/** See the CheckedHoveredImage attribute */
+	void SetCheckedHoveredImage(const FSlateBrush* Brush);
+	/** See the CheckedPressedImage attribute */
+	void SetCheckedPressedImage(const FSlateBrush* Brush);
+	
+	/** See the UndeterminedImage attribute */
+	void SetUndeterminedImage(const FSlateBrush* Brush);
+	/** See the UndeterminedHoveredImage attribute */
+	void SetUndeterminedHoveredImage(const FSlateBrush* Brush);
+	/** See the UndeterminedPressedImage attribute */
+	void SetUndeterminedPressedImage(const FSlateBrush* Brush);
 
 protected:
 
@@ -140,9 +199,22 @@ protected:
 	 * @return	The name of the image to display
 	 */
 	const FSlateBrush* OnGetCheckImage() const;
-
-
+	
+	const FSlateBrush* GetUncheckedImage() const;
+	const FSlateBrush* GetUncheckedHoveredImage() const;
+	const FSlateBrush* GetUncheckedPressedImage() const;
+	
+	const FSlateBrush* GetCheckedImage() const;
+	const FSlateBrush* GetCheckedHoveredImage() const;
+	const FSlateBrush* GetCheckedPressedImage() const;
+	
+	const FSlateBrush* GetUndeterminedImage() const;
+	const FSlateBrush* GetUndeterminedHoveredImage() const;
+	const FSlateBrush* GetUndeterminedPressedImage() const;
+	
 protected:
+	
+	const FCheckBoxStyle* Style;
 
 	/** True if this check box is currently in a pressed state */
 	bool bIsPressed;
