SLATE/UMG - Adding a way to set the various mouse events of the SBorder after initialization.

diff --git a/Engine/Source/Runtime/Slate/Private/Widgets/Layout/SBorder.cpp b/Engine/Source/Runtime/Slate/Private/Widgets/Layout/SBorder.cpp
--- a/Engine/Source/Runtime/Slate/Private/Widgets/Layout/SBorder.cpp
+++ b/Engine/Source/Runtime/Slate/Private/Widgets/Layout/SBorder.cpp
@@ -241,3 +241,23 @@ void SBorder::SetBorderImage(const TAttribute<const FSlateBrush*>& InBorderImage
 {
 	BorderImage = InBorderImage;
 }
+
+void SBorder::SetOnMouseButtonDown(FPointerEventHandler EventHandler)
+{
+	MouseButtonDownHandler = EventHandler;
+}
+
+void SBorder::SetOnMouseButtonUp(FPointerEventHandler EventHandler)
+{
+	MouseButtonUpHandler = EventHandler;
+}
+
+void SBorder::SetOnMouseMove(FPointerEventHandler EventHandler)
+{
+	MouseMoveHandler = EventHandler;
+}
+
+void SBorder::SetOnMouseDoubleClick(FPointerEventHandler EventHandler)
+{
+	MouseDoubleClickHandler = EventHandler;
+}
diff --git a/Engine/Source/Runtime/Slate/Public/Widgets/Layout/SBorder.h b/Engine/Source/Runtime/Slate/Public/Widgets/Layout/SBorder.h
--- a/Engine/Source/Runtime/Slate/Public/Widgets/Layout/SBorder.h
+++ b/Engine/Source/Runtime/Slate/Public/Widgets/Layout/SBorder.h
@@ -107,6 +107,18 @@ public:
 	/** See BorderImage attribute */
 	void SetBorderImage(const TAttribute<const FSlateBrush*>& InBorderImage);
 
+	/** See OnMouseButtonDown event */
+	void SetOnMouseButtonDown(FPointerEventHandler EventHandler);
+
+	/** See OnMouseButtonUp event */
+	void SetOnMouseButtonUp(FPointerEventHandler EventHandler);
+
+	/** See OnMouseMove event */
+	void SetOnMouseMove(FPointerEventHandler EventHandler);
+
+	/** See OnMouseDoubleClick event */
+	void SetOnMouseDoubleClick(FPointerEventHandler EventHandler);
+
 public:
 	// SWidget interface
 	virtual int32 OnPaint( const FGeometry& AllottedGeometry, const FSlateRect& MyClippingRect, FSlateWindowElementList& OutDrawElements, int32 LayerId, const FWidgetStyle& InWidgetStyle, bool bParentEnabled ) const override;
