Slate/UMG - Exposing setters for item sizes on the STableViewBase.

diff --git a/Engine/Source/Runtime/Slate/Private/Widgets/Views/SListPanel.cpp b/Engine/Source/Runtime/Slate/Private/Widgets/Views/SListPanel.cpp
--- a/Engine/Source/Runtime/Slate/Private/Widgets/Views/SListPanel.cpp
+++ b/Engine/Source/Runtime/Slate/Private/Widgets/Views/SListPanel.cpp
@@ -228,3 +228,12 @@ bool SListPanel::ShouldArrangeHorizontally() const
 	return ItemWidth.Get() > 0;
 }
 
+void SListPanel::SetItemHeight(TAttribute<float> Height)
+{
+	ItemHeight = Height;
+}
+
+void SListPanel::SetItemWidth(TAttribute<float> Width)
+{
+	ItemWidth = Width;
+}
diff --git a/Engine/Source/Runtime/Slate/Private/Widgets/Views/SListPanel.h b/Engine/Source/Runtime/Slate/Private/Widgets/Views/SListPanel.h
--- a/Engine/Source/Runtime/Slate/Private/Widgets/Views/SListPanel.h
+++ b/Engine/Source/Runtime/Slate/Private/Widgets/Views/SListPanel.h
@@ -85,6 +85,12 @@ public:
 
 	/** Returns true if this list panel is pending a refresh, false otherwise */
 	bool IsRefreshPending() const;
+
+	/** See ItemHeight attribute */
+	void SetItemHeight(TAttribute<float> Height);
+
+	/** See ItemWidth attribute */
+	void SetItemWidth(TAttribute<float> Width);
 	
 protected:
 
diff --git a/Engine/Source/Runtime/Slate/Private/Widgets/Views/STableViewBase.cpp b/Engine/Source/Runtime/Slate/Private/Widgets/Views/STableViewBase.cpp
--- a/Engine/Source/Runtime/Slate/Private/Widgets/Views/STableViewBase.cpp
+++ b/Engine/Source/Runtime/Slate/Private/Widgets/Views/STableViewBase.cpp
@@ -31,7 +31,7 @@ void STableViewBase::ConstructChildren( const TAttribute<float>& InItemWidth, co
 			.FillWidth(1)
 			[
 				SAssignNew( ItemsPanel, SListPanel )
-				.ItemWidth(InItemWidth)
+				.ItemWidth( InItemWidth )
 				.ItemHeight( InItemHeight )
 				.NumDesiredItems( this, &STableViewBase::GetNumItemsBeingObserved )
 			]
@@ -50,7 +50,7 @@ void STableViewBase::ConstructChildren( const TAttribute<float>& InItemWidth, co
 		: StaticCastSharedRef<SWidget>
 		(
 			SAssignNew( ItemsPanel, SListPanel )
-			.ItemWidth(InItemWidth)
+			.ItemWidth( InItemWidth )
 			.ItemHeight( InItemHeight )
 			.NumDesiredItems( this, &STableViewBase::GetNumItemsBeingObserved )
 		);
@@ -615,6 +615,16 @@ float STableViewBase::GetItemHeight() const
 	return ItemsPanel->GetItemHeight();
 }
 
+void STableViewBase::SetItemHeight(TAttribute<float> Height)
+{
+	ItemsPanel->SetItemHeight(Height);
+}
+
+void STableViewBase::SetItemWidth(TAttribute<float> Width)
+{
+	ItemsPanel->SetItemWidth(Width);
+}
+
 float STableViewBase::GetNumLiveWidgets() const
 {
 	return ItemsPanel->GetChildren()->Num();
diff --git a/Engine/Source/Runtime/Slate/Public/Widgets/Views/STableViewBase.h b/Engine/Source/Runtime/Slate/Public/Widgets/Views/STableViewBase.h
--- a/Engine/Source/Runtime/Slate/Public/Widgets/Views/STableViewBase.h
+++ b/Engine/Source/Runtime/Slate/Public/Widgets/Views/STableViewBase.h
@@ -23,6 +23,12 @@ public:
 	/** Create the child widgets that comprise the list */
 	void ConstructChildren( const TAttribute<float>& InItemWidth, const TAttribute<float>& InItemHeight, const TSharedPtr<SHeaderRow>& InColumnHeaders, const TSharedPtr<SScrollBar>& InScrollBar  );
 
+	/** Sets the item height */
+	void SetItemHeight(TAttribute<float> Height);
+
+	/** Sets the item width */
+	void SetItemWidth(TAttribute<float> Width);
+
 	/**
 	 * Invoked by the scrollbar when the user scrolls.
 	 *
