Fixed SMultiLineEditableText performance when pasting large chunks of text.

TTP# 336464 - Editor: Finish the Multiline Editable Text Block

PasteTextFromClipboard calls TypeChar to handle adding the new text, and TypeChar was calling UpdateIfNeeded; this meant that the layout was being re-flowed for every character in the pasted text. This is now done once by FinishChangingText instead.

diff --git a/Engine/Source/Runtime/Slate/Private/Widgets/Text/SMultiLineEditableText.cpp b/Engine/Source/Runtime/Slate/Private/Widgets/Text/SMultiLineEditableText.cpp
--- a/Engine/Source/Runtime/Slate/Private/Widgets/Text/SMultiLineEditableText.cpp
+++ b/Engine/Source/Runtime/Slate/Private/Widgets/Text/SMultiLineEditableText.cpp
@@ -426,6 +426,8 @@ void SMultiLineEditableText::FinishChangingText()
 		
 		LastSetText = EditedText;
 
+		TextLayout->UpdateIfNeeded();
+
 		// Let outsiders know that the text content has been changed
 		OnTextChanged.ExecuteIfBound(EditedText);
 	}
@@ -498,8 +500,6 @@ void SMultiLineEditableText::BackspaceChar()
 			}
 		}
 
-		TextLayout->UpdateIfNeeded();
-
 		CursorInfo.SetCursorLocationAndCalculateAlignment(TextLayout, FinalCursorLocation);
 
 		ClearSelection();
@@ -593,8 +593,6 @@ void SMultiLineEditableText::DeleteChar()
 				//do nothing to the cursor as the FinalCursorLocation is already correct
 			}
 
-			TextLayout->UpdateIfNeeded();
-
 			CursorInfo.SetCursorLocationAndCalculateAlignment(TextLayout, FinalCursorLocation);
 
 			ClearSelection();
@@ -638,8 +636,6 @@ void SMultiLineEditableText::TypeChar( const int32 Character )
 		ClearSelection();
 		const FTextLocation FinalCursorLocation = FTextLocation( CursorInteractionPosition.GetLineIndex(), FMath::Min( CursorInteractionPosition.GetOffset() + 1, Line.Text->Len() ) );
 
-		TextLayout->UpdateIfNeeded();
-
 		CursorInfo.SetCursorLocationAndCalculateAlignment(TextLayout, FinalCursorLocation);
 		UpdateCursorHighlight();
 	}
@@ -711,8 +707,6 @@ void SMultiLineEditableText::DeleteSelectedText()
 		const FTextLocation FinalCursorLocation = FTextLocation(SelectionBeginningLineIndex, SelectionBeginningLineOffset);
 		CursorInfo.SetCursorLocationAndCalculateAlignment(TextLayout, FinalCursorLocation);
 		UpdateCursorHighlight();
-
-		TextLayout->UpdateIfNeeded();
 	}
 }
 
