SLATE - STextBlock shadow no longer attempts to draw if the shadow color is totally transparent.

diff --git a/Engine/Source/Runtime/Slate/Private/Widgets/Text/STextBlock.cpp b/Engine/Source/Runtime/Slate/Private/Widgets/Text/STextBlock.cpp
--- a/Engine/Source/Runtime/Slate/Private/Widgets/Text/STextBlock.cpp
+++ b/Engine/Source/Runtime/Slate/Private/Widgets/Text/STextBlock.cpp
@@ -92,8 +92,10 @@ int32 STextBlock::OnPaint( const FGeometry& AllottedGeometry, const FSlateRect&
 	bool bEnabled = ShouldBeEnabled( bParentEnabled );
 	ESlateDrawEffect::Type DrawEffects = bEnabled ? ESlateDrawEffect::None : ESlateDrawEffect::DisabledEffect;
 
+	FLinearColor CurShadowColor = GetShadowColorAndOpacity();
 	const FVector2D CurShadowOffset = GetShadowOffset();
-	const bool ShouldDropShadow = CurShadowOffset.Size() > 0;
+
+	const bool ShouldDropShadow = CurShadowOffset.Size() > 0 && CurShadowColor.A > 0;
 
 	FSlateFontInfo FontInfo = GetFont();
 
@@ -178,7 +180,7 @@ int32 STextBlock::OnPaint( const FGeometry& AllottedGeometry, const FSlateRect&
 			FontInfo,
 			ClippingRect,
 			DrawEffects,
-			GetShadowColorAndOpacity() * InWidgetStyle.GetColorAndOpacityTint()
+			CurShadowColor * InWidgetStyle.GetColorAndOpacityTint()
 		);
 	}
 
