SLATE/UMG - Exposing more post initialization setters for EditableText attributes.  Tab/spaces.

diff --git a/Engine/Source/Runtime/Slate/Private/Widgets/Input/SEditableTextBox.cpp b/Engine/Source/Runtime/Slate/Private/Widgets/Input/SEditableTextBox.cpp
--- a/Engine/Source/Runtime/Slate/Private/Widgets/Input/SEditableTextBox.cpp
+++ b/Engine/Source/Runtime/Slate/Private/Widgets/Input/SEditableTextBox.cpp
@@ -164,4 +164,19 @@ FReply SEditableTextBox::OnKeyDown( const FGeometry& MyGeometry, const FKeyboard
 	}
 
 	return FReply::Unhandled();
-}
\ No newline at end of file
+}
+
+void SEditableTextBox::SetHintText(const TAttribute< FText >& InHintText)
+{
+	EditableText->SetHintText(InHintText);
+}
+
+void SEditableTextBox::SetIsReadOnly(TAttribute< bool > InIsReadOnly)
+{
+	EditableText->SetIsReadOnly(InIsReadOnly);
+}
+
+void SEditableTextBox::SetIsPassword(TAttribute< bool > InIsPassword)
+{
+	EditableText->SetIsPassword(InIsPassword);
+}
diff --git a/Engine/Source/Runtime/Slate/Public/Widgets/Input/SEditableTextBox.h b/Engine/Source/Runtime/Slate/Public/Widgets/Input/SEditableTextBox.h
--- a/Engine/Source/Runtime/Slate/Public/Widgets/Input/SEditableTextBox.h
+++ b/Engine/Source/Runtime/Slate/Public/Widgets/Input/SEditableTextBox.h
@@ -34,10 +34,10 @@ public:
 		SLATE_STYLE_ARGUMENT( FEditableTextBoxStyle, Style )
 
 		/** Sets the text content for this editable text box widget */
- 		SLATE_ATTRIBUTE( FText, Text )
+		SLATE_ATTRIBUTE( FText, Text )
 
 		/** Hint text that appears when there is no text in the text box */
- 		SLATE_ATTRIBUTE( FText, HintText )
+		SLATE_ATTRIBUTE( FText, HintText )
 
 		/** Font color and opacity (overrides Style) */
 		SLATE_ATTRIBUTE( FSlateFontInfo, Font )
@@ -112,6 +112,15 @@ public:
 	 * @param  InNewText  The new text string
 	 */
 	void SetText( const TAttribute< FText >& InNewText );
+	
+	/** See the HintText attribute */
+	void SetHintText( const TAttribute< FText >& InHintText );
+	
+	/** See the IsReadOnly attribute */
+	void SetIsReadOnly( TAttribute< bool > InIsReadOnly );
+	
+	/** See the IsPassword attribute */
+	void SetIsPassword( TAttribute< bool > InIsPassword );
 
 	/**
 	 * If InError is a non-empty string the TextBox will the ErrorReporting provided during construction
