SLATE - Fixing the reset pointer table type trait so that tables containing UObjects can not get locked to scroll into view.

diff --git a/Engine/Source/Runtime/Slate/Public/Framework/Views/TableViewTypeTraits.h b/Engine/Source/Runtime/Slate/Public/Framework/Views/TableViewTypeTraits.h
--- a/Engine/Source/Runtime/Slate/Public/Framework/Views/TableViewTypeTraits.h
+++ b/Engine/Source/Runtime/Slate/Public/Framework/Views/TableViewTypeTraits.h
@@ -228,7 +228,7 @@ public:
 		return InPtr != NULL;
 	}
 
-	static void ResetPtr( T* InPtr )
+	static void ResetPtr( T*& InPtr )
 	{
 		InPtr = NULL;
 	}
