SLATE - Adding a deprecated version IsChecked that accepts a bool.

diff --git a/Engine/Source/Runtime/Slate/Public/Widgets/Input/SCheckBox.h b/Engine/Source/Runtime/Slate/Public/Widgets/Input/SCheckBox.h
--- a/Engine/Source/Runtime/Slate/Public/Widgets/Input/SCheckBox.h
+++ b/Engine/Source/Runtime/Slate/Public/Widgets/Input/SCheckBox.h
@@ -69,6 +69,22 @@ public:
 		/** Whether the check box is currently in a checked state */
 		SLATE_ATTRIBUTE( ESlateCheckBoxState::Type, IsChecked )
 
+		// TODO Remove this function when IsChecked(bool InIsChecked) is removed.  It has to be here to prevent ambiguous conversions.
+		/** Whether the check box is currently in a checked state */
+		FArguments& IsChecked(ESlateCheckBoxState::Type InIsChecked)
+		{
+			_IsChecked = InIsChecked;
+			return Me();
+		}
+
+		/** Whether the check box is currently in a checked state */
+		DEPRECATED(4.3, "Please use IsChecked(TAttribute<ESlateCheckBoxState::Type>)")
+		FArguments& IsChecked(bool InIsChecked)
+		{
+			_IsChecked = InIsChecked ? ESlateCheckBoxState::Checked : ESlateCheckBoxState::Unchecked;
+			return Me();
+		}
+
 		/** How the content of the toggle button should align within the given space*/
 		SLATE_ARGUMENT( EHorizontalAlignment, HAlign )
 
