Fixed code example in an SListView comment

It was out of date to what was actually required.

diff --git a/Engine/Source/Runtime/Slate/Public/Widgets/Views/SListView.h b/Engine/Source/Runtime/Slate/Public/Widgets/Views/SListView.h
--- a/Engine/Source/Runtime/Slate/Public/Widgets/Views/SListView.h
+++ b/Engine/Source/Runtime/Slate/Public/Widgets/Views/SListView.h
@@ -22,7 +22,7 @@
  *
  * A sample implementation of OnGenerateWidgetForList would simply return a TextBlock with the corresponding text:
  *
- * TSharedRef<ITableRow> OnGenerateWidgetForList( FString* InItem )
+ * TSharedRef<ITableRow> OnGenerateWidgetForList( FString* InItem, const TSharedRef<STableViewBase>& OwnerTable )
  * {
  *     return SNew(STextBlock).Text( (*InItem) )
  * }
