Fixed dragging derived table row types with uninitialised member

This addresses TTP#338731 - EDITOR: Class Viewer: Drag and drop into level viewport frequently fails

diff --git a/Engine/Source/Runtime/Slate/Public/Widgets/Views/STableRow.h b/Engine/Source/Runtime/Slate/Public/Widgets/Views/STableRow.h
--- a/Engine/Source/Runtime/Slate/Public/Widgets/Views/STableRow.h
+++ b/Engine/Source/Runtime/Slate/Public/Widgets/Views/STableRow.h
@@ -84,7 +84,7 @@ public:
 	 */
 	void Construct(const typename STableRow<ItemType>::FArguments& InArgs, const TSharedRef<STableViewBase>& InOwnerTableView)
 	{
-		bProcessingSelectionTouch = false;
+		/** Note: Please initialize any state in ConstructInternal, not here. This is because STableRow derivatives call ConstructInternal directly to avoid constructing children. **/
 
 		ConstructInternal(InArgs, InOwnerTableView);
 
@@ -673,6 +673,8 @@ protected:
 	 */
 	void ConstructInternal(FArguments const& InArgs, TSharedRef<STableViewBase> const& InOwnerTableView)
 	{
+		bProcessingSelectionTouch = false;
+
 		check(InArgs._Style);
 		Style = InArgs._Style;
 
