Fixed SWindowTitleBar incorrectly using HasMinimizeBox rather than HasMaximizeBox when creating the maximize button.

https://answers.unrealengine.com/questions/60742/slate-swindow-supportsmaximize-and-supportsminimiz.html

Also removed the duplicate HasMaximizeButton and HasMinimizeButton functions, since these were:
 a) Added after HasMaximizeBox and HasMinimizeBox
 b) Not const-correct
 c) Only used in SWindowTitleBar

diff --git a/Engine/Source/Runtime/Slate/Private/Framework/Application/SWindowTitleBar.h b/Engine/Source/Runtime/Slate/Private/Framework/Application/SWindowTitleBar.h
--- a/Engine/Source/Runtime/Slate/Private/Framework/Application/SWindowTitleBar.h
+++ b/Engine/Source/Runtime/Slate/Private/Framework/Application/SWindowTitleBar.h
@@ -173,7 +173,7 @@ protected:
 
 		MaximizeRestoreButton = SNew(SButton)
 				.IsFocusable(false)
-				.IsEnabled(OwnerWindow->HasMinimizeBox())
+				.IsEnabled(OwnerWindow->HasMaximizeBox())
 				.ContentPadding(0.0f)
 				.OnClicked(this, &SWindowTitleBar::MaximizeRestoreButton_OnClicked)
 				.Cursor(EMouseCursor::Default)
@@ -474,7 +474,7 @@ private:
 
 		if (NativeWindow.IsValid() && NativeWindow->IsMaximized())
 		{
-			if (!OwnerWindow->HasMaximizeButton())
+			if (!OwnerWindow->HasMaximizeBox())
 			{
 				return &Style->MaximizeButtonStyle.Disabled;
 			}
@@ -503,7 +503,7 @@ private:
 		}
 		else
 		{
-			if (!OwnerWindow->HasMaximizeButton())
+			if (!OwnerWindow->HasMaximizeBox())
 			{
 				return &Style->MaximizeButtonStyle.Disabled;
 			}
@@ -562,7 +562,7 @@ private:
 
 		TSharedPtr<FGenericWindow> NativeWindow = OwnerWindow->GetNativeWindow();
 
-		if (!OwnerWindow->HasMinimizeButton())
+		if (!OwnerWindow->HasMinimizeBox())
 		{
 			return &Style->MinimizeButtonStyle.Disabled;
 		}
diff --git a/Engine/Source/Runtime/SlateCore/Public/Widgets/SWindow.h b/Engine/Source/Runtime/SlateCore/Public/Widgets/SWindow.h
--- a/Engine/Source/Runtime/SlateCore/Public/Widgets/SWindow.h
+++ b/Engine/Source/Runtime/SlateCore/Public/Widgets/SWindow.h
@@ -518,16 +518,6 @@ public:
 		bHasMinimizeButton = false;
 	}
 
-	bool HasMaximizeButton( )
-	{
-		return bHasMaximizeButton;
-	}
-
-	bool HasMinimizeButton( )
-	{
-		return bHasMinimizeButton;
-	}
-
 	void SetTitleBar( const TSharedPtr<IWindowTitleBar> InTitleBar )
 	{
 		TitleBar = InTitleBar;
