Improved cursor movement and SetText handling

TTP# 336464 - Editor: Finish the Multiline Editable Text Block

SetText improvements:
- Split up the SMultiLineEditableText::SetText function so that it behaves the same as the SEditableText version (calling SetText updates the BoundText attribute).
- Added an SetEditableText function, which handles updating the TextLayout to match the text provided.
- Updated the existing code to call the appropriate version (most needed to call SetEditableText).
- Tidied up the logic surrounding the old LastSetText variable. This is now called BoundTextLastTick and is used *only* to update the editable text if the text from the source delegate is changed.
- Made sure that everything performing a text comparison was using FText::EqualTo (rather than using == on FString, which is case insensitive).
- Made sure that the multiline editable text controls still worked correctly, both with and without a delegate providing the source text.

Improved cursor movement:
- Neither SMultiLineEditableText or SEditableText were handling cursor movement correctly when you had text selected (they were moving based on the last cursor position, rather than the extent of the selection range).
- I've implemented the various cursor movement rules for both versions to mirror the rules used by MS Word.

SMultiLineEditableText::IsAtWordStart now uses the word break iterator so it correctly handles all cultures (which don't break words on whitepsace)

SMultiLineEditableText now draws the cursor at the correct end of the selection highlight.

diff --git a/Engine/Source/Runtime/Core/Private/Internationalization/LegacyText.cpp b/Engine/Source/Runtime/Core/Private/Internationalization/LegacyText.cpp
--- a/Engine/Source/Runtime/Core/Private/Internationalization/LegacyText.cpp
+++ b/Engine/Source/Runtime/Core/Private/Internationalization/LegacyText.cpp
@@ -5,6 +5,11 @@
 #if !UE_ENABLE_ICU
 #include "Text.h"
 
+bool FText::IsWhitespace( const TCHAR Char )
+{
+	return FChar::IsWhitespace(Char);
+}
+
 FText FText::AsDate( const FDateTime& DateTime, const EDateTimeStyle::Type DateStyle, const TSharedPtr<FCulture, ESPMode::ThreadSafe>& TargetCulture )
 {
 	checkf(FInternationalization::Get().IsInitialized() == true, TEXT("FInternationalization is not initialized. An FText formatting method was likely used in static object initialization - this is not supported."));
diff --git a/Engine/Source/Runtime/Core/Public/GenericPlatform/MicrosoftPlatformString.h b/Engine/Source/Runtime/Core/Public/GenericPlatform/MicrosoftPlatformString.h
--- a/Engine/Source/Runtime/Core/Public/GenericPlatform/MicrosoftPlatformString.h
+++ b/Engine/Source/Runtime/Core/Public/GenericPlatform/MicrosoftPlatformString.h
@@ -290,6 +290,13 @@ struct FMicrosoftPlatformString : public FGenericPlatformString
 	{
 		return _tcslen( (const WIDECHAR*)String );
 	}
+
+	static const ANSICHAR* GetEncodingName()
+	{
+		return "UTF-16LE";
+	}
+
+	static const bool IsUnicodeEncoded = true;
 };
 
 #if !USE_SECURE_CRT
diff --git a/Engine/Source/Runtime/Slate/Private/Framework/Text/TextLayout.cpp b/Engine/Source/Runtime/Slate/Private/Framework/Text/TextLayout.cpp
--- a/Engine/Source/Runtime/Slate/Private/Framework/Text/TextLayout.cpp
+++ b/Engine/Source/Runtime/Slate/Private/Framework/Text/TextLayout.cpp
@@ -540,11 +540,11 @@ FTextLocation FTextLayout::GetTextLocationAt( const FLineView& LineView, const F
 	}
 
 	const int32 LineTextLength = LineModels[ LineView.ModelIndex ].Text->Len();
-	if ( LineTextLength == 0 )
+	if ( LineTextLength == 0 || !LineView.Blocks.Num() )
 	{
 		return FTextLocation( LineView.ModelIndex, 0 );
 	}
-	else if (Relative.X < LineView.Offset.X)
+	else if (Relative.X < LineView.Blocks[0]->GetLocationOffset().X)
 	{
 		return FTextLocation(LineView.ModelIndex, LineView.Range.BeginIndex);
 	}
@@ -966,6 +966,13 @@ void FTextLayout::GetAsText(FString& DisplayText) const
 	}
 }
 
+void FTextLayout::GetAsText(FText& DisplayText) const
+{
+	FString DisplayString;
+	GetAsText(DisplayString);
+	DisplayText = FText::FromString(DisplayString);
+}
+
 void GetRangeAsTextFromLine(FString& DisplayText, const FTextLayout::FLineModel& LineModel, const FTextRange& Range)
 {
 	for (int RunIndex = 0; RunIndex < LineModel.Runs.Num(); RunIndex++)
diff --git a/Engine/Source/Runtime/Slate/Private/Widgets/Input/SEditableText.cpp b/Engine/Source/Runtime/Slate/Private/Widgets/Input/SEditableText.cpp
--- a/Engine/Source/Runtime/Slate/Private/Widgets/Input/SEditableText.cpp
+++ b/Engine/Source/Runtime/Slate/Private/Widgets/Input/SEditableText.cpp
@@ -636,71 +636,76 @@ static int32 OffsetToTestBasedOnDirection( int8 Direction )
 
 FReply SEditableText::MoveCursor( ECursorMoveMethod::Type Method, const FVector2D& Direction, ECursorAction::Type Action )
 {
-	const int32 OldCaretPosition = CaretPosition;
+	bool bAllowMoveCursor = true;
+	int32 OldCaretPosition = CaretPosition;
 	int32 NewCaretPosition = OldCaretPosition;
 
-	const int32 StringLength = EditedText.ToString().Len();
-	const FString& EditedTextAsString = EditedText.ToString();
-
-	// If control is held down, jump to beginning of the current word (or, the previous word, if
-	// the caret is currently on a whitespace character
-	switch( Method ) 
+	// If we have selected text, the cursor needs to:
+	// a) Jump to the start of the selection if moving the cursor Left
+	// b) Jump to the end of the selection if moving the cursor Right
+	// This is done regardless of whether the selection was made left-to-right, or right-to-left
+	// This also needs to be done *before* moving to word boundaries, as the start point needs to be 
+	// the start or end of the selection depending on the above rules
+	if (Action == ECursorAction::MoveCursor && Method != ECursorMoveMethod::ScreenPosition && AnyTextSelected())
 	{
-	case ECursorMoveMethod::Word:
+		if (Method == ECursorMoveMethod::CharacterHorizontal)
 		{
-			const FString EditedTextString = EditedText.ToString();
-
-			if (Direction.X > 0)
-			{
-				// Scan right for text
-				int32 CurCharIndex = CaretPosition;
-				while( CurCharIndex < StringLength && !FText::IsWhitespace( EditedTextString[ CurCharIndex ] ) )
-				{
-					++CurCharIndex;
-				}
+			// If we're moving the cursor horizontally, we just snap to the start or end of the selection rather than 
+			// move the cursor by the normal movement rules
+			bAllowMoveCursor = false;
+		}
 
-				// Scan right for whitespace
-				while( CurCharIndex < StringLength && FText::IsWhitespace( EditedTextString[ CurCharIndex ] ) )
-				{
-					++CurCharIndex;
-				}
+		// Work out which edge of the selection we need to start at
+		bool bSnapToSelectionStart = false;
+		switch (Method)
+		{
+		case ECursorMoveMethod::CharacterHorizontal:
+		case ECursorMoveMethod::Word:
+			bSnapToSelectionStart = Direction.X < 0.0f;
+			break;
 
-				NewCaretPosition = CurCharIndex;
-			}
-			else
-			{
-				// Scan left for whitespace
-				int32 CurCharIndex = CaretPosition;
-				while( CurCharIndex > 0 && FText::IsWhitespace( EditedTextString[ CurCharIndex - 1 ] ) )
-				{
-					--CurCharIndex;
-				}
+		default:
+			break;
+		}
 
-				// Scan left for text
-				while( CurCharIndex > 0 && !FText::IsWhitespace( EditedTextString[ CurCharIndex - 1 ] ) )
-				{
-					--CurCharIndex;
-				}
+		// Adjust the current cursor position - also set the new cursor position so that the bAllowMoveCursor == false case is handled
+		OldCaretPosition = bSnapToSelectionStart ? Selection.GetMinIndex() : Selection.GetMaxIndex();
+		NewCaretPosition = OldCaretPosition;
 
-				NewCaretPosition = CurCharIndex;
-			}
+		// If we're snapping to a word boundary, but the selection was already at a word boundary, don't let the cursor move any more
+		if (Method == ECursorMoveMethod::Word && IsAtWordStart(NewCaretPosition))
+		{
+			bAllowMoveCursor = false;
 		}
-		break;
+	}
 
-	case ECursorMoveMethod::CharacterHorizontal:
-		NewCaretPosition = FMath::Clamp<int32>( CaretPosition + Direction.X, 0, StringLength );
-		break;
+	if (bAllowMoveCursor)
+	{
+		const int32 StringLength = EditedText.ToString().Len();
+
+		// If control is held down, jump to beginning of the current word (or, the previous word, if
+		// the caret is currently on a whitespace character
+		switch( Method ) 
+		{
+		case ECursorMoveMethod::Word:
+			NewCaretPosition = ScanForWordBoundary( CaretPosition, Direction.X );
+			break;
+
+		case ECursorMoveMethod::CharacterHorizontal:
+			NewCaretPosition = FMath::Clamp<int32>( CaretPosition + Direction.X, 0, StringLength );
+			break;
 	
-	case ECursorMoveMethod::CharacterVertical:
-		return FReply::Unhandled();
+		case ECursorMoveMethod::CharacterVertical:
+			return FReply::Unhandled();
 	
-	case ECursorMoveMethod::ScreenPosition:
-		NewCaretPosition = FindClickedCharacterIndex( Direction );
-		break;
+		case ECursorMoveMethod::ScreenPosition:
+			NewCaretPosition = FindClickedCharacterIndex( Direction );
+			break;
 
-	default:
-		checkSlow(false, "Unknown ECursorMoveMethod value");
-		break;
+		default:
+			checkSlow(false, "Unknown ECursorMoveMethod value");
+			break;
+		}
 	}
 
 	SetCaretPosition( NewCaretPosition );
@@ -804,7 +809,7 @@ void SEditableText::SelectWordAt(const FVector2D& LocalPosition)
 
 	// Select the word that the user double-clicked on, and move the caret to the end of that word
 	{
-		const FString EditedTextString = EditedText.ToString();
+		const FString& EditedTextString = EditedText.ToString();
 
 		// Find beginning of the word
 		int32 FirstWordCharIndex = ClickedCharIndex;
@@ -815,7 +820,7 @@ void SEditableText::SelectWordAt(const FVector2D& LocalPosition)
 
 		// Find end of the word
 		int32 LastWordCharIndex = ClickedCharIndex;
-		while( LastWordCharIndex < EditedText.ToString().Len() && !FText::IsWhitespace( EditedTextString[ LastWordCharIndex ] ) )
+		while( LastWordCharIndex < EditedTextString.Len() && !FText::IsWhitespace( EditedTextString[ LastWordCharIndex ] ) )
 		{
 			++LastWordCharIndex;
 		}
@@ -1700,6 +1705,63 @@ int32 SEditableText::FindClickedCharacterIndex( const FVector2D& InLocalCursorPo
 }
 
 
+int32 SEditableText::ScanForWordBoundary( const int32 Location, int8 Direction ) const
+{
+	const FString& EditedTextString = EditedText.ToString();
+	const int32 StringLength = EditedTextString.Len();
+
+	if (Direction > 0)
+	{
+		// Scan right for text
+		int32 CurCharIndex = Location;
+		while( CurCharIndex < StringLength && !FText::IsWhitespace( EditedTextString[ CurCharIndex ] ) )
+		{
+			++CurCharIndex;
+		}
+
+		// Scan right for whitespace
+		while( CurCharIndex < StringLength && FText::IsWhitespace( EditedTextString[ CurCharIndex ] ) )
+		{
+			++CurCharIndex;
+		}
+
+		return CurCharIndex;
+	}
+	else
+	{
+		// Scan left for whitespace
+		int32 CurCharIndex = Location;
+		while( CurCharIndex > 0 && FText::IsWhitespace( EditedTextString[ CurCharIndex - 1 ] ) )
+		{
+			--CurCharIndex;
+		}
+
+		// Scan left for text
+		while( CurCharIndex > 0 && !FText::IsWhitespace( EditedTextString[ CurCharIndex - 1 ] ) )
+		{
+			--CurCharIndex;
+		}
+
+		return CurCharIndex;
+	}
+}
+
+
+bool SEditableText::IsAtWordStart( const int32 Location ) const
+{
+	const FString& EditedTextString = EditedText.ToString();
+
+	if (Location > 0)
+	{
+		const TCHAR CharBeforeCursor = EditedTextString[Location - 1];
+		const TCHAR CharAtCursor = EditedTextString[Location];
+		return FText::IsWhitespace(CharBeforeCursor) && !FText::IsWhitespace(CharAtCursor);
+	}
+
+	return false;
+}
+
+
 void SEditableText::PushUndoState( const SEditableText::FUndoState& InUndoState )
 {
 	// If we've already undone some state, then we'll remove any undo state beyond the level that
diff --git a/Engine/Source/Runtime/Slate/Private/Widgets/Input/SMultiLineEditableTextBox.cpp b/Engine/Source/Runtime/Slate/Private/Widgets/Input/SMultiLineEditableTextBox.cpp
--- a/Engine/Source/Runtime/Slate/Private/Widgets/Input/SMultiLineEditableTextBox.cpp
+++ b/Engine/Source/Runtime/Slate/Private/Widgets/Input/SMultiLineEditableTextBox.cpp
@@ -44,6 +44,7 @@ void SMultiLineEditableTextBox::Construct( const FArguments& InArgs )
 					.OnTextCommitted( InArgs._OnTextCommitted )
 					.Justification(InArgs._Justification)
 					.LineHeightPercentage(InArgs._LineHeightPercentage)
+					.Margin(InArgs._Margin)
 					.WrapTextAt(InArgs._WrapTextAt)
 					.AutoWrapText(InArgs._AutoWrapText)
 				]
diff --git a/Engine/Source/Runtime/Slate/Private/Widgets/Testing/STestSuite.cpp b/Engine/Source/Runtime/Slate/Private/Widgets/Testing/STestSuite.cpp
--- a/Engine/Source/Runtime/Slate/Private/Widgets/Testing/STestSuite.cpp
+++ b/Engine/Source/Runtime/Slate/Private/Widgets/Testing/STestSuite.cpp
@@ -1693,7 +1693,9 @@ public:
 							+SHorizontalBox::Slot()
 							[
 								SNew(SMultiLineEditableTextBox)
+								.Margin(10)
 								.Text(MultilineEditableText)
+								//.Justification(ETextJustify::Right)
 								.Font(FSlateFontInfo(FPaths::EngineContentDir() / TEXT("Slate/Fonts/Roboto-Regular.ttf"), 12))
 								.AutoWrapText(true)
 								.HintText(LOCTEXT("TypehereTextHint", "Type Here"))
diff --git a/Engine/Source/Runtime/Slate/Private/Widgets/Text/SMultiLineEditableText.cpp b/Engine/Source/Runtime/Slate/Private/Widgets/Text/SMultiLineEditableText.cpp
--- a/Engine/Source/Runtime/Slate/Private/Widgets/Text/SMultiLineEditableText.cpp
+++ b/Engine/Source/Runtime/Slate/Private/Widgets/Text/SMultiLineEditableText.cpp
@@ -68,8 +68,6 @@ void SMultiLineEditableText::FCursorInfo::RestoreFromUndo(const FCursorInfo& Und
 	SetCursorLocationAndAlignment(UndoData.CursorPosition, UndoData.CursorAlignment);
 }
 
-FNoChildren SMultiLineEditableText::FSlateCursorRunHighlighter::NoChildrenInstance;
-
 SMultiLineEditableText::FSlateCursorRunHighlighter::FSlateCursorRunHighlighter(const FCursorInfo* InCursorInfo)
 	: CursorInfo(InCursorInfo)
 {
@@ -83,88 +81,110 @@ void SMultiLineEditableText::FSlateCursorRunHighlighter::OnArrangeChildren( cons
 
 FChildren* SMultiLineEditableText::FSlateCursorRunHighlighter::GetChildren()
 {
+	static FNoChildren NoChildrenInstance;
 	return &NoChildrenInstance;
 }
 
 int32 SMultiLineEditableText::FSlateCursorRunHighlighter::OnPaint( const FTextLayout::FLineView& Line, const TSharedRef< ISlateRun >& Run, const TSharedRef< ILayoutBlock >& Block, const FTextBlockStyle& DefaultStyle, const FGeometry& AllottedGeometry, const FSlateRect& MyClippingRect, FSlateWindowElementList& OutDrawElements, int32 LayerId, const FWidgetStyle& InWidgetStyle, bool bParentEnabled ) const 
+{
+	LayerId = OnPaintCursor(Line, Run, Block, DefaultStyle, AllottedGeometry, MyClippingRect, OutDrawElements, LayerId, InWidgetStyle, bParentEnabled, CursorInfo->GetCursorAlignment());
+
+	return Run->OnPaint( Line, Block, DefaultStyle, AllottedGeometry, MyClippingRect, OutDrawElements, LayerId, InWidgetStyle, bParentEnabled );
+}
+
+int32 SMultiLineEditableText::FSlateCursorRunHighlighter::OnPaintCursor( const FTextLayout::FLineView& Line, const TSharedRef< ISlateRun >& Run, const TSharedRef< ILayoutBlock >& Block, const FTextBlockStyle& DefaultStyle, const FGeometry& AllottedGeometry, const FSlateRect& MyClippingRect, FSlateWindowElementList& OutDrawElements, int32 LayerId, const FWidgetStyle& InWidgetStyle, const bool bParentEnabled, const ECursorAlignment InCursorAlignment ) const
 {
 	FVector2D Location( Block->GetLocationOffset() );
 	Location.Y = Line.Offset.Y;
 
-	if ( Mode == EMode::Selection )
-	{
-		const FColor SelectionBackgroundColorAndOpacity = ((FLinearColor::White - DefaultStyle.ColorAndOpacity.GetColor(InWidgetStyle))*0.5f + FLinearColor(-0.2f, -0.05f, 0.15f)) * InWidgetStyle.GetColorAndOpacityTint();
+	FLinearColor CursorColorAndOpacity = InWidgetStyle.GetForegroundColor();
 
-		float HighlightWidth = Block->GetSize().X;
-		if (HighlightWidth == 0)
-		{
-			const TSharedRef< FSlateFontMeasure > FontMeasureService = FSlateApplication::Get().GetRenderer()->GetFontMeasureService();
-			const FVector2D SpaceSize = FontMeasureService->Measure(FString(TEXT("  ")), DefaultStyle.Font);
-			HighlightWidth = SpaceSize.X;
-		}
+	const float FontMaxCharHeight = FTextEditHelper::GetFontHeight(DefaultStyle.Font);
+	const float CursorWidth = FTextEditHelper::CalculateCaretWidth(FontMaxCharHeight);
+	const double CurrentTime = FSlateApplication::Get().GetCurrentTime();
 
-		// Draw the actual highlight rectangle
-		FSlateDrawElement::MakeBox(
-			OutDrawElements,
-			++LayerId,
-			FPaintGeometry(AllottedGeometry.AbsolutePosition + Location, FVector2D(HighlightWidth, Line.Size.Y), AllottedGeometry.Scale),
-			&DefaultStyle.HighlightShape,
-			MyClippingRect,
-			bParentEnabled ? ESlateDrawEffect::None : ESlateDrawEffect::DisabledEffect,
-			SelectionBackgroundColorAndOpacity
-		);
+	// The cursor is always visible (i.e. not blinking) when we're interacting with it; otherwise it might get lost.
+	const bool bForceCursorVisible = (CurrentTime - CursorInfo->GetLastCursorInteractionTime()) < EditableTextDefs::CaretBlinkPauseTime;
+	float CursorOpacity = (bForceCursorVisible)
+		? 1.0f
+		: FMath::RoundToFloat( FMath::MakePulsatingValue( CurrentTime, EditableTextDefs::BlinksPerSecond ));
 
-		//FLinearColor InvertedForeground = FLinearColor::White - InWidgetStyle.GetForegroundColor();
-		//InvertedForeground.A = InWidgetStyle.GetForegroundColor().A;
+	CursorOpacity *= CursorOpacity;	// Squared falloff, because it looks more interesting
+	CursorColorAndOpacity.A = CursorOpacity;
 
-		//FWidgetStyle WidgetStyle( InWidgetStyle );
-		//WidgetStyle.SetForegroundColor( InvertedForeground );
+	// @todo: Slate Styles - make this brush part of the widget style
+	const FSlateBrush* StyleInfo = FCoreStyle::Get().GetBrush("EditableText.SelectionBackground");
 
-		return Run->OnPaint( Line, Block, DefaultStyle, AllottedGeometry, MyClippingRect, OutDrawElements, LayerId, InWidgetStyle, bParentEnabled );
-	}
-	else
-	{
-		FLinearColor CursorColorAndOpacity = InWidgetStyle.GetForegroundColor();
+	const FVector2D OptionalWidth = InCursorAlignment == ECursorAlignment::Right ? FVector2D(Block->GetSize().X, 0) : FVector2D::ZeroVector;
 
-		const float FontMaxCharHeight = FTextEditHelper::GetFontHeight(DefaultStyle.Font);
-		const float CursorWidth = FTextEditHelper::CalculateCaretWidth(FontMaxCharHeight);
-		const double CurrentTime = FSlateApplication::Get().GetCurrentTime();
+	FSlateDrawElement::MakeBox(
+		OutDrawElements,
+		LayerId,
+		FPaintGeometry(AllottedGeometry.AbsolutePosition + Location + OptionalWidth, FVector2D(CursorWidth * AllottedGeometry.Scale, Line.TextSize.Y), AllottedGeometry.Scale),
+		StyleInfo,
+		MyClippingRect,
+		bParentEnabled ? ESlateDrawEffect::None : ESlateDrawEffect::DisabledEffect,
+		CursorColorAndOpacity*InWidgetStyle.GetColorAndOpacityTint());
 
-		// The cursor is always visible (i.e. not blinking) when we're interacting with it; otherwise it might get lost.
-		const bool bForceCursorVisible = (CurrentTime - CursorInfo->GetLastCursorInteractionTime()) < EditableTextDefs::CaretBlinkPauseTime;
-		float CursorOpacity = (bForceCursorVisible)
-			? 1.0f
-			: FMath::RoundToFloat( FMath::MakePulsatingValue( CurrentTime, EditableTextDefs::BlinksPerSecond ));
+	return LayerId;
+}
 
-		CursorOpacity *= CursorOpacity;	// Squared falloff, because it looks more interesting
-		CursorColorAndOpacity.A = CursorOpacity;
+TSharedRef< SMultiLineEditableText::FSlateCursorRunHighlighter > SMultiLineEditableText::FSlateCursorRunHighlighter::Create(const FCursorInfo* InCursorInfo)
+{
+	return MakeShareable( new SMultiLineEditableText::FSlateCursorRunHighlighter(InCursorInfo) );
+}
 
-		// @todo: Slate Styles - make this brush part of the widget style
-		const FSlateBrush* StyleInfo = FCoreStyle::Get().GetBrush("EditableText.SelectionBackground");
+SMultiLineEditableText::FSlateSelectionRunHighlighter::FSlateSelectionRunHighlighter(const FCursorInfo* InCursorInfo)
+	: SMultiLineEditableText::FSlateCursorRunHighlighter(InCursorInfo)
+	, CursorAlignment(ECursorAlignment::Left)
+{
+}
 
-		const FVector2D OptionalWidth = CursorInfo->GetCursorAlignment() == ECursorAlignment::Right ? FVector2D(Block->GetSize().X, 0) : FVector2D::ZeroVector;
+int32 SMultiLineEditableText::FSlateSelectionRunHighlighter::OnPaint( const FTextLayout::FLineView& Line, const TSharedRef< ISlateRun >& Run, const TSharedRef< ILayoutBlock >& Block, const FTextBlockStyle& DefaultStyle, const FGeometry& AllottedGeometry, const FSlateRect& MyClippingRect, FSlateWindowElementList& OutDrawElements, int32 LayerId, const FWidgetStyle& InWidgetStyle, bool bParentEnabled ) const 
+{
+	FVector2D Location( Block->GetLocationOffset() );
+	Location.Y = Line.Offset.Y;
 
-		FSlateDrawElement::MakeBox(
-			OutDrawElements,
-			LayerId,
-			FPaintGeometry(AllottedGeometry.AbsolutePosition + Location + OptionalWidth, FVector2D(CursorWidth * AllottedGeometry.Scale, Line.TextSize.Y), AllottedGeometry.Scale),
-			StyleInfo,
-			MyClippingRect,
-			bParentEnabled ? ESlateDrawEffect::None : ESlateDrawEffect::DisabledEffect,
-			CursorColorAndOpacity*InWidgetStyle.GetColorAndOpacityTint());
+	const FColor SelectionBackgroundColorAndOpacity = ((FLinearColor::White - DefaultStyle.ColorAndOpacity.GetColor(InWidgetStyle))*0.5f + FLinearColor(-0.2f, -0.05f, 0.15f)) * InWidgetStyle.GetColorAndOpacityTint();
 
-		return Run->OnPaint( Line, Block, DefaultStyle, AllottedGeometry, MyClippingRect, OutDrawElements, LayerId, InWidgetStyle, bParentEnabled );
+	float HighlightWidth = Block->GetSize().X;
+	if (HighlightWidth == 0)
+	{
+		const TSharedRef< FSlateFontMeasure > FontMeasureService = FSlateApplication::Get().GetRenderer()->GetFontMeasureService();
+		const FVector2D SpaceSize = FontMeasureService->Measure(FString(TEXT("  ")), DefaultStyle.Font);
+		HighlightWidth = SpaceSize.X;
 	}
-}
 
-TSharedRef< SMultiLineEditableText::FSlateCursorRunHighlighter > SMultiLineEditableText::FSlateCursorRunHighlighter::Create(const FCursorInfo* InCursorInfo)
-{
-	return MakeShareable( new SMultiLineEditableText::FSlateCursorRunHighlighter(InCursorInfo) );
+	// Draw the actual highlight rectangle
+	FSlateDrawElement::MakeBox(
+		OutDrawElements,
+		++LayerId,
+		FPaintGeometry(AllottedGeometry.AbsolutePosition + Location, FVector2D(HighlightWidth, Line.Size.Y), AllottedGeometry.Scale),
+		&DefaultStyle.HighlightShape,
+		MyClippingRect,
+		bParentEnabled ? ESlateDrawEffect::None : ESlateDrawEffect::DisabledEffect,
+		SelectionBackgroundColorAndOpacity
+	);
+
+	//FLinearColor InvertedForeground = FLinearColor::White - InWidgetStyle.GetForegroundColor();
+	//InvertedForeground.A = InWidgetStyle.GetForegroundColor().A;
+
+	//FWidgetStyle WidgetStyle( InWidgetStyle );
+	//WidgetStyle.SetForegroundColor( InvertedForeground );
+
+	// If the cursor is within the range of this line view, we need to draw it
+	const FTextLocation CursorLocation = CursorInfo->GetCursorLocation();
+	if(Line.ModelIndex == CursorLocation.GetLineIndex() && Line.Range.InclusiveContains(CursorLocation.GetOffset()))
+	{
+		LayerId = OnPaintCursor(Line, Run, Block, DefaultStyle, AllottedGeometry, MyClippingRect, OutDrawElements, LayerId, InWidgetStyle, bParentEnabled, CursorAlignment);
+	}
+
+	return Run->OnPaint( Line, Block, DefaultStyle, AllottedGeometry, MyClippingRect, OutDrawElements, LayerId, InWidgetStyle, bParentEnabled );
 }
 
-void SMultiLineEditableText::FSlateCursorRunHighlighter::SetMode( EMode InMode )
+TSharedRef< SMultiLineEditableText::FSlateSelectionRunHighlighter > SMultiLineEditableText::FSlateSelectionRunHighlighter::Create(const FCursorInfo* InCursorInfo)
 {
-	Mode = InMode;
+	return MakeShareable( new SMultiLineEditableText::FSlateSelectionRunHighlighter(InCursorInfo) );
 }
 
 SMultiLineEditableText::SMultiLineEditableText()
@@ -206,12 +226,22 @@ void SMultiLineEditableText::Construct( const FArguments& InArgs )
 	OnTextChanged = InArgs._OnTextChanged;
 	OnTextCommitted = InArgs._OnTextCommitted;
 
-	CursorSelectionHighlighter = FSlateCursorRunHighlighter::Create(&CursorInfo);
+	CursorRunHighlighter = FSlateCursorRunHighlighter::Create(&CursorInfo);
+	SelectionRunHighlighter = FSlateSelectionRunHighlighter::Create(&CursorInfo);
 	TextLayout = FSlateTextLayout::Create();
 
 	BoundText = InArgs._Text;
 
-	SetText(BoundText.Get(FText::GetEmpty()));
+	{
+		const FText& TextToSet = BoundText.Get(FText::GetEmpty());
+		SetEditableText(TextToSet, true); // force the set to ensure we create an empty run (if BoundText returns an empty string)
+
+		// Update the cached BoundText value to prevent it triggering another SetEditableText update again next Tick
+		if (BoundText.IsBound())
+		{
+			BoundTextLastTick = TextToSet;
+		}
+	}
 
 	TextLayout->SetWrappingWidth( WrapTextAt.Get() );
 	TextLayout->SetMargin( Margin.Get() );
@@ -252,27 +282,43 @@ void SMultiLineEditableText::Construct( const FArguments& InArgs )
 
 void SMultiLineEditableText::SetText(const TAttribute< FText >& InText)
 {
-	const bool bHasPreviouslySetText = LastSetText.IsSet();
-	LastSetText = InText.Get();
+	// NOTE: This will unbind any getter that is currently assigned to the Text attribute!  You should only be calling
+	//       SetText() if you know what you're doing.
+	BoundText = InText;
 
-	FString EditedText;
-	TextLayout->GetAsText(EditedText);
+	const FText& TextToSet = BoundText.Get(FText::GetEmpty());
 
-	check(LastSetText.IsSet());
-	const FText& LastSetTextValue = LastSetText.GetValue();
-	const FString& TextString = LastSetTextValue.ToString();
-	const bool HasTextChanged = !bHasPreviouslySetText || EditedText != TextString;
+	// Update the cached BoundText value to prevent it triggering another SetEditableText update again next Tick
+	if (BoundText.IsBound())
+	{
+		BoundTextLastTick = TextToSet;
+	}
 
-	if (HasTextChanged)
+	// Update the internal editable text
+	if (SetEditableText(TextToSet))
+	{
+		// Let outsiders know that the text content has been changed
+		OnTextChanged.ExecuteIfBound(TextToSet);
+	}
+}
+
+bool SMultiLineEditableText::SetEditableText(const FText& TextToSet, const bool bForce)
+{
+	FText EditedText;
+	TextLayout->GetAsText(EditedText);
+
+	const bool HasTextChanged = !EditedText.EqualTo(TextToSet);
+	if (HasTextChanged || bForce)
 	{
-		const int32 TextLength = TextString.Len();
+		const FString& TextToSetString = TextToSet.ToString();
+		const int32 TextToSetLength = TextToSetString.Len();
 
 		TextLayout->ClearLines();
 
-		if(TextLength)
+		if(TextToSetLength)
 		{
 			// Iterate over each line break candidate, adding ranges from after the end of the last line added to before the newline or end of string.
-			FLineBreakIterator LBI(TextString);
+			FLineBreakIterator LBI(TextToSetString);
 			int32 RangeBegin = LBI.GetCurrentPosition();
 			TSharedPtr< FString > LineText;
 
@@ -289,16 +335,16 @@ void SMultiLineEditableText::SetText(const TAttribute< FText >& InText)
 				const int32 BreakingCharacterIndex = BreakIndex - 1;
 				if (BreakingCharacterIndex >= RangeBegin) // Valid index check.
 				{
-					const TCHAR BreakingCharacter = TextString[BreakingCharacterIndex];
+					const TCHAR BreakingCharacter = TextToSetString[BreakingCharacterIndex];
 					// If line break candidate is after a newline, add the range as a new line.
 					if (FChar::IsLinebreak(BreakingCharacter))
 					{
-						LineText = MakeShareable(new FString(BreakingCharacterIndex - RangeBegin, (*TextString) + RangeBegin));
+						LineText = MakeShareable(new FString(BreakingCharacterIndex - RangeBegin, (*TextToSetString) + RangeBegin));
 					}
 					// If the line break candidate is the end of the string, add the range as a new line.
-					else if (BreakIndex == TextLength)
+					else if (BreakIndex == TextToSetLength)
 					{
-						LineText = MakeShareable(new FString(BreakIndex - RangeBegin, (*TextString) + RangeBegin));
+						LineText = MakeShareable(new FString(BreakIndex - RangeBegin, (*TextToSetString) + RangeBegin));
 					}
 					else
 					{
@@ -332,9 +378,10 @@ void SMultiLineEditableText::SetText(const TAttribute< FText >& InText)
 			TextLayout->AddLine(LineText.ToSharedRef(), Runs);
 		}
 
-		// Let outsiders know that the text content has been changed
-		OnTextChanged.ExecuteIfBound(LastSetTextValue);
+		return true;
 	}
+
+	return false;
 }
 
 FReply SMultiLineEditableText::OnKeyboardFocusReceived( const FGeometry& MyGeometry, const FKeyboardFocusEvent& InKeyboardFocusEvent )
@@ -381,10 +428,10 @@ void SMultiLineEditableText::OnKeyboardFocusLost( const FKeyboardFocusEvent& InK
 		//Always clear the local undo chain on commit.
 		ClearUndoStates();
 
-		FString EditedText;
+		FText EditedText;
 		TextLayout->GetAsText(EditedText);
 
-		OnTextCommitted.ExecuteIfBound(FText::FromString(EditedText), TextAction);
+		OnTextCommitted.ExecuteIfBound(EditedText, TextAction);
 		RemoveCursorHighlight();
 	}
 }
@@ -408,12 +455,12 @@ void SMultiLineEditableText::FinishChangingText()
 	check(bIsChangingText);
 	bIsChangingText = false;
 
-	FString EditedTextString;
-	TextLayout->GetAsText(EditedTextString);
-	FText EditedText = FText::FromString(EditedTextString);
+	FText EditedText;
+	TextLayout->GetAsText(EditedText);
 
 	// Has the text changed?
-	if (EditedText.ToString() != StateBeforeChangingText.Text.ToString())
+	const bool HasTextChanged = !EditedText.EqualTo(StateBeforeChangingText.Text);
+	if (HasTextChanged)
 	{
 		// Save text state
 		SaveText(EditedText);
@@ -423,8 +470,6 @@ void SMultiLineEditableText::FinishChangingText()
 
 		// We're done with this state data now.  Clear out any old data.
 		StateBeforeChangingText = FUndoState();
-		
-		LastSetText = EditedText;
 
 		TextLayout->UpdateIfNeeded();
 
@@ -712,35 +757,83 @@ void SMultiLineEditableText::DeleteSelectedText()
 
 FReply SMultiLineEditableText::MoveCursor( ECursorMoveMethod::Type Method, const FVector2D& Direction, ECursorAction::Type Action )
 {
+	bool bAllowMoveCursor = true;
 	FTextLocation NewCursorPosition;
-	const FTextLocation CursorPosition = CursorInfo.GetCursorInteractionLocation();
-	ECursorEOLMode CursorEOLMode = ECursorEOLMode::HardLines;
+	FTextLocation CursorPosition = CursorInfo.GetCursorInteractionLocation();
+
+	// If we have selected text, the cursor needs to:
+	// a) Jump to the start of the selection if moving the cursor Left or Up
+	// b) Jump to the end of the selection if moving the cursor Right or Down
+	// This is done regardless of whether the selection was made left-to-right, or right-to-left
+	// This also needs to be done *before* moving to word boundaries, or moving vertically, as the 
+	// start point needs to be the start or end of the selection depending on the above rules
+	if (Action == ECursorAction::MoveCursor && Method != ECursorMoveMethod::ScreenPosition && AnyTextSelected())
+	{
+		if (Method == ECursorMoveMethod::CharacterHorizontal)
+		{
+			// If we're moving the cursor horizontally, we just snap to the start or end of the selection rather than 
+			// move the cursor by the normal movement rules
+			bAllowMoveCursor = false;
+		}
 
-	switch (Method)
+		// Work out which edge of the selection we need to start at
+		bool bSnapToSelectionStart = false;
+		switch (Method)
+		{
+		case ECursorMoveMethod::CharacterHorizontal:
+		case ECursorMoveMethod::Word:
+			bSnapToSelectionStart = Direction.X < 0.0f;
+			break;
+	
+		case ECursorMoveMethod::CharacterVertical:
+			bSnapToSelectionStart = Direction.Y < 0.0f;
+			break;
+
+		default:
+			break;
+		}
+
+		// Adjust the current cursor position - also set the new cursor position so that the bAllowMoveCursor == false case is handled
+		const FTextSelection Selection(SelectionStart.GetValue(), CursorPosition);
+		CursorPosition = bSnapToSelectionStart ? Selection.GetBeginning() : Selection.GetEnd();
+		NewCursorPosition = CursorPosition;
+
+		// If we're snapping to a word boundary, but the selection was already at a word boundary, don't let the cursor move any more
+		if (Method == ECursorMoveMethod::Word && IsAtWordStart(NewCursorPosition))
+		{
+			bAllowMoveCursor = false;
+		}
+	}
+
+	ECursorEOLMode CursorEOLMode = ECursorEOLMode::HardLines;
+	if (bAllowMoveCursor)
 	{
-	case ECursorMoveMethod::CharacterHorizontal:
-		NewCursorPosition = TranslatedLocation( CursorPosition, Direction.X );
-		PreferredCursorScreenOffsetInLine = TextLayout->GetLocationAt( NewCursorPosition ).X;
-		break;
+		switch (Method)
+		{
+		case ECursorMoveMethod::CharacterHorizontal:
+			NewCursorPosition = TranslatedLocation( CursorPosition, Direction.X );
+			PreferredCursorScreenOffsetInLine = TextLayout->GetLocationAt( NewCursorPosition ).X;
+			break;
 
-	case ECursorMoveMethod::Word:
-		NewCursorPosition = ScanForWordBoundary( CursorPosition, Direction.X );
-		PreferredCursorScreenOffsetInLine = TextLayout->GetLocationAt( NewCursorPosition ).X;
-		break;
+		case ECursorMoveMethod::Word:
+			NewCursorPosition = ScanForWordBoundary( CursorPosition, Direction.X );
+			PreferredCursorScreenOffsetInLine = TextLayout->GetLocationAt( NewCursorPosition ).X;
+			break;
 	
-	case ECursorMoveMethod::CharacterVertical:
-		NewCursorPosition = TranslateLocationVertical( CursorPosition, Direction.Y );
-		break;
+		case ECursorMoveMethod::CharacterVertical:
+			NewCursorPosition = TranslateLocationVertical( CursorPosition, Direction.Y );
+			break;
 	
-	case ECursorMoveMethod::ScreenPosition:
-		NewCursorPosition = TextLayout->GetTextLocationAt( Direction * TextLayout->GetScale() );
-		CursorEOLMode = ECursorEOLMode::SoftLines;
-		PreferredCursorScreenOffsetInLine = Direction.X;
-		break;
+		case ECursorMoveMethod::ScreenPosition:
+			NewCursorPosition = TextLayout->GetTextLocationAt( Direction * TextLayout->GetScale() );
+			CursorEOLMode = ECursorEOLMode::SoftLines;
+			PreferredCursorScreenOffsetInLine = Direction.X;
+			break;
 
-	default:
-		checkSlow(false, "Unknown ECursorMoveMethod value");
-		break;
+		default:
+			checkSlow(false, "Unknown ECursorMoveMethod value");
+			break;
+		}
 	}
 
 	if (Action == ECursorAction::SelectText)
@@ -769,11 +862,19 @@ void SMultiLineEditableText::UpdateCursorHighlight()
 	TextLayout->ClearHighlights();
 
 	const FTextLocation CursorInteractionPosition = CursorInfo.GetCursorInteractionLocation();
-	FTextLocation SelectionLocation = SelectionStart.Get( CursorInteractionPosition );
+	const FTextLocation SelectionLocation = SelectionStart.Get( CursorInteractionPosition );
 
-	if ( SelectionLocation != CursorInteractionPosition )
+	bool bDrawCursor = true;
+
+	const bool bHasSelection = SelectionLocation != CursorInteractionPosition;
+	if ( bHasSelection )
 	{
-		FTextSelection Selection( SelectionLocation, CursorInteractionPosition );
+		bDrawCursor = false;
+
+		const FTextLocation CursorInteractionPosition = CursorInfo.GetCursorInteractionLocation();
+		const FTextLocation SelectionLocation = SelectionStart.Get( CursorInteractionPosition );
+
+		const FTextSelection Selection( SelectionLocation, CursorInteractionPosition );
 
 		const int32 SelectionBeginningLineIndex = Selection.GetBeginning().GetLineIndex();
 		const int32 SelectionBeginningLineOffset = Selection.GetBeginning().GetOffset();
@@ -781,50 +882,76 @@ void SMultiLineEditableText::UpdateCursorHighlight()
 		const int32 SelectionEndLineIndex = Selection.GetEnd().GetLineIndex();
 		const int32 SelectionEndLineOffset = Selection.GetEnd().GetOffset();
 
+		const bool bCursorIsBeforeSelection = CursorInteractionPosition < SelectionLocation;
+		SelectionRunHighlighter->SetCursorAlignment((bCursorIsBeforeSelection) ? ECursorAlignment::Left : ECursorAlignment::Right);
+
 		if ( SelectionBeginningLineIndex == SelectionEndLineIndex )
 		{
 			const FTextRange Range(SelectionBeginningLineOffset, SelectionEndLineOffset);
-			//if (!Range.IsEmpty())
-			//{
-				TextLayout->AddHighlight(FTextHighlight(SelectionBeginningLineIndex, Range, CursorSelectionHighlighter.ToSharedRef()));
-			//}
+			if (!Range.IsEmpty())
+			{
+				TextLayout->AddHighlight(FTextHighlight(SelectionBeginningLineIndex, Range, SelectionRunHighlighter.ToSharedRef()));
+			}
+			else
+			{
+				// We have no selection range for this line, but we're supposed to draw the cursor on this line - let the normal cursor handling take care of that
+				bDrawCursor = true;
+			}
 		}
 		else
 		{
 			const TArray< FTextLayout::FLineModel >& Lines = TextLayout->GetLineModels();
 
+			bool bHasAddedBeginHighlight = false;
+			bool bHasAddedEndHighlight = false;
+
 			for (int LineIndex = SelectionBeginningLineIndex; LineIndex <= SelectionEndLineIndex; LineIndex++)
 			{
 				if ( LineIndex == SelectionBeginningLineIndex )
 				{
+					check(!bHasAddedBeginHighlight);
+					bHasAddedBeginHighlight = true;
+
 					const FTextRange Range(SelectionBeginningLineOffset, Lines[LineIndex].Text->Len());
-					//if (!Range.IsEmpty())
-					//{
-						TextLayout->AddHighlight(FTextHighlight(LineIndex, Range, CursorSelectionHighlighter.ToSharedRef()));
-					//}
+					if (!Range.IsEmpty())
+					{
+						TextLayout->AddHighlight(FTextHighlight(LineIndex, Range, SelectionRunHighlighter.ToSharedRef()));
+					}
+					else if (bCursorIsBeforeSelection)
+					{
+						// We have no selection range for this line, but we're supposed to draw the cursor on this line - let the normal cursor handling take care of that
+						bDrawCursor = true;
+					}
 				}
 				else if ( LineIndex == SelectionEndLineIndex )
 				{
+					check(!bHasAddedEndHighlight);
+					bHasAddedEndHighlight = true;
+
 					const FTextRange Range(0, SelectionEndLineOffset);
-					//if (!Range.IsEmpty())
-					//{
-						TextLayout->AddHighlight(FTextHighlight(LineIndex, Range, CursorSelectionHighlighter.ToSharedRef()));
-					//}
+					if (!Range.IsEmpty())
+					{
+						TextLayout->AddHighlight(FTextHighlight(LineIndex, Range, SelectionRunHighlighter.ToSharedRef()));
+					}
+					else if (!bCursorIsBeforeSelection)
+					{
+						// We have no selection range for this line, but we're supposed to draw the cursor on this line - let the normal cursor handling take care of that
+						bDrawCursor = true;
+					}
 				}
 				else
 				{
 					const FTextRange Range(0, Lines[LineIndex].Text->Len());
-					//if (!Range.IsEmpty())
-					//{
-						TextLayout->AddHighlight(FTextHighlight(LineIndex, Range, CursorSelectionHighlighter.ToSharedRef()));
-					//}
+					if (!Range.IsEmpty())
+					{
+						TextLayout->AddHighlight(FTextHighlight(LineIndex, Range, SelectionRunHighlighter.ToSharedRef()));
+					}
 				}
 			}
 		}
-
-		CursorSelectionHighlighter->SetMode( FSlateCursorRunHighlighter::EMode::Selection );
 	}
-	else
+	
+	if ( bDrawCursor )
 	{
 		// The cursor mode uses the literal position rather than the interaction position
 		const FTextLocation CursorPosition = CursorInfo.GetCursorLocation();
@@ -834,20 +961,18 @@ void SMultiLineEditableText::UpdateCursorHighlight()
 
 		if (LineTextLength == 0)
 		{
-			TextLayout->AddHighlight( FTextHighlight( CursorPosition.GetLineIndex(), FTextRange( 0, 0 ), CursorSelectionHighlighter.ToSharedRef() ) );
+			TextLayout->AddHighlight(FTextHighlight(CursorPosition.GetLineIndex(), FTextRange(0, 0), CursorRunHighlighter.ToSharedRef()));
 		}
 		else if (CursorPosition.GetOffset() == LineTextLength)
 		{
-			TextLayout->AddHighlight(FTextHighlight(CursorPosition.GetLineIndex(), FTextRange(LineTextLength - 1, LineTextLength), CursorSelectionHighlighter.ToSharedRef()));
+			TextLayout->AddHighlight(FTextHighlight(CursorPosition.GetLineIndex(), FTextRange(LineTextLength - 1, LineTextLength), CursorRunHighlighter.ToSharedRef()));
 		}
 		else
 		{
-			TextLayout->AddHighlight(FTextHighlight(CursorPosition.GetLineIndex(), FTextRange(CursorPosition.GetOffset(), CursorPosition.GetOffset() + 1), CursorSelectionHighlighter.ToSharedRef()));
+			TextLayout->AddHighlight(FTextHighlight(CursorPosition.GetLineIndex(), FTextRange(CursorPosition.GetOffset(), CursorPosition.GetOffset() + 1), CursorRunHighlighter.ToSharedRef()));
 		}
-
-		CursorSelectionHighlighter->SetMode( FSlateCursorRunHighlighter::EMode::Cursor );
 	}
-};
+}
 
 void SMultiLineEditableText::RemoveCursorHighlight()
 {
@@ -1306,20 +1431,21 @@ void SMultiLineEditableText::Undo()
 				// NOTE: It's important the no code called here creates or destroys undo states!
 				const FUndoState& UndoState = UndoStates[UndoStateIndex];
 
-				SetText(UndoState.Text);
+				SaveText(UndoState.Text);
+
+				if (SetEditableText(UndoState.Text))
+				{
+					// Let outsiders know that the text content has been changed
+					OnTextChanged.ExecuteIfBound(UndoState.Text);
+				}
+
 				CursorInfo = UndoState.CursorInfo.CreateUndo();
 				SelectionStart = UndoState.SelectionStart;
 
 				UpdateCursorHighlight();
-
-				SaveText(UndoState.Text);
 			}
 
 			CurrentUndoLevel = UndoStateIndex;
-
-			// Let outsiders know that the text content has been changed
-			check(LastSetText.IsSet());
-			OnTextChanged.ExecuteIfBound(LastSetText.GetValue());
 		}
 	}
 }
@@ -1338,13 +1464,18 @@ void SMultiLineEditableText::Redo()
 				// NOTE: It's important the no code called here creates or destroys undo states!
 				const FUndoState& UndoState = UndoStates[NextUndoLevel];
 
-				SetText(UndoState.Text);
+				SaveText(UndoState.Text);
+
+				if (SetEditableText(UndoState.Text))
+				{
+					// Let outsiders know that the text content has been changed
+					OnTextChanged.ExecuteIfBound(UndoState.Text);
+				}
+
 				CursorInfo.RestoreFromUndo(UndoState.CursorInfo);
 				SelectionStart = UndoState.SelectionStart;
 
 				UpdateCursorHighlight();
-
-				SaveText(UndoState.Text);
 			}
 
 			CurrentUndoLevel = NextUndoLevel;
@@ -1357,10 +1488,6 @@ void SMultiLineEditableText::Redo()
 				// Pop last undo state that we created on initial undo
 				UndoStates.RemoveAt(UndoStates.Num() - 1);
 			}
-
-			// Let outsiders know that the text content has been changed
-			check(LastSetText.IsSet());
-			OnTextChanged.ExecuteIfBound(LastSetText.GetValue());
 		}
 	}
 }
@@ -1396,10 +1523,10 @@ void SMultiLineEditableText::ClearUndoStates()
 void SMultiLineEditableText::MakeUndoState(SMultiLineEditableText::FUndoState& OutUndoState)
 {
 	//@todo save and restoring the whole document is not ideal [3/31/2014 justin.sargent]
-	FString EditedText;
+	FText EditedText;
 	TextLayout->GetAsText(EditedText);
 
-	OutUndoState.Text = FText::FromString(EditedText);
+	OutUndoState.Text = EditedText;
 	OutUndoState.CursorInfo = CursorInfo;
 	OutUndoState.SelectionStart = SelectionStart;
 }
@@ -1473,7 +1600,7 @@ void SMultiLineEditableText::OnWindowClosed(const TSharedRef<SWindow>&)
 
 void SMultiLineEditableText::LoadText()
 {
-	SetText(BoundText.Get());
+	SetText(BoundText);
 	TextLayout->UpdateIfNeeded();
 }
 
@@ -1485,10 +1612,12 @@ void SMultiLineEditableText::Tick( const FGeometry& AllottedGeometry, const doub
 
 	if (!bShouldAppearFocused && BoundText.IsBound())
 	{
-		const FText& UpdatedBoundText = BoundText.Get();
-		if (!LastSetText.IsSet() || UpdatedBoundText.CompareTo(LastSetText.GetValue()) != 0)
+		const FText& TextToSet = BoundText.Get(FText::GetEmpty());
+		if (!BoundTextLastTick.EqualTo(TextToSet))
 		{
-			SetText(UpdatedBoundText);
+			// The source text has changed, so update the internal editable text
+			SetEditableText(TextToSet);
+			BoundTextLastTick = TextToSet;
 		}
 	}
 }
@@ -1743,8 +1872,8 @@ bool SMultiLineEditableText::IsAtEndOfLine( const FTextLocation& Location ) cons
 
 bool SMultiLineEditableText::IsAtWordStart( const FTextLocation& Location ) const
 {
-	const TCHAR CharBeforeCursor = GetCharacterAt(TranslatedLocation(Location, -1));
-	const TCHAR CharAfterCursor = GetCharacterAt(TranslatedLocation(Location, +1));
+	const FTextSelection WordUnderCursor = TextLayout->GetWordAt(Location);
+	const FTextLocation WordStart = WordUnderCursor.GetBeginning();
 
-	return FText::IsWhitespace(CharBeforeCursor) && !FText::IsWhitespace(CharAfterCursor);
+	return WordStart.IsValid() && WordStart == Location;
 }
diff --git a/Engine/Source/Runtime/Slate/Public/Framework/Text/TextLayout.h b/Engine/Source/Runtime/Slate/Public/Framework/Text/TextLayout.h
--- a/Engine/Source/Runtime/Slate/Public/Framework/Text/TextLayout.h
+++ b/Engine/Source/Runtime/Slate/Public/Framework/Text/TextLayout.h
@@ -297,6 +297,8 @@ public:
 
 	void GetAsText(FString& DisplayText) const;
 
+	void GetAsText(FText& DisplayText) const;
+
 	void GetSelectionAsText(FString& DisplayText, const FTextSelection& Selection) const;
 
 	FTextSelection GetWordAt(const FTextLocation& Location) const;
diff --git a/Engine/Source/Runtime/Slate/Public/Widgets/Input/SEditableText.h b/Engine/Source/Runtime/Slate/Public/Widgets/Input/SEditableText.h
--- a/Engine/Source/Runtime/Slate/Public/Widgets/Input/SEditableText.h
+++ b/Engine/Source/Runtime/Slate/Public/Widgets/Input/SEditableText.h
@@ -373,6 +373,12 @@ protected:
 	 */
 	int32 FindClickedCharacterIndex( const FVector2D& InLocalCursorPosition ) const;
 
+	/** Find the closest word boundary */
+	int32 ScanForWordBoundary( const int32 Location, int8 Direction ) const; 
+
+	/** Are we currently at the beginning of a word */
+	bool IsAtWordStart( const int32 Location ) const;
+
 	/**
 	 * Adds the specified undo state to the undo stack
 	 *
diff --git a/Engine/Source/Runtime/Slate/Public/Widgets/Input/SMultiLineEditableTextBox.h b/Engine/Source/Runtime/Slate/Public/Widgets/Input/SMultiLineEditableTextBox.h
--- a/Engine/Source/Runtime/Slate/Public/Widgets/Input/SMultiLineEditableTextBox.h
+++ b/Engine/Source/Runtime/Slate/Public/Widgets/Input/SMultiLineEditableTextBox.h
@@ -30,6 +30,7 @@ public:
 		, _SelectAllTextOnCommit( false )
 		, _BackgroundColor()		
 		, _Padding()
+		, _Margin()
 		, _ErrorReporting()
 		{}
 
@@ -98,6 +99,10 @@ public:
 		/** Padding between the box/border and the text widget inside (overrides Style) */
 		SLATE_ATTRIBUTE( FMargin, Padding )
 
+		/** The amount of blank space left around the edges of text area. 
+			This is different to Padding because this area is still considered part of the text area, and as such, can still be interacted with */
+		SLATE_ATTRIBUTE( FMargin, Margin )
+
 		/** Provide a alternative mechanism for error reporting. */
 		SLATE_ARGUMENT( TSharedPtr<class IErrorReportingWidget>, ErrorReporting )
 
diff --git a/Engine/Source/Runtime/Slate/Public/Widgets/Text/SMultiLineEditableText.h b/Engine/Source/Runtime/Slate/Public/Widgets/Text/SMultiLineEditableText.h
--- a/Engine/Source/Runtime/Slate/Public/Widgets/Text/SMultiLineEditableText.h
+++ b/Engine/Source/Runtime/Slate/Public/Widgets/Text/SMultiLineEditableText.h
@@ -188,22 +188,13 @@ private:
 		}
 	};
 
+	/** Run highlighter used to draw the cursor */
 	class FSlateCursorRunHighlighter : public ISlateRunHighlighter
 	{
 	public:
 
 		static TSharedRef< FSlateCursorRunHighlighter > Create(const FCursorInfo* InCursorInfo);
 
-	public:
-
-		enum class EMode : uint8
-		{
-			Cursor,
-			Selection,
-		};
-
-	public:
-
 		virtual ~FSlateCursorRunHighlighter() {}
 
 		virtual int32 OnPaint( const FTextLayout::FLineView& Line, const TSharedRef< ISlateRun >& Run, const TSharedRef< ILayoutBlock >& Block, const FTextBlockStyle& DefaultStyle, const FGeometry& AllottedGeometry, const FSlateRect& MyClippingRect, FSlateWindowElementList& OutDrawElements, int32 LayerId, const FWidgetStyle& InWidgetStyle, bool bParentEnabled ) const override;
@@ -212,20 +203,37 @@ private:
 
 		virtual void OnArrangeChildren( const TSharedRef< ILayoutBlock >& Block, const FGeometry& AllottedGeometry, FArrangedChildren& ArrangedChildren ) const override;
 
-		void SetMode( EMode Mode );
+	protected:
 
-	private:
+		int32 OnPaintCursor( const FTextLayout::FLineView& Line, const TSharedRef< ISlateRun >& Run, const TSharedRef< ILayoutBlock >& Block, const FTextBlockStyle& DefaultStyle, const FGeometry& AllottedGeometry, const FSlateRect& MyClippingRect, FSlateWindowElementList& OutDrawElements, int32 LayerId, const FWidgetStyle& InWidgetStyle, const bool bParentEnabled, const ECursorAlignment InCursorAlignment ) const;
 
 		FSlateCursorRunHighlighter(const FCursorInfo* InCursorInfo);
 
-	private:
-
 		/** Cursor data that this highlighter is tracking */
 		const FCursorInfo* CursorInfo;
+	};
+
+	/** Run highlighter used to draw selection ranges */
+	class FSlateSelectionRunHighlighter : public FSlateCursorRunHighlighter
+	{
+	public:
+
+		static TSharedRef< FSlateSelectionRunHighlighter > Create(const FCursorInfo* InCursorInfo);
 
-		EMode Mode;
+		virtual ~FSlateSelectionRunHighlighter() {}
 
-		static FNoChildren NoChildrenInstance;
+		virtual int32 OnPaint( const FTextLayout::FLineView& Line, const TSharedRef< ISlateRun >& Run, const TSharedRef< ILayoutBlock >& Block, const FTextBlockStyle& DefaultStyle, const FGeometry& AllottedGeometry, const FSlateRect& MyClippingRect, FSlateWindowElementList& OutDrawElements, int32 LayerId, const FWidgetStyle& InWidgetStyle, bool bParentEnabled ) const override;
+
+		void SetCursorAlignment(const ECursorAlignment InCursorAlignment)
+		{
+			CursorAlignment = InCursorAlignment;
+		}
+
+	protected:
+
+		FSlateSelectionRunHighlighter(const FCursorInfo* InCursorInfo);
+
+		ECursorAlignment CursorAlignment;
 	};
 
 private:
@@ -354,13 +362,24 @@ private:
 
 	void SaveText(const FText& TextToSave);
 
+	/**
+	 * Sets the current editable text for this text block
+	 * Note: Doesn't update the value of BoundText, nor does it call OnTextChanged
+	 *
+	 * @param TextToSet		The new text to set in the internal TextLayout
+	 * @param bForce		True to force the update, even if the text currently matches what's in the TextLayout
+	 *
+	 * @return true if the text was updated, false if the text wasn't update (because it was already up-to-date)
+	 */
+	bool SetEditableText(const FText& TextToSet, const bool bForce = false);
+
 private:
 
 	/** The text displayed in this text block */
 	TAttribute<FText> BoundText;
 
-	/** The text that is currently in TextLayout; used to avoid un-needed updates */
-	TOptional<FText> LastSetText;
+	/** The state of BoundText last Tick() (only used when BoundText is bound to a delegate providing the source text) */
+	FText BoundTextLastTick;
 
 	/** In control of the layout and wrapping of the BoundText */
 	TSharedPtr< FSlateTextLayout > TextLayout;
@@ -381,7 +400,9 @@ private:
 	TAttribute< ETextJustify::Type > Justification; 
 	TAttribute< float > LineHeightPercentage;
 
-	TSharedPtr< FSlateCursorRunHighlighter > CursorSelectionHighlighter;
+	TSharedPtr< FSlateCursorRunHighlighter > CursorRunHighlighter;
+
+	TSharedPtr< FSlateSelectionRunHighlighter > SelectionRunHighlighter;
 
 	/** That start of the selection when there is a selection. The end is implicitly wherever the cursor happens to be. */
 	TOptional<FTextLocation> SelectionStart;
