Reverted some places using FText::EqualTo to comapre text in the multi-line editable text

TTP# 336464 - Editor: Finish the Multiline Editable Text Block

They now use a case-sensitive string comparison instead, as this will perform an ordinal comparision and not consider the strings identical if they're just "conceptually" identical.

diff --git a/Engine/Source/Runtime/Slate/Private/Widgets/Text/SMultiLineEditableText.cpp b/Engine/Source/Runtime/Slate/Private/Widgets/Text/SMultiLineEditableText.cpp
--- a/Engine/Source/Runtime/Slate/Private/Widgets/Text/SMultiLineEditableText.cpp
+++ b/Engine/Source/Runtime/Slate/Private/Widgets/Text/SMultiLineEditableText.cpp
@@ -307,7 +307,7 @@ bool SMultiLineEditableText::SetEditableText(const FText& TextToSet, const bool
 	FText EditedText;
 	TextLayout->GetAsText(EditedText);
 
-	const bool HasTextChanged = !EditedText.EqualTo(TextToSet);
+	const bool HasTextChanged = !EditedText.ToString().Equals(TextToSet.ToString(), ESearchCase::CaseSensitive);
 	if (HasTextChanged || bForce)
 	{
 		const FString& TextToSetString = TextToSet.ToString();
@@ -459,7 +459,7 @@ void SMultiLineEditableText::FinishChangingText()
 	TextLayout->GetAsText(EditedText);
 
 	// Has the text changed?
-	const bool HasTextChanged = !EditedText.EqualTo(StateBeforeChangingText.Text);
+	const bool HasTextChanged = !EditedText.ToString().Equals(StateBeforeChangingText.Text.ToString(), ESearchCase::CaseSensitive);
 	if (HasTextChanged)
 	{
 		// Save text state
@@ -1613,7 +1613,7 @@ void SMultiLineEditableText::Tick( const FGeometry& AllottedGeometry, const doub
 	if (!bShouldAppearFocused && BoundText.IsBound())
 	{
 		const FText& TextToSet = BoundText.Get(FText::GetEmpty());
-		if (!BoundTextLastTick.EqualTo(TextToSet))
+		if (!BoundTextLastTick.ToString().Equals(TextToSet.ToString(), ESearchCase::CaseSensitive))
 		{
 			// The source text has changed, so update the internal editable text
 			SetEditableText(TextToSet);
