Fixed: Slate Tree view scrolling can get temporarily stuck in 'springy mode' with blank space showing and no mouse buttons down
Fixes TTP 338095. Now consuming all button clicks while RMB-scrolling.

diff --git a/Engine/Source/Runtime/Slate/Private/Widgets/Views/STableViewBase.cpp b/Engine/Source/Runtime/Slate/Private/Widgets/Views/STableViewBase.cpp
--- a/Engine/Source/Runtime/Slate/Private/Widgets/Views/STableViewBase.cpp
+++ b/Engine/Source/Runtime/Slate/Private/Widgets/Views/STableViewBase.cpp
@@ -9,7 +9,6 @@ namespace ListConstants
 	static const float OvershootBounceRate = 350.0f;
 }
 
-/** Create the child widgets that comprise the list */
 void STableViewBase::ConstructChildren( const TAttribute<float>& InItemWidth, const TAttribute<float>& InItemHeight, const TSharedPtr<SHeaderRow>& InHeaderRow, const TSharedPtr<SScrollBar>& InScrollBar )
 {
 	HeaderRow = InHeaderRow;
@@ -84,11 +83,7 @@ void STableViewBase::ConstructChildren( const TAttribute<float>& InItemWidth, co
 	
 }
 
-/**
- * See SWidget::SupportsKeyboardFocus().
- *
- * @return  True if this widget can take keyboard focus
- */
+
 bool STableViewBase::SupportsKeyboardFocus() const
 {
 	// The ListView is focusable.
@@ -151,13 +146,7 @@ static FEndOfListResult ComputeOffsetForEndOfList( const FGeometry& ListPanelGeo
 	return FEndOfListResult( OffsetFromEndOfList, ItemsAboveView );
 }
 
-/**
- * See SWidget::Tick()
- *
- * @param  AllottedGeometry The space allotted for this widget
- * @param  InCurrentTime  Current absolute real time
- * @param  InDeltaTime  Real time passed since last tick
- */
+
 void STableViewBase::Tick( const FGeometry& AllottedGeometry, const double InCurrentTime, const float InDeltaTime )
 {
 	if (ItemsPanel.IsValid())
@@ -234,25 +223,14 @@ void STableViewBase::Tick( const FGeometry& AllottedGeometry, const double InCur
 	}
 }
 
-/**
- * Invoked by the scrollbar when the user scrolls.
- *
- * @param InScrollOffsetFraction  The location to which the user scrolled as a fraction (between 0 and 1) of total height of the content.
- */
+
 void STableViewBase::ScrollBar_OnUserScrolled( float InScrollOffsetFraction )
 {
 	const double ClampedScrollOffsetInItems = FMath::Clamp<double>( InScrollOffsetFraction, 0.0, 1.0 )* GetNumItemsBeingObserved();
 	ScrollTo( ClampedScrollOffsetInItems );
 }
 
-/**
- * See SWidget::OnMouseButtonDown()
- *
- * @param MyGeometry The Geometry of the widget receiving the event
- * @param MouseEvent Information about the input event
- *
- * @return Whether the event was handled along with possible requests for the system to take action.
- */
+
 FReply STableViewBase::OnMouseButtonDown( const FGeometry& MyGeometry, const FPointerEvent& MouseEvent )
 {
 	// Zero the scroll velocity so the list stops immediately on mouse down, even if the user does not drag
@@ -269,17 +247,26 @@ FReply STableViewBase::OnMouseButtonDown( const FGeometry& MyGeometry, const FPo
 		// they reach our scroll threshold
 		return FReply::Handled();
 	}
+	else if ( this->HasMouseCapture() )
+	{
+		// Consume all mouse buttons while we are RMB-dragging.
+		return FReply::Handled();
+	}
 	return FReply::Unhandled();			
 }
 
-/**
- * See SWidget::OnMouseButtonUp()
- *
- * @param MyGeometry The Geometry of the widget receiving the event
- * @param MouseEvent Information about the input event
- *
- * @return Whether the event was handled along with possible requests for the system to take action.
- */
+FReply STableViewBase::OnMouseButtonDoubleClick( const FGeometry& InMyGeometry, const FPointerEvent& InMouseEvent )
+{
+	if ( this->HasMouseCapture() )
+	{
+		// Consume all other mouse buttons while we are RMB-dragging.
+		return FReply::Handled();
+	}
+	return FReply::Unhandled();			
+
+}
+
+
 FReply STableViewBase::OnMouseButtonUp( const FGeometry& MyGeometry, const FPointerEvent& MouseEvent )
 {
 	if ( MouseEvent.GetEffectingButton() == EKeys::RightMouseButton )
@@ -309,14 +296,6 @@ FReply STableViewBase::OnMouseButtonUp( const FGeometry& MyGeometry, const FPoin
 }
 
 
-/**
- * See SWidget::OnMouseMove()
- *
- * @param MyGeometry The Geometry of the widget receiving the event
- * @param MouseEvent Information about the input event
- *
- * @return Whether the event was handled along with possible requests for the system to take action.
- */
 FReply STableViewBase::OnMouseMove( const FGeometry& MyGeometry, const FPointerEvent& MouseEvent )
 {	
 	if( MouseEvent.IsMouseButtonDown( EKeys::RightMouseButton ) )
@@ -368,14 +347,6 @@ void STableViewBase::OnMouseLeave( const FPointerEvent& MouseEvent )
 }
 
 
-/**
- * See SWidget::OnMouseWheel()
- *
- * @param MyGeometry The Geometry of the widget receiving the event
- * @param MouseEvent Information about the input event
- *
- * @return Whether the event was handled along with possible requests for the system to take action.
- */
 FReply STableViewBase::OnMouseWheel( const FGeometry& MyGeometry, const FPointerEvent& MouseEvent )
 {
 	if( !MouseEvent.IsControlDown() )
@@ -392,14 +363,7 @@ FReply STableViewBase::OnMouseWheel( const FGeometry& MyGeometry, const FPointer
 	return FReply::Unhandled();
 }
 
-/**
- * See SWidget::OnKeyDown().
- *
- * @param MyGeometry The Geometry of the widget receiving the event
- * @param  InKeyboardEvent  Keyboard event
- *
- * @return  Returns whether the event was handled, along with other possible actions
- */
+
 FReply STableViewBase::OnKeyDown( const FGeometry& MyGeometry, const FKeyboardEvent& InKeyboardEvent )
 {
 	if ( InKeyboardEvent.IsControlDown() && InKeyboardEvent.GetKey() == EKeys::End )
@@ -470,7 +434,6 @@ FReply STableViewBase::OnTouchEnded( const FGeometry& MyGeometry, const FPointer
 }
 
 
-/** @return The number of Widgets we currently have generated. */
 int32 STableViewBase::GetNumGeneratedChildren() const
 {
 	return (ItemsPanel.IsValid())
diff --git a/Engine/Source/Runtime/Slate/Public/Widgets/Views/STableViewBase.h b/Engine/Source/Runtime/Slate/Public/Widgets/Views/STableViewBase.h
--- a/Engine/Source/Runtime/Slate/Public/Widgets/Views/STableViewBase.h
+++ b/Engine/Source/Runtime/Slate/Public/Widgets/Views/STableViewBase.h
@@ -42,6 +42,7 @@ public:
 	virtual bool SupportsKeyboardFocus() const override;
 	virtual void Tick( const FGeometry& AllottedGeometry, const double InCurrentTime, const float InDeltaTime ) override;
 	virtual FReply OnMouseButtonDown( const FGeometry& MyGeometry, const FPointerEvent& MouseEvent ) override;
+	virtual FReply OnMouseButtonDoubleClick( const FGeometry& InMyGeometry, const FPointerEvent& InMouseEvent ) override;
 	virtual FReply OnMouseButtonUp( const FGeometry& MyGeometry, const FPointerEvent& MouseEvent ) override;
 	virtual void OnMouseLeave( const FPointerEvent& MouseEvent ) override;
 	virtual FReply OnMouseMove( const FGeometry& MyGeometry, const FPointerEvent& MouseEvent ) override;
