Slate - Exopsing a little more information in the scrollbox, such as an event for when the user scrolls.  Needed for the credits screen.

diff --git a/Engine/Source/Runtime/Slate/Private/Widgets/Layout/SScrollBox.cpp b/Engine/Source/Runtime/Slate/Private/Widgets/Layout/SScrollBox.cpp
--- a/Engine/Source/Runtime/Slate/Private/Widgets/Layout/SScrollBox.cpp
+++ b/Engine/Source/Runtime/Slate/Private/Widgets/Layout/SScrollBox.cpp
@@ -20,8 +20,6 @@ public:
 		}
 	SLATE_END_ARGS()
 
-
-
 	void Construct( const FArguments& InArgs, const TArray<SScrollBox::FSlot*>& InSlots )
 	{
 		PhysicalOffset = 0;
@@ -32,13 +30,10 @@ public:
 		}
 	}
 
-
-
 public:
 
 	virtual void OnArrangeChildren( const FGeometry& AllottedGeometry, FArrangedChildren& ArrangedChildren ) const override
 	{
-		
 		float CurChildOffset = -PhysicalOffset;
 
 		for( int32 SlotIndex=0; SlotIndex < Children.Num(); ++SlotIndex )
@@ -60,12 +55,8 @@ public:
 				CurChildOffset += ThisSlotDesiredHeight;
 			}
 		}
-
-
 	}
 
-
-
 	virtual FVector2D ComputeDesiredSize() const override
 	{
 		FVector2D ThisDesiredSize = FVector2D::ZeroVector;
@@ -80,17 +71,13 @@ public:
 		return ThisDesiredSize;
 	}
 
-
-
 	virtual FChildren* GetChildren() override
 	{
 		return &Children;
 	}
 
-
 	float PhysicalOffset;
 	TPanelChildren<SScrollBox::FSlot> Children;	
-	
 };
 
 
@@ -105,22 +92,26 @@ void SScrollBox::Construct( const FArguments& InArgs )
 	AmountScrolledWhileRightMouseDown = 0;
 	bShowSoftwareCursor = false;
 	SoftwareCursorPosition = FVector2D::ZeroVector;
+	OnUserScrolled = InArgs._OnUserScrolled;
 
 	TSharedPtr<SHorizontalBox> PanelAndScrollbar;
 	this->ChildSlot
 	[
 		SAssignNew(PanelAndScrollbar, SHorizontalBox)
-		+SHorizontalBox::Slot()
+
+		+ SHorizontalBox::Slot()
 		.FillWidth(1)
 		[
 			SNew(SOverlay)
-			+SOverlay::Slot()
+
+			+ SOverlay::Slot()
 			.Padding( FMargin(0.0f, 0.0f, 0.0f, 1.0f) )
 			[
 				// Scroll panel that presents the scrolled content
 				SAssignNew( ScrollPanel, SScrollPanel, InArgs.Slots )
 			]
-			+SOverlay::Slot()
+
+			+ SOverlay::Slot()
 			.HAlign(HAlign_Fill)
 			.VAlign(VAlign_Top)
 			[
@@ -130,7 +121,8 @@ void SScrollBox::Construct( const FArguments& InArgs )
 				.ColorAndOpacity(this, &SScrollBox::GetTopShadowOpacity)
 				.Image( &InArgs._Style->TopShadowBrush )	
 			]
-			+SOverlay::Slot()
+
+			+ SOverlay::Slot()
 			.HAlign(HAlign_Fill)
 			.VAlign(VAlign_Bottom)
 			[
@@ -171,8 +163,6 @@ void SScrollBox::Construct( const FArguments& InArgs )
 	ScrollBar->SetState( 0.0f, 1.0f );
 }
 
-
-
 /** Adds a slot to SScrollBox */
 SScrollBox::FSlot& SScrollBox::AddSlot()
 {
@@ -182,8 +172,6 @@ SScrollBox::FSlot& SScrollBox::AddSlot()
 	return NewSlot;
 }
 
-
-
 /** Removes a slot at the specified location */
 void SScrollBox::RemoveSlot( const TSharedRef<SWidget>& WidgetToRemove )
 {
@@ -208,6 +196,11 @@ bool SScrollBox::IsRightClickScrolling() const
 	return AmountScrolledWhileRightMouseDown >= SlatePanTriggerDistance && this->ScrollBar->IsNeeded();
 }
 
+float SScrollBox::GetScrollOffset()
+{
+	return DesiredScrollOffset;
+}
+
 void SScrollBox::SetScrollOffset( float NewScrollOffset )
 {
 	DesiredScrollOffset = NewScrollOffset;
@@ -299,7 +292,6 @@ FReply SScrollBox::OnMouseButtonUp( const FGeometry& MyGeometry, const FPointerE
 	return FReply::Unhandled();
 }
 
-
 FReply SScrollBox::OnMouseMove( const FGeometry& MyGeometry, const FPointerEvent& MouseEvent )
 {	
 	if( MouseEvent.IsMouseButtonDown( EKeys::RightMouseButton ) )
@@ -372,17 +364,18 @@ bool SScrollBox::ScrollBy( const FGeometry& AllottedGeometry, float ScrollAmount
 	const float ContentSize = ScrollPanel->GetDesiredSize().Y;
 	const FGeometry ScrollPanelGeometry = FindChildGeometry( AllottedGeometry, ScrollPanel.ToSharedRef() );
 
+	float PreviousScrollOffset = DesiredScrollOffset;
+
 	DesiredScrollOffset += ScrollAmount;
 
-	const float ScrollMin = 0.0f, ScrollMax = ContentSize - ScrollPanelGeometry.Size.Y;
-	if (DesiredScrollOffset < ScrollMin || DesiredScrollOffset > ScrollMax)
-	{
-		DesiredScrollOffset = FMath::Clamp( DesiredScrollOffset, ScrollMin, ScrollMax );
-		return false;
-	}
-	return true;
-}
+	const float ScrollMin = 0.0f;
+	const float ScrollMax = ContentSize - ScrollPanelGeometry.Size.Y;
+	DesiredScrollOffset = FMath::Clamp(DesiredScrollOffset, ScrollMin, ScrollMax);
+
+	OnUserScrolled.ExecuteIfBound(DesiredScrollOffset);
 
+	return DesiredScrollOffset != PreviousScrollOffset;
+}
 
 FReply SScrollBox::OnDragDetected( const FGeometry& MyGeometry, const FPointerEvent& MouseEvent )
 {
@@ -427,6 +420,7 @@ int32 SScrollBox::OnPaint( const FGeometry& AllottedGeometry, const FSlateRect&
 void SScrollBox::ScrollBar_OnUserScrolled( float InScrollOffsetFraction )
 {
 	DesiredScrollOffset = InScrollOffsetFraction * ScrollPanel->GetDesiredSize().Y;
+	OnUserScrolled.ExecuteIfBound(DesiredScrollOffset);
 }
 
 const float ShadowFadeDistance = 32.0f;
diff --git a/Engine/Source/Runtime/Slate/Public/Widgets/Layout/SScrollBox.h b/Engine/Source/Runtime/Slate/Public/Widgets/Layout/SScrollBox.h
--- a/Engine/Source/Runtime/Slate/Public/Widgets/Layout/SScrollBox.h
+++ b/Engine/Source/Runtime/Slate/Public/Widgets/Layout/SScrollBox.h
@@ -28,6 +28,7 @@ public:
 	SLATE_BEGIN_ARGS(SScrollBox)
 		: _Style( &FCoreStyle::Get().GetWidgetStyle<FScrollBoxStyle>("ScrollBox") )
 		, _ExternalScrollbar()
+		, _OnUserScrolled()
 		{}
 		
 		SLATE_SUPPORTS_SLOT( FSlot )
@@ -37,6 +38,10 @@ public:
 
 		/** Custom scroll bar */
 		SLATE_ARGUMENT( TSharedPtr<SScrollBar>, ExternalScrollbar )
+
+		/** Called when the button is clicked */
+		SLATE_EVENT(FOnUserScrolled, OnUserScrolled)
+
 	SLATE_END_ARGS()
 
 
@@ -58,6 +63,8 @@ public:
 		        the right mouse button and dragging. */
 	bool IsRightClickScrolling() const;
 
+	float GetScrollOffset();
+
 	void SetScrollOffset( float NewScrollOffset );
 
 public:
@@ -117,6 +124,9 @@ private:
 
 	/**	Whether the software cursor should be drawn in the viewport */
 	bool bShowSoftwareCursor;
+
+	/** Fired when the user scrolls the scrollbox */
+	FOnUserScrolled OnUserScrolled;
 };
 
 
