Added EButtonTouchMethod::PreciseTap for supporting buttons inside of multitouch-aware lists. Button must be tapped without moving one's finger, as moving the finger begins scrolling.

diff --git a/Engine/Source/Runtime/Slate/Private/Framework/Testing/STableViewTesting.cpp b/Engine/Source/Runtime/Slate/Private/Framework/Testing/STableViewTesting.cpp
--- a/Engine/Source/Runtime/Slate/Private/Framework/Testing/STableViewTesting.cpp
+++ b/Engine/Source/Runtime/Slate/Private/Framework/Testing/STableViewTesting.cpp
@@ -315,6 +315,7 @@ public:
 				.Padding(this, &SItemEditor::GetVariableHeight)
 				[
 					SNew( SButton )
+					.TouchMethod( EButtonTouchMethod::PreciseTap )
 					.OnClicked( this, &SItemEditor::OnAddChild )
 					[
 						SNew(STextBlock)
diff --git a/Engine/Source/Runtime/Slate/Private/Widgets/Input/SButton.cpp b/Engine/Source/Runtime/Slate/Private/Widgets/Input/SButton.cpp
--- a/Engine/Source/Runtime/Slate/Private/Widgets/Input/SButton.cpp
+++ b/Engine/Source/Runtime/Slate/Private/Widgets/Input/SButton.cpp
@@ -54,6 +54,7 @@ void SButton::Construct( const FArguments& InArgs )
 	OnClicked = InArgs._OnClicked;
 
 	ClickMethod = InArgs._ClickMethod;
+	TouchMethod = InArgs._TouchMethod;
 
 	HoveredSound = InArgs._HoveredSoundOverride.Get(Style->HoveredSlateSound);
 	PressedSound = InArgs._PressedSoundOverride.Get(Style->PressedSlateSound);
@@ -116,7 +117,7 @@ FReply SButton::OnKeyDown( const FGeometry& MyGeometry, const FKeyboardEvent& In
 FReply SButton::OnMouseButtonDown( const FGeometry& MyGeometry, const FPointerEvent& MouseEvent )
 {
 	FReply Reply = FReply::Unhandled();
-	if ( MouseEvent.GetEffectingButton() == EKeys::LeftMouseButton )
+	if ( MouseEvent.GetEffectingButton() == EKeys::LeftMouseButton || MouseEvent.IsTouchEvent() )
 	{
 		bIsPressed = true;
 
@@ -133,6 +134,11 @@ FReply SButton::OnMouseButtonDown( const FGeometry& MyGeometry, const FPointerEv
 			//You should ALWAYS handle the OnClicked event.
 			ensure(Reply.IsEventHandled() == true);
 		}
+		else if (TouchMethod == EButtonTouchMethod::PreciseTap && MouseEvent.IsTouchEvent() )
+		{
+			// do not capture the pointer for precise taps
+			// 
+		}
 		else
 		{
 			//we need to capture the mouse for MouseUp events
@@ -154,13 +160,13 @@ FReply SButton::OnMouseButtonDoubleClick( const FGeometry& InMyGeometry, const F
 FReply SButton::OnMouseButtonUp( const FGeometry& MyGeometry, const FPointerEvent& MouseEvent )
 {
 	FReply Reply = FReply::Unhandled();
-	if ( MouseEvent.GetEffectingButton() == EKeys::LeftMouseButton )
+	if ( MouseEvent.GetEffectingButton() == EKeys::LeftMouseButton || MouseEvent.IsTouchEvent() )
 	{
 		bIsPressed = false;
 
 		if( ClickMethod == EButtonClickMethod::MouseDown )
 		{
-			// NOTE: If we're configured to click on mouse-down, then we never capture the mouse thus
+			// NOTE: If we're configured to click on mouse-down/precise-tap, then we never capture the mouse thus
 			//       may never receive an OnMouseButtonUp() call.  We make sure that our bIsPressed
 			//       state is reset by overriding OnMouseLeave().
 		}
@@ -169,9 +175,15 @@ FReply SButton::OnMouseButtonUp( const FGeometry& MyGeometry, const FPointerEven
 			const bool bIsUnderMouse = MyGeometry.IsUnderLocation(MouseEvent.GetScreenSpacePosition());
 			if( bIsUnderMouse )
 			{
+				// If we asked for a precide tap, all we need is for the user to have not moved their pointer very far.
+				// Precise taps only occur for mobile devices, 
+				const bool bTriggerForTouchEvent = MouseEvent.IsTouchEvent() && TouchMethod == EButtonTouchMethod::PreciseTap;
+
 				// If we were asked to allow the button to be clicked on mouse up, regardless of whether the user
 				// pressed the button down first, then we'll allow the click to proceed without an active capture
-				if( ( ClickMethod == EButtonClickMethod::MouseUp || HasMouseCapture() ) && OnClicked.IsBound() == true )
+				const bool bTriggerForMouseEvent = ( ClickMethod == EButtonClickMethod::MouseUp || HasMouseCapture() );
+
+				if( (bTriggerForTouchEvent || bTriggerForMouseEvent) && OnClicked.IsBound() == true )
 				{
 					Reply = OnClicked.Execute();
 				}
@@ -199,6 +211,10 @@ FReply SButton::OnMouseButtonUp( const FGeometry& MyGeometry, const FPointerEven
 
 FReply SButton::OnMouseMove( const FGeometry& MyGeometry, const FPointerEvent& MouseEvent )
 {
+	if ( MouseEvent.IsTouchEvent() && TouchMethod == EButtonTouchMethod::PreciseTap && MouseEvent.GetCursorDelta().SizeSquared() != 0 )
+	{
+		bIsPressed = false;
+	}
 	return FReply::Unhandled();
 }
 
@@ -219,13 +235,12 @@ void SButton::OnMouseLeave( const FPointerEvent& MouseEvent )
 
 	// If we're setup to click on mouse-down, then we never capture the mouse and may not receive a
 	// mouse up event, so we need to make sure our pressed state is reset properly here
-	if( ClickMethod == EButtonClickMethod::MouseDown )
+	if( ClickMethod == EButtonClickMethod::MouseDown || ( TouchMethod == EButtonTouchMethod::PreciseTap && MouseEvent.IsTouchEvent() ) )
 	{
 		bIsPressed = false;
 	}
 }
 
-
 void SButton::PlayPressedSound() const
 {
 	FSlateApplication::Get().PlaySound( PressedSound );
diff --git a/Engine/Source/Runtime/Slate/Public/Widgets/Input/SButton.h b/Engine/Source/Runtime/Slate/Public/Widgets/Input/SButton.h
--- a/Engine/Source/Runtime/Slate/Public/Widgets/Input/SButton.h
+++ b/Engine/Source/Runtime/Slate/Public/Widgets/Input/SButton.h
@@ -19,6 +19,7 @@ public:
 		, _ContentPadding(FMargin(4.0, 2.0))
 		, _Text()
 		, _ClickMethod( EButtonClickMethod::DownAndUp )
+		, _TouchMethod( EButtonTouchMethod::DownAndUp )
 		, _DesiredSizeScale( FVector2D(1,1) )
 		, _ContentScale( FVector2D(1,1) )
 		, _ButtonColorAndOpacity(FLinearColor::White)
@@ -53,6 +54,9 @@ public:
 		/** Sets the rules to use for determining whether the button was clicked.  This is an advanced setting and generally should be left as the default. */
 		SLATE_ARGUMENT( EButtonClickMethod::Type, ClickMethod )
 
+		/** How should the button be clicked with touch events? */
+		SLATE_ARGUMENT( EButtonTouchMethod::Type, TouchMethod )
+
 		SLATE_ATTRIBUTE( FVector2D, DesiredSizeScale )
 
 		SLATE_ATTRIBUTE( FVector2D, ContentScale )
@@ -161,6 +165,9 @@ protected:
 	/** Sets whether a click should be triggered on mouse down, mouse up, or that both a mouse down and up are required. */
 	EButtonClickMethod::Type ClickMethod;
 
+	/** How should the button be clicked with touch events? */
+	EButtonTouchMethod::Type TouchMethod;
+
 	/** Can this button be focused? */
 	bool bIsFocusable;
 
diff --git a/Engine/Source/Runtime/SlateCore/Public/Types/SlateEnums.h b/Engine/Source/Runtime/SlateCore/Public/Types/SlateEnums.h
--- a/Engine/Source/Runtime/SlateCore/Public/Types/SlateEnums.h
+++ b/Engine/Source/Runtime/SlateCore/Public/Types/SlateEnums.h
@@ -35,6 +35,23 @@ namespace EButtonClickMethod
 	};
 }
 
+UENUM()
+namespace EButtonTouchMethod
+{
+	/** Ways in which touch interactions trigger a "Clicked" event. */
+	enum Type
+	{
+		/** Most buttons behave this way. */
+		DownAndUp,
+
+		/**
+		 * Inside a list, buttons can only be clicked with precise tap.
+		 * Moving the pointer will scroll the list.
+		 */
+		PreciseTap
+	};
+}
+
 UENUM()
 namespace EFocusMoveDirection
 {
diff --git a/Engine/Source/Runtime/SlateCore/Public/Widgets/SWidget.h b/Engine/Source/Runtime/SlateCore/Public/Widgets/SWidget.h
--- a/Engine/Source/Runtime/SlateCore/Public/Widgets/SWidget.h
+++ b/Engine/Source/Runtime/SlateCore/Public/Widgets/SWidget.h
@@ -390,21 +390,21 @@ public:
 	/**
 	 * Called when a touchpad touch is started (finger down)
 	 * 
-	 * @param TouchEvent	The touch event generated
+	 * @param InTouchEvent	The touch event generated
 	 */
 	virtual FReply OnTouchStarted( const FGeometry& MyGeometry, const FPointerEvent& InTouchEvent );
 
 	/**
 	 * Called when a touchpad touch is moved  (finger moved)
 	 * 
-	 * @param TouchEvent	The touch event generated
+	 * @param InTouchEvent	The touch event generated
 	 */
 	virtual FReply OnTouchMoved( const FGeometry& MyGeometry, const FPointerEvent& InTouchEvent );
 
 	/**
 	 * Called when a touchpad touch is ended (finger lifted)
 	 * 
-	 * @param TouchEvent	The touch event generated
+	 * @param InTouchEvent	The touch event generated
 	 */
 	virtual FReply OnTouchEnded( const FGeometry& MyGeometry, const FPointerEvent& InTouchEvent );
 
