Made shutdown code safer to eliminate crash.
#ttp 330012 - Editor Regular Crash - Right Click Close on close icon produces a 100% crash
#branch UE4
#proj Runtime.Slate
#change Checked validity of docking area TSharedPtr before using it.  In this case, it's possible for the SDockingArea to have already been deleted when the tabs are shutdown.

diff --git a/Engine/Source/Runtime/Slate/Private/Framework/Docking/SDockingTabWell.cpp b/Engine/Source/Runtime/Slate/Private/Framework/Docking/SDockingTabWell.cpp
--- a/Engine/Source/Runtime/Slate/Private/Framework/Docking/SDockingTabWell.cpp
+++ b/Engine/Source/Runtime/Slate/Private/Framework/Docking/SDockingTabWell.cpp
@@ -464,13 +464,18 @@ void SDockingTabWell::RemoveAndDestroyTab(const TSharedRef<SDockTab>& TabToRemov
 		
 		if ( ensure(ParentTabStack.IsValid()) )
 		{
+			TSharedPtr<SDockingArea> DockAreaPtr = ParentTabStack->GetDockArea();
+
 			ParentTabStack->OnTabClosed( TabToRemove );
 			
 			// We might be closing down an entire dock area, if this is a major tab.
 			// Use this opportunity to save its layout
 			if (RemovalMethod == SDockingNode::TabRemoval_Closed)
 			{
-				ParentTabStack->GetDockArea()->GetTabManager()->GetPrivateApi().OnTabClosing( TabToRemove );
+				if (DockAreaPtr.IsValid())
+				{
+					DockAreaPtr->GetTabManager()->GetPrivateApi().OnTabClosing( TabToRemove );
+				}
 			}
 
 			if (Tabs.Num() == 0)
@@ -481,10 +486,12 @@ void SDockingTabWell::RemoveAndDestroyTab(const TSharedRef<SDockTab>& TabToRemov
 			{
 				RefreshParentContent();
 			}
-		}
-
 
-		GetDockArea()->CleanUp( RemovalMethod );
+			if (DockAreaPtr.IsValid())
+			{
+				DockAreaPtr->CleanUp( RemovalMethod );
+			}
+		}
 	}
 }
 
