Fixed case when spin box could receive keyboard focus during NotifyValueCommitted call
TTP# 339953

diff --git a/Engine/Source/Runtime/Slate/Public/Widgets/Input/SSpinBox.h b/Engine/Source/Runtime/Slate/Public/Widgets/Input/SSpinBox.h
--- a/Engine/Source/Runtime/Slate/Public/Widgets/Input/SSpinBox.h
+++ b/Engine/Source/Runtime/Slate/Public/Widgets/Input/SSpinBox.h
@@ -388,7 +388,7 @@ public:
 
 	virtual FReply OnKeyboardFocusReceived( const FGeometry& MyGeometry, const FKeyboardFocusEvent& InKeyboardFocusEvent ) override
 	{
-		if ( InKeyboardFocusEvent.GetCause() == EKeyboardFocusCause::Keyboard || InKeyboardFocusEvent.GetCause() == EKeyboardFocusCause::SetDirectly )
+		if ( !bDragging && (InKeyboardFocusEvent.GetCause() == EKeyboardFocusCause::Keyboard || InKeyboardFocusEvent.GetCause() == EKeyboardFocusCause::SetDirectly) )
 		{
 			EnterTextMode();
 			return FReply::Handled().SetKeyboardFocus( EditableText.ToSharedRef(), InKeyboardFocusEvent.GetCause() );
