Slate - SCheckbox now allows its Content slot to be dynamically set.

diff --git a/Engine/Source/Runtime/Slate/Private/Widgets/Input/SCheckBox.cpp b/Engine/Source/Runtime/Slate/Private/Widgets/Input/SCheckBox.cpp
--- a/Engine/Source/Runtime/Slate/Private/Widgets/Input/SCheckBox.cpp
+++ b/Engine/Source/Runtime/Slate/Private/Widgets/Input/SCheckBox.cpp
@@ -50,12 +50,17 @@ void SCheckBox::Construct( const SCheckBox::FArguments& InArgs )
 				.Image( this, &SCheckBox::OnGetCheckImage )
 				.ColorAndOpacity( ForegroundColor )
 			]
-			+SHorizontalBox::Slot()
+			+ SHorizontalBox::Slot()
 			.AutoWidth()
 			.Padding( Padding )
 			.VAlign( VAlign_Center )
 			[
-				InArgs._Content.Widget
+				SAssignNew(ContentContainer, SBorder)
+				.BorderImage(FStyleDefaults::GetNoBrush())
+				.Padding(0.0f)
+				[
+					InArgs._Content.Widget
+				]
 			]
 		];
 	}
@@ -65,7 +70,7 @@ void SCheckBox::Construct( const SCheckBox::FArguments& InArgs )
 		// Toggle buttons have a visual appearance that is similar to a Slate button
 		this->ChildSlot
 		[
-			SNew( SBorder )
+			SAssignNew( ContentContainer, SBorder )
 			// Make sure we aren't trying to compute the desired size when the check box is collapsed
 			.Visibility(InArgs._Visibility)
 			// Bind the border background to our method that gets a slate brush for the current state of the control
@@ -350,6 +355,11 @@ void SCheckBox::PlayHoverSound() const
 	FSlateApplication::Get().PlaySound( HoveredSound );
 }
 
+void SCheckBox::SetContent(const TSharedRef< SWidget >& InContent)
+{
+	ContentContainer->SetContent(InContent);
+}
+
 void SCheckBox::SetStyle(const FCheckBoxStyle* InStyle)
 {
 	Style = InStyle;
diff --git a/Engine/Source/Runtime/Slate/Public/Widgets/Input/SCheckBox.h b/Engine/Source/Runtime/Slate/Public/Widgets/Input/SCheckBox.h
--- a/Engine/Source/Runtime/Slate/Public/Widgets/Input/SCheckBox.h
+++ b/Engine/Source/Runtime/Slate/Public/Widgets/Input/SCheckBox.h
@@ -191,6 +191,9 @@ public:
 	/** See the IsChecked attribute */
 	void SetIsChecked(TAttribute<ESlateCheckBoxState::Type> InIsChecked);
 	
+	/** See the Content slot */
+	void SetContent(const TSharedRef< SWidget >& InContent);
+	
 	/** See the Style attribute */
 	void SetStyle(const FCheckBoxStyle* InStyle);
 	
@@ -293,4 +296,8 @@ protected:
 
 	/** The Sound to play when the check box is unchecked */
 	FSlateSound UncheckedSound;
+
+protected:
+	/** When in toggle button mode, this will hold the pointer to the toggle button's border */
+	TSharedPtr<SBorder> ContentContainer;
 };
