Menu Anchors now close the menu they have created upon destruction.
#TTP 337946 - EDITOR: Empty list after renaming Action Mapping

#branch UE4

#proj Runtime.Slate

#add Added destructor to the SMenuAnchor, upon the menu anchor being destroyed it will attempt to close its child menu.

diff --git a/Engine/Source/Runtime/Slate/Private/Widgets/Input/SMenuAnchor.cpp b/Engine/Source/Runtime/Slate/Private/Widgets/Input/SMenuAnchor.cpp
--- a/Engine/Source/Runtime/Slate/Private/Widgets/Input/SMenuAnchor.cpp
+++ b/Engine/Source/Runtime/Slate/Private/Widgets/Input/SMenuAnchor.cpp
@@ -277,4 +277,27 @@ SMenuAnchor::SMenuAnchor()
 	, bDismissedThisTick( false )
 	, Method(CreateNewWindow)
 {
-}
\ No newline at end of file
+}
+
+SMenuAnchor::~SMenuAnchor()
+{
+	TSharedPtr<SWindow> PopupWindow = PopupWindowPtr.Pin();
+	if (PopupWindow.IsValid())
+	{
+		// Close the Popup window.
+		if (Method == CreateNewWindow)
+		{
+			// Request that the popup be closed.
+			PopupWindow->RequestDestroyWindow();
+		}
+		// Close the popup overlay
+		else if (Method == UseCurrentWindow && PopupLayerSlot != NULL)
+		{
+			PopupWindow->RemovePopupLayerSlot(PopupLayerSlot->GetWidget());
+		}
+		
+		// We no longer have a popup open, so reset all the tracking state associated.
+		PopupWindowPtr.Reset();
+		PopupLayerSlot = NULL;
+	}
+}
diff --git a/Engine/Source/Runtime/Slate/Public/Widgets/Input/SMenuAnchor.h b/Engine/Source/Runtime/Slate/Public/Widgets/Input/SMenuAnchor.h
--- a/Engine/Source/Runtime/Slate/Public/Widgets/Input/SMenuAnchor.h
+++ b/Engine/Source/Runtime/Slate/Public/Widgets/Input/SMenuAnchor.h
@@ -38,6 +38,7 @@ public:
 
 	SLATE_END_ARGS()
 
+	virtual ~SMenuAnchor();
 	
 	/**
 	 * Construct this widget
