Fixed widgets having an incorrect interaction offset when using center or right aligned justification

TTP #339914 - Slate: RichTextBlock: Hyperlinks within centered aligned text have their hit box offset

Both SRichTextBlock and SMultiLineEditableText were offsetting the OnPaint calls to compensate for the justification, but they weren't also offsetting the OnArrangeChildren calls.

diff --git a/Engine/Source/Runtime/Slate/Private/Widgets/Text/SMultiLineEditableText.cpp b/Engine/Source/Runtime/Slate/Private/Widgets/Text/SMultiLineEditableText.cpp
--- a/Engine/Source/Runtime/Slate/Private/Widgets/Text/SMultiLineEditableText.cpp
+++ b/Engine/Source/Runtime/Slate/Private/Widgets/Text/SMultiLineEditableText.cpp
@@ -1713,7 +1713,24 @@ FChildren* SMultiLineEditableText::GetChildren()
 
 void SMultiLineEditableText::OnArrangeChildren( const FGeometry& AllottedGeometry, FArrangedChildren& ArrangedChildren ) const
 {
-	TextLayout->ArrangeChildren( AllottedGeometry, ArrangedChildren );
+	const ETextJustify::Type TextJustification = TextLayout->GetJustification();
+
+	if ( TextJustification == ETextJustify::Right )
+	{
+		const FVector2D TextLayoutSize = TextLayout->GetSize();
+		const FVector2D Offset( ( AllottedGeometry.Size.X - TextLayoutSize.X ), 0 );
+		TextLayout->ArrangeChildren( AllottedGeometry.MakeChild( Offset, TextLayoutSize ), ArrangedChildren );
+	}
+	else if ( TextJustification == ETextJustify::Center )
+	{
+		const FVector2D TextLayoutSize = TextLayout->GetSize();
+		const FVector2D Offset( ( AllottedGeometry.Size.X - TextLayoutSize.X ) / 2, 0 );
+		TextLayout->ArrangeChildren( AllottedGeometry.MakeChild( Offset, TextLayoutSize ), ArrangedChildren );
+	}
+	else
+	{
+		TextLayout->ArrangeChildren( AllottedGeometry, ArrangedChildren );
+	}
 }
 
 bool SMultiLineEditableText::SupportsKeyboardFocus() const
diff --git a/Engine/Source/Runtime/Slate/Private/Widgets/Text/SRichTextBlock.cpp b/Engine/Source/Runtime/Slate/Private/Widgets/Text/SRichTextBlock.cpp
--- a/Engine/Source/Runtime/Slate/Private/Widgets/Text/SRichTextBlock.cpp
+++ b/Engine/Source/Runtime/Slate/Private/Widgets/Text/SRichTextBlock.cpp
@@ -115,7 +115,24 @@ FChildren* SRichTextBlock::GetChildren()
 
 void SRichTextBlock::OnArrangeChildren( const FGeometry& AllottedGeometry, FArrangedChildren& ArrangedChildren ) const
 {
-	TextLayout->ArrangeChildren( AllottedGeometry, ArrangedChildren );
+	const ETextJustify::Type TextJustification = TextLayout->GetJustification();
+
+	if ( TextJustification == ETextJustify::Right )
+	{
+		const FVector2D TextLayoutSize = TextLayout->GetSize();
+		const FVector2D Offset( ( AllottedGeometry.Size.X - TextLayoutSize.X ), 0 );
+		TextLayout->ArrangeChildren( AllottedGeometry.MakeChild( Offset, TextLayoutSize ), ArrangedChildren );
+	}
+	else if ( TextJustification == ETextJustify::Center )
+	{
+		const FVector2D TextLayoutSize = TextLayout->GetSize();
+		const FVector2D Offset( ( AllottedGeometry.Size.X - TextLayoutSize.X ) / 2, 0 );
+		TextLayout->ArrangeChildren( AllottedGeometry.MakeChild( Offset, TextLayoutSize ), ArrangedChildren );
+	}
+	else
+	{
+		TextLayout->ArrangeChildren( AllottedGeometry, ArrangedChildren );
+	}
 }
 
 TSharedPtr< ITextDecorator > SRichTextBlock::TryGetDecorator( const TArray< TSharedRef< ITextDecorator > >& InDecorators, const FString& InText, const FTextRunParseResults& TextRun ) const
