Implemented multi-line editable text IME support

TTP# 336464 - Editor: Finish the Multiline Editable Text Block

This is essentially the same as the version for SEditableText, but it also has to deal with converting the FTextLocation points used within a text layout, into a int32 as used by the IME interface (as the IME always works with the document as a single flat string).

To make this easier, there's now a FTextLayout::FTextOffsetLocations helper class which provides functions to let you convert between the two coordiante spaces. A FTextOffsetLocations can either be generated as an optional parameter to FTextLayout::GetAsText, or, via GetTextOffsetLocations if you don't need the string itself.

This also fixes some general IME issues that were affecting both the multi-line and single-line editable text:
 1) TSF based IMEs no longer allow you to move the cursor outside of the composition range (this is consistent with how the other IME backends work).
 2) IMM based IMEs will now remove the composition text from the document if you choose to cancel your composition.
 3) All IMEs will now finish their current composition and start a new one if you move the cursor via the mouse.

diff --git a/Engine/Source/Runtime/Core/Private/Windows/WindowsTextInputMethodSystem.cpp b/Engine/Source/Runtime/Core/Private/Windows/WindowsTextInputMethodSystem.cpp
--- a/Engine/Source/Runtime/Core/Private/Windows/WindowsTextInputMethodSystem.cpp
+++ b/Engine/Source/Runtime/Core/Private/Windows/WindowsTextInputMethodSystem.cpp
@@ -12,12 +12,20 @@ namespace
 {
 	FString GetIMMStringAsFString(HIMC IMMContext, const DWORD StringType)
 	{
-		const LONG StringNeededSize = ::ImmGetCompositionString(IMMContext, StringType, nullptr, 0);
-		TCHAR* RawString = reinterpret_cast<TCHAR*>(malloc(StringNeededSize));
-		::ImmGetCompositionString(IMMContext, StringType, RawString, StringNeededSize);
-		const FString Str(StringNeededSize / sizeof(TCHAR), RawString);
-		free(RawString);
-		return Str;
+		// Get the internal buffer of the string, we're going to use it as scratch space
+		FString OutString;
+		TArray<TCHAR>& OutStringBuffer = OutString.GetCharArray();
+				
+		// Work out the maximum size required and resize the buffer so it can hold enough data
+		const LONG StringNeededSizeBytes = ::ImmGetCompositionString(IMMContext, StringType, nullptr, 0);
+		const int32 StringNeededSizeTCHARs = static_cast<int32>(StringNeededSizeBytes) / sizeof(TCHAR);
+		OutStringBuffer.SetNumUninitialized(StringNeededSizeTCHARs + 1); // +1 for null
+
+		// Get directly into the string buffer, and then null terminate the FString
+		::ImmGetCompositionString(IMMContext, StringType, OutStringBuffer.GetData(), StringNeededSizeBytes);
+		OutStringBuffer[StringNeededSizeTCHARs] = 0;
+
+		return OutString;
 	}
 
 	class FTextInputMethodChangeNotifier : public ITextInputMethodChangeNotifier
@@ -233,6 +241,7 @@ void FWindowsTextInputMethodSystem::BeginIMMComposition()
 	FInternalContext& InternalContext = ContextToInternalContextMap[ActiveContext];
 				
 	InternalContext.IMMContext.IsComposing = true;
+	InternalContext.IMMContext.IsDeactivating = false;
 	ActiveContext->BeginComposition();
 
 	uint32 SelectionBeginIndex = 0;
@@ -254,6 +263,7 @@ void FWindowsTextInputMethodSystem::EndIMMComposition()
 
 	FInternalContext& InternalContext = ContextToInternalContextMap[ActiveContext];
 	InternalContext.IMMContext.IsComposing = false;
+	InternalContext.IMMContext.IsDeactivating = false;
 	ActiveContext->EndComposition();
 }
 
@@ -529,6 +539,7 @@ void FWindowsTextInputMethodSystem::ActivateContext(const TSharedRef<ITextInputM
 
 	// IMM Implementation
 	InternalContext.IMMContext.IsComposing = false;
+	InternalContext.IMMContext.IsDeactivating = false;
 
 	// TSF Implementation
 	FCOMPtr<FTextStoreACP>& TextStore = InternalContext.TSFContext;
@@ -574,8 +585,9 @@ void FWindowsTextInputMethodSystem::DeactivateContext(const TSharedRef<ITextInpu
 	const TSharedPtr<FGenericWindow> GenericWindow = Context->GetWindow();
 	const HWND Hwnd = GenericWindow.IsValid() ? reinterpret_cast<HWND>(GenericWindow->GetOSWindowHandle()) : nullptr;
 	HIMC IMMContext = ::ImmGetContext(Hwnd);
-	// Request the composition is canceled to ensure that the composition input UI is closed, and that a WM_IME_ENDCOMPOSITION message is sent
-	::ImmNotifyIME(IMMContext, NI_COMPOSITIONSTR, CPS_CANCEL, 0);
+	InternalContext.IMMContext.IsDeactivating = true;
+	// Request the composition is completed to ensure that the composition input UI is closed, and that a WM_IME_ENDCOMPOSITION message is sent
+	::ImmNotifyIME(IMMContext, NI_COMPOSITIONSTR, CPS_COMPLETE, 0);
 	::ImmReleaseContext(Hwnd, IMMContext);
 
 	// General Implementation
@@ -652,22 +664,60 @@ int32 FWindowsTextInputMethodSystem::ProcessMessage(HWND hwnd, uint32 msg, WPARA
 
 				UpdateIMMWindowPositions(IMMContext);
 
+				const bool bHasBeenCanceled = !lParam;
 				const bool bHasCompositionStringFlag = !!(lParam & GCS_COMPSTR);
 				const bool bHasResultStringFlag = !!(lParam & GCS_RESULTSTR);
 				const bool bHasNoMoveCaretFlag = !!(lParam & CS_NOMOVECARET);
 				const bool bHasCursorPosFlag = !!(lParam & GCS_CURSORPOS);
 
+				// Canceled, so remove the compositing string
+				if(bHasBeenCanceled)
+				{
+					UE_LOG(LogWindowsTextInputMethodSystem, Verbose, TEXT("WM_IME_COMPOSITION Composition Canceled"));
+
+					// We might have been canceled from a mouse move, so we need to take the current selection so we can restore it properly
+					// otherwise calling SetTextInRange can cause the cursor/selection to jump around
+					uint32 SelectionBeginIndex = 0;
+					uint32 SelectionLength = 0;
+					ITextInputMethodContext::ECaretPosition SelectionCaretPosition = ITextInputMethodContext::ECaretPosition::Ending;
+					ActiveContext->GetSelectionRange(SelectionBeginIndex, SelectionLength, SelectionCaretPosition);
+
+					// Clear Composition
+					ActiveContext->SetTextInRange(InternalContext.IMMContext.CompositionBeginIndex, InternalContext.IMMContext.CompositionLength, TEXT(""));
+					
+					// Restore any previous selection
+					ActiveContext->SetSelectionRange(SelectionBeginIndex, SelectionLength, SelectionCaretPosition);
+
+					EndIMMComposition();
+				}
+
 				// Check Result
 				if(bHasResultStringFlag)
 				{
+					// If we're being deactivated, so we need to take the current selection so we can restore it properly
+					// otherwise calling SetTextInRange can cause the cursor/selection to jump around
+					uint32 SelectionBeginIndex = 0;
+					uint32 SelectionLength = 0;
+					ITextInputMethodContext::ECaretPosition SelectionCaretPosition = ITextInputMethodContext::ECaretPosition::Ending;
+					ActiveContext->GetSelectionRange(SelectionBeginIndex, SelectionLength, SelectionCaretPosition);
+
 					const FString ResultString = GetIMMStringAsFString(IMMContext, GCS_RESULTSTR);
 					UE_LOG(LogWindowsTextInputMethodSystem, Verbose, TEXT("WM_IME_COMPOSITION Result String: %s"), *ResultString);
 
 					// Update Result
 					ActiveContext->SetTextInRange(InternalContext.IMMContext.CompositionBeginIndex, InternalContext.IMMContext.CompositionLength, ResultString);
 
-					// Once we get a result, we're done; set the caret to the end of the result and end the current composition
-					ActiveContext->SetSelectionRange(InternalContext.IMMContext.CompositionBeginIndex + ResultString.Len(), 0, ITextInputMethodContext::ECaretPosition::Ending);
+					if(InternalContext.IMMContext.IsDeactivating)
+					{
+						// Restore any previous selection
+						ActiveContext->SetSelectionRange(SelectionBeginIndex, SelectionLength, SelectionCaretPosition);
+					}
+					else
+					{
+						// Once we get a result, we're done; set the caret to the end of the result and end the current composition
+						ActiveContext->SetSelectionRange(InternalContext.IMMContext.CompositionBeginIndex + ResultString.Len(), 0, ITextInputMethodContext::ECaretPosition::Ending);
+					}
+
 					EndIMMComposition();
 				}
 
diff --git a/Engine/Source/Runtime/Core/Public/Windows/WindowsTextInputMethodSystem.h b/Engine/Source/Runtime/Core/Public/Windows/WindowsTextInputMethodSystem.h
--- a/Engine/Source/Runtime/Core/Public/Windows/WindowsTextInputMethodSystem.h
+++ b/Engine/Source/Runtime/Core/Public/Windows/WindowsTextInputMethodSystem.h
@@ -95,6 +95,7 @@ private:
 		struct
 		{
 			bool IsComposing;
+			bool IsDeactivating;
 			int32 CompositionBeginIndex;
 			uint32 CompositionLength;
 		} IMMContext;
diff --git a/Engine/Source/Runtime/Slate/Private/Framework/Text/TextLayout.cpp b/Engine/Source/Runtime/Slate/Private/Framework/Text/TextLayout.cpp
--- a/Engine/Source/Runtime/Slate/Private/Framework/Text/TextLayout.cpp
+++ b/Engine/Source/Runtime/Slate/Private/Framework/Text/TextLayout.cpp
@@ -950,27 +950,64 @@ bool FTextLayout::IsEmpty() const
 	return (LineModels.Num() == 0 || (LineModels.Num() == 1 && LineModels[0].Text->Len() == 0));
 }
 
-void FTextLayout::GetAsText(FString& DisplayText) const
+void FTextLayout::GetAsText(FString& DisplayText, FTextOffsetLocations* const OutTextOffsetLocations) const
 {
+	GetAsTextAndOffsets(&DisplayText, OutTextOffsetLocations);
+}
+
+void FTextLayout::GetAsText(FText& DisplayText, FTextOffsetLocations* const OutTextOffsetLocations) const
+{
+	FString DisplayString;
+	GetAsText(DisplayString, OutTextOffsetLocations);
+	DisplayText = FText::FromString(DisplayString);
+}
+
+void FTextLayout::GetTextOffsetLocations(FTextOffsetLocations& OutTextOffsetLocations) const
+{
+	GetAsTextAndOffsets(nullptr, &OutTextOffsetLocations);
+}
+
+void FTextLayout::GetAsTextAndOffsets(FString* const OutDisplayText, FTextOffsetLocations* const OutTextOffsetLocations) const
+{
+	int32 DisplayTextLength = 0;
+
+	if (OutTextOffsetLocations)
+	{
+		OutTextOffsetLocations->OffsetData.Reserve(LineModels.Num());
+	}
+
 	for (int LineModelIndex = 0; LineModelIndex < LineModels.Num(); LineModelIndex++)
 	{
 		const FLineModel& LineModel = LineModels[LineModelIndex];
 
+		// Append \n to the end of the previous line
+		if (LineModelIndex > 0)
+		{
+			if (OutDisplayText)
+			{
+				OutDisplayText->AppendChar('\n');
+			}
+			++DisplayTextLength;
+		}
+
+		int32 LineLength = 0;
 		for (int RunIndex = 0; RunIndex < LineModel.Runs.Num(); RunIndex++)
 		{
 			const FRunModel& Run = LineModel.Runs[RunIndex];
-			Run.AppendText(DisplayText);
+			if (OutDisplayText)
+			{
+				Run.AppendText(*OutDisplayText);
+			}
+			LineLength += Run.GetTextRange().Len();
 		}
 
-		DisplayText.AppendChar('\n');
-	}
-}
+		if (OutTextOffsetLocations)
+		{
+			OutTextOffsetLocations->OffsetData.Add(FTextOffsetLocations::FOffsetEntry(DisplayTextLength, LineLength));
+		}
 
-void FTextLayout::GetAsText(FText& DisplayText) const
-{
-	FString DisplayString;
-	GetAsText(DisplayString);
-	DisplayText = FText::FromString(DisplayString);
+		DisplayTextLength += LineLength;
+	}
 }
 
 void GetRangeAsTextFromLine(FString& DisplayText, const FTextLayout::FLineModel& LineModel, const FTextRange& Range)
@@ -1458,3 +1495,39 @@ FTextLayout::FRunModel::FRunModel( const TSharedRef< IRun >& InRun ) : Run( InRu
 {
 
 }
+
+int32 FTextLayout::FTextOffsetLocations::TextLocationToOffset(const FTextLocation& InLocation) const
+{
+	const int32 LineIndex = InLocation.GetLineIndex();
+	if(LineIndex >= 0 && LineIndex < OffsetData.Num())
+	{
+		const FOffsetEntry& OffsetEntry = OffsetData[LineIndex];
+		return OffsetEntry.FlatStringIndex + InLocation.GetOffset();
+	}
+	return INDEX_NONE;
+}
+
+FTextLocation FTextLayout::FTextOffsetLocations::OffsetToTextLocation(const int32 InOffset) const
+{
+	for(int32 OffsetEntryIndex = 0; OffsetEntryIndex < OffsetData.Num(); ++OffsetEntryIndex)
+	{
+		const FOffsetEntry& OffsetEntry = OffsetData[OffsetEntryIndex];
+		const FTextRange FlatLineRange(OffsetEntry.FlatStringIndex, OffsetEntry.FlatStringIndex + OffsetEntry.DocumentLineLength);
+		if(FlatLineRange.InclusiveContains(InOffset))
+		{
+			const int32 OffsetWithinLine = InOffset - OffsetEntry.FlatStringIndex;
+			return FTextLocation(OffsetEntryIndex, OffsetWithinLine);
+		}
+	}
+	return FTextLocation();
+}
+
+int32 FTextLayout::FTextOffsetLocations::GetTextLength() const
+{
+	if(OffsetData.Num())
+	{
+		const FOffsetEntry& OffsetEntry = OffsetData.Last();
+		return OffsetEntry.FlatStringIndex + OffsetEntry.DocumentLineLength;
+	}
+	return 0;
+}
diff --git a/Engine/Source/Runtime/Slate/Private/Widgets/Input/SEditableText.cpp b/Engine/Source/Runtime/Slate/Private/Widgets/Input/SEditableText.cpp
--- a/Engine/Source/Runtime/Slate/Private/Widgets/Input/SEditableText.cpp
+++ b/Engine/Source/Runtime/Slate/Private/Widgets/Input/SEditableText.cpp
@@ -295,10 +295,10 @@ void SEditableText::Tick( const FGeometry& AllottedGeometry, const double InCurr
 bool SEditableText::FTextInputMethodContext::IsReadOnly()
 {
 	bool Return = true;
-	if(OwningWidget.IsValid())
+	const TSharedPtr<SEditableText> OwningWidgetPtr = OwningWidget.Pin();
+	if(OwningWidgetPtr.IsValid())
 	{
-		const TSharedRef<SEditableText> This = OwningWidget.Pin().ToSharedRef();
-		Return = This->GetIsReadOnly();
+		Return = OwningWidgetPtr->GetIsReadOnly();
 	}
 	return Return;
 }
@@ -306,10 +306,10 @@ bool SEditableText::FTextInputMethodContext::IsReadOnly()
 uint32 SEditableText::FTextInputMethodContext::GetTextLength()
 {
 	uint32 TextLength = 0;
-	if(OwningWidget.IsValid())
+	const TSharedPtr<SEditableText> OwningWidgetPtr = OwningWidget.Pin();
+	if(OwningWidgetPtr.IsValid())
 	{
-		const TSharedRef<SEditableText> This = OwningWidget.Pin().ToSharedRef();
-		const FString& EditedText = This->EditedText.ToString();
+		const FString& EditedText = OwningWidgetPtr->EditedText.ToString();
 		TextLength = EditedText.Len();
 	}
 	return TextLength;
@@ -317,25 +317,25 @@ uint32 SEditableText::FTextInputMethodContext::GetTextLength()
 
 void SEditableText::FTextInputMethodContext::GetSelectionRange(uint32& BeginIndex, uint32& Length, ECaretPosition& OutCaretPosition)
 {
-	if(OwningWidget.IsValid())
+	const TSharedPtr<SEditableText> OwningWidgetPtr = OwningWidget.Pin();
+	if(OwningWidgetPtr.IsValid())
 	{
-		const TSharedRef<SEditableText> This = OwningWidget.Pin().ToSharedRef();
-		if(This->Selection.StartIndex != This->Selection.FinishIndex)
+		if(OwningWidgetPtr->Selection.StartIndex != OwningWidgetPtr->Selection.FinishIndex)
 		{
-			BeginIndex = This->Selection.GetMinIndex();
-			Length = This->Selection.GetMaxIndex() - BeginIndex;
+			BeginIndex = OwningWidgetPtr->Selection.GetMinIndex();
+			Length = OwningWidgetPtr->Selection.GetMaxIndex() - BeginIndex;
 		}
 		else
 		{
-			BeginIndex = This->CaretPosition;
+			BeginIndex = OwningWidgetPtr->CaretPosition;
 			Length = 0;
 		}
 
-		if(This->CaretPosition == BeginIndex + Length)
+		if(OwningWidgetPtr->CaretPosition == BeginIndex + Length)
 		{
 			OutCaretPosition = ITextInputMethodContext::ECaretPosition::Ending;
 		}
-		else if(This->CaretPosition == BeginIndex)
+		else if(OwningWidgetPtr->CaretPosition == BeginIndex)
 		{
 			OutCaretPosition = ITextInputMethodContext::ECaretPosition::Beginning;
 		}
@@ -344,9 +344,9 @@ void SEditableText::FTextInputMethodContext::GetSelectionRange(uint32& BeginInde
 
 void SEditableText::FTextInputMethodContext::SetSelectionRange(const uint32 BeginIndex, const uint32 Length, const ECaretPosition InCaretPosition)
 {
-	if(OwningWidget.IsValid())
+	const TSharedPtr<SEditableText> OwningWidgetPtr = OwningWidget.Pin();
+	if(OwningWidgetPtr.IsValid())
 	{
-		const TSharedRef<SEditableText> This = OwningWidget.Pin().ToSharedRef();
 		const uint32 MinIndex = BeginIndex;
 		const uint32 MaxIndex = MinIndex + Length;
 
@@ -369,45 +369,45 @@ void SEditableText::FTextInputMethodContext::SetSelectionRange(const uint32 Begi
 			break;
 		}
 
-		This->SetCaretPosition(SelectionEndIndex);
-		This->ClearSelection();
-		This->SelectText(SelectionBeginIndex);
+		OwningWidgetPtr->SetCaretPosition(SelectionEndIndex);
+		OwningWidgetPtr->ClearSelection();
+		OwningWidgetPtr->SelectText(SelectionBeginIndex);
 	}
 }
 
 void SEditableText::FTextInputMethodContext::GetTextInRange(const uint32 BeginIndex, const uint32 Length, FString& OutString)
 {
-	if(OwningWidget.IsValid())
+	const TSharedPtr<SEditableText> OwningWidgetPtr = OwningWidget.Pin();
+	if(OwningWidgetPtr.IsValid())
 	{
-		const TSharedRef<SEditableText> This = OwningWidget.Pin().ToSharedRef();
-		const FString& EditedText = This->EditedText.ToString();
+		const FString& EditedText = OwningWidgetPtr->EditedText.ToString();
 		OutString = EditedText.Mid(BeginIndex, Length);
 	}
 }
 
 void SEditableText::FTextInputMethodContext::SetTextInRange(const uint32 BeginIndex, const uint32 Length, const FString& InString)
 {
-	if(OwningWidget.IsValid())
+	const TSharedPtr<SEditableText> OwningWidgetPtr = OwningWidget.Pin();
+	if(OwningWidgetPtr.IsValid())
 	{
-		const TSharedRef<SEditableText> This = OwningWidget.Pin().ToSharedRef();
-		const FString& OldText = This->EditedText.ToString();
+		const FString& OldText = OwningWidgetPtr->EditedText.ToString();
 		
 		// We don't use Start/FinishEditing text here because the whole IME operation handles that.
 		// Also, we don't want to support undo for individual characters added during an IME context
-		This->EditedText = FText::FromString(OldText.Left( BeginIndex ) + InString + OldText.Mid( BeginIndex + Length ));
-		This->SaveText();
+		OwningWidgetPtr->EditedText = FText::FromString(OldText.Left( BeginIndex ) + InString + OldText.Mid( BeginIndex + Length ));
+		OwningWidgetPtr->SaveText();
 
-		This->OnTextChanged.ExecuteIfBound(This->EditedText);
+		OwningWidgetPtr->OnTextChanged.ExecuteIfBound(OwningWidgetPtr->EditedText);
 	}
 }
 
 int32 SEditableText::FTextInputMethodContext::GetCharacterIndexFromPoint(const FVector2D& Point)
 {
 	int32 CharacterIndex = INDEX_NONE;
-	if(OwningWidget.IsValid())
+	const TSharedPtr<SEditableText> OwningWidgetPtr = OwningWidget.Pin();
+	if(OwningWidgetPtr.IsValid())
 	{
-		const TSharedRef<SEditableText> This = OwningWidget.Pin().ToSharedRef();
-		CharacterIndex = This->FindClickedCharacterIndex(CachedGeometry, Point);
+		CharacterIndex = OwningWidgetPtr->FindClickedCharacterIndex(CachedGeometry, Point);
 	}
 	return CharacterIndex;
 }
@@ -415,19 +415,18 @@ int32 SEditableText::FTextInputMethodContext::GetCharacterIndexFromPoint(const F
 bool SEditableText::FTextInputMethodContext::GetTextBounds(const uint32 BeginIndex, const uint32 Length, FVector2D& Position, FVector2D& Size)
 {
 	bool IsClipped = false;
-	if(OwningWidget.IsValid())
+	const TSharedPtr<SEditableText> OwningWidgetPtr = OwningWidget.Pin();
+	if(OwningWidgetPtr.IsValid())
 	{
-		const TSharedRef<SEditableText> This = OwningWidget.Pin().ToSharedRef();
-
-		const FSlateFontInfo& FontInfo = This->Font.Get();
-		const FString VisibleText = This->GetStringToRender();
+		const FSlateFontInfo& FontInfo = OwningWidgetPtr->Font.Get();
+		const FString VisibleText = OwningWidgetPtr->GetStringToRender();
 		const TSharedRef< FSlateFontMeasure > FontMeasureService = FSlateApplication::Get().GetRenderer()->GetFontMeasureService();
 
 		const float TextLeft = FontMeasureService->Measure( VisibleText.Left( BeginIndex ), FontInfo ).X;
 		const float TextRight = FontMeasureService->Measure( VisibleText.Left( BeginIndex + Length ), FontInfo ).X;
 
-		const float ScrollAreaLeft = This->ScrollHelper.ToScrollerSpace( FVector2D::ZeroVector ).X;
-		const float ScrollAreaRight = This->ScrollHelper.ToScrollerSpace( CachedGeometry.Size ).X;
+		const float ScrollAreaLeft = OwningWidgetPtr->ScrollHelper.ToScrollerSpace( FVector2D::ZeroVector ).X;
+		const float ScrollAreaRight = OwningWidgetPtr->ScrollHelper.ToScrollerSpace( CachedGeometry.Size ).X;
 
 		const int32 FirstVisibleCharacter = FontMeasureService->FindCharacterIndexAtOffset( VisibleText, FontInfo, ScrollAreaLeft );
 		const int32 LastVisibleCharacter = FontMeasureService->FindCharacterIndexAtOffset( VisibleText, FontInfo, ScrollAreaRight );
@@ -442,17 +441,17 @@ bool SEditableText::FTextInputMethodContext::GetTextBounds(const uint32 BeginInd
 
 		const float TextVertOffset = EditableTextDefs::TextVertOffsetPercent * FontHeight;
 
-		Position = CachedGeometry.LocalToAbsolute( This->ScrollHelper.FromScrollerSpace( FVector2D( TextLeft, DrawPositionY + TextVertOffset ) ) );
-		Size = CachedGeometry.Scale * This->ScrollHelper.SizeFromScrollerSpace( FVector2D( TextRight - TextLeft, CachedGeometry.Size.Y ) );
+		Position = CachedGeometry.LocalToAbsolute( OwningWidgetPtr->ScrollHelper.FromScrollerSpace( FVector2D( TextLeft, DrawPositionY + TextVertOffset ) ) );
+		Size = CachedGeometry.Scale * OwningWidgetPtr->ScrollHelper.SizeFromScrollerSpace( FVector2D( TextRight - TextLeft, CachedGeometry.Size.Y ) );
 	}
 	return IsClipped;
 }
 
 void SEditableText::FTextInputMethodContext::GetScreenBounds(FVector2D& Position, FVector2D& Size)
 {
-	if(OwningWidget.IsValid())
+	const TSharedPtr<SEditableText> OwningWidgetPtr = OwningWidget.Pin();
+	if(OwningWidgetPtr.IsValid())
 	{
-		const TSharedRef<SEditableText> This = OwningWidget.Pin().ToSharedRef();
 		Position = CachedGeometry.AbsolutePosition;
 		Size = CachedGeometry.GetDrawSize();
 	}
@@ -461,10 +460,10 @@ void SEditableText::FTextInputMethodContext::GetScreenBounds(FVector2D& Position
 TSharedPtr<FGenericWindow> SEditableText::FTextInputMethodContext::GetWindow()
 {
 	TSharedPtr<FGenericWindow> Window;
-	if(OwningWidget.IsValid())
+	const TSharedPtr<SEditableText> OwningWidgetPtr = OwningWidget.Pin();
+	if(OwningWidgetPtr.IsValid())
 	{
-		const TSharedRef<SEditableText> This = OwningWidget.Pin().ToSharedRef();
-		const TSharedPtr<SWindow> SlateWindow = FSlateApplication::Get().FindWidgetWindow( This );
+		const TSharedPtr<SWindow> SlateWindow = FSlateApplication::Get().FindWidgetWindow( OwningWidgetPtr.ToSharedRef() );
 		Window = SlateWindow.IsValid() ? SlateWindow->GetNativeWindow() : nullptr;
 	}
 	return Window;
@@ -476,10 +475,10 @@ void SEditableText::FTextInputMethodContext::BeginComposition()
 	{
 		IsComposing = true;
 
-		if(OwningWidget.IsValid())
+		const TSharedPtr<SEditableText> OwningWidgetPtr = OwningWidget.Pin();
+		if(OwningWidgetPtr.IsValid())
 		{
-			const TSharedRef<SEditableText> This = OwningWidget.Pin().ToSharedRef();
-			This->StartChangingText();
+			OwningWidgetPtr->StartChangingText();
 		}
 	}
 }
@@ -496,10 +495,10 @@ void SEditableText::FTextInputMethodContext::EndComposition()
 {
 	if(IsComposing)
 	{
-		if(OwningWidget.IsValid())
+		const TSharedPtr<SEditableText> OwningWidgetPtr = OwningWidget.Pin();
+		if(OwningWidgetPtr.IsValid())
 		{
-			const TSharedRef<SEditableText> This = OwningWidget.Pin().ToSharedRef();
-			This->FinishChangingText();
+			OwningWidgetPtr->FinishChangingText();
 		}
 
 		IsComposing = false;
@@ -636,6 +635,12 @@ static int32 OffsetToTestBasedOnDirection( int8 Direction )
 
 FReply SEditableText::MoveCursor( ECursorMoveMethod::Type Method, const FVector2D& Direction, ECursorAction::Type Action )
 {
+	// We can't use the keyboard to move the cursor while composing, as the IME has control over it
+	if(TextInputMethodContext->IsComposing && Method != ECursorMoveMethod::ScreenPosition)
+	{
+		return FReply::Handled();
+	}
+
 	bool bAllowMoveCursor = true;
 	int32 OldCaretPosition = CaretPosition;
 	int32 NewCaretPosition = OldCaretPosition;
@@ -719,6 +724,18 @@ FReply SEditableText::MoveCursor( ECursorMoveMethod::Type Method, const FVector2
 		ClearSelection();
 	}
 
+	// If we've moved the cursor while composing, we need to end the current composition session
+	// Note: You should only be able to do this via the mouse due to the check at the top of this function
+	if(TextInputMethodContext->IsComposing)
+	{
+		ITextInputMethodSystem* const TextInputMethodSystem = FSlateApplication::Get().GetTextInputMethodSystem();
+		if(TextInputMethodSystem)
+		{
+			TextInputMethodSystem->DeactivateContext(TextInputMethodContext.ToSharedRef());
+			TextInputMethodSystem->ActivateContext(TextInputMethodContext.ToSharedRef());
+		}
+	}
+
 	return FReply::Handled();
 }
 
@@ -1243,10 +1260,12 @@ int32 SEditableText::OnPaint( const FGeometry& AllottedGeometry, const FSlateRec
 	}
 
 	// Draw composition background
-	if( TextInputMethodContext->IsComposing )
+	// We only draw the composition highlight if the cursor is within the composition range
+	const FTextRange CompositionRange(TextInputMethodContext->CompositionBeginIndex, TextInputMethodContext->CompositionBeginIndex + TextInputMethodContext->CompositionLength);
+	if( TextInputMethodContext->IsComposing && CompositionRange.InclusiveContains(CaretPosition) )
 	{
-		const float CompositionLeft = FontMeasureService->Measure( VisibleText.Left( TextInputMethodContext->CompositionBeginIndex ), FontInfo ).X;
-		const float CompositionRight = FontMeasureService->Measure( VisibleText.Left( TextInputMethodContext->CompositionBeginIndex + TextInputMethodContext->CompositionLength ), FontInfo ).X;
+		const float CompositionLeft = FontMeasureService->Measure( VisibleText.Left( CompositionRange.BeginIndex ), FontInfo ).X;
+		const float CompositionRight = FontMeasureService->Measure( VisibleText.Left( CompositionRange.EndIndex ), FontInfo ).X;
 		const float CompositionTop = 0.0f;
 		const float CompositionBottom = FontMeasureService->Measure( VisibleText.Mid( TextInputMethodContext->CompositionBeginIndex, TextInputMethodContext->CompositionLength ), FontInfo ).Y;
 
diff --git a/Engine/Source/Runtime/Slate/Private/Widgets/Text/SMultiLineEditableText.cpp b/Engine/Source/Runtime/Slate/Private/Widgets/Text/SMultiLineEditableText.cpp
--- a/Engine/Source/Runtime/Slate/Private/Widgets/Text/SMultiLineEditableText.cpp
+++ b/Engine/Source/Runtime/Slate/Private/Widgets/Text/SMultiLineEditableText.cpp
@@ -87,12 +87,12 @@ FChildren* SMultiLineEditableText::FSlateCursorRunHighlighter::GetChildren()
 
 int32 SMultiLineEditableText::FSlateCursorRunHighlighter::OnPaint( const FTextLayout::FLineView& Line, const TSharedRef< ISlateRun >& Run, const TSharedRef< ILayoutBlock >& Block, const FTextBlockStyle& DefaultStyle, const FGeometry& AllottedGeometry, const FSlateRect& MyClippingRect, FSlateWindowElementList& OutDrawElements, int32 LayerId, const FWidgetStyle& InWidgetStyle, bool bParentEnabled ) const 
 {
-	LayerId = OnPaintCursor(Line, Run, Block, DefaultStyle, AllottedGeometry, MyClippingRect, OutDrawElements, LayerId, InWidgetStyle, bParentEnabled, CursorInfo->GetCursorAlignment());
+	LayerId = OnPaintCursor(Line, Run, Block, DefaultStyle, AllottedGeometry, MyClippingRect, OutDrawElements, LayerId, InWidgetStyle, bParentEnabled, 0.0f, CursorInfo->GetCursorAlignment());
 
 	return Run->OnPaint( Line, Block, DefaultStyle, AllottedGeometry, MyClippingRect, OutDrawElements, LayerId, InWidgetStyle, bParentEnabled );
 }
 
-int32 SMultiLineEditableText::FSlateCursorRunHighlighter::OnPaintCursor( const FTextLayout::FLineView& Line, const TSharedRef< ISlateRun >& Run, const TSharedRef< ILayoutBlock >& Block, const FTextBlockStyle& DefaultStyle, const FGeometry& AllottedGeometry, const FSlateRect& MyClippingRect, FSlateWindowElementList& OutDrawElements, int32 LayerId, const FWidgetStyle& InWidgetStyle, const bool bParentEnabled, const ECursorAlignment InCursorAlignment ) const
+int32 SMultiLineEditableText::FSlateCursorRunHighlighter::OnPaintCursor( const FTextLayout::FLineView& Line, const TSharedRef< ISlateRun >& Run, const TSharedRef< ILayoutBlock >& Block, const FTextBlockStyle& DefaultStyle, const FGeometry& AllottedGeometry, const FSlateRect& MyClippingRect, FSlateWindowElementList& OutDrawElements, int32 LayerId, const FWidgetStyle& InWidgetStyle, const bool bParentEnabled, const float InCursorOffset, const ECursorAlignment InCursorAlignment ) const
 {
 	FVector2D Location( Block->GetLocationOffset() );
 	Location.Y = Line.Offset.Y;
@@ -113,15 +113,16 @@ int32 SMultiLineEditableText::FSlateCursorRunHighlighter::OnPaintCursor( const F
 	CursorColorAndOpacity.A = CursorOpacity;
 
 	// @todo: Slate Styles - make this brush part of the widget style
-	const FSlateBrush* StyleInfo = FCoreStyle::Get().GetBrush("EditableText.SelectionBackground");
+	const FSlateBrush* CursorBrush = FCoreStyle::Get().GetBrush("EditableText.SelectionBackground");
 
-	const FVector2D OptionalWidth = InCursorAlignment == ECursorAlignment::Right ? FVector2D(Block->GetSize().X, 0) : FVector2D::ZeroVector;
+	FVector2D OptionalWidth = InCursorAlignment == ECursorAlignment::Right ? FVector2D(Block->GetSize().X, 0) : FVector2D::ZeroVector;
+	OptionalWidth += FVector2D(InCursorOffset, 0.0f);
 
 	FSlateDrawElement::MakeBox(
 		OutDrawElements,
 		LayerId,
 		FPaintGeometry(AllottedGeometry.AbsolutePosition + Location + OptionalWidth, FVector2D(CursorWidth * AllottedGeometry.Scale, Line.TextSize.Y), AllottedGeometry.Scale),
-		StyleInfo,
+		CursorBrush,
 		MyClippingRect,
 		bParentEnabled ? ESlateDrawEffect::None : ESlateDrawEffect::DisabledEffect,
 		CursorColorAndOpacity*InWidgetStyle.GetColorAndOpacityTint());
@@ -147,36 +148,32 @@ int32 SMultiLineEditableText::FSlateSelectionRunHighlighter::OnPaint( const FTex
 
 	const FColor SelectionBackgroundColorAndOpacity = ((FLinearColor::White - DefaultStyle.ColorAndOpacity.GetColor(InWidgetStyle))*0.5f + FLinearColor(-0.2f, -0.05f, 0.15f)) * InWidgetStyle.GetColorAndOpacityTint();
 
-	float HighlightWidth = Block->GetSize().X;
-	if (HighlightWidth == 0)
-	{
-		const TSharedRef< FSlateFontMeasure > FontMeasureService = FSlateApplication::Get().GetRenderer()->GetFontMeasureService();
-		const FVector2D SpaceSize = FontMeasureService->Measure(FString(TEXT("  ")), DefaultStyle.Font);
-		HighlightWidth = SpaceSize.X;
+	const float HighlightWidth = Block->GetSize().X;
+	if (HighlightWidth)
+	{
+		// Draw the actual highlight rectangle
+		FSlateDrawElement::MakeBox(
+			OutDrawElements,
+			++LayerId,
+			FPaintGeometry(AllottedGeometry.AbsolutePosition + Location, FVector2D(HighlightWidth, Line.Size.Y), AllottedGeometry.Scale),
+			&DefaultStyle.HighlightShape,
+			MyClippingRect,
+			bParentEnabled ? ESlateDrawEffect::None : ESlateDrawEffect::DisabledEffect,
+			SelectionBackgroundColorAndOpacity
+		);
 	}
 
-	// Draw the actual highlight rectangle
-	FSlateDrawElement::MakeBox(
-		OutDrawElements,
-		++LayerId,
-		FPaintGeometry(AllottedGeometry.AbsolutePosition + Location, FVector2D(HighlightWidth, Line.Size.Y), AllottedGeometry.Scale),
-		&DefaultStyle.HighlightShape,
-		MyClippingRect,
-		bParentEnabled ? ESlateDrawEffect::None : ESlateDrawEffect::DisabledEffect,
-		SelectionBackgroundColorAndOpacity
-	);
-
 	//FLinearColor InvertedForeground = FLinearColor::White - InWidgetStyle.GetForegroundColor();
 	//InvertedForeground.A = InWidgetStyle.GetForegroundColor().A;
 
 	//FWidgetStyle WidgetStyle( InWidgetStyle );
 	//WidgetStyle.SetForegroundColor( InvertedForeground );
 
-	// If the cursor is within the range of this line view, we need to draw it
+	// If the cursor is within the range of this block, we need to draw it
 	const FTextLocation CursorLocation = CursorInfo->GetCursorLocation();
-	if(Line.ModelIndex == CursorLocation.GetLineIndex() && Line.Range.InclusiveContains(CursorLocation.GetOffset()))
+	if(Line.ModelIndex == CursorLocation.GetLineIndex() && Block->GetTextRange().InclusiveContains(CursorLocation.GetOffset()))
 	{
-		LayerId = OnPaintCursor(Line, Run, Block, DefaultStyle, AllottedGeometry, MyClippingRect, OutDrawElements, LayerId, InWidgetStyle, bParentEnabled, CursorAlignment);
+		LayerId = OnPaintCursor(Line, Run, Block, DefaultStyle, AllottedGeometry, MyClippingRect, OutDrawElements, LayerId, InWidgetStyle, bParentEnabled, 0.0f, CursorAlignment);
 	}
 
 	return Run->OnPaint( Line, Block, DefaultStyle, AllottedGeometry, MyClippingRect, OutDrawElements, LayerId, InWidgetStyle, bParentEnabled );
@@ -187,6 +184,71 @@ TSharedRef< SMultiLineEditableText::FSlateSelectionRunHighlighter > SMultiLineEd
 	return MakeShareable( new SMultiLineEditableText::FSlateSelectionRunHighlighter(InCursorInfo) );
 }
 
+SMultiLineEditableText::FSlateCompositionRunHighlighter::FSlateCompositionRunHighlighter(const FCursorInfo* InCursorInfo)
+	: SMultiLineEditableText::FSlateCursorRunHighlighter(InCursorInfo)
+{
+}
+
+int32 SMultiLineEditableText::FSlateCompositionRunHighlighter::OnPaint( const FTextLayout::FLineView& Line, const TSharedRef< ISlateRun >& Run, const TSharedRef< ILayoutBlock >& Block, const FTextBlockStyle& DefaultStyle, const FGeometry& AllottedGeometry, const FSlateRect& MyClippingRect, FSlateWindowElementList& OutDrawElements, int32 LayerId, const FWidgetStyle& InWidgetStyle, bool bParentEnabled ) const 
+{
+	FVector2D Location( Block->GetLocationOffset() );
+	Location.Y = Line.Offset.Y;
+
+	const float CompositionWidth = Block->GetSize().X;
+	if (CompositionWidth)
+	{
+		const FLinearColor LineColorAndOpacity = InWidgetStyle.GetForegroundColor();
+
+		// @todo: Slate Styles - make this brush part of the widget style
+		const FSlateBrush* CompositionBrush = FCoreStyle::Get().GetBrush("EditableText.CompositionBackground");
+
+		// Draw composition background
+		FSlateDrawElement::MakeBox(
+			OutDrawElements,
+			++LayerId,
+			FPaintGeometry(AllottedGeometry.AbsolutePosition + Location, FVector2D(CompositionWidth, Line.TextSize.Y), AllottedGeometry.Scale),
+			CompositionBrush,
+			MyClippingRect,
+			bParentEnabled ? ESlateDrawEffect::None : ESlateDrawEffect::DisabledEffect,
+			LineColorAndOpacity * InWidgetStyle.GetColorAndOpacityTint()
+		);
+	}
+
+	// If the cursor is within the range of this block, we need to draw it
+	const FTextLocation CursorInteractionLocation = CursorInfo->GetCursorInteractionLocation();
+	if(Line.ModelIndex == CursorInteractionLocation.GetLineIndex() && Block->GetTextRange().InclusiveContains(CursorInteractionLocation.GetOffset()))
+	{
+		// Unlike selection where the cursor will always be at one end of the block, composition may place the cursor anywhere within the block,
+		// so we need to calculate the offset (in scaled screen space) required from the start of the block to the position of the cursor
+		// We always do this check using the interaction position of the cursor, as this ensures we draw the cursor at the correct place, 
+		// even when it's to the right of a character (which is why we also pass ECursorAlignment::Left into OnPaintCursor, otherwise we'd doubly compensate)
+
+		// Get the text for the current run
+		FString RunText;
+		Run->AppendText(RunText);
+
+		// Trim that text down to the text for this block
+		const FTextRange BlockTextRange = Block->GetTextRange();
+		const FString BlockText = RunText.Mid(BlockTextRange.BeginIndex, BlockTextRange.EndIndex - BlockTextRange.BeginIndex);
+
+		// Then trim that down so we have the text up-to the interaction point of the cursor
+		const FString TextUpToCursor = BlockText.Left(CursorInteractionLocation.GetOffset() - BlockTextRange.BeginIndex);
+
+		// Finally, measure that text so we can get an accurate offset for the cursor position within this block
+		const TSharedRef< FSlateFontMeasure > FontMeasureService = FSlateApplication::Get().GetRenderer()->GetFontMeasureService();
+		const FVector2D CursorOffset = FontMeasureService->Measure(TextUpToCursor, DefaultStyle.Font, AllottedGeometry.Scale);
+
+		LayerId = OnPaintCursor(Line, Run, Block, DefaultStyle, AllottedGeometry, MyClippingRect, OutDrawElements, LayerId, InWidgetStyle, bParentEnabled, CursorOffset.X, ECursorAlignment::Left);
+	}
+
+	return Run->OnPaint( Line, Block, DefaultStyle, AllottedGeometry, MyClippingRect, OutDrawElements, LayerId, InWidgetStyle, bParentEnabled );
+}
+
+TSharedRef< SMultiLineEditableText::FSlateCompositionRunHighlighter > SMultiLineEditableText::FSlateCompositionRunHighlighter::Create(const FCursorInfo* InCursorInfo)
+{
+	return MakeShareable( new SMultiLineEditableText::FSlateCompositionRunHighlighter(InCursorInfo) );
+}
+
 SMultiLineEditableText::SMultiLineEditableText()
 	: PreferredCursorScreenOffsetInLine(0)
 	, bSelectAllTextWhenFocused(false)
@@ -201,6 +263,15 @@ SMultiLineEditableText::SMultiLineEditableText()
 	
 }
 
+SMultiLineEditableText::~SMultiLineEditableText()
+{
+	ITextInputMethodSystem* const TextInputMethodSystem = FSlateApplication::IsInitialized() ? FSlateApplication::Get().GetTextInputMethodSystem() : nullptr;
+	if(TextInputMethodSystem)
+	{
+		TextInputMethodSystem->UnregisterContext(TextInputMethodContext.ToSharedRef());
+	}
+}
+
 void SMultiLineEditableText::Construct( const FArguments& InArgs )
 {
 	CursorInfo = FCursorInfo();
@@ -228,6 +299,7 @@ void SMultiLineEditableText::Construct( const FArguments& InArgs )
 
 	CursorRunHighlighter = FSlateCursorRunHighlighter::Create(&CursorInfo);
 	SelectionRunHighlighter = FSlateSelectionRunHighlighter::Create(&CursorInfo);
+	CompositionRunHighlighter = FSlateCompositionRunHighlighter::Create(&CursorInfo);
 	TextLayout = FSlateTextLayout::Create();
 
 	BoundText = InArgs._Text;
@@ -278,6 +350,17 @@ void SMultiLineEditableText::Construct( const FArguments& InArgs )
 	// build context menu extender
 	MenuExtender = MakeShareable(new FExtender);
 	MenuExtender->AddMenuExtension("EditText", EExtensionHook::Before, TSharedPtr<FUICommandList>(), InArgs._ContextMenuExtender);
+
+	TextInputMethodContext = MakeShareable( new FTextInputMethodContext( StaticCastSharedRef<SMultiLineEditableText>(AsShared()) ) );
+	ITextInputMethodSystem* const TextInputMethodSystem = FSlateApplication::Get().GetTextInputMethodSystem();
+	if(TextInputMethodSystem)
+	{
+		TextInputMethodChangeNotifier = TextInputMethodSystem->RegisterContext(TextInputMethodContext.ToSharedRef());
+	}
+	if(TextInputMethodChangeNotifier.IsValid())
+	{
+		TextInputMethodChangeNotifier->NotifyLayoutChanged(ITextInputMethodChangeNotifier::ELayoutChangeType::Created);
+	}
 }
 
 void SMultiLineEditableText::SetText(const TAttribute< FText >& InText)
@@ -304,8 +387,7 @@ void SMultiLineEditableText::SetText(const TAttribute< FText >& InText)
 
 bool SMultiLineEditableText::SetEditableText(const FText& TextToSet, const bool bForce)
 {
-	FText EditedText;
-	TextLayout->GetAsText(EditedText);
+	const FText EditedText = GetEditableText();
 
 	const bool HasTextChanged = !EditedText.ToString().Equals(TextToSet.ToString(), ESearchCase::CaseSensitive);
 	if (HasTextChanged || bForce)
@@ -384,10 +466,29 @@ bool SMultiLineEditableText::SetEditableText(const FText& TextToSet, const bool
 	return false;
 }
 
+FText SMultiLineEditableText::GetEditableText() const
+{
+	FText EditedText;
+	TextLayout->GetAsText(EditedText);
+	return EditedText;
+}
+
 FReply SMultiLineEditableText::OnKeyboardFocusReceived( const FGeometry& MyGeometry, const FKeyboardFocusEvent& InKeyboardFocusEvent )
 {
-	UpdateCursorHighlight();
-	return SWidget::OnKeyboardFocusReceived( MyGeometry, InKeyboardFocusEvent );
+	// Skip the focus received code if it's due to the context menu closing
+	if ( !ContextMenuWindow.IsValid() )
+	{
+		ITextInputMethodSystem* const TextInputMethodSystem = FSlateApplication::Get().GetTextInputMethodSystem();
+		if(TextInputMethodSystem)
+		{
+			TextInputMethodSystem->ActivateContext(TextInputMethodContext.ToSharedRef());
+		}
+
+		UpdateCursorHighlight();
+		return SWidget::OnKeyboardFocusReceived( MyGeometry, InKeyboardFocusEvent );
+	}
+
+	return FReply::Handled();
 }
 
 void SMultiLineEditableText::OnKeyboardFocusLost( const FKeyboardFocusEvent& InKeyboardFocusEvent )
@@ -395,18 +496,18 @@ void SMultiLineEditableText::OnKeyboardFocusLost( const FKeyboardFocusEvent& InK
 	// Skip the focus lost code if it's due to the context menu opening
 	if (!ContextMenuWindow.IsValid())
 	{
-		//ITextInputMethodSystem* const TextInputMethodSystem = FSlateApplication::Get().GetTextInputMethodSystem();
-		//if (TextInputMethodSystem)
-		//{
-		//	TextInputMethodSystem->DeactivateContext(TextInputMethodContext.ToSharedRef());
-		//}
+		ITextInputMethodSystem* const TextInputMethodSystem = FSlateApplication::Get().GetTextInputMethodSystem();
+		if (TextInputMethodSystem)
+		{
+			TextInputMethodSystem->DeactivateContext(TextInputMethodContext.ToSharedRef());
+		}
 
 		// When focus is lost let anyone who is interested that text was committed
 		if (InKeyboardFocusEvent.GetCause() != EKeyboardFocusCause::WindowActivate) //Don't clear selection when activating a new window or cant copy and paste on right click )
 		{
 			ClearSelection();
 		}
-
+		
 		// See if user explicitly tabbed away or moved focus
 		ETextCommit::Type TextAction;
 		switch (InKeyboardFocusEvent.GetCause())
@@ -428,8 +529,7 @@ void SMultiLineEditableText::OnKeyboardFocusLost( const FKeyboardFocusEvent& InK
 		//Always clear the local undo chain on commit.
 		ClearUndoStates();
 
-		FText EditedText;
-		TextLayout->GetAsText(EditedText);
+		const FText EditedText = GetEditableText();
 
 		OnTextCommitted.ExecuteIfBound(EditedText, TextAction);
 		RemoveCursorHighlight();
@@ -455,8 +555,7 @@ void SMultiLineEditableText::FinishChangingText()
 	check(bIsChangingText);
 	bIsChangingText = false;
 
-	FText EditedText;
-	TextLayout->GetAsText(EditedText);
+	const FText EditedText = GetEditableText();
 
 	// Has the text changed?
 	const bool HasTextChanged = !EditedText.ToString().Equals(StateBeforeChangingText.Text.ToString(), ESearchCase::CaseSensitive);
@@ -757,6 +856,12 @@ void SMultiLineEditableText::DeleteSelectedText()
 
 FReply SMultiLineEditableText::MoveCursor( ECursorMoveMethod::Type Method, const FVector2D& Direction, ECursorAction::Type Action )
 {
+	// We can't use the keyboard to move the cursor while composing, as the IME has control over it
+	if(TextInputMethodContext->IsComposing && Method != ECursorMoveMethod::ScreenPosition)
+	{
+		return FReply::Handled();
+	}
+
 	bool bAllowMoveCursor = true;
 	FTextLocation NewCursorPosition;
 	FTextLocation CursorPosition = CursorInfo.GetCursorInteractionLocation();
@@ -854,6 +959,18 @@ FReply SMultiLineEditableText::MoveCursor( ECursorMoveMethod::Type Method, const
 	CursorInfo.SetCursorLocationAndCalculateAlignment(TextLayout, NewCursorPosition, CursorEOLMode);
 	UpdateCursorHighlight();
 
+	// If we've moved the cursor while composing, we need to end the current composition session
+	// Note: You should only be able to do this via the mouse due to the check at the top of this function
+	if(TextInputMethodContext->IsComposing)
+	{
+		ITextInputMethodSystem* const TextInputMethodSystem = FSlateApplication::Get().GetTextInputMethodSystem();
+		if(TextInputMethodSystem)
+		{
+			TextInputMethodSystem->DeactivateContext(TextInputMethodContext.ToSharedRef());
+			TextInputMethodSystem->ActivateContext(TextInputMethodContext.ToSharedRef());
+		}
+	}
+
 	return FReply::Handled();
 }
 
@@ -866,13 +983,38 @@ void SMultiLineEditableText::UpdateCursorHighlight()
 
 	bool bDrawCursor = true;
 
+	const bool bIsComposing = TextInputMethodContext->IsComposing;
 	const bool bHasSelection = SelectionLocation != CursorInteractionPosition;
-	if ( bHasSelection )
+
+	if ( bIsComposing )
 	{
-		bDrawCursor = false;
+		FTextLayout::FTextOffsetLocations OffsetLocations;
+		TextLayout->GetTextOffsetLocations(OffsetLocations);
 
-		const FTextLocation CursorInteractionPosition = CursorInfo.GetCursorInteractionLocation();
-		const FTextLocation SelectionLocation = SelectionStart.Get( CursorInteractionPosition );
+		const FTextLocation CompositionBeginLocation = OffsetLocations.OffsetToTextLocation(TextInputMethodContext->CompositionBeginIndex);
+		const FTextLocation CompositionEndLocation = OffsetLocations.OffsetToTextLocation(TextInputMethodContext->CompositionBeginIndex + TextInputMethodContext->CompositionLength);
+
+		// Composition should never span more than one (hard) line
+		if ( CompositionBeginLocation.GetLineIndex() == CompositionEndLocation.GetLineIndex() )
+		{
+			const FTextRange Range(CompositionBeginLocation.GetOffset(), CompositionEndLocation.GetOffset());
+
+			// We only draw the composition highlight if the cursor is within the composition range
+			bDrawCursor = !(CompositionBeginLocation.GetLineIndex() == CursorInteractionPosition.GetLineIndex() && Range.InclusiveContains(CursorInteractionPosition.GetOffset()));
+			if (!Range.IsEmpty() && !bDrawCursor)
+			{
+				TextLayout->AddHighlight(FTextHighlight(CompositionBeginLocation.GetLineIndex(), Range, CompositionRunHighlighter.ToSharedRef()));
+			}
+			else
+			{
+				// We have no composition range for this line, but we're supposed to draw the cursor on this line - let the normal cursor handling take care of that
+				bDrawCursor = true;
+			}
+		}
+	}
+	else if ( bHasSelection )
+	{
+		bDrawCursor = false;
 
 		const FTextSelection Selection( SelectionLocation, CursorInteractionPosition );
 
@@ -902,16 +1044,10 @@ void SMultiLineEditableText::UpdateCursorHighlight()
 		{
 			const TArray< FTextLayout::FLineModel >& Lines = TextLayout->GetLineModels();
 
-			bool bHasAddedBeginHighlight = false;
-			bool bHasAddedEndHighlight = false;
-
 			for (int LineIndex = SelectionBeginningLineIndex; LineIndex <= SelectionEndLineIndex; LineIndex++)
 			{
 				if ( LineIndex == SelectionBeginningLineIndex )
 				{
-					check(!bHasAddedBeginHighlight);
-					bHasAddedBeginHighlight = true;
-
 					const FTextRange Range(SelectionBeginningLineOffset, Lines[LineIndex].Text->Len());
 					if (!Range.IsEmpty())
 					{
@@ -925,9 +1061,6 @@ void SMultiLineEditableText::UpdateCursorHighlight()
 				}
 				else if ( LineIndex == SelectionEndLineIndex )
 				{
-					check(!bHasAddedEndHighlight);
-					bHasAddedEndHighlight = true;
-
 					const FTextRange Range(0, SelectionEndLineOffset);
 					if (!Range.IsEmpty())
 					{
@@ -1380,31 +1513,43 @@ void SMultiLineEditableText::PasteTextFromClipboard()
 		FString PastedText;
 		FPlatformMisc::ClipboardPaste(PastedText);
 
-		for (int32 CurCharIndex = 0; CurCharIndex < PastedText.Len(); ++CurCharIndex)
+		InsertTextAtCursor(PastedText);
+
+		FinishChangingText();
+	}
+}
+
+void SMultiLineEditableText::InsertTextAtCursor(const FString& InString)
+{
+	if (InString.Len())
+	{
+		for (const TCHAR* Character = InString.GetCharArray().GetData(); *Character; ++Character)
 		{
-			const TCHAR& Character = PastedText.GetCharArray()[CurCharIndex];
-			if (Character == '\n')
+			if (*Character == '\n')
 			{
 				OnEnter();
 			}
 			else
 			{
-				TypeChar(Character);
+				TypeChar(*Character);
 			}
 		}
-
-		FinishChangingText();
 	}
 }
 
 bool SMultiLineEditableText::CanExecuteUndo() const
 {
+	if(TextInputMethodContext->IsComposing)
+	{
+		return false;
+	}
+
 	return true;
 }
 
 void SMultiLineEditableText::Undo()
 {
-	if (!IsReadOnly.Get() && UndoStates.Num() > 0/* && !TextInputMethodContext->IsComposing*/)
+	if (!IsReadOnly.Get() && UndoStates.Num() > 0 && !TextInputMethodContext->IsComposing)
 	{
 		// Restore from undo state
 		int32 UndoStateIndex;
@@ -1454,7 +1599,7 @@ void SMultiLineEditableText::Redo()
 {
 	// Is there anything to redo?  If we've haven't tried to undo since the last time new
 	// undo state was added, then CurrentUndoLevel will be INDEX_NONE
-	if (!IsReadOnly.Get() && CurrentUndoLevel != INDEX_NONE /*&& !(TextInputMethodContext->IsComposing)*/)
+	if (!IsReadOnly.Get() && CurrentUndoLevel != INDEX_NONE && !(TextInputMethodContext->IsComposing))
 	{
 		const int32 NextUndoLevel = CurrentUndoLevel + 1;
 		if (UndoStates.Num() > NextUndoLevel)
@@ -1523,8 +1668,7 @@ void SMultiLineEditableText::ClearUndoStates()
 void SMultiLineEditableText::MakeUndoState(SMultiLineEditableText::FUndoState& OutUndoState)
 {
 	//@todo save and restoring the whole document is not ideal [3/31/2014 justin.sargent]
-	FText EditedText;
-	TextLayout->GetAsText(EditedText);
+	const FText EditedText = GetEditableText();
 
 	OutUndoState.Text = EditedText;
 	OutUndoState.CursorInfo = CursorInfo;
@@ -1606,6 +1750,12 @@ void SMultiLineEditableText::LoadText()
 
 void SMultiLineEditableText::Tick( const FGeometry& AllottedGeometry, const double InCurrentTime, const float InDeltaTime )
 {
+	if(TextInputMethodChangeNotifier.IsValid() && TextInputMethodContext.IsValid() && TextInputMethodContext->CachedGeometry != AllottedGeometry)
+	{
+		TextInputMethodContext->CachedGeometry = AllottedGeometry;
+		TextInputMethodChangeNotifier->NotifyLayoutChanged(ITextInputMethodChangeNotifier::ELayoutChangeType::Changed);
+	}
+
 	TextLayout->SetScale( AllottedGeometry.Scale );
 
 	const bool bShouldAppearFocused = HasKeyboardFocus() || ContextMenuWindow.IsValid();
@@ -1894,3 +2044,263 @@ bool SMultiLineEditableText::IsAtWordStart( const FTextLocation& Location ) cons
 
 	return WordStart.IsValid() && WordStart == Location;
 }
+
+bool SMultiLineEditableText::FTextInputMethodContext::IsReadOnly()
+{
+	bool Return = true;
+	const TSharedPtr<SMultiLineEditableText> OwningWidgetPtr = OwningWidget.Pin();
+	if(OwningWidgetPtr.IsValid())
+	{
+		Return = OwningWidgetPtr->GetIsReadOnly();
+	}
+	return Return;
+}
+
+uint32 SMultiLineEditableText::FTextInputMethodContext::GetTextLength()
+{
+	uint32 TextLength = 0;
+	const TSharedPtr<SMultiLineEditableText> OwningWidgetPtr = OwningWidget.Pin();
+	if(OwningWidgetPtr.IsValid())
+	{
+		FTextLayout::FTextOffsetLocations OffsetLocations;
+		OwningWidgetPtr->TextLayout->GetTextOffsetLocations(OffsetLocations);
+
+		TextLength = OffsetLocations.GetTextLength();
+	}
+	return TextLength;
+}
+
+void SMultiLineEditableText::FTextInputMethodContext::GetSelectionRange(uint32& BeginIndex, uint32& Length, ECaretPosition& OutCaretPosition)
+{
+	const TSharedPtr<SMultiLineEditableText> OwningWidgetPtr = OwningWidget.Pin();
+	if(OwningWidgetPtr.IsValid())
+	{
+		const FTextLocation CursorInteractionPosition = OwningWidgetPtr->CursorInfo.GetCursorInteractionLocation();
+		const FTextLocation SelectionLocation = OwningWidgetPtr->SelectionStart.Get(CursorInteractionPosition);
+
+		FTextLayout::FTextOffsetLocations OffsetLocations;
+		OwningWidgetPtr->TextLayout->GetTextOffsetLocations(OffsetLocations);
+
+		const bool bHasSelection = SelectionLocation != CursorInteractionPosition;
+		if(bHasSelection)
+		{
+			// We need to translate the selection into "editable text" space
+			const FTextSelection Selection(SelectionLocation, CursorInteractionPosition);
+
+			const FTextLocation& BeginningOfSelectionInDocumentSpace = Selection.GetBeginning();
+			const int32 BeginningOfSelectionInEditableTextSpace = OffsetLocations.TextLocationToOffset(BeginningOfSelectionInDocumentSpace);
+
+			const FTextLocation& EndOfSelectionInDocumentSpace = Selection.GetEnd();
+			const int32 EndOfSelectionInEditableTextSpace = OffsetLocations.TextLocationToOffset(EndOfSelectionInDocumentSpace);
+
+			BeginIndex = BeginningOfSelectionInEditableTextSpace;
+			Length = EndOfSelectionInEditableTextSpace - BeginningOfSelectionInEditableTextSpace;
+
+			const bool bCursorIsBeforeSelection = CursorInteractionPosition < SelectionLocation;
+			OutCaretPosition = (bCursorIsBeforeSelection) ? ITextInputMethodContext::ECaretPosition::Beginning : ITextInputMethodContext::ECaretPosition::Ending;
+		}
+		else
+		{
+			// We need to translate the cursor position into "editable text" space
+			const int32 CursorInteractionPositionInEditableTextSpace = OffsetLocations.TextLocationToOffset(CursorInteractionPosition);
+
+			BeginIndex = CursorInteractionPositionInEditableTextSpace;
+			Length = 0;
+
+			OutCaretPosition = ITextInputMethodContext::ECaretPosition::Beginning;
+		}
+	}
+}
+
+void SMultiLineEditableText::FTextInputMethodContext::SetSelectionRange(const uint32 BeginIndex, const uint32 Length, const ECaretPosition InCaretPosition)
+{
+	const TSharedPtr<SMultiLineEditableText> OwningWidgetPtr = OwningWidget.Pin();
+	if(OwningWidgetPtr.IsValid())
+	{
+		const uint32 MinIndex = BeginIndex;
+		const uint32 MaxIndex = MinIndex + Length;
+
+		FTextLayout::FTextOffsetLocations OffsetLocations;
+		OwningWidgetPtr->TextLayout->GetTextOffsetLocations(OffsetLocations);
+
+		// We need to translate the indices into document space
+		const FTextLocation MinTextLocation = OffsetLocations.OffsetToTextLocation(MinIndex);
+		const FTextLocation MaxTextLocation = OffsetLocations.OffsetToTextLocation(MaxIndex);
+
+		OwningWidgetPtr->ClearSelection();
+
+		switch(InCaretPosition)
+		{
+		case ITextInputMethodContext::ECaretPosition::Beginning:
+			{
+				OwningWidgetPtr->CursorInfo.SetCursorLocationAndCalculateAlignment(OwningWidgetPtr->TextLayout, MinTextLocation);
+				OwningWidgetPtr->SelectionStart = MaxTextLocation;
+			}
+			break;
+
+		case ITextInputMethodContext::ECaretPosition::Ending:
+			{
+				OwningWidgetPtr->SelectionStart = MinTextLocation;
+				OwningWidgetPtr->CursorInfo.SetCursorLocationAndCalculateAlignment(OwningWidgetPtr->TextLayout, MaxTextLocation);
+			}
+			break;
+		}
+
+		OwningWidgetPtr->UpdateCursorHighlight();
+	}
+}
+
+void SMultiLineEditableText::FTextInputMethodContext::GetTextInRange(const uint32 BeginIndex, const uint32 Length, FString& OutString)
+{
+	const TSharedPtr<SMultiLineEditableText> OwningWidgetPtr = OwningWidget.Pin();
+	if(OwningWidgetPtr.IsValid())
+	{
+		const FText EditedText = OwningWidgetPtr->GetEditableText();
+		OutString = EditedText.ToString().Mid(BeginIndex, Length);
+	}
+}
+
+void SMultiLineEditableText::FTextInputMethodContext::SetTextInRange(const uint32 BeginIndex, const uint32 Length, const FString& InString)
+{
+	const TSharedPtr<SMultiLineEditableText> OwningWidgetPtr = OwningWidget.Pin();
+	if(OwningWidgetPtr.IsValid())
+	{
+		// We don't use Start/FinishEditing text here because the whole IME operation handles that.
+		// Also, we don't want to support undo for individual characters added during an IME context
+		const FText OldEditedText = OwningWidgetPtr->GetEditableText();
+
+		// We do this as a select, delete, and insert as it's the simplest way to keep the text layout correct
+		SetSelectionRange(BeginIndex, Length, ITextInputMethodContext::ECaretPosition::Beginning);
+		OwningWidgetPtr->DeleteSelectedText();
+		OwningWidgetPtr->InsertTextAtCursor(InString);
+
+		// Has the text changed?
+		const FText EditedText = OwningWidgetPtr->GetEditableText();
+		const bool HasTextChanged = !EditedText.ToString().Equals(OldEditedText.ToString(), ESearchCase::CaseSensitive);
+		if(HasTextChanged)
+		{
+			OwningWidgetPtr->SaveText(EditedText);
+			OwningWidgetPtr->TextLayout->UpdateIfNeeded();
+			OwningWidgetPtr->OnTextChanged.ExecuteIfBound(EditedText);
+		}
+	}
+}
+
+int32 SMultiLineEditableText::FTextInputMethodContext::GetCharacterIndexFromPoint(const FVector2D& Point)
+{
+	int32 CharacterIndex = INDEX_NONE;
+	const TSharedPtr<SMultiLineEditableText> OwningWidgetPtr = OwningWidget.Pin();
+	if(OwningWidgetPtr.IsValid())
+	{
+		const FTextLocation CharacterPosition = OwningWidgetPtr->TextLayout->GetTextLocationAt(Point * OwningWidgetPtr->TextLayout->GetScale());
+
+		FTextLayout::FTextOffsetLocations OffsetLocations;
+		OwningWidgetPtr->TextLayout->GetTextOffsetLocations(OffsetLocations);
+
+		CharacterIndex = OffsetLocations.TextLocationToOffset(CharacterPosition);
+	}
+	return CharacterIndex;
+}
+
+bool SMultiLineEditableText::FTextInputMethodContext::GetTextBounds(const uint32 BeginIndex, const uint32 Length, FVector2D& Position, FVector2D& Size)
+{
+	bool IsClipped = false;
+	const TSharedPtr<SMultiLineEditableText> OwningWidgetPtr = OwningWidget.Pin();
+	if(OwningWidgetPtr.IsValid())
+	{
+		FTextLayout::FTextOffsetLocations OffsetLocations;
+		OwningWidgetPtr->TextLayout->GetTextOffsetLocations(OffsetLocations);
+
+		const FTextLocation BeginLocation = OffsetLocations.OffsetToTextLocation(BeginIndex);
+		const FTextLocation EndLocation = OffsetLocations.OffsetToTextLocation(BeginIndex + Length);
+
+		const FVector2D BeginPosition = OwningWidgetPtr->TextLayout->GetLocationAt(BeginLocation);
+		const FVector2D EndPosition = OwningWidgetPtr->TextLayout->GetLocationAt(EndLocation);
+
+		if(BeginPosition.Y == EndPosition.Y)
+		{
+			// The text range is contained within a single line
+			Position = BeginPosition;
+			Size = EndPosition - BeginPosition;
+		}
+		else
+		{
+			// If the two positions aren't on the same line, then we assume the worst case scenario, and make the size as wide as the text area itself
+			Position = FVector2D(0.0f, BeginPosition.Y);
+			Size = FVector2D(OwningWidgetPtr->TextLayout->GetDrawSize().X, EndPosition.Y - BeginPosition.Y);
+		}
+
+		// Translate the position (which is in local space) into screen (absolute) space
+		// Note: The local positions are pre-scaled, so we don't scale them again here
+		FGeometry CachedGeometryNoScale = CachedGeometry;
+		CachedGeometryNoScale.Scale = 1.0f;
+		Position = CachedGeometryNoScale.LocalToAbsolute(Position);
+	}
+	return IsClipped;
+}
+
+void SMultiLineEditableText::FTextInputMethodContext::GetScreenBounds(FVector2D& Position, FVector2D& Size)
+{
+	const TSharedPtr<SMultiLineEditableText> OwningWidgetPtr = OwningWidget.Pin();
+	if(OwningWidgetPtr.IsValid())
+	{
+		Position = CachedGeometry.AbsolutePosition;
+		Size = CachedGeometry.GetDrawSize();
+	}
+}
+
+TSharedPtr<FGenericWindow> SMultiLineEditableText::FTextInputMethodContext::GetWindow()
+{
+	TSharedPtr<FGenericWindow> Window;
+	const TSharedPtr<SMultiLineEditableText> OwningWidgetPtr = OwningWidget.Pin();
+	if(OwningWidgetPtr.IsValid())
+	{
+		const TSharedPtr<SWindow> SlateWindow = FSlateApplication::Get().FindWidgetWindow( OwningWidgetPtr.ToSharedRef() );
+		Window = SlateWindow.IsValid() ? SlateWindow->GetNativeWindow() : nullptr;
+	}
+	return Window;
+}
+
+void SMultiLineEditableText::FTextInputMethodContext::BeginComposition()
+{
+	if(!IsComposing)
+	{
+		IsComposing = true;
+
+		const TSharedPtr<SMultiLineEditableText> OwningWidgetPtr = OwningWidget.Pin();
+		if(OwningWidgetPtr.IsValid())
+		{
+			OwningWidgetPtr->StartChangingText();
+			OwningWidgetPtr->UpdateCursorHighlight();
+		}
+	}
+}
+
+void SMultiLineEditableText::FTextInputMethodContext::UpdateCompositionRange(const int32 InBeginIndex, const uint32 InLength)
+{
+	check(IsComposing);
+
+	CompositionBeginIndex = InBeginIndex;
+	CompositionLength = InLength;
+
+	const TSharedPtr<SMultiLineEditableText> OwningWidgetPtr = OwningWidget.Pin();
+	if(OwningWidgetPtr.IsValid())
+	{
+		OwningWidgetPtr->UpdateCursorHighlight();
+	}
+}
+
+void SMultiLineEditableText::FTextInputMethodContext::EndComposition()
+{
+	if(IsComposing)
+	{
+		const TSharedPtr<SMultiLineEditableText> OwningWidgetPtr = OwningWidget.Pin();
+		if(OwningWidgetPtr.IsValid())
+		{
+			OwningWidgetPtr->FinishChangingText();
+			OwningWidgetPtr->UpdateCursorHighlight();
+		}
+
+		IsComposing = false;
+	}
+}
diff --git a/Engine/Source/Runtime/Slate/Public/Framework/Text/TextLayout.h b/Engine/Source/Runtime/Slate/Public/Framework/Text/TextLayout.h
--- a/Engine/Source/Runtime/Slate/Public/Framework/Text/TextLayout.h
+++ b/Engine/Source/Runtime/Slate/Public/Framework/Text/TextLayout.h
@@ -218,6 +218,37 @@ public:
 		int32 ModelIndex;
 	};
 
+	/** A mapping between the offsets into the text as a flat string (with line-breaks), and the internal lines used within a text layout */
+	struct FTextOffsetLocations
+	{
+	friend class FTextLayout;
+
+	public:
+		int32 TextLocationToOffset(const FTextLocation& InLocation) const;
+		FTextLocation OffsetToTextLocation(const int32 InOffset) const;
+		int32 GetTextLength() const;
+
+	private:
+		struct FOffsetEntry
+		{
+			FOffsetEntry(const int32 InFlatStringIndex, const int32 InDocumentLineLength)
+				: FlatStringIndex(InFlatStringIndex)
+				, DocumentLineLength(InDocumentLineLength)
+			{
+			}
+
+			/** Index in the flat string for this entry */
+			int32 FlatStringIndex;
+
+			/** The length of the line in the document (not including any trailing \n character) */
+			int32 DocumentLineLength;
+		};
+
+		/** This array contains one entry for each line in the document; 
+			the array index is the line number, and the entry contains the index in the flat string that marks the start of the line, 
+			along with the length of the line (not including any trailing \n character) */
+		TArray<FOffsetEntry> OffsetData;
+	};
 
 public:
 
@@ -295,9 +326,12 @@ public:
 
 	bool IsEmpty() const;
 
-	void GetAsText(FString& DisplayText) const;
+	void GetAsText(FString& DisplayText, FTextOffsetLocations* const OutTextOffsetLocations = nullptr) const;
+
+	void GetAsText(FText& DisplayText, FTextOffsetLocations* const OutTextOffsetLocations = nullptr) const;
 
-	void GetAsText(FText& DisplayText) const;
+	/** Constructs an array containing the mappings between the text that would be returned by GetAsText, and the internal FTextLocation points used within this text layout */
+	void GetTextOffsetLocations(FTextOffsetLocations& OutTextOffsetLocations) const;
 
 	void GetSelectionAsText(FString& DisplayText, const FTextSelection& Selection) const;
 
@@ -350,6 +384,8 @@ private:
 
 	FBreakCandidate CreateBreakCandidate( int32& OutRunIndex, FLineModel& Line, int32 PreviousBreak, int32 CurrentBreak );
 
+	void GetAsTextAndOffsets(FString* const OutDisplayText, FTextOffsetLocations* const OutTextOffsetLocations) const;
+
 protected:
 
 	/** The models for the lines of text. A LineModel represents a single string with no manual breaks. */
diff --git a/Engine/Source/Runtime/Slate/Public/Widgets/Text/SMultiLineEditableText.h b/Engine/Source/Runtime/Slate/Public/Widgets/Text/SMultiLineEditableText.h
--- a/Engine/Source/Runtime/Slate/Public/Widgets/Text/SMultiLineEditableText.h
+++ b/Engine/Source/Runtime/Slate/Public/Widgets/Text/SMultiLineEditableText.h
@@ -2,8 +2,8 @@
 
 #pragma once
 
-#include "MultiBox/MultiBoxExtender.h"
 #include "ITextEditorWidget.h"
+#include "ITextInputMethodSystem.h"
 
 /** An editable text widget that supports multiple lines and soft word-wrapping. */
 class SLATE_API SMultiLineEditableText : public SWidget, public ITextEditorWidget
@@ -64,6 +64,7 @@ public:
 	SLATE_END_ARGS()
 
 	SMultiLineEditableText();
+	virtual ~SMultiLineEditableText();
 
 	void Construct( const FArguments& InArgs );
 
@@ -85,6 +86,48 @@ public:
 
 private:
 	
+	/** 
+	 * Note: The IME interface for the multiline editable text uses the pre-flowed version of the string since the IME APIs are designed to work with flat strings
+	 *		 This means we have to do a bit of juggling to convert between the two
+	 */
+	friend class FTextInputMethodContext;
+	class FTextInputMethodContext : public ITextInputMethodContext
+	{
+	public:
+		FTextInputMethodContext(const TWeakPtr<SMultiLineEditableText>& InOwningWidget)
+			:	OwningWidget(InOwningWidget)
+			,	IsComposing(false)
+			,	CompositionBeginIndex(INDEX_NONE)
+			,	CompositionLength(0)
+		{}
+
+		virtual ~FTextInputMethodContext() {}
+
+	private:
+		virtual bool IsReadOnly() override;
+		virtual uint32 GetTextLength() override;
+		virtual void GetSelectionRange(uint32& BeginIndex, uint32& Length, ECaretPosition& CaretPosition) override;
+		virtual void SetSelectionRange(const uint32 BeginIndex, const uint32 Length, const ECaretPosition CaretPosition) override;
+		virtual void GetTextInRange(const uint32 BeginIndex, const uint32 Length, FString& OutString) override;
+		virtual void SetTextInRange(const uint32 BeginIndex, const uint32 Length, const FString& InString) override;
+		virtual int32 GetCharacterIndexFromPoint(const FVector2D& Point) override;
+		virtual bool GetTextBounds(const uint32 BeginIndex, const uint32 Length, FVector2D& Position, FVector2D& Size) override;
+		virtual void GetScreenBounds(FVector2D& Position, FVector2D& Size) override;
+		virtual TSharedPtr<FGenericWindow> GetWindow() override;
+		virtual void BeginComposition() override;
+		virtual void UpdateCompositionRange(const int32 InBeginIndex, const uint32 InLength) override;
+		virtual void EndComposition() override;
+
+	private:
+		TWeakPtr<SMultiLineEditableText> OwningWidget;
+
+	public:
+		FGeometry CachedGeometry;
+		bool IsComposing;
+		int32 CompositionBeginIndex;
+		uint32 CompositionLength;
+	};
+
 	enum class ECursorAlignment : uint8
 	{
 		/** Visually align the cursor to the left of the character its placed at, and insert text before the character */
@@ -205,7 +248,7 @@ private:
 
 	protected:
 
-		int32 OnPaintCursor( const FTextLayout::FLineView& Line, const TSharedRef< ISlateRun >& Run, const TSharedRef< ILayoutBlock >& Block, const FTextBlockStyle& DefaultStyle, const FGeometry& AllottedGeometry, const FSlateRect& MyClippingRect, FSlateWindowElementList& OutDrawElements, int32 LayerId, const FWidgetStyle& InWidgetStyle, const bool bParentEnabled, const ECursorAlignment InCursorAlignment ) const;
+		int32 OnPaintCursor( const FTextLayout::FLineView& Line, const TSharedRef< ISlateRun >& Run, const TSharedRef< ILayoutBlock >& Block, const FTextBlockStyle& DefaultStyle, const FGeometry& AllottedGeometry, const FSlateRect& MyClippingRect, FSlateWindowElementList& OutDrawElements, int32 LayerId, const FWidgetStyle& InWidgetStyle, const bool bParentEnabled, const float InCursorOffset, const ECursorAlignment InCursorAlignment ) const;
 
 		FSlateCursorRunHighlighter(const FCursorInfo* InCursorInfo);
 
@@ -236,6 +279,22 @@ private:
 		ECursorAlignment CursorAlignment;
 	};
 
+	/** Run highlighter used to draw the composition range */
+	class FSlateCompositionRunHighlighter : public FSlateCursorRunHighlighter
+	{
+	public:
+
+		static TSharedRef< FSlateCompositionRunHighlighter > Create(const FCursorInfo* InCursorInfo);
+
+		virtual ~FSlateCompositionRunHighlighter() {}
+
+		virtual int32 OnPaint( const FTextLayout::FLineView& Line, const TSharedRef< ISlateRun >& Run, const TSharedRef< ILayoutBlock >& Block, const FTextBlockStyle& DefaultStyle, const FGeometry& AllottedGeometry, const FSlateRect& MyClippingRect, FSlateWindowElementList& OutDrawElements, int32 LayerId, const FWidgetStyle& InWidgetStyle, bool bParentEnabled ) const override;
+
+	protected:
+
+		FSlateCompositionRunHighlighter(const FCursorInfo* InCursorInfo);
+	};
+
 private:
 
 	// BEGIN ITextEditorWidget interface
@@ -313,6 +372,9 @@ private:
 
 	bool DoesClipboardHaveAnyText() const;
 
+	/** Insert the given text at the current cursor position, correctly taking into account new line characters */
+	void InsertTextAtCursor(const FString& InString);
+
 	/**
 	 * Given a location and a Direction to offset, return a new location.
 	 *
@@ -373,6 +435,12 @@ private:
 	 */
 	bool SetEditableText(const FText& TextToSet, const bool bForce = false);
 
+	/**
+	 * Gets the current editable text for this text block
+	 * Note: We don't store text in this form (it's stored as lines in the text layout) so every call to this function has to reconstruct it
+	 */
+	FText GetEditableText() const;
+
 private:
 
 	/** The text displayed in this text block */
@@ -404,6 +472,8 @@ private:
 
 	TSharedPtr< FSlateSelectionRunHighlighter > SelectionRunHighlighter;
 
+	TSharedPtr< FSlateCompositionRunHighlighter > CompositionRunHighlighter;
+
 	/** That start of the selection when there is a selection. The end is implicitly wherever the cursor happens to be. */
 	TOptional<FTextLocation> SelectionStart;
 
@@ -454,4 +524,10 @@ private:
 
 	/** Weak pointer to context menu window that's currently open, if there is one */
 	TWeakPtr< SWindow > ContextMenuWindow;
+
+	/** Implemented context object for text input method systems. */
+	TSharedPtr<FTextInputMethodContext> TextInputMethodContext;
+
+	/** Notification interface object for text input method systems. */
+	TSharedPtr<ITextInputMethodChangeNotifier> TextInputMethodChangeNotifier;
 };
