Slate TableViews now require that a touch was initiated in them before they will respond to touch scrolling.
Fixes TTP#338093 SLATE: MOBILE: Tree views get permanently stuck in a scrolling mode if the mouse is moved outside of a PIE window with 'mouse as touch' enabled

diff --git a/Engine/Source/Runtime/Slate/Private/Widgets/Views/STableViewBase.cpp b/Engine/Source/Runtime/Slate/Private/Widgets/Views/STableViewBase.cpp
--- a/Engine/Source/Runtime/Slate/Private/Widgets/Views/STableViewBase.cpp
+++ b/Engine/Source/Runtime/Slate/Private/Widgets/Views/STableViewBase.cpp
@@ -231,6 +231,25 @@ void STableViewBase::ScrollBar_OnUserScrolled( float InScrollOffsetFraction )
 }
 
 
+FReply STableViewBase::OnPreviewMouseButtonDown( const FGeometry& MyGeometry, const FPointerEvent& MouseEvent )
+{
+	if (MouseEvent.IsTouchEvent())
+	{
+		// Clear any inertia 
+		this->InertialScrollManager.ClearScrollVelocity();
+		// We have started a new interaction; track how far the user has moved since they put their finger down.
+		AmountScrolledWhileRightMouseDown = 0;
+		// Someone put their finger down in this list, so they probably want to drag the list.
+		bStartedTouchInteraction = true;
+		return FReply::Unhandled();
+	}
+	else
+	{
+		return FReply::Unhandled();
+	}
+}
+
+
 FReply STableViewBase::OnMouseButtonDown( const FGeometry& MyGeometry, const FPointerEvent& MouseEvent )
 {
 	// Zero the scroll velocity so the list stops immediately on mouse down, even if the user does not drag
@@ -339,6 +358,7 @@ FReply STableViewBase::OnMouseMove( const FGeometry& MyGeometry, const FPointerE
 
 void STableViewBase::OnMouseLeave( const FPointerEvent& MouseEvent )
 {
+	bStartedTouchInteraction = false;
 	if( FSlateApplication::Get().GetMouseCaptor().Get() != this )
 	{
 		// No longer scrolling (unless we have mouse capture)
@@ -392,36 +412,41 @@ FCursorReply STableViewBase::OnCursorQuery( const FGeometry& MyGeometry, const F
 
 FReply STableViewBase::OnTouchStarted( const FGeometry& MyGeometry, const FPointerEvent& InTouchEvent )
 {
-	// Clear any inertia 
-	this->InertialScrollManager.ClearScrollVelocity();
-	// We have started a new interaction; track how far the user has moved since they put their finger down.
-	AmountScrolledWhileRightMouseDown = 0;
+	// See OnPreviewMouseButtonDown()
+	//     if (MouseEvent.IsTouchEvent())
 
 	return FReply::Unhandled();
 }
 
 FReply STableViewBase::OnTouchMoved( const FGeometry& MyGeometry, const FPointerEvent& InTouchEvent )
 {
-	const float ScrollByAmount = InTouchEvent.GetCursorDelta().Y;
-	AmountScrolledWhileRightMouseDown += FMath::Abs( ScrollByAmount );
+	if (bStartedTouchInteraction)
+	{
+		const float ScrollByAmount = InTouchEvent.GetCursorDelta().Y;
+		AmountScrolledWhileRightMouseDown += FMath::Abs( ScrollByAmount );
 
-	this->InertialScrollManager.AddScrollSample( ScrollByAmount, FPlatformTime::Seconds());
-	const float AmountScrolled = this->ScrollBy( MyGeometry, ScrollByAmount, EAllowOverscroll::Yes );
+		this->InertialScrollManager.AddScrollSample( ScrollByAmount, FPlatformTime::Seconds());
+		const float AmountScrolled = this->ScrollBy( MyGeometry, ScrollByAmount, EAllowOverscroll::Yes );
 
-	if (AmountScrolledWhileRightMouseDown > SlatePanTriggerDistance)
+		if (AmountScrolledWhileRightMouseDown > SlatePanTriggerDistance)
+		{
+			// The user has moved the list some amount; they are probably
+			// trying to scroll. From now on, the list assumes the user is scrolling
+			// until they lift their finger.
+			return FReply::Handled().CaptureMouse( AsShared() );
+		}
+		return FReply::Handled();
+	}
+	else
 	{
-		// The user has moved the list some amount; they are probably
-		// trying to scroll. From now on, the list assumes the user is scrolling
-		// until they lift their finger.
-		return FReply::Handled().CaptureMouse( AsShared() );
+		return FReply::Handled();
 	}
-
-	return FReply::Handled();
 }
 
 FReply STableViewBase::OnTouchEnded( const FGeometry& MyGeometry, const FPointerEvent& InTouchEvent )
 {
 	AmountScrolledWhileRightMouseDown = 0;
+	bStartedTouchInteraction = false;
 
 	if (HasMouseCapture())
 	{
@@ -493,6 +518,7 @@ int32 STableViewBase::OnPaint( const FGeometry& AllottedGeometry, const FSlateRe
 STableViewBase::STableViewBase( ETableViewMode::Type InTableViewMode )
 	: TableViewMode( InTableViewMode )
 	, ScrollOffset( 0 )
+	, bStartedTouchInteraction( false )
 	, AmountScrolledWhileRightMouseDown( 0 )
 	, LastGenerateResults( 0,0,0,false )
 	, bWasAtEndOfList(false)
diff --git a/Engine/Source/Runtime/Slate/Public/Widgets/Views/STableViewBase.h b/Engine/Source/Runtime/Slate/Public/Widgets/Views/STableViewBase.h
--- a/Engine/Source/Runtime/Slate/Public/Widgets/Views/STableViewBase.h
+++ b/Engine/Source/Runtime/Slate/Public/Widgets/Views/STableViewBase.h
@@ -41,6 +41,7 @@ public:
 	virtual void OnMouseCaptureLost() override;
 	virtual bool SupportsKeyboardFocus() const override;
 	virtual void Tick( const FGeometry& AllottedGeometry, const double InCurrentTime, const float InDeltaTime ) override;
+	virtual FReply OnPreviewMouseButtonDown( const FGeometry& MyGeometry, const FPointerEvent& MouseEvent ) override;
 	virtual FReply OnMouseButtonDown( const FGeometry& MyGeometry, const FPointerEvent& MouseEvent ) override;
 	virtual FReply OnMouseButtonDoubleClick( const FGeometry& InMyGeometry, const FPointerEvent& InMouseEvent ) override;
 	virtual FReply OnMouseButtonUp( const FGeometry& MyGeometry, const FPointerEvent& MouseEvent ) override;
@@ -215,6 +216,9 @@ protected:
 	/** Scroll offset from the beginning of the list in items */
 	double ScrollOffset;
 
+	/** Did the user start an interaction in this list? */
+	bool bStartedTouchInteraction;
+
 	/** How much we scrolled while the rmb has been held */
 	float AmountScrolledWhileRightMouseDown;
 
