Slate - UniformGrid - Exposing a way to set the slot padding after initialization.  Max Column/Row values are now calculated and cached each frame.  This was done to allow updates to be displayed in real time for UMG.

diff --git a/Engine/Source/Runtime/Slate/Private/Widgets/Layout/SUniformGridPanel.cpp b/Engine/Source/Runtime/Slate/Private/Widgets/Layout/SUniformGridPanel.cpp
--- a/Engine/Source/Runtime/Slate/Private/Widgets/Layout/SUniformGridPanel.cpp
+++ b/Engine/Source/Runtime/Slate/Private/Widgets/Layout/SUniformGridPanel.cpp
@@ -17,39 +17,37 @@ void SUniformGridPanel::Construct( const FArguments& InArgs )
 	{
 		FSlot* ChildSlot = InArgs.Slots[ChildIndex];
 		Children.Add( ChildSlot );
-		// A single cell at (N,M) means our grid size is (N+1, M+1)
-		NumColumns = FMath::Max( ChildSlot->Column+1, NumColumns );
-		NumRows = FMath::Max( ChildSlot->Row+1, NumRows );
 	}
 }
 
-
 void SUniformGridPanel::OnArrangeChildren( const FGeometry& AllottedGeometry, FArrangedChildren& ArrangedChildren ) const
 {
-	const FVector2D CellSize(AllottedGeometry.Size.X / NumColumns, AllottedGeometry.Size.Y / NumRows);
-	const FMargin& CurrentSlotPadding(SlotPadding.Get());
-	for (int32 ChildIndex=0; ChildIndex < Children.Num(); ++ChildIndex)
+	if ( Children.Num() > 0 )
 	{
-		const FSlot& Child = Children[ChildIndex];
-		const EVisibility ChildVisibility = Child.Widget->GetVisibility();
-		if ( ArrangedChildren.Accepts(ChildVisibility) )
+		const FVector2D CellSize(AllottedGeometry.Size.X / NumColumns, AllottedGeometry.Size.Y / NumRows);
+		const FMargin& CurrentSlotPadding(SlotPadding.Get());
+		for ( int32 ChildIndex=0; ChildIndex < Children.Num(); ++ChildIndex )
 		{
-			// Do the standard arrangement of elements within a slot
-			// Takes care of alignment and padding.
-			AlignmentArrangeResult XAxisResult = AlignChild<Orient_Horizontal>( CellSize.X, Child, CurrentSlotPadding );
-			AlignmentArrangeResult YAxisResult = AlignChild<Orient_Vertical>( CellSize.Y, Child, CurrentSlotPadding );
-
-			ArrangedChildren.AddWidget( ChildVisibility,
-				AllottedGeometry.MakeChild(Child.Widget,
-				FVector2D(CellSize.X*Child.Column + XAxisResult.Offset, CellSize.Y*Child.Row + YAxisResult.Offset),
-				FVector2D(XAxisResult.Size, YAxisResult.Size)
-			));
+			const FSlot& Child = Children[ChildIndex];
+			const EVisibility ChildVisibility = Child.Widget->GetVisibility();
+			if ( ArrangedChildren.Accepts(ChildVisibility) )
+			{
+				// Do the standard arrangement of elements within a slot
+				// Takes care of alignment and padding.
+				AlignmentArrangeResult XAxisResult = AlignChild<Orient_Horizontal>(CellSize.X, Child, CurrentSlotPadding);
+				AlignmentArrangeResult YAxisResult = AlignChild<Orient_Vertical>(CellSize.Y, Child, CurrentSlotPadding);
+
+				ArrangedChildren.AddWidget(ChildVisibility,
+					AllottedGeometry.MakeChild(Child.Widget,
+					FVector2D(CellSize.X*Child.Column + XAxisResult.Offset, CellSize.Y*Child.Row + YAxisResult.Offset),
+					FVector2D(XAxisResult.Size, YAxisResult.Size)
+					));
+			}
+
 		}
-		
 	}
 }
 
-
 FVector2D SUniformGridPanel::ComputeDesiredSize() const
 {
 	FVector2D MaxChildDesiredSize = FVector2D::ZeroVector;
@@ -58,9 +56,17 @@ FVector2D SUniformGridPanel::ComputeDesiredSize() const
 	const float CachedMinDesiredSlotWidth = MinDesiredSlotWidth.Get();
 	const float CachedMinDesiredSlotHeight = MinDesiredSlotHeight.Get();
 	
+	NumColumns = 0;
+	NumRows = 0;
+
 	for ( int32 ChildIndex=0; ChildIndex < Children.Num(); ++ChildIndex )
 	{
 		const FSlot& Child = Children[ ChildIndex ];
+
+		// A single cell at (N,M) means our grid size is (N+1, M+1)
+		NumColumns = FMath::Max(Child.Column + 1, NumColumns);
+		NumRows = FMath::Max(Child.Row + 1, NumRows);
+
 		if (Child.Widget->GetVisibility() != EVisibility::Collapsed)
 		{
 			FVector2D ChildDesiredSize = Child.Widget->GetDesiredSize() + SlotPaddingDesiredSize;
@@ -76,21 +82,20 @@ FVector2D SUniformGridPanel::ComputeDesiredSize() const
 	return FVector2D( NumColumns*MaxChildDesiredSize.X, NumRows*MaxChildDesiredSize.Y );
 }
 
-
 FChildren* SUniformGridPanel::GetChildren()
 {
 	return &Children;
 }
 
+void SUniformGridPanel::SetSlotPadding(TAttribute<FMargin> InSlotPadding)
+{
+	SlotPadding = InSlotPadding;
+}
 
 SUniformGridPanel::FSlot& SUniformGridPanel::AddSlot( int32 Column, int32 Row )
 {
 	FSlot& NewSlot = SUniformGridPanel::Slot(Column, Row);
 
-	// A single cell at (N,M) means our grid size is (N+1, M+1)
-	NumColumns = FMath::Max( Column+1, NumColumns );
-	NumRows = FMath::Max( Row+1, NumRows );
-
 	Children.Add( &NewSlot );
 
 	return NewSlot;
@@ -115,4 +120,4 @@ void SUniformGridPanel::ClearChildren()
 	NumColumns = 0;
 	NumRows = 0;
 	Children.Empty();
-}
\ No newline at end of file
+}
diff --git a/Engine/Source/Runtime/Slate/Public/Widgets/Layout/SUniformGridPanel.h b/Engine/Source/Runtime/Slate/Public/Widgets/Layout/SUniformGridPanel.h
--- a/Engine/Source/Runtime/Slate/Public/Widgets/Layout/SUniformGridPanel.h
+++ b/Engine/Source/Runtime/Slate/Public/Widgets/Layout/SUniformGridPanel.h
@@ -58,6 +58,9 @@ public:
 	virtual FChildren* GetChildren() override;
 	// END SPanel INTERFACE
 
+	/** See SlotPadding attribute */
+	void SetSlotPadding(TAttribute<FMargin> InSlotPadding);
+
 	/**
 	 * Dynamically add a new slot to the UI at specified Column and Row.
 	 *
@@ -79,8 +82,11 @@ public:
 private:
 	TPanelChildren<FSlot> Children;
 	TAttribute<FMargin> SlotPadding;
-	int32 NumColumns;
-	int32 NumRows;
+	
+	/** These values are recomputed and cached during compute desired size, as they may have changed since the previous frame. */
+	mutable int32 NumColumns;
+	mutable int32 NumRows;
+
 	TAttribute<float> MinDesiredSlotWidth;
 	TAttribute<float> MinDesiredSlotHeight;
 };
\ No newline at end of file
