Slate - Adding UObject wrapper styles for progress bar, editable text box, and scroll box so that they are exposed for UMG.

diff --git a/Engine/Source/Runtime/Slate/Private/Framework/Styling/EditableTextBoxWidgetStyle.cpp b/Engine/Source/Runtime/Slate/Private/Framework/Styling/EditableTextBoxWidgetStyle.cpp
new file mode 100644
--- /dev/null
+++ b/Engine/Source/Runtime/Slate/Private/Framework/Styling/EditableTextBoxWidgetStyle.cpp
@@ -0,0 +1,10 @@
+// Copyright 1998-2014 Epic Games, Inc. All Rights Reserved.
+
+#include "SlatePrivatePCH.h"
+
+
+UEditableTextBoxWidgetStyle::UEditableTextBoxWidgetStyle( const class FPostConstructInitializeProperties& PCIP )
+	: Super(PCIP)
+{
+	
+}
\ No newline at end of file
diff --git a/Engine/Source/Runtime/Slate/Private/Framework/Styling/ProgressWidgetStyle.cpp b/Engine/Source/Runtime/Slate/Private/Framework/Styling/ProgressWidgetStyle.cpp
new file mode 100644
--- /dev/null
+++ b/Engine/Source/Runtime/Slate/Private/Framework/Styling/ProgressWidgetStyle.cpp
@@ -0,0 +1,10 @@
+// Copyright 1998-2014 Epic Games, Inc. All Rights Reserved.
+
+#include "SlatePrivatePCH.h"
+
+
+UProgressWidgetStyle::UProgressWidgetStyle( const class FPostConstructInitializeProperties& PCIP )
+	: Super(PCIP)
+{
+	
+}
\ No newline at end of file
diff --git a/Engine/Source/Runtime/Slate/Private/Framework/Styling/ScrollBoxWidgetStyle.cpp b/Engine/Source/Runtime/Slate/Private/Framework/Styling/ScrollBoxWidgetStyle.cpp
new file mode 100644
--- /dev/null
+++ b/Engine/Source/Runtime/Slate/Private/Framework/Styling/ScrollBoxWidgetStyle.cpp
@@ -0,0 +1,10 @@
+// Copyright 1998-2014 Epic Games, Inc. All Rights Reserved.
+
+#include "SlatePrivatePCH.h"
+
+
+UScrollBoxWidgetStyle::UScrollBoxWidgetStyle( const class FPostConstructInitializeProperties& PCIP )
+	: Super(PCIP)
+{
+	
+}
\ No newline at end of file
diff --git a/Engine/Source/Runtime/Slate/Private/SlatePrivatePCH.h b/Engine/Source/Runtime/Slate/Private/SlatePrivatePCH.h
--- a/Engine/Source/Runtime/Slate/Private/SlatePrivatePCH.h
+++ b/Engine/Source/Runtime/Slate/Private/SlatePrivatePCH.h
@@ -28,3 +28,6 @@
 #include "ComboButtonWidgetStyle.h"
 #include "TextBlockWidgetStyle.h"
 #include "EditableTextWidgetStyle.h"
+#include "EditableTextBoxWidgetStyle.h"
+#include "ProgressWidgetStyle.h"
+#include "ScrollBoxWidgetStyle.h"
diff --git a/Engine/Source/Runtime/Slate/Public/Framework/Styling/EditableTextBoxWidgetStyle.h b/Engine/Source/Runtime/Slate/Public/Framework/Styling/EditableTextBoxWidgetStyle.h
new file mode 100644
--- /dev/null
+++ b/Engine/Source/Runtime/Slate/Public/Framework/Styling/EditableTextBoxWidgetStyle.h
@@ -0,0 +1,24 @@
+// Copyright 1998-2014 Epic Games, Inc. All Rights Reserved.
+
+#pragma once
+
+#include "SlateWidgetStyleContainerBase.h"
+#include "EditableTextBoxWidgetStyle.generated.h"
+
+/**
+ */
+UCLASS(hidecategories=Object, MinimalAPI)
+class UEditableTextBoxWidgetStyle : public USlateWidgetStyleContainerBase
+{
+	GENERATED_UCLASS_BODY()
+
+public:
+	/** The actual data describing the button's appearance. */
+	UPROPERTY(Category=Appearance, EditAnywhere, meta=(ShowOnlyInnerProperties))
+	FEditableTextBoxStyle EditableTextBoxStyle;
+
+	virtual const struct FSlateWidgetStyle* const GetStyle() const override
+	{
+		return static_cast< const struct FSlateWidgetStyle* >( &EditableTextBoxStyle );
+	}
+};
diff --git a/Engine/Source/Runtime/Slate/Public/Framework/Styling/ProgressWidgetStyle.h b/Engine/Source/Runtime/Slate/Public/Framework/Styling/ProgressWidgetStyle.h
new file mode 100644
--- /dev/null
+++ b/Engine/Source/Runtime/Slate/Public/Framework/Styling/ProgressWidgetStyle.h
@@ -0,0 +1,24 @@
+// Copyright 1998-2014 Epic Games, Inc. All Rights Reserved.
+
+#pragma once
+
+#include "SlateWidgetStyleContainerBase.h"
+#include "ProgressWidgetStyle.generated.h"
+
+/**
+ */
+UCLASS(hidecategories=Object, MinimalAPI)
+class UProgressWidgetStyle : public USlateWidgetStyleContainerBase
+{
+	GENERATED_UCLASS_BODY()
+
+public:
+	/** The actual data describing the button's appearance. */
+	UPROPERTY(Category=Appearance, EditAnywhere, meta=(ShowOnlyInnerProperties))
+	FProgressBarStyle ProgressBarStyle;
+
+	virtual const struct FSlateWidgetStyle* const GetStyle() const override
+	{
+		return static_cast< const struct FSlateWidgetStyle* >( &ProgressBarStyle );
+	}
+};
diff --git a/Engine/Source/Runtime/Slate/Public/Framework/Styling/ScrollBoxWidgetStyle.h b/Engine/Source/Runtime/Slate/Public/Framework/Styling/ScrollBoxWidgetStyle.h
new file mode 100644
--- /dev/null
+++ b/Engine/Source/Runtime/Slate/Public/Framework/Styling/ScrollBoxWidgetStyle.h
@@ -0,0 +1,24 @@
+// Copyright 1998-2014 Epic Games, Inc. All Rights Reserved.
+
+#pragma once
+
+#include "SlateWidgetStyleContainerBase.h"
+#include "ScrollBoxWidgetStyle.generated.h"
+
+/**
+ */
+UCLASS(hidecategories=Object, MinimalAPI)
+class UScrollBoxWidgetStyle : public USlateWidgetStyleContainerBase
+{
+	GENERATED_UCLASS_BODY()
+
+public:
+	/** The actual data describing the scrollbox's appearance. */
+	UPROPERTY(Category=Appearance, EditAnywhere, meta=(ShowOnlyInnerProperties))
+	FScrollBoxStyle ScrollBoxStyle;
+
+	virtual const struct FSlateWidgetStyle* const GetStyle() const override
+	{
+		return static_cast< const struct FSlateWidgetStyle* >( &ScrollBoxStyle );
+	}
+};
diff --git a/Engine/Source/Runtime/SlateCore/Private/Styling/SlateTypes.cpp b/Engine/Source/Runtime/SlateCore/Private/Styling/SlateTypes.cpp
--- a/Engine/Source/Runtime/SlateCore/Private/Styling/SlateTypes.cpp
+++ b/Engine/Source/Runtime/SlateCore/Private/Styling/SlateTypes.cpp
@@ -223,6 +223,10 @@ FEditableTextBoxStyle::FEditableTextBoxStyle()
 {
 }
 
+FEditableTextBoxStyle::~FEditableTextBoxStyle()
+{
+}
+
 void FEditableTextBoxStyle::GetResources( TArray< const FSlateBrush* >& OutBrushes ) const
 {
 	OutBrushes.Add( &BackgroundImageNormal );
diff --git a/Engine/Source/Runtime/SlateCore/Public/Styling/SlateTypes.h b/Engine/Source/Runtime/SlateCore/Public/Styling/SlateTypes.h
--- a/Engine/Source/Runtime/SlateCore/Public/Styling/SlateTypes.h
+++ b/Engine/Source/Runtime/SlateCore/Public/Styling/SlateTypes.h
@@ -502,7 +502,7 @@ struct SLATECORE_API FEditableTextBoxStyle : public FSlateWidgetStyle
 
 	FEditableTextBoxStyle();
 
-	virtual ~FEditableTextBoxStyle() {}
+	virtual ~FEditableTextBoxStyle();
 
 	virtual void GetResources( TArray< const FSlateBrush* >& OutBrushes ) const override;
 
