Slate - Exposing a function to dynamically set the style of an editable textbox in Slate.  Only updates the brushes at the moment.

diff --git a/Engine/Source/Runtime/Slate/Private/Widgets/Input/SEditableTextBox.cpp b/Engine/Source/Runtime/Slate/Private/Widgets/Input/SEditableTextBox.cpp
--- a/Engine/Source/Runtime/Slate/Private/Widgets/Input/SEditableTextBox.cpp
+++ b/Engine/Source/Runtime/Slate/Private/Widgets/Input/SEditableTextBox.cpp
@@ -11,10 +11,8 @@
 void SEditableTextBox::Construct( const FArguments& InArgs )
 {
 	check (InArgs._Style);
-	BorderImageNormal = &InArgs._Style->BackgroundImageNormal;
-	BorderImageHovered = &InArgs._Style->BackgroundImageHovered;
-	BorderImageFocused = &InArgs._Style->BackgroundImageFocused;
-	BorderImageReadOnly = &InArgs._Style->BackgroundImageReadOnly;
+	SetStyle(InArgs._Style);
+
 	TAttribute<FMargin> Padding = InArgs._Padding.IsSet() ? InArgs._Padding : InArgs._Style->Padding;
 	TAttribute<FSlateFontInfo> Font = InArgs._Font.IsSet() ? InArgs._Font : InArgs._Style->Font;
 	TAttribute<FSlateColor> ForegroundColor = InArgs._ForegroundColor.IsSet() ? InArgs._ForegroundColor : InArgs._Style->ForegroundColor;
@@ -28,7 +26,8 @@ void SEditableTextBox::Construct( const FArguments& InArgs )
 		.Padding( 0 )
 		[
 			SAssignNew( Box, SHorizontalBox)
-			+SHorizontalBox::Slot()
+
+			+ SHorizontalBox::Slot()
 			.VAlign(VAlign_Fill)
 			.HAlign(HAlign_Fill)
 			.FillWidth(1)
@@ -68,12 +67,22 @@ void SEditableTextBox::Construct( const FArguments& InArgs )
 
 }
 
+void SEditableTextBox::SetStyle(const FEditableTextBoxStyle* InStyle)
+{
+	Style = InStyle;
+
+	if ( Style == NULL )
+	{
+		FArguments Defaults;
+		Style = Defaults._Style;
+	}
+
+	BorderImageNormal = &Style->BackgroundImageNormal;
+	BorderImageHovered = &Style->BackgroundImageHovered;
+	BorderImageFocused = &Style->BackgroundImageFocused;
+	BorderImageReadOnly = &Style->BackgroundImageReadOnly;
+}
 
-/**
- * Sets the text string currently being edited 
- *
- * @param  InNewText  The new text string
- */
 void SEditableTextBox::SetText( const TAttribute< FText >& InNewText )
 {
 	EditableText->SetText( InNewText );
diff --git a/Engine/Source/Runtime/Slate/Public/Widgets/Input/SEditableTextBox.h b/Engine/Source/Runtime/Slate/Public/Widgets/Input/SEditableTextBox.h
--- a/Engine/Source/Runtime/Slate/Public/Widgets/Input/SEditableTextBox.h
+++ b/Engine/Source/Runtime/Slate/Public/Widgets/Input/SEditableTextBox.h
@@ -106,6 +106,9 @@ public:
 		return EditableText->GetText();
 	}
 
+	/** See attribute Style */
+	void SetStyle(const FEditableTextBoxStyle* InStyle);
+
 	/**
 	 * Sets the text string currently being edited 
 	 *
@@ -136,6 +139,7 @@ public:
 	virtual FReply OnKeyDown( const FGeometry& MyGeometry, const FKeyboardEvent& InKeyboardEvent ) override;
 
 protected:
+	const FEditableTextBoxStyle* Style;
 
 	/** Editable text widget */
 	TSharedPtr< SEditableText > EditableText;
