Slate - Exposing a way to pass in a style for the scrollbar on a ScrollBox.  Also adding a UObject exposed Scrollbar Style.

diff --git a/Engine/Source/Runtime/Slate/Private/Framework/Styling/ScrollBarWidgetStyle.cpp b/Engine/Source/Runtime/Slate/Private/Framework/Styling/ScrollBarWidgetStyle.cpp
new file mode 100644
--- /dev/null
+++ b/Engine/Source/Runtime/Slate/Private/Framework/Styling/ScrollBarWidgetStyle.cpp
@@ -0,0 +1,10 @@
+// Copyright 1998-2014 Epic Games, Inc. All Rights Reserved.
+
+#include "SlatePrivatePCH.h"
+
+
+UScrollBarWidgetStyle::UScrollBarWidgetStyle( const class FPostConstructInitializeProperties& PCIP )
+	: Super(PCIP)
+{
+	
+}
\ No newline at end of file
diff --git a/Engine/Source/Runtime/Slate/Private/SlatePrivatePCH.h b/Engine/Source/Runtime/Slate/Private/SlatePrivatePCH.h
--- a/Engine/Source/Runtime/Slate/Private/SlatePrivatePCH.h
+++ b/Engine/Source/Runtime/Slate/Private/SlatePrivatePCH.h
@@ -31,3 +31,4 @@
 #include "EditableTextBoxWidgetStyle.h"
 #include "ProgressWidgetStyle.h"
 #include "ScrollBoxWidgetStyle.h"
+#include "ScrollBarWidgetStyle.h"
diff --git a/Engine/Source/Runtime/Slate/Private/Widgets/Layout/SScrollBox.cpp b/Engine/Source/Runtime/Slate/Private/Widgets/Layout/SScrollBox.cpp
--- a/Engine/Source/Runtime/Slate/Private/Widgets/Layout/SScrollBox.cpp
+++ b/Engine/Source/Runtime/Slate/Private/Widgets/Layout/SScrollBox.cpp
@@ -149,6 +149,7 @@ void SScrollBox::Construct( const FArguments& InArgs )
 		ScrollBar = 
 			SNew( SScrollBar)
 			.Thickness(FVector2D(5.0f, 5.0f))
+			.Style( InArgs._BarStyle )
 			.OnUserScrolled(this, &SScrollBox::ScrollBar_OnUserScrolled);
 
 		PanelAndScrollbar->AddSlot()
diff --git a/Engine/Source/Runtime/Slate/Public/Framework/Styling/ScrollBarWidgetStyle.h b/Engine/Source/Runtime/Slate/Public/Framework/Styling/ScrollBarWidgetStyle.h
new file mode 100644
--- /dev/null
+++ b/Engine/Source/Runtime/Slate/Public/Framework/Styling/ScrollBarWidgetStyle.h
@@ -0,0 +1,24 @@
+// Copyright 1998-2014 Epic Games, Inc. All Rights Reserved.
+
+#pragma once
+
+#include "SlateWidgetStyleContainerBase.h"
+#include "ScrollBarWidgetStyle.generated.h"
+
+/**
+ */
+UCLASS(hidecategories=Object, MinimalAPI)
+class UScrollBarWidgetStyle : public USlateWidgetStyleContainerBase
+{
+	GENERATED_UCLASS_BODY()
+
+public:
+	/** The actual data describing the scrollbox's appearance. */
+	UPROPERTY(Category=Appearance, EditAnywhere, meta=(ShowOnlyInnerProperties))
+	FScrollBarStyle ScrollBarStyle;
+
+	virtual const struct FSlateWidgetStyle* const GetStyle() const override
+	{
+		return static_cast< const struct FSlateWidgetStyle* >( &ScrollBarStyle );
+	}
+};
diff --git a/Engine/Source/Runtime/Slate/Public/Widgets/Layout/SScrollBox.h b/Engine/Source/Runtime/Slate/Public/Widgets/Layout/SScrollBox.h
--- a/Engine/Source/Runtime/Slate/Public/Widgets/Layout/SScrollBox.h
+++ b/Engine/Source/Runtime/Slate/Public/Widgets/Layout/SScrollBox.h
@@ -27,6 +27,7 @@ public:
 
 	SLATE_BEGIN_ARGS(SScrollBox)
 		: _Style( &FCoreStyle::Get().GetWidgetStyle<FScrollBoxStyle>("ScrollBox") )
+		, _BarStyle( &FCoreStyle::Get().GetWidgetStyle<FScrollBarStyle>("ScrollBar") )
 		, _ExternalScrollbar()
 		, _OnUserScrolled()
 		{}
@@ -36,6 +37,9 @@ public:
 		/** Style used to draw this scrollbox */
 		SLATE_STYLE_ARGUMENT( FScrollBoxStyle, Style )
 
+		/** Style used to draw this scrollbox's scrollbar */
+		SLATE_STYLE_ARGUMENT( FScrollBarStyle, BarStyle )
+
 		/** Custom scroll bar */
 		SLATE_ARGUMENT( TSharedPtr<SScrollBar>, ExternalScrollbar )
 
