Restore call to F<Platform>Application's DestroyApplication

diff --git a/Engine/Source/Runtime/Slate/Private/Framework/Application/SlateApplication.cpp b/Engine/Source/Runtime/Slate/Private/Framework/Application/SlateApplication.cpp
--- a/Engine/Source/Runtime/Slate/Private/Framework/Application/SlateApplication.cpp
+++ b/Engine/Source/Runtime/Slate/Private/Framework/Application/SlateApplication.cpp
@@ -308,6 +308,21 @@ void FSlateApplication::Create()
 	PlatformApplication->SetMessageHandler( CurrentApplication.ToSharedRef() );
 }
 
+void FSlateApplication::Shutdown()
+{
+	if (FSlateApplication::IsInitialized())
+	{
+		CurrentApplication->OnShutdown();
+		CurrentApplication->DestroyRenderer();
+		CurrentApplication->Renderer.Reset();
+		PlatformApplication->DestroyApplication();
+		PlatformApplication.Reset();
+		CurrentApplication.Reset();
+		CurrentBaseApplication.Reset();
+	}
+}
+
+
 TSharedPtr<FSlateApplication> FSlateApplication::CurrentApplication = NULL;
 
 FSlateApplication::FSlateApplication()
diff --git a/Engine/Source/Runtime/Slate/Public/Framework/Application/SlateApplication.h b/Engine/Source/Runtime/Slate/Public/Framework/Application/SlateApplication.h
--- a/Engine/Source/Runtime/Slate/Public/Framework/Application/SlateApplication.h
+++ b/Engine/Source/Runtime/Slate/Public/Framework/Application/SlateApplication.h
@@ -161,18 +161,7 @@ public:
 		return *CurrentApplication;
 	}
 
-	static void Shutdown()
-	{
-		if ( FSlateApplication::IsInitialized() )
-		{
-			CurrentApplication->OnShutdown();
-			CurrentApplication->DestroyRenderer();
-			CurrentApplication->Renderer.Reset();
-			PlatformApplication.Reset();
-			CurrentApplication.Reset();
-			CurrentBaseApplication.Reset();
-		}
-	}
+	static void Shutdown();
 
 	/** @return the global tab manager */
 	static TSharedRef<class FGlobalTabmanager> GetGlobalTabManager();
