Added heartbeat analytics to track idle time

diff --git a/Engine/Source/Editor/UnrealEd/Private/UnrealEdMisc.cpp b/Engine/Source/Editor/UnrealEd/Private/UnrealEdMisc.cpp
--- a/Engine/Source/Editor/UnrealEd/Private/UnrealEdMisc.cpp
+++ b/Engine/Source/Editor/UnrealEd/Private/UnrealEdMisc.cpp
@@ -313,6 +313,13 @@ void FUnrealEdMisc::OnInit()
 
 	// Send Project Analytics
 	InitEngineAnalytics();
+	
+	// Setup a timer for a heartbeat event to track if users are actually using the editor or it is idle.
+	float Seconds = 60.0f;
+	FTimerDelegate Delegate;
+	Delegate.BindRaw( this, &FUnrealEdMisc::EditorAnalyticsHeartbeat );
+	GEditor->GetTimerManager()->SetTimer( Delegate, Seconds, true );
+
 }
 
 void FUnrealEdMisc::InitEngineAnalytics()
@@ -383,6 +390,22 @@ void FUnrealEdMisc::InitEngineAnalytics()
 	}
 }
 
+void FUnrealEdMisc::EditorAnalyticsHeartbeat()
+{
+	static double LastHeartbeatTime = FPlatformTime::Seconds();
+	
+	double LastInteractionTime = FSlateApplication::Get().GetLastUserInteractionTime();
+	
+	// Did the user interact since the last heartbeat
+	bool bIdle = LastInteractionTime < LastHeartbeatTime;
+	
+	TArray< FAnalyticsEventAttribute > Attributes;
+	Attributes.Add(FAnalyticsEventAttribute(TEXT("Idle"), bIdle));
+	FEngineAnalytics::GetProvider().RecordEvent(TEXT("Editor.Usage.Heartbeat"), Attributes);
+	
+	LastHeartbeatTime = FPlatformTime::Seconds();
+}
+
 void FUnrealEdMisc::TickAssetAnalytics()
 {
 	if( bIsAssetAnalyticsPending )
diff --git a/Engine/Source/Editor/UnrealEd/Public/UnrealEdMisc.h b/Engine/Source/Editor/UnrealEd/Public/UnrealEdMisc.h
--- a/Engine/Source/Editor/UnrealEd/Public/UnrealEdMisc.h
+++ b/Engine/Source/Editor/UnrealEd/Public/UnrealEdMisc.h
@@ -260,6 +260,8 @@ private:
 	/** Initialize engine analytics */
 	void InitEngineAnalytics();
 
+	/** Called when the heartbeat event should be sent to engine analytics */
+	void EditorAnalyticsHeartbeat();
 
 private:
 
diff --git a/Engine/Source/Runtime/Slate/Private/Framework/Application/SlateApplication.cpp b/Engine/Source/Runtime/Slate/Private/Framework/Application/SlateApplication.cpp
--- a/Engine/Source/Runtime/Slate/Private/Framework/Application/SlateApplication.cpp
+++ b/Engine/Source/Runtime/Slate/Private/Framework/Application/SlateApplication.cpp
@@ -2924,6 +2924,8 @@ bool FSlateApplication::ProcessKeyDownEvent( FKeyboardEvent& InKeyboardEvent )
 {
 	FReply Reply = FReply::Unhandled();
 
+	LastUserInteractionTime = this->GetCurrentTime();
+	
 	if (IsDragDropping() && InKeyboardEvent.GetKey() == EKeys::Escape)
 	{
 		// Pressing ESC while drag and dropping terminates the drag drop.
@@ -2935,7 +2937,7 @@ bool FSlateApplication::ProcessKeyDownEvent( FKeyboardEvent& InKeyboardEvent )
 		const int32 EventCount = InKeyboardEvent.IsRepeat() ? SlateApplicationDefs::NumRepeatsPerActualRepeat : 1;
 		for( int32 CurEventIndex = 0; CurEventIndex < EventCount; ++CurEventIndex )
 		{
-			LastUserInteractionTimeForThrottling = this->GetCurrentTime();
+			LastUserInteractionTimeForThrottling = LastUserInteractionTime;
 
 			// If we are inspecting, pressing ESC exits inspection mode.
 			if ( InKeyboardEvent.GetKey() == EKeys::Escape )
@@ -3026,10 +3028,12 @@ bool FSlateApplication::ProcessKeyUpEvent( FKeyboardEvent& InKeyboardEvent )
 {
 	FReply Reply = FReply::Unhandled();
 
+	LastUserInteractionTime = this->GetCurrentTime();
+	
 	const int32 EventCount = InKeyboardEvent.IsRepeat() ? SlateApplicationDefs::NumRepeatsPerActualRepeat : 1;
 	for( int32 CurEventIndex = 0; CurEventIndex < EventCount; ++CurEventIndex )
 	{
-		LastUserInteractionTimeForThrottling = this->GetCurrentTime();
+		LastUserInteractionTimeForThrottling = LastUserInteractionTime;
 
 		// Bubble the keyboard event
 		FWidgetPath EventPath = FocusedWidgetPath.ToWidgetPath();
@@ -3126,8 +3130,9 @@ bool FSlateApplication::OnMouseDown( const TSharedPtr< FGenericWindow >& Platfor
 
 bool FSlateApplication::ProcessMouseButtonDownEvent( const TSharedPtr< FGenericWindow >& PlatformWindow, FPointerEvent& MouseEvent )
 {
-	LastUserInteractionTimeForThrottling = this->GetCurrentTime();
-
+	LastUserInteractionTime = this->GetCurrentTime();
+	LastUserInteractionTimeForThrottling = LastUserInteractionTime;
+	
 	PlatformApplication->SetCapture( PlatformWindow );
 	PressedMouseButtons.Add( MouseEvent.GetEffectingButton() );
 
@@ -3277,8 +3282,9 @@ bool FSlateApplication::OnMouseDoubleClick( const TSharedPtr< FGenericWindow >&
 
 bool FSlateApplication::ProcessMouseButtonDoubleClickEvent( const TSharedPtr< FGenericWindow >& PlatformWindow, FPointerEvent& InMouseEvent )
 {
-	LastUserInteractionTimeForThrottling = this->GetCurrentTime();
-
+	LastUserInteractionTime = this->GetCurrentTime();
+	LastUserInteractionTimeForThrottling = LastUserInteractionTime;
+	
 	PlatformApplication->SetCapture( PlatformWindow );
 	PressedMouseButtons.Add( InMouseEvent.GetEffectingButton() );
 
@@ -3332,7 +3338,8 @@ bool FSlateApplication::OnMouseUp( const EMouseButtons::Type Button )
 
 bool FSlateApplication::ProcessMouseButtonUpEvent( FPointerEvent& MouseEvent )
 {
-	LastUserInteractionTimeForThrottling = this->GetCurrentTime();
+	LastUserInteractionTime = this->GetCurrentTime();
+	LastUserInteractionTimeForThrottling = LastUserInteractionTime;
 	PressedMouseButtons.Remove( MouseEvent.GetEffectingButton() );
 
 	if (DragDetector.DetectDragForWidget.IsValid() && MouseEvent.GetEffectingButton() == DragDetector.DetectDragButton )
@@ -3449,6 +3456,8 @@ bool FSlateApplication::ProcessMouseWheelOrGestureEvent( FPointerEvent& InWheelE
 		return false;
 	}
 
+	LastUserInteractionTime = this->GetCurrentTime();
+	
 	// NOTE: We intentionally don't reset LastUserInteractionTimeForThrottling here so that the UI can be responsive while scrolling
 
 	FWidgetPath EventPath = ( MouseCaptor.IsValid() )
@@ -3545,6 +3554,9 @@ bool FSlateApplication::ProcessMouseMoveEvent( FPointerEvent& MouseEvent )
 		// Detecting a mouse move of zero delta is our way of filtering out synthesized move events
 		const bool AllowSpawningOfToolTips = true;
 		UpdateToolTip( AllowSpawningOfToolTips );
+		
+		// Guard against synthesized mouse moves and only track user interaction if the cursor pos changed
+		LastUserInteractionTime = this->GetCurrentTime();
 	}
 
 	FWidgetPath WidgetsUnderCursor = LocateWindowUnderMouse( MouseEvent.GetScreenSpacePosition(), GetInteractiveTopLevelWindows() );
@@ -3859,6 +3871,8 @@ bool FSlateApplication::OnControllerAnalog( EControllerButtons::Type Button, int
 
 void FSlateApplication::ProcessControllerAnalogValueChangedEvent( FControllerEvent& ControllerEvent )
 {
+	LastUserInteractionTime = this->GetCurrentTime();
+	
 	CALL_WIDGET_FUNCTION(ControllerEvent, OnControllerAnalogValueChanged);
 }
 
@@ -3886,6 +3900,8 @@ bool FSlateApplication::OnControllerButtonPressed( EControllerButtons::Type Butt
 
 void FSlateApplication::ProcessControllerButtonPressedEvent( FControllerEvent& ControllerEvent )
 {
+	LastUserInteractionTime = this->GetCurrentTime();
+	
 	CALL_WIDGET_FUNCTION(ControllerEvent, OnControllerButtonPressed);
 }
 
@@ -3913,6 +3929,8 @@ bool FSlateApplication::OnControllerButtonReleased( EControllerButtons::Type But
 
 void FSlateApplication::ProcessControllerButtonReleasedEvent( FControllerEvent& ControllerEvent )
 {
+	LastUserInteractionTime = this->GetCurrentTime();
+	
 	CALL_WIDGET_FUNCTION(ControllerEvent, OnControllerButtonReleased);
 
 }
@@ -4012,6 +4030,8 @@ bool FSlateApplication::OnMotionDetected(const FVector& Tilt, const FVector& Rot
 
 void FSlateApplication::ProcessMotionDetectedEvent( FMotionEvent& MotionEvent )
 {
+	LastUserInteractionTime = this->GetCurrentTime();
+	
 	CALL_WIDGET_FUNCTION(MotionEvent, OnMotionDetected);
 }
 
@@ -4133,6 +4153,9 @@ bool FSlateApplication::ProcessWindowActivatedEvent( const FWindowActivateEvent&
 
 	if ( ActivateEvent.GetActivationType() != FWindowActivateEvent::EA_Deactivate )
 	{
+		// Only window activate considered a user interaction
+		LastUserInteractionTime = this->GetCurrentTime();
+		
 		// Do not process activation messages unless we have no modal windows or the current window is modal
 		if( !ActiveModalWindow.IsValid() || ActivateEvent.GetAffectedWindow() == ActiveModalWindow || ActivateEvent.GetAffectedWindow()->IsDescendantOf(ActiveModalWindow) )
 		{
@@ -4369,6 +4392,8 @@ EDropEffect::Type FSlateApplication::OnDragEnter( const TSharedRef< SWindow >& W
 
 bool FSlateApplication::ProcessDragEnterEvent( TSharedRef<SWindow> WindowEntered, FDragDropEvent& DragDropEvent )
 {
+	LastUserInteractionTime = this->GetCurrentTime();
+	
 	FWidgetPath WidgetsUnderCursor = LocateWindowUnderMouse( DragDropEvent.GetScreenSpacePosition(), GetInteractiveTopLevelWindows() );
 	DragDropEvent.SetEventPath( WidgetsUnderCursor );
 
diff --git a/Engine/Source/Runtime/Slate/Public/Framework/Application/SlateApplication.h b/Engine/Source/Runtime/Slate/Public/Framework/Application/SlateApplication.h
--- a/Engine/Source/Runtime/Slate/Public/Framework/Application/SlateApplication.h
+++ b/Engine/Source/Runtime/Slate/Public/Framework/Application/SlateApplication.h
@@ -899,6 +899,9 @@ public:
 	/** Sets the handler for otherwise unhandled key down events. This is used by the editor to provide a global action list, if the key was not consumed by any widget. */
 	void SetUnhandledKeyDownEventHandler( const FOnKeyboardEvent& NewHandler );
 
+	/** @return the last time a user interacted with a keyboard, mouse, touch device, or joystick */
+	double GetLastUserInteractionTime() const { return LastUserInteractionTime; }
+	
 public:
 
 	// Begin FSlateApplicationBase interface
@@ -1169,6 +1172,9 @@ private:
 	FThrottleRequest UserInteractionResponsivnessThrottle;
 
 	/** The last real time that the user pressed a key or mouse button */
+	double LastUserInteractionTime;
+
+	/** Subset of LastUserInteractionTime that is used only when consdering when to throttle */
 	double LastUserInteractionTimeForThrottling;
 
 
