Slate/UMG - Changing BarStyle to ScrollBarStyle

diff --git a/Engine/Source/Runtime/Slate/Private/Widgets/Layout/SScrollBox.cpp b/Engine/Source/Runtime/Slate/Private/Widgets/Layout/SScrollBox.cpp
--- a/Engine/Source/Runtime/Slate/Private/Widgets/Layout/SScrollBox.cpp
+++ b/Engine/Source/Runtime/Slate/Private/Widgets/Layout/SScrollBox.cpp
@@ -139,7 +139,7 @@ void SScrollBox::Construct( const FArguments& InArgs )
 
 	if (InArgs._ExternalScrollbar.IsValid())
 	{
-		// An external scroll bar was sepecified by the user
+		// An external scroll bar was specified by the user
 		ScrollBar = InArgs._ExternalScrollbar;
 		ScrollBar->SetOnUserScrolled(FOnUserScrolled::CreateSP(this, &SScrollBox::ScrollBar_OnUserScrolled));
 	}
@@ -149,7 +149,7 @@ void SScrollBox::Construct( const FArguments& InArgs )
 		ScrollBar = 
 			SNew( SScrollBar)
 			.Thickness(FVector2D(5.0f, 5.0f))
-			.Style( InArgs._BarStyle )
+			.Style( InArgs._ScrollBarStyle )
 			.OnUserScrolled(this, &SScrollBox::ScrollBar_OnUserScrolled);
 
 		PanelAndScrollbar->AddSlot()
diff --git a/Engine/Source/Runtime/Slate/Public/Widgets/Layout/SScrollBox.h b/Engine/Source/Runtime/Slate/Public/Widgets/Layout/SScrollBox.h
--- a/Engine/Source/Runtime/Slate/Public/Widgets/Layout/SScrollBox.h
+++ b/Engine/Source/Runtime/Slate/Public/Widgets/Layout/SScrollBox.h
@@ -27,7 +27,7 @@ public:
 
 	SLATE_BEGIN_ARGS(SScrollBox)
 		: _Style( &FCoreStyle::Get().GetWidgetStyle<FScrollBoxStyle>("ScrollBox") )
-		, _BarStyle( &FCoreStyle::Get().GetWidgetStyle<FScrollBarStyle>("ScrollBar") )
+		, _ScrollBarStyle( &FCoreStyle::Get().GetWidgetStyle<FScrollBarStyle>("ScrollBar") )
 		, _ExternalScrollbar()
 		, _OnUserScrolled()
 		{}
@@ -38,7 +38,7 @@ public:
 		SLATE_STYLE_ARGUMENT( FScrollBoxStyle, Style )
 
 		/** Style used to draw this scrollbox's scrollbar */
-		SLATE_STYLE_ARGUMENT( FScrollBarStyle, BarStyle )
+		SLATE_STYLE_ARGUMENT( FScrollBarStyle, ScrollBarStyle )
 
 		/** Custom scroll bar */
 		SLATE_ARGUMENT( TSharedPtr<SScrollBar>, ExternalScrollbar )
diff --git a/Engine/Source/Runtime/UMG/Private/Components/ScrollBox.cpp b/Engine/Source/Runtime/UMG/Private/Components/ScrollBox.cpp
--- a/Engine/Source/Runtime/UMG/Private/Components/ScrollBox.cpp
+++ b/Engine/Source/Runtime/UMG/Private/Components/ScrollBox.cpp
@@ -52,12 +52,12 @@ TSharedRef<SWidget> UScrollBox::RebuildWidget()
 	if ( BarStylePtr == NULL )
 	{
 		SScrollBox::FArguments Defaults;
-		BarStylePtr = Defaults._BarStyle;
+		BarStylePtr = Defaults._ScrollBarStyle;
 	}
 
 	MyScrollBox = SNew(SScrollBox)
 		.Style(StylePtr)
-		.BarStyle(BarStylePtr);
+		.ScrollBarStyle(BarStylePtr);
 
 	for ( UPanelSlot* Slot : Slots )
 	{
