Fixed bugs in overscroll, so that the bounce and overscroll amount is smoother and more consistent.
Fixes TTP 339241 iOS style scrolling has odd velocity and drag properties when dragging/velocity past the end of the list
Problem was that smooth scroll offset was at some point changed to use fractional items, but I was using the assumption that it was using Slate Units.

diff --git a/Engine/Source/Runtime/Slate/Private/Widgets/Views/SListPanel.cpp b/Engine/Source/Runtime/Slate/Private/Widgets/Views/SListPanel.cpp
--- a/Engine/Source/Runtime/Slate/Private/Widgets/Views/SListPanel.cpp
+++ b/Engine/Source/Runtime/Slate/Private/Widgets/Views/SListPanel.cpp
@@ -15,6 +15,7 @@ void SListPanel::Construct( const FArguments& InArgs )
 {
 	PreferredRowNum = 0;
 	SmoothScrollOffsetInItems = 0;
+	OverscrollAmount = 0;
 	ItemWidth = InArgs._ItemWidth;
 	ItemHeight = InArgs._ItemHeight;
 	NumDesiredItems = InArgs._NumDesiredItems;
@@ -57,7 +58,7 @@ void SListPanel::OnArrangeChildren( const FGeometry& AllottedGeometry, FArranged
 		float WidthSoFar = 0;
 		float LocalItemWidth = ItemWidth.Get();
 		float LocalItemHeight = ItemHeight.Get();
-		float HeightSoFar = -FMath::FloorToInt(SmoothScrollOffsetInItems * LocalItemHeight);
+		float HeightSoFar = -FMath::FloorToInt(SmoothScrollOffsetInItems * LocalItemHeight) - OverscrollAmount;
 
 		for( int32 ItemIndex = 0; ItemIndex < Children.Num(); ++ItemIndex )
 		{
@@ -79,7 +80,7 @@ void SListPanel::OnArrangeChildren( const FGeometry& AllottedGeometry, FArranged
 		if (Children.Num() > 0)
 		{
 			// This is a normal list, arrange items vertically
-			float HeightSoFar = -FMath::FloorToInt(SmoothScrollOffsetInItems * Children[0].Widget->GetDesiredSize().Y);
+			float HeightSoFar = -FMath::FloorToInt(SmoothScrollOffsetInItems * Children[0].Widget->GetDesiredSize().Y)-OverscrollAmount;
 			for( int32 ItemIndex = 0; ItemIndex < Children.Num(); ++ItemIndex )
 			{
 					const FVector2D ItemDesiredSize = Children[ItemIndex].Widget->GetDesiredSize();
@@ -178,6 +179,11 @@ void SListPanel::SmoothScrollOffset(float InOffsetInItems)
 	SmoothScrollOffsetInItems = InOffsetInItems;
 }
 
+void SListPanel::SetOverscrollAmount( float InOverscrollAmount )
+{
+	OverscrollAmount = InOverscrollAmount;
+}
+
 /** Remove all the children from this panel */
 void SListPanel::ClearItems()
 {
diff --git a/Engine/Source/Runtime/Slate/Private/Widgets/Views/SListPanel.h b/Engine/Source/Runtime/Slate/Private/Widgets/Views/SListPanel.h
--- a/Engine/Source/Runtime/Slate/Private/Widgets/Views/SListPanel.h
+++ b/Engine/Source/Runtime/Slate/Private/Widgets/Views/SListPanel.h
@@ -67,6 +67,9 @@ public:
 
 	/** Fraction of the first item that we should offset by to simulate smooth scrolling. */
 	void SmoothScrollOffset(float OffsetInItems);
+
+	/** Set how much we should appear to have scrolled past the beginning/end of the list. */
+	void SetOverscrollAmount( float InOverscrollAmount );
 	
 	/** Remove all the children from this panel */
 	void ClearItems();
@@ -117,6 +120,9 @@ protected:
 	 */
 	float SmoothScrollOffsetInItems;
 
+	/** Amount scrolled past beginning/end of list in Slate Units. */
+	float OverscrollAmount;
+
 	/** The preferred number of rows that this widget should have */
 	int32 PreferredRowNum;
 
diff --git a/Engine/Source/Runtime/Slate/Private/Widgets/Views/STableViewBase.cpp b/Engine/Source/Runtime/Slate/Private/Widgets/Views/STableViewBase.cpp
--- a/Engine/Source/Runtime/Slate/Private/Widgets/Views/STableViewBase.cpp
+++ b/Engine/Source/Runtime/Slate/Private/Widgets/Views/STableViewBase.cpp
@@ -5,8 +5,8 @@
 
 namespace ListConstants
 {
-	static const float OvershootMax = 50.0f;
-	static const float OvershootBounceRate = 350.0f;
+	static const float OvershootMax = 150.0f;
+	static const float OvershootBounceRate = 250.0f;
 }
 
 void STableViewBase::ConstructChildren( const TAttribute<float>& InItemWidth, const TAttribute<float>& InItemHeight, const TSharedPtr<SHeaderRow>& InHeaderRow, const TSharedPtr<SScrollBar>& InScrollBar )
@@ -198,8 +198,9 @@ void STableViewBase::Tick( const FGeometry& AllottedGeometry, const double InCur
 			
 			
 			ScrollOffset = FMath::Max(0.0, ScrollOffset);
+			ItemsPanel->SmoothScrollOffset( FMath::Fractional(ScrollOffset / GetNumItemsWide()) );
 			const float OverscrollAmount = Overscroll.GetOverscroll();
-			ItemsPanel->SmoothScrollOffset( FMath::Fractional(ScrollOffset / GetNumItemsWide()) +  OverscrollAmount );
+			ItemsPanel->SetOverscrollAmount( OverscrollAmount );
 			
 
 			UpdateSelectionSet();
@@ -525,7 +526,7 @@ STableViewBase::STableViewBase( ETableViewMode::Type InTableViewMode )
 	, SelectionMode( ESelectionMode::Multi )
 	, SoftwareCursorPosition( ForceInitToZero )
 	, bShowSoftwareCursor( false )
-	, Overscroll( ListConstants::OvershootMax )
+	, Overscroll()
 	, bItemsNeedRefresh( true )	
 {
 }
@@ -669,17 +670,15 @@ TSharedRef<class SWidget> STableViewBase::GetScrollWidget()
 }
 
 
-STableViewBase::FOverscroll::FOverscroll( const float InMaxOverscroll )
+STableViewBase::FOverscroll::FOverscroll()
 : OverscrollAmount( 0.0f )
-, MaxOverscroll( InMaxOverscroll )
 {
 }
 
 float STableViewBase::FOverscroll::ScrollBy( float Delta )
 {
 	const float ValueBeforeDeltaApplied = OverscrollAmount;
-	const float OverscrollMagnitude = OverscrollAmount*OverscrollAmount + 1.0f;
-	const float EasedDelta = Delta / OverscrollMagnitude;
+	const float EasedDelta = Delta / (FMath::Abs(OverscrollAmount/ListConstants::OvershootMax)+1.0f);
 	OverscrollAmount = FMath::Clamp(OverscrollAmount + EasedDelta, -ListConstants::OvershootMax, ListConstants::OvershootMax);
 
 	return ValueBeforeDeltaApplied - OverscrollAmount;
@@ -692,8 +691,7 @@ float STableViewBase::FOverscroll::GetOverscroll() const
 
 void STableViewBase::FOverscroll::UpdateOverscroll( float InDeltaTime )
 {
-	const float OverscrollMagnitude = OverscrollAmount*OverscrollAmount;
-	const float PullForce = OverscrollMagnitude + 1.0f;
+	const float PullForce = FMath::Abs(OverscrollAmount) + 1.0f;
 	const float EasedDelta = ListConstants::OvershootBounceRate * InDeltaTime * PullForce;
 
 	if ( OverscrollAmount > 0 )
diff --git a/Engine/Source/Runtime/Slate/Public/Widgets/Views/SListView.h b/Engine/Source/Runtime/Slate/Public/Widgets/Views/SListView.h
--- a/Engine/Source/Runtime/Slate/Public/Widgets/Views/SListView.h
+++ b/Engine/Source/Runtime/Slate/Public/Widgets/Views/SListView.h
@@ -1129,7 +1129,7 @@ protected:
 		}
 		else if ( AllowOverscroll == EAllowOverscroll::Yes && Overscroll.ShouldApplyOverscroll( ScrollOffset == 0, bWasAtEndOfList, ScrollByAmountInSlateUnits ) )
 		{
-			const float UnclampedScrollDelta = FMath::Sign(ScrollByAmountInSlateUnits) * AbsScrollByAmount / 10;				
+			const float UnclampedScrollDelta = FMath::Sign(ScrollByAmountInSlateUnits) * AbsScrollByAmount;				
 			const float ActuallyScrolledBy = Overscroll.ScrollBy( UnclampedScrollDelta );
 			if (ActuallyScrolledBy != 0.0f)
 			{
diff --git a/Engine/Source/Runtime/Slate/Public/Widgets/Views/STableViewBase.h b/Engine/Source/Runtime/Slate/Public/Widgets/Views/STableViewBase.h
--- a/Engine/Source/Runtime/Slate/Public/Widgets/Views/STableViewBase.h
+++ b/Engine/Source/Runtime/Slate/Public/Widgets/Views/STableViewBase.h
@@ -253,7 +253,7 @@ protected:
 	{
 	public:
 
-		FOverscroll( const float InMaxOverscroll );
+		FOverscroll();
 
 		/** @return The Amount actually scrolled */
 		float ScrollBy( float Delta );
@@ -278,9 +278,6 @@ protected:
 	private:
 		/** How much we've over-scrolled above/below the beginning/end of the list. */
 		float OverscrollAmount;
-
-		/** The maximum amount that we can scroll past the edge. */
-		const float MaxOverscroll;
 	} Overscroll;
 
 private:
diff --git a/Engine/Source/Runtime/Slate/Public/Widgets/Views/STileView.h b/Engine/Source/Runtime/Slate/Public/Widgets/Views/STileView.h
--- a/Engine/Source/Runtime/Slate/Public/Widgets/Views/STileView.h
+++ b/Engine/Source/Runtime/Slate/Public/Widgets/Views/STileView.h
@@ -269,7 +269,7 @@ protected:
 		}
 		else if ( AllowOverscroll == EAllowOverscroll::Yes && S::Overscroll.ShouldApplyOverscroll( S::ScrollOffset == 0, S::bWasAtEndOfList, ScrollByAmountInSlateUnits ) )
 		{
-			const float UnclampedScrollDelta = ScrollByAmountInSlateUnits / (10 * GetNumItemsWide());
+			const float UnclampedScrollDelta = ScrollByAmountInSlateUnits / GetNumItemsWide();
 			const float ActuallyScrolledBy = S::Overscroll.ScrollBy( UnclampedScrollDelta );
 			if (ActuallyScrolledBy != 0.0f)
 			{
