Slate: Fixed default test suite scenario - a nomad tab was being spawned as a floating child of a non-global tab manager, which should be impossible.

diff --git a/Engine/Source/Runtime/Slate/Private/Widgets/Testing/STestSuite.cpp b/Engine/Source/Runtime/Slate/Private/Widgets/Testing/STestSuite.cpp
--- a/Engine/Source/Runtime/Slate/Private/Widgets/Testing/STestSuite.cpp
+++ b/Engine/Source/Runtime/Slate/Private/Widgets/Testing/STestSuite.cpp
@@ -3413,17 +3413,6 @@ TSharedRef<SDockTab> SpawnTestSuite1( const FSpawnTabArgs& Args )
 		)
 	)
 #if PLATFORM_SUPPORTS_MULTIPLE_NATIVE_WINDOWS
-	->AddArea
-	(
-		// This area will get a 400x600 window at 10,10
-		FTabManager::NewArea(400,600)
-		->SetWindow( FVector2D(10,10), false )
-		->Split
-		(
-			// The area contains a single tab with the widget reflector.
-			FTabManager::NewStack()->AddTab( "WidgetReflector", ETabState::OpenedTab )
-		)
-	)
 	->AddArea
 	(
 		FTabManager::NewArea(320,240)
@@ -3527,21 +3516,7 @@ TSharedRef<SDockTab> SpawnTestSuite2( const FSpawnTabArgs& Args )
 			->AddTab("GridPanelTest", ETabState::OpenedTab)
 			->AddTab("DPIScalingTest", ETabState::OpenedTab)
 		)
-	)
-#if PLATFORM_SUPPORTS_MULTIPLE_NATIVE_WINDOWS
-	->AddArea
-	(
-		// This area will get a 400x600 window at 10,10
-		FTabManager::NewArea(400,600)
-		->SetWindow( FVector2D(10,10), false )
-		->Split
-		(
-			// The area contains a single tab with the widget reflector.
-			FTabManager::NewStack()->AddTab( "WidgetReflector", ETabState::OpenedTab )
-		)
-	)
-#endif
-	;
+	);
 
 	TSharedRef<SDockTab> TestSuite2Tab = 
 		SNew(SDockTab)
@@ -3620,7 +3595,22 @@ void RestoreSlateTestSuite()
 			->AddTab( "TestSuite2", ETabState::OpenedTab )
 			->AddTab( "TestSuite1", ETabState::OpenedTab )
 		)		
-	);
+	)
+	#if PLATFORM_SUPPORTS_MULTIPLE_NATIVE_WINDOWS
+	->AddArea
+	(
+		// This area will get a 400x600 window at 10,10
+		FTabManager::NewArea(400,600)
+		->SetWindow( FVector2D(10,10), false )
+		->Split
+		(
+			// The area contains a single tab with the widget reflector.
+			FTabManager::NewStack()->AddTab( "WidgetReflector", ETabState::OpenedTab )
+		)
+	)
+	#endif
+	;
+
 
 	FGlobalTabmanager::Get()->RestoreFrom( Layout, TSharedPtr<SWindow>() );
 }
