Slate: Fixed floating child tabs being offset by (50,50) whenever their corresponding major tab was relocated.

diff --git a/Engine/Source/Runtime/Slate/Private/Framework/Docking/TabManager.cpp b/Engine/Source/Runtime/Slate/Private/Framework/Docking/TabManager.cpp
--- a/Engine/Source/Runtime/Slate/Private/Framework/Docking/TabManager.cpp
+++ b/Engine/Source/Runtime/Slate/Private/Framework/Docking/TabManager.cpp
@@ -1668,7 +1668,7 @@ void FGlobalTabmanager::OnTabRelocated( const TSharedRef<SDockTab>& RelocatedTab
 				{
 					TSharedRef<SWindow> NewChildWindow = SNew(SWindow)
 					.AutoCenter(EAutoCenter::None)
-					.ScreenPosition(OldChildWindow->GetPositionInScreen() + FVector2D(50,50))
+					.ScreenPosition(OldChildWindow->GetPositionInScreen())
 					.ClientSize(OldChildWindow->GetSizeInScreen())
 					.SupportsMinimize(false) .SupportsMaximize(false)
 					.CreateTitleBar(false)
