Removed some dead code from slate tab manager.

diff --git a/Engine/Source/Runtime/Slate/Private/Framework/Docking/SDockingSplitter.cpp b/Engine/Source/Runtime/Slate/Private/Framework/Docking/SDockingSplitter.cpp
--- a/Engine/Source/Runtime/Slate/Private/Framework/Docking/SDockingSplitter.cpp
+++ b/Engine/Source/Runtime/Slate/Private/Framework/Docking/SDockingSplitter.cpp
@@ -282,43 +282,13 @@ TArray< TSharedRef<SDockTab> > SDockingSplitter::GetAllChildTabs() const
 	return ChildTabs;
 }
 
+
 EOrientation SDockingSplitter::GetOrientation() const
 {
 	return Splitter->GetOrientation();
 }
 
 
-bool SDockingSplitter::ClearRedundantNodes( const TSharedRef<SDockingSplitter>& ParentNode )
-{
-	for (int32 ChildIndex=0; ChildIndex < ParentNode->Children.Num(); ++ChildIndex)
-	{
-		const TSharedRef<SDockingNode> ChildNode = ParentNode->Children[ChildIndex];
-		if (ChildNode->GetNodeType() == SDockingNode::DockSplitter || ChildNode->GetNodeType() == SDockingNode::DockArea)
-		{
-			const TSharedRef<SDockingSplitter>& ChildSplitter = StaticCastSharedRef<SDockingSplitter>(ChildNode);
-			if ( ChildSplitter->GetOrientation() == ParentNode->GetOrientation() || ChildSplitter->Children.Num() == 1 )
-			{
-				// We found a child splitter with the same orientation as ours.
-				// Clean up by bumping children up to our level.
-				ParentNode->RemoveChildAt( ChildIndex );
-				
-				for (int32 GrandchildIndex=0; GrandchildIndex < ChildSplitter->Children.Num(); ++GrandchildIndex)
-				{
-					ParentNode->AddChildNode(ChildSplitter->Children[GrandchildIndex], ChildIndex + GrandchildIndex);
-				}
-
-				// We fixed some stuff up, and we need to re-run this routine
-				return true;
-			}
-		}
-	}
-
-	// There were no redundant splitters.
-	return false;
-
-
-}
-
 TSharedPtr<FTabManager::FLayoutNode> SDockingSplitter::GatherPersistentLayout() const
 {
 	// Assume that all the nodes were dragged out, and there's no meaningful layout data to be gathered.
diff --git a/Engine/Source/Runtime/Slate/Private/Framework/Docking/SDockingSplitter.h b/Engine/Source/Runtime/Slate/Private/Framework/Docking/SDockingSplitter.h
--- a/Engine/Source/Runtime/Slate/Private/Framework/Docking/SDockingSplitter.h
+++ b/Engine/Source/Runtime/Slate/Private/Framework/Docking/SDockingSplitter.h
@@ -89,13 +89,6 @@ protected:
 
 	float ComputeChildCoefficientTotal() const;
 
-	/**
-	 * @param ParentNode    DockSplitter node whose children might be redundant
-	 *
-	 * @return true if we found some redundant splitters and need to do another pass on this list; false if we found no redundancies.
-	 */
-	static bool ClearRedundantNodes( const TSharedRef<SDockingSplitter>& ParentNode );
-
 	/** The SSplitter widget that SDockingSplitter wraps. */
 	TSharedPtr<SSplitter> Splitter;
 	
