Slate: Fix for multiple sets of window chrome widgets showing up (one for each tab).
Fixes TTP #340442. Long-standing bug, but the new level browser's use case exposes it.

diff --git a/Engine/Source/Runtime/Slate/Private/Framework/Docking/SDockingArea.cpp b/Engine/Source/Runtime/Slate/Private/Framework/Docking/SDockingArea.cpp
--- a/Engine/Source/Runtime/Slate/Private/Framework/Docking/SDockingArea.cpp
+++ b/Engine/Source/Runtime/Slate/Private/Framework/Docking/SDockingArea.cpp
@@ -27,7 +27,13 @@ void SDockingArea::Construct( const FArguments& InArgs, const TSharedRef<FTabMan
 			SAssignNew(Splitter, SSplitter)
 			. Orientation( PersistentNode->GetOrientation() )
 		]
-
+		
+		+ SOverlay::Slot()
+		// Houses the minimize, maximize, restore, and icon
+		.Expose(WindowControlsArea)
+		.HAlign(HAlign_Fill)
+		.VAlign(VAlign_Top)
+		
 		+SOverlay::Slot()
 		.HAlign(HAlign_Left)
 		.VAlign(VAlign_Fill)
@@ -190,7 +196,7 @@ void SDockingArea::HideCross()
 void SDockingArea::CleanUp( ELayoutModification RemovalMethod )
 {
 	const ECleanupRetVal CleanupResult = CleanUpNodes();
-
+	
 	if ( CleanupResult != VisibleTabsUnderNode )
 	{
 		bIsCenterTargetVisible = true;
@@ -221,6 +227,11 @@ void SDockingArea::CleanUp( ELayoutModification RemovalMethod )
 	{
 		bIsCenterTargetVisible = false;
 	}
+
+	// In some cases a dock area will control the window,
+	// and we need to move some of the tabs out of the way
+	// to make room for window chrome.
+	MakeRoomForWindowChrome();
 }
 
 void SDockingArea::SetParentWindow( const TSharedRef<SWindow>& NewParentWindow )
@@ -228,6 +239,12 @@ void SDockingArea::SetParentWindow( const TSharedRef<SWindow>& NewParentWindow )
 	if( bManageParentWindow )
 	{
 		NewParentWindow->SetRequestDestroyWindowOverride( FRequestDestroyWindowOverride::CreateSP( this, &SDockingArea::OnOwningWindowBeingDestroyed ) );
+
+		TSharedPtr<IWindowTitleBar> TitleBar;
+		TSharedRef<SWidget> TitleBarWidget = FSlateApplication::Get().MakeWindowTitleBar(NewParentWindow, SNullWidget::NullWidget, HAlign_Fill, TitleBar);
+
+		(*WindowControlsArea)[TitleBarWidget];
+		NewParentWindow->SetTitleBar(TitleBar);
 	}
 	ParentWindowPtr = NewParentWindow;
 }
@@ -407,5 +424,36 @@ void SDockingArea::OnLiveTabAdded()
 {
 	bIsCenterTargetVisible = false;
 	SDockingNode::OnLiveTabAdded();
+
+	CleanUp(SDockingNode::TabRemoval_None);
+}
+
+void SDockingArea::MakeRoomForWindowChrome()
+{
+	TArray< TSharedRef<SDockingNode> > AllNodes = this->GetChildNodesRecursively();
+	if (AllNodes.Num() > 0)
+	{
+		// Clear out all the reserved space.
+		for (auto SomeNode : AllNodes)
+		{
+			if (SomeNode->GetNodeType() == DockTabStack)
+			{
+				auto SomeTabStack = StaticCastSharedRef<SDockingTabStack>(SomeNode);
+				SomeTabStack->ClearReservedSpace();
+			}
+		}
+
+		if (this->ParentWindowPtr.IsValid())
+		{
+			// Reserve some space for the minimize, restore, and close controls
+			TSharedRef<SDockingTabStack> WindowControlHousing = this->FindTabStackToHouseWindowControls();
+			WindowControlHousing->ReserveSpaceForWindowChrome(SDockingTabStack::EChromeElement::Controls);
+
+			// Reserve some space for the app icons
+			TSharedRef<SDockingTabStack> IconHousing = this->FindTabStackToHouseWindowIcon();
+			IconHousing->ReserveSpaceForWindowChrome(SDockingTabStack::EChromeElement::Icon);
+		}
+	}
+
 }
 
diff --git a/Engine/Source/Runtime/Slate/Private/Framework/Docking/SDockingArea.h b/Engine/Source/Runtime/Slate/Private/Framework/Docking/SDockingArea.h
--- a/Engine/Source/Runtime/Slate/Private/Framework/Docking/SDockingArea.h
+++ b/Engine/Source/Runtime/Slate/Private/Framework/Docking/SDockingArea.h
@@ -89,12 +89,25 @@ private:
 
 	virtual void OnLiveTabAdded() override;
 
+	/**
+	 * If this dock area controls a window, then we need
+	 * to reserve some room in the upper left and upper right tab wells
+	 * so that there is no overlap with the window chrome.
+	 */
+	void MakeRoomForWindowChrome();
+
 private:	
 	
 	/** The window this dock area is embedded within.  If bIsManagingParentWindow is true, the dock area will also
 	    destroy the window when the last tab goes away. */
 	TWeakPtr<SWindow> ParentWindowPtr;
 
+	/**
+	 * We don't want to waste a lot of space for the minimize, restore, close buttons and other windows controls.
+	 * DockAreas that manage a parent window will use this slot to house those controls.
+	 */
+	SOverlay::FOverlaySlot* WindowControlsArea;
+
 	/** True if this docking area should close the parent window when the last tab in this docking area goes away */
 	bool bManageParentWindow;
 
diff --git a/Engine/Source/Runtime/Slate/Private/Framework/Docking/SDockingSplitter.cpp b/Engine/Source/Runtime/Slate/Private/Framework/Docking/SDockingSplitter.cpp
--- a/Engine/Source/Runtime/Slate/Private/Framework/Docking/SDockingSplitter.cpp
+++ b/Engine/Source/Runtime/Slate/Private/Framework/Docking/SDockingSplitter.cpp
@@ -103,7 +103,7 @@ SDockingNode::ECleanupRetVal SDockingSplitter::CleanUpNodes()
 				TSharedRef<SDockingTabStack> ChildAsStack = StaticCastSharedRef<SDockingTabStack>(ChildNode);
 				if (ChildNodePurpose == NoTabsUnderNode)
 				{
-					// This child node presents no tabs and keeps no tab hisotry.
+					// This child node presents no tabs and keeps no tab history.
 					RemoveChildAt(ChildIndex);
 				}
 				else
@@ -314,3 +314,71 @@ TSharedPtr<FTabManager::FLayoutNode> SDockingSplitter::GatherPersistentLayout()
 	
 }
 
+
+TSharedRef<SDockingTabStack> SDockingSplitter::FindTabStackToHouseWindowControls() const
+{
+#if PLATFORM_MAC
+	return StaticCastSharedRef<SDockingTabStack>(this->FindTabStack(ETabStackToFind::UpperLeft));
+#else
+	return StaticCastSharedRef<SDockingTabStack>(this->FindTabStack(ETabStackToFind::UpperRight));
+#endif
+}
+
+TSharedRef<SDockingTabStack> SDockingSplitter::FindTabStackToHouseWindowIcon() const
+{
+	return StaticCastSharedRef<SDockingTabStack>(this->FindTabStack(ETabStackToFind::UpperLeft));
+}
+
+TSharedRef<SDockingNode> SDockingSplitter::FindTabStack(ETabStackToFind FindMe) const
+{
+	auto FindFirstVisibleChild = [=]() -> TSharedRef<SDockingNode>
+	{
+		for (auto ChildNode : Children)
+		{
+			if (ChildNode->GetVisibility() == EVisibility::Visible)
+			{
+				return ChildNode;
+			}
+		}
+
+		// We might find and modify some invisible nodes. Not a problem.
+		return Children[0];
+	};
+
+	auto FindLastVisibleChild = [=]() -> TSharedRef<SDockingNode>
+	{
+		for (int32 i = Children.Num() - 1; i >= 0; --i)
+		{
+			const TSharedRef<SDockingNode>& ChildNode = Children[i];
+			if (ChildNode->GetVisibility() == EVisibility::Visible)
+			{
+				return ChildNode;
+			}
+		}
+
+		// We might find and modify some invisible nodes. Not a problem.
+		return Children.Last();
+	};
+
+	// We want the top-most node that is on the left-most or the right-most.
+	// In the case of left-most, just grab the first element until we hit a leaf.
+	// In the case of right-most grab the first element for vertical splits and the last for horizontal.
+	TSharedRef<SDockingNode> Candidate = (FindMe == ETabStackToFind::UpperLeft || this->GetOrientation() == Orient_Vertical)
+		? FindFirstVisibleChild()
+		: Children.Last();
+
+	const SDockingNode::Type CandidateType = Candidate->GetNodeType();
+	if (CandidateType == DockTabStack)
+	{
+		return Candidate;
+	}
+	else if ( ensure(CandidateType == DockArea || CandidateType == DockSplitter) )
+	{
+		return StaticCastSharedRef<SDockingSplitter>(Candidate)->FindTabStack(FindMe);
+	}
+	else
+	{
+		return Candidate;
+	}
+
+}
\ No newline at end of file
diff --git a/Engine/Source/Runtime/Slate/Private/Framework/Docking/SDockingSplitter.h b/Engine/Source/Runtime/Slate/Private/Framework/Docking/SDockingSplitter.h
--- a/Engine/Source/Runtime/Slate/Private/Framework/Docking/SDockingSplitter.h
+++ b/Engine/Source/Runtime/Slate/Private/Framework/Docking/SDockingSplitter.h
@@ -78,6 +78,15 @@ public:
 
 	virtual TSharedPtr<FTabManager::FLayoutNode> GatherPersistentLayout() const override;
 
+	/**
+	 * The tab stack which would be appropriate to use for showing the minimize/maximize/close buttons.
+	 * On Windows it is the upper right most. On Mac the upper left most.
+	 */
+	TSharedRef<SDockingTabStack> FindTabStackToHouseWindowControls() const;
+
+	/** Which tab stack is appropriate for showing the app icon. */
+	TSharedRef<SDockingTabStack> FindTabStackToHouseWindowIcon() const;
+
 protected:
 
 	/** Is the docking direction (left, right, above, below) match the orientation (horizontal vs. vertical) */
@@ -89,6 +98,15 @@ protected:
 
 	float ComputeChildCoefficientTotal() const;
 
+	enum class ETabStackToFind
+	{
+		UpperLeft,
+		UpperRight
+	};
+
+	/** Helper: Finds the upper left or the upper right tab stack in the hierarchy */
+	TSharedRef<SDockingNode> FindTabStack(ETabStackToFind FindMe) const;
+
 	/** The SSplitter widget that SDockingSplitter wraps. */
 	TSharedPtr<SSplitter> Splitter;
 	
diff --git a/Engine/Source/Runtime/Slate/Private/Framework/Docking/SDockingTabStack.cpp b/Engine/Source/Runtime/Slate/Private/Framework/Docking/SDockingTabStack.cpp
--- a/Engine/Source/Runtime/Slate/Private/Framework/Docking/SDockingTabStack.cpp
+++ b/Engine/Source/Runtime/Slate/Private/Framework/Docking/SDockingTabStack.cpp
@@ -697,6 +697,43 @@ TSharedPtr<FTabManager::FLayoutNode> SDockingTabStack::GatherPersistentLayout()
 }
 
 
+void SDockingTabStack::ClearReservedSpace()
+{
+	(*TitleBarSlot).Padding(0);
+}
+
+
+void SDockingTabStack::ReserveSpaceForWindowChrome(EChromeElement Element)
+{
+	// @todo: It would be nice to read these from the desired size of the title bar
+	//        instead of hard-coding, but at least the long-standing TTP is fixed!
+	#if PLATFORM_MAC
+		static const FMargin ControlsPadding = FMargin(64, 0, 0, 0);
+		static const FMargin IconPadding = FMargin(0);
+	#else
+		static const FMargin ControlsPadding = FMargin(0, 0, 96, 0);
+		static const FMargin IconPadding = FMargin(32, 0, 0, 0);
+	#endif
+
+	const FMargin CurrentPadding = TitleBarSlot->SlotPadding.Get();
+	switch (Element)
+	{
+	case EChromeElement::Controls:
+		(*TitleBarSlot).Padding(CurrentPadding + ControlsPadding);
+		break;
+
+	case EChromeElement::Icon:
+		(*TitleBarSlot).Padding(CurrentPadding + IconPadding);
+		break;
+
+	default:
+		ensure(false);
+		break;
+	}
+}
+
+
+
 TSharedRef< SDockingTabStack > SDockingTabStack::CreateNewTabStackBySplitting( const SDockingNode::RelativeDirection Direction )
 {
 	TSharedPtr<SDockingSplitter> ParentNode = ParentNodePtr.Pin();
@@ -720,23 +757,7 @@ void SDockingTabStack::SetParentNode( TSharedRef<class SDockingSplitter> InParen
 	// place the title bar widgets into our content instead!
 	const TSharedPtr<SDockingArea>& DockArea = GetDockArea();
 
-	if (DockArea.IsValid() && DockArea->GetParentWindow().IsValid())
-	{
-		// @todo mainframe: Really we only want these to show up for tab stacks that are along the top of the window,
-		//                  and only the first one!  Currently, all SDockingAreas with a parent window set will get
-		//                  title area widgets added!
-		const TSharedRef<SWindow>& ParentWindow = DockArea->GetParentWindow().ToSharedRef();
-
-		TSharedPtr<IWindowTitleBar> TitleBar;
-		TSharedRef<SWidget> TitleBarWidget = FSlateApplication::Get().MakeWindowTitleBar(ParentWindow, TitleBarContent, HAlign_Fill, TitleBar);
-
-		(*TitleBarSlot)[TitleBarWidget];
-		ParentWindow->SetTitleBar(TitleBar);
-	}
-	else
-	{
-		(*TitleBarSlot)[TitleBarContent.ToSharedRef()];
-	}
+	(*TitleBarSlot)[TitleBarContent.ToSharedRef()];
 }
 
 
diff --git a/Engine/Source/Runtime/Slate/Private/Framework/Docking/SDockingTabStack.h b/Engine/Source/Runtime/Slate/Private/Framework/Docking/SDockingTabStack.h
--- a/Engine/Source/Runtime/Slate/Private/Framework/Docking/SDockingTabStack.h
+++ b/Engine/Source/Runtime/Slate/Private/Framework/Docking/SDockingTabStack.h
@@ -67,6 +67,24 @@ public:
 
 	virtual TSharedPtr<FTabManager::FLayoutNode> GatherPersistentLayout() const override;
 
+	/** Elements for which we might want to reserve space. */
+	enum class EChromeElement
+	{
+		Icon,
+		Controls
+	};
+
+	/**
+	 * Remove all the space that might have been reserved for
+	 * various window chrome elements (app icons, minimize, close, etc.)
+	 */
+	void ClearReservedSpace();
+	/**
+	 * Add some extra padding so that the corresponding window chrome element
+	 * does not overlap our tabs.
+	 */
+	void ReserveSpaceForWindowChrome(EChromeElement);
+
 public:
 
 	virtual void Tick( const FGeometry& AllottedGeometry, const double InCurrentTime, const float InDeltaTime ) override;
