Slate : STextBlock's DesiredSize now correctly responds to negative shadow offsets.

diff --git a/Engine/Source/Runtime/Slate/Private/Widgets/Text/STextBlock.cpp b/Engine/Source/Runtime/Slate/Private/Widgets/Text/STextBlock.cpp
--- a/Engine/Source/Runtime/Slate/Private/Widgets/Text/STextBlock.cpp
+++ b/Engine/Source/Runtime/Slate/Private/Widgets/Text/STextBlock.cpp
@@ -272,7 +272,9 @@ void STextBlock::CacheDesiredSize()
 			TextMeasurement.X = FMath::Min(TextMeasurement.X, WrappingWidth);
 		}
 
-		this->Advanced_SetDesiredSize(TextMeasurement + GetShadowOffset());
+		const FVector2D CurrentShadowOffset = GetShadowOffset();
+		const FVector2D AbsoluteShadowOffset(FMath::Abs(CurrentShadowOffset.X), FMath::Abs(CurrentShadowOffset.Y));
+		this->Advanced_SetDesiredSize(TextMeasurement + AbsoluteShadowOffset);
 
 		// Update cached values
 		CachedWrapTextWidth = WrappingWidth;
