Fixed bug introduced by CL 2186421, where windows not being managed by a DockArea would not get any window chrome.

diff --git a/Engine/Source/Runtime/Slate/Private/Framework/Docking/SDockingArea.cpp b/Engine/Source/Runtime/Slate/Private/Framework/Docking/SDockingArea.cpp
--- a/Engine/Source/Runtime/Slate/Private/Framework/Docking/SDockingArea.cpp
+++ b/Engine/Source/Runtime/Slate/Private/Framework/Docking/SDockingArea.cpp
@@ -239,13 +239,17 @@ void SDockingArea::SetParentWindow( const TSharedRef<SWindow>& NewParentWindow )
 	if( bManageParentWindow )
 	{
 		NewParentWindow->SetRequestDestroyWindowOverride( FRequestDestroyWindowOverride::CreateSP( this, &SDockingArea::OnOwningWindowBeingDestroyed ) );
+	}
 
+	// Even thought we don't manage the parent window's lifetime, we are still responsible for making its window chrome.
+	{
 		TSharedPtr<IWindowTitleBar> TitleBar;
 		TSharedRef<SWidget> TitleBarWidget = FSlateApplication::Get().MakeWindowTitleBar(NewParentWindow, SNullWidget::NullWidget, HAlign_Fill, TitleBar);
 
 		(*WindowControlsArea)[TitleBarWidget];
 		NewParentWindow->SetTitleBar(TitleBar);
 	}
+
 	ParentWindowPtr = NewParentWindow;
 }
 
