FSlateApplication::SetJoystickCaptorToGameViewport now captures all joysticks

diff --git a/Engine/Source/Runtime/Slate/Private/Framework/Application/SlateApplication.cpp b/Engine/Source/Runtime/Slate/Private/Framework/Application/SlateApplication.cpp
--- a/Engine/Source/Runtime/Slate/Private/Framework/Application/SlateApplication.cpp
+++ b/Engine/Source/Runtime/Slate/Private/Framework/Application/SlateApplication.cpp
@@ -1361,7 +1361,7 @@ void FSlateApplication::SetJoystickCaptorToGameViewport()
 		FWidgetPath PathToWidget;
 		FSlateWindowHelper::FindPathToWidget(SlateWindows, CurrentGameViewportWidget.ToSharedRef(), /*OUT*/ PathToWidget);
 
-		FReply Temp = FReply::Handled().CaptureJoystick(CurrentGameViewportWidget.ToSharedRef());
+		FReply Temp = FReply::Handled().CaptureJoystick(CurrentGameViewportWidget.ToSharedRef(), true);
 
 		ProcessReply(PathToWidget, Temp, NULL, NULL);
 	}
