Added support for "lightweight highlighters"

TTP# 336464 - Editor: Finish the Multiline Editable Text Block

These are a kind of highlighter that exist at the line level, rather than the run level. This means that they can't do the more complex things like adjusting the run rendering (eg, changing the text color), but it does mean that they can be overlapped and don't require a re-flow of the entire layout when they're added or removed.

To avoid confusion with the existing "highlighter" concept, and to better describe what the classes do, the existing IRunHighlighter has been renamed to IRunRenderer, and the newly added "lightweight highlighters" are called ILineHighlighter.

FTextLayout has been updated to flow the lightweight highlighters separately to the main layout flow, and the dirty flag has been split so that we can skip the layout re-flow if only the highlighters have been changed.

FSlateTextLayout has been updated to first paint any "underlay" highlighters, then the line view blocks (potentially with a custom run renderer), then any "overlay" highlighters.

The GetChildren and OnArrangeChildren have been removed from ISlateRunRenderer (and not added to ISlateLineHighlighter) because they were never being used (FSlateTextLayout wasn't calling them), and were seemingly just adding complexity.

diff --git a/Engine/Source/Runtime/Slate/Private/Framework/Text/SlateHyperlinkRun.cpp b/Engine/Source/Runtime/Slate/Private/Framework/Text/SlateHyperlinkRun.cpp
--- a/Engine/Source/Runtime/Slate/Private/Framework/Text/SlateHyperlinkRun.cpp
+++ b/Engine/Source/Runtime/Slate/Private/Framework/Text/SlateHyperlinkRun.cpp
@@ -51,7 +51,7 @@ int8 FSlateHyperlinkRun::GetKerning( int32 CurrentIndex, float Scale ) const
 	return 0;
 }
 
-TSharedRef< ILayoutBlock > FSlateHyperlinkRun::CreateBlock( int32 StartIndex, int32 EndIndex, FVector2D Size, const TSharedPtr< IRunHighlighter >& Highlighter )
+TSharedRef< ILayoutBlock > FSlateHyperlinkRun::CreateBlock( int32 StartIndex, int32 EndIndex, FVector2D Size, const TSharedPtr< IRunRenderer >& Renderer )
 {
 	TSharedRef< SWidget > Widget = SNew( SRichTextHyperlink, ViewModel )
 		.Style( &Style )
@@ -60,7 +60,7 @@ TSharedRef< ILayoutBlock > FSlateHyperlinkRun::CreateBlock( int32 StartIndex, in
 
 	Children.Add( Widget );
 
-	return FWidgetLayoutBlock::Create( SharedThis( this ), Widget, FTextRange( StartIndex, EndIndex ), Size, Highlighter );
+	return FWidgetLayoutBlock::Create( SharedThis( this ), Widget, FTextRange( StartIndex, EndIndex ), Size, Renderer );
 }
 
 void FSlateHyperlinkRun::OnNavigate()
diff --git a/Engine/Source/Runtime/Slate/Private/Framework/Text/SlateImageRun.cpp b/Engine/Source/Runtime/Slate/Private/Framework/Text/SlateImageRun.cpp
--- a/Engine/Source/Runtime/Slate/Private/Framework/Text/SlateImageRun.cpp
+++ b/Engine/Source/Runtime/Slate/Private/Framework/Text/SlateImageRun.cpp
@@ -75,9 +75,9 @@ int32 FSlateImageRun::OnPaint( const FTextLayout::FLineView& Line, const TShared
 	return LayerId;
 }
 
-TSharedRef< ILayoutBlock > FSlateImageRun::CreateBlock( int32 BeginIndex, int32 EndIndex, FVector2D Size, const TSharedPtr< IRunHighlighter >& Highlighter )
+TSharedRef< ILayoutBlock > FSlateImageRun::CreateBlock( int32 BeginIndex, int32 EndIndex, FVector2D Size, const TSharedPtr< IRunRenderer >& Renderer )
 {
-	return FDefaultLayoutBlock::Create( SharedThis( this ), FTextRange( BeginIndex, EndIndex ), Size, Highlighter );
+	return FDefaultLayoutBlock::Create( SharedThis( this ), FTextRange( BeginIndex, EndIndex ), Size, Renderer );
 }
 
 int8 FSlateImageRun::GetKerning( int32 CurrentIndex, float Scale ) const 
diff --git a/Engine/Source/Runtime/Slate/Private/Framework/Text/SlateSimpleRunHighlighter.cpp b/Engine/Source/Runtime/Slate/Private/Framework/Text/SlateSimpleRunHighlighter.cpp
deleted file mode 100644
diff --git a/Engine/Source/Runtime/Slate/Private/Framework/Text/SlateTextHighlightRunRenderer.cpp b/Engine/Source/Runtime/Slate/Private/Framework/Text/SlateTextHighlightRunRenderer.cpp
new file mode 100644
--- /dev/null
+++ b/Engine/Source/Runtime/Slate/Private/Framework/Text/SlateTextHighlightRunRenderer.cpp
@@ -0,0 +1,40 @@
+// Copyright 1998-2014 Epic Games, Inc. All Rights Reserved.
+
+#include "SlatePrivatePCH.h"
+#include "SlateTextHighlightRunRenderer.h"
+
+FSlateTextHighlightRunRenderer::FSlateTextHighlightRunRenderer()
+{
+
+}
+
+int32 FSlateTextHighlightRunRenderer::OnPaint( const FTextLayout::FLineView& Line, const TSharedRef< ISlateRun >& Run, const TSharedRef< ILayoutBlock >& Block, const FTextBlockStyle& DefaultStyle, const FGeometry& AllottedGeometry, const FSlateRect& MyClippingRect, FSlateWindowElementList& OutDrawElements, int32 LayerId, const FWidgetStyle& InWidgetStyle, bool bParentEnabled ) const 
+{
+	FVector2D Location( Block->GetLocationOffset() );
+	Location.Y = Line.Offset.Y;
+
+	// Draw the actual highlight rectangle
+	FSlateDrawElement::MakeBox(
+		OutDrawElements,
+		++LayerId,
+		FPaintGeometry( AllottedGeometry.AbsolutePosition + Location, FVector2D( Block->GetSize().X, Line.Size.Y ), AllottedGeometry.Scale ),
+		&DefaultStyle.HighlightShape,
+		MyClippingRect,
+		bParentEnabled ? ESlateDrawEffect::None : ESlateDrawEffect::DisabledEffect,
+		InWidgetStyle.GetColorAndOpacityTint() * DefaultStyle.HighlightColor
+		);
+
+	FLinearColor InvertedForeground = FLinearColor::White - InWidgetStyle.GetForegroundColor();
+	InvertedForeground.A = InWidgetStyle.GetForegroundColor().A;
+
+	FWidgetStyle WidgetStyle( InWidgetStyle );
+	WidgetStyle.SetForegroundColor( InvertedForeground );
+
+	return Run->OnPaint( Line, Block, DefaultStyle, AllottedGeometry, MyClippingRect, OutDrawElements, LayerId, WidgetStyle, bParentEnabled );
+}
+
+TSharedRef< FSlateTextHighlightRunRenderer > FSlateTextHighlightRunRenderer::Create()
+{
+	return MakeShareable( new FSlateTextHighlightRunRenderer() );
+}
+
diff --git a/Engine/Source/Runtime/Slate/Private/Framework/Text/SlateTextLayout.cpp b/Engine/Source/Runtime/Slate/Private/Framework/Text/SlateTextLayout.cpp
--- a/Engine/Source/Runtime/Slate/Private/Framework/Text/SlateTextLayout.cpp
+++ b/Engine/Source/Runtime/Slate/Private/Framework/Text/SlateTextLayout.cpp
@@ -42,28 +42,29 @@ int32 FSlateTextLayout::OnPaint( const FTextBlockStyle& DefaultTextStyle, const
 	const FSlateRect ClippingRect = AllottedGeometry.GetClippingRect().IntersectionWith(MyClippingRect);
 	const ESlateDrawEffect::Type DrawEffects = bParentEnabled ? ESlateDrawEffect::None : ESlateDrawEffect::DisabledEffect;
 
-	const int32 DebugLayer = LayerId;
-	const int32 TextLayer = LayerId + 1;
-
 	static bool ShowDebug = false;
 	FLinearColor BlockHue( 0, 1.0f, 1.0f, 0.5 );
-	int32 MaxLayerId = LayerId;
+	int32 HighestLayerId = LayerId;
 
-	for (int LineIndex = 0; LineIndex < LineViews.Num(); LineIndex++)
+	for (const FTextLayout::FLineView& LineView : LineViews)
 	{
-		const FTextLayout::FLineView& LineView = LineViews[ LineIndex ];
+		// Render any underlays for this line
+		const int32 HighestUnderlayLayerId = OnPaintHighlights( LineView, LineView.UnderlayHighlights, DefaultTextStyle, AllottedGeometry, ClippingRect, OutDrawElements, LayerId, InWidgetStyle, bParentEnabled );
 
-		for (int BlockIndex = 0; BlockIndex < LineView.Blocks.Num(); BlockIndex++)
-		{
-			const TSharedRef< ILayoutBlock > Block = LineView.Blocks[ BlockIndex ];
+		const int32 BlockDebugLayer = HighestUnderlayLayerId;
+		const int32 TextLayer = BlockDebugLayer + 1;
+		int32 HighestBlockLayerId = TextLayer;
 
+		// Render every block for this line
+		for (const TSharedRef< ILayoutBlock >& Block : LineView.Blocks)
+		{
 			if ( ShowDebug )
 			{
 				BlockHue.R += 50.0f;
 
 				FSlateDrawElement::MakeBox(
 					OutDrawElements, 
-					DebugLayer,
+					BlockDebugLayer,
 					FPaintGeometry( AllottedGeometry.AbsolutePosition + Block->GetLocationOffset(), Block->GetSize(), AllottedGeometry.Scale ),
 					&DefaultTextStyle.HighlightShape,
 					ClippingRect,
@@ -72,25 +73,44 @@ int32 FSlateTextLayout::OnPaint( const FTextBlockStyle& DefaultTextStyle, const
 					);
 			}
 
-			int32 BlockLayerId = TextLayer;
 			const TSharedRef< ISlateRun > Run = StaticCastSharedRef< ISlateRun >( Block->GetRun() );
 
-			int32 BlocksLayerId = MaxLayerId;
-			const TSharedPtr< ISlateRunHighlighter > Highlighter = StaticCastSharedPtr< ISlateRunHighlighter >( Block->GetHighlighter() );
-			if ( Highlighter.IsValid() )
+			int32 HighestRunLayerId = TextLayer;
+			const TSharedPtr< ISlateRunRenderer > RunRenderer = StaticCastSharedPtr< ISlateRunRenderer >( Block->GetRenderer() );
+			if ( RunRenderer.IsValid() )
 			{
-				BlocksLayerId = Highlighter->OnPaint( LineView, Run, Block, DefaultTextStyle, AllottedGeometry, ClippingRect, OutDrawElements, BlockLayerId, InWidgetStyle, bParentEnabled );
+				HighestRunLayerId = RunRenderer->OnPaint( LineView, Run, Block, DefaultTextStyle, AllottedGeometry, ClippingRect, OutDrawElements, TextLayer, InWidgetStyle, bParentEnabled );
 			}
 			else
 			{
-				BlocksLayerId = Run->OnPaint( LineView, Block, DefaultTextStyle, AllottedGeometry, ClippingRect, OutDrawElements, BlockLayerId, InWidgetStyle, bParentEnabled );
+				HighestRunLayerId = Run->OnPaint( LineView, Block, DefaultTextStyle, AllottedGeometry, ClippingRect, OutDrawElements, TextLayer, InWidgetStyle, bParentEnabled );
 			}
 
-			MaxLayerId = FMath::Max( MaxLayerId, BlocksLayerId );
+			HighestBlockLayerId = FMath::Max( HighestBlockLayerId, HighestRunLayerId );
+		}
+
+		// Render any overlays for this line
+		const int32 HighestOverlayLayerId = OnPaintHighlights( LineView, LineView.OverlayHighlights, DefaultTextStyle, AllottedGeometry, ClippingRect, OutDrawElements, HighestBlockLayerId, InWidgetStyle, bParentEnabled );
+		HighestLayerId = FMath::Max( HighestLayerId, HighestOverlayLayerId );
+	}
+
+	return HighestLayerId;
+}
+
+int32 FSlateTextLayout::OnPaintHighlights( const FTextLayout::FLineView& LineView, const TArray<FLineViewHighlight>& Highlights, const FTextBlockStyle& DefaultTextStyle, const FGeometry& AllottedGeometry, const FSlateRect& ClippingRect, FSlateWindowElementList& OutDrawElements, const int32 LayerId, const FWidgetStyle& InWidgetStyle, bool bParentEnabled ) const
+{
+	int32 CurrentLayerId = LayerId;
+
+	for (const FLineViewHighlight& Highlight : Highlights)
+	{
+		const TSharedPtr< ISlateLineHighlighter > LineHighlighter = StaticCastSharedPtr< ISlateLineHighlighter >( Highlight.Highlighter );
+		if (LineHighlighter.IsValid())
+		{
+			CurrentLayerId = LineHighlighter->OnPaint( LineView, Highlight.OffsetX, Highlight.Width, DefaultTextStyle, AllottedGeometry, ClippingRect, OutDrawElements, CurrentLayerId, InWidgetStyle, bParentEnabled );
 		}
 	}
 
-	return MaxLayerId;
+	return CurrentLayerId;
 }
 
 void FSlateTextLayout::EndLayout()
diff --git a/Engine/Source/Runtime/Slate/Private/Framework/Text/SlateTextRun.cpp b/Engine/Source/Runtime/Slate/Private/Framework/Text/SlateTextRun.cpp
--- a/Engine/Source/Runtime/Slate/Private/Framework/Text/SlateTextRun.cpp
+++ b/Engine/Source/Runtime/Slate/Private/Framework/Text/SlateTextRun.cpp
@@ -60,9 +60,9 @@ int8 FSlateTextRun::GetKerning(int32 CurrentIndex, float Scale) const
 	return FontMeasure->GetKerning( Style.Font, Scale, (*Text)[ PreviousIndex ], (*Text)[ CurrentIndex ] );
 }
 
-TSharedRef< ILayoutBlock > FSlateTextRun::CreateBlock( int32 BeginIndex, int32 EndIndex, FVector2D Size, const TSharedPtr< IRunHighlighter >& Highlighter )
+TSharedRef< ILayoutBlock > FSlateTextRun::CreateBlock( int32 BeginIndex, int32 EndIndex, FVector2D Size, const TSharedPtr< IRunRenderer >& Renderer )
 {
-	return FDefaultLayoutBlock::Create( SharedThis( this ), FTextRange( BeginIndex, EndIndex ), Size, Highlighter );
+	return FDefaultLayoutBlock::Create( SharedThis( this ), FTextRange( BeginIndex, EndIndex ), Size, Renderer );
 }
 
 int32 FSlateTextRun::OnPaint( const FTextLayout::FLineView& Line, const TSharedRef< ILayoutBlock >& Block, const FTextBlockStyle& DefaultStyle, const FGeometry& AllottedGeometry, const FSlateRect& MyClippingRect, FSlateWindowElementList& OutDrawElements, int32 LayerId, const FWidgetStyle& InWidgetStyle, bool bParentEnabled ) const 
diff --git a/Engine/Source/Runtime/Slate/Private/Framework/Text/SlateWidgetRun.cpp b/Engine/Source/Runtime/Slate/Private/Framework/Text/SlateWidgetRun.cpp
--- a/Engine/Source/Runtime/Slate/Private/Framework/Text/SlateWidgetRun.cpp
+++ b/Engine/Source/Runtime/Slate/Private/Framework/Text/SlateWidgetRun.cpp
@@ -43,9 +43,9 @@ int8 FSlateWidgetRun::GetKerning( int32 CurrentIndex, float Scale ) const
 	return 0;
 }
 
-TSharedRef< ILayoutBlock > FSlateWidgetRun::CreateBlock( int32 StartIndex, int32 EndIndex, FVector2D Size, const TSharedPtr< IRunHighlighter >& Highlighter )
+TSharedRef< ILayoutBlock > FSlateWidgetRun::CreateBlock( int32 StartIndex, int32 EndIndex, FVector2D Size, const TSharedPtr< IRunRenderer >& Renderer )
 {
-	return FDefaultLayoutBlock::Create( SharedThis( this ), FTextRange( StartIndex, EndIndex ), Size, Highlighter );
+	return FDefaultLayoutBlock::Create( SharedThis( this ), FTextRange( StartIndex, EndIndex ), Size, Renderer );
 }
 
 int32 FSlateWidgetRun::OnPaint( const FTextLayout::FLineView& Line, const TSharedRef< ILayoutBlock >& Block, const FTextBlockStyle& DefaultStyle, const FGeometry& AllottedGeometry, const FSlateRect& MyClippingRect, FSlateWindowElementList& OutDrawElements, int32 LayerId, const FWidgetStyle& InWidgetStyle, bool bParentEnabled ) const 
diff --git a/Engine/Source/Runtime/Slate/Private/Framework/Text/TextLayout.cpp b/Engine/Source/Runtime/Slate/Private/Framework/Text/TextLayout.cpp
--- a/Engine/Source/Runtime/Slate/Private/Framework/Text/TextLayout.cpp
+++ b/Engine/Source/Runtime/Slate/Private/Framework/Text/TextLayout.cpp
@@ -109,7 +109,7 @@ FTextLayout::FBreakCandidate FTextLayout::CreateBreakCandidate( int32& OutRunInd
 	return BreakCandidate;
 }
 
-void FTextLayout::CreateLineViewBlocks( int32 LineModelIndex, const int32 StopIndex, int32& OutRunIndex, int32& OutHighlightIndex, int32& OutPreviousBlockEnd, TArray< TSharedRef< ILayoutBlock > >& OutSoftLine )
+void FTextLayout::CreateLineViewBlocks( int32 LineModelIndex, const int32 StopIndex, int32& OutRunIndex, int32& OutRendererIndex, int32& OutPreviousBlockEnd, TArray< TSharedRef< ILayoutBlock > >& OutSoftLine )
 {
 	const FLineModel& LineModel = LineModels[ LineModelIndex ];
 
@@ -121,40 +121,40 @@ void FTextLayout::CreateLineViewBlocks( int32 LineModelIndex, const int32 StopIn
 		const FRunModel& Run = LineModel.Runs[ OutRunIndex ];
 		const FTextRange RunRange = Run.GetTextRange();
 
-		TSharedPtr< IRunHighlighter > BlockHighlighter = NULL;
+		TSharedPtr< IRunRenderer > BlockRenderer = nullptr;
 
 		int32 BlockStopIndex = RunRange.EndIndex;
-		if ( OutHighlightIndex != INDEX_NONE )
+		if ( OutRendererIndex != INDEX_NONE )
 		{
-			//Grab the currently active highlighter
-			const FTextHighlight& Highlight = LineModel.Highlights[OutHighlightIndex];
+			// Grab the currently active renderer
+			const FTextRunRenderer& Renderer = LineModel.RunRenderers[OutRendererIndex];
 
-			//Check to see if the last block was highlighted with the same highlighter
-			if ( OutPreviousBlockEnd >= Highlight.Range.BeginIndex )
+			// Check to see if the last block was rendered with the same renderer
+			if ( OutPreviousBlockEnd >= Renderer.Range.BeginIndex )
 			{
-				//If the highlighter ends before our run...
-				if ( Highlight.Range.EndIndex <= RunRange.EndIndex )
+				//If the renderer ends before our run...
+				if ( Renderer.Range.EndIndex <= RunRange.EndIndex )
 				{
-					//Adjust the stopping point of the block to be the end of the highlight,
-					//since highlights need their own block segments
-					BlockStopIndex = Highlight.Range.EndIndex;
-					BlockHighlighter = Highlight.Highlighter;
+					// Adjust the stopping point of the block to be the end of the renderer range,
+					// since highlights need their own block segments
+					BlockStopIndex = Renderer.Range.EndIndex;
+					BlockRenderer = Renderer.Renderer;
 				}
 				else
 				{
-					//This whole run is encompassed by the highlighter
-					BlockHighlighter = Highlight.Highlighter;
+					// This whole run is encompassed by the renderer
+					BlockRenderer = Renderer.Renderer;
 				}
 			}
 			else
 			{
-				//Does the highlight begin before our run ends?
-				if ( Highlight.Range.BeginIndex <= RunRange.EndIndex )
+				// Does the renderer range begin before our run ends?
+				if ( Renderer.Range.BeginIndex <= RunRange.EndIndex )
 				{
-					// then adjust the current block stopping point to just before the highlight begins,
-					//since highlights need their own block segments
-					BlockStopIndex = Highlight.Range.BeginIndex;
-					BlockHighlighter = NULL;
+					// then adjust the current block stopping point to just before the renderer range begins,
+					// since renderers need their own block segments
+					BlockStopIndex = Renderer.Range.BeginIndex;
+					BlockRenderer = nullptr;
 				}
 			}
 		}
@@ -174,7 +174,7 @@ void FTextLayout::CreateLineViewBlocks( int32 LineModelIndex, const int32 StopIn
 
 			FBlockDefinition BlockDefine;
 			BlockDefine.ActualRange = FTextRange(BlockBeginIndex, BlockStopIndex);
-			BlockDefine.Highlighter = BlockHighlighter;
+			BlockDefine.Renderer = BlockRenderer;
 
 			OutSoftLine.Add( Run.CreateBlock( BlockDefine, Scale ) );
 			OutPreviousBlockEnd = BlockStopIndex;
@@ -183,7 +183,7 @@ void FTextLayout::CreateLineViewBlocks( int32 LineModelIndex, const int32 StopIn
 		{
 			FBlockDefinition BlockDefine;
 			BlockDefine.ActualRange = RunRange;
-			BlockDefine.Highlighter = BlockHighlighter;
+			BlockDefine.Renderer = BlockRenderer;
 
 			OutSoftLine.Add( Run.CreateBlock( BlockDefine, Scale ) );
 			OutPreviousBlockEnd = RunRange.EndIndex;
@@ -201,13 +201,13 @@ void FTextLayout::CreateLineViewBlocks( int32 LineModelIndex, const int32 StopIn
 			++OutRunIndex;
 		}
 
-		if ( OutHighlightIndex != INDEX_NONE && BlockStopIndex == LineModel.Highlights[ OutHighlightIndex ].Range.EndIndex )
+		if ( OutRendererIndex != INDEX_NONE && BlockStopIndex == LineModel.RunRenderers[ OutRendererIndex ].Range.EndIndex )
 		{
-			++OutHighlightIndex;
+			++OutRendererIndex;
 
-			if ( OutHighlightIndex >= LineModel.Highlights.Num() )
+			if ( OutRendererIndex >= LineModel.RunRenderers.Num() )
 			{
-				OutHighlightIndex = INDEX_NONE;
+				OutRendererIndex = INDEX_NONE;
 			}
 		}
 
@@ -289,6 +289,42 @@ void FTextLayout::JustifyLayout()
 	}
 }
 
+void FTextLayout::JustifyHighlights()
+{
+	if ( Justification == ETextJustify::Left )
+	{
+		return;
+	}
+
+	const float LayoutWidthNoMargin = DrawSize.X - Margin.GetTotalSpaceAlong<Orient_Horizontal>() * Scale;
+
+	for (int LineViewIndex = 0; LineViewIndex < LineViews.Num(); LineViewIndex++)
+	{
+		FLineView& LineView = LineViews[ LineViewIndex ];
+		const float ExtraSpace = LayoutWidthNoMargin - LineView.Size.X;
+
+		float HighlightOffset = 0;
+		if ( Justification == ETextJustify::Center )
+		{
+			HighlightOffset = ExtraSpace / 2;
+		}
+		else if ( Justification == ETextJustify::Right )
+		{
+			HighlightOffset = ExtraSpace;
+		}
+
+		for (FLineViewHighlight& Highlight : LineView.UnderlayHighlights)
+		{
+			Highlight.OffsetX += HighlightOffset;
+		}
+
+		for (FLineViewHighlight& Highlight : LineView.OverlayHighlights)
+		{
+			Highlight.OffsetX += HighlightOffset;
+		}
+	}
+}
+
 void FTextLayout::FlowLayout()
 {
 	check( WrappingWidth >= 0 );
@@ -307,17 +343,17 @@ void FTextLayout::FlowLayout()
 		int32 CurrentRunIndex = 0;
 		int32 PreviousBlockEnd = 0;
 
-		int32 CurrentHighlightIndex = 0;
-		if ( CurrentHighlightIndex >= LineModel.Highlights.Num() )
+		int32 CurrentRendererIndex = 0;
+		if ( CurrentRendererIndex >= LineModel.RunRenderers.Num() )
 		{
-			CurrentHighlightIndex = INDEX_NONE;
+			CurrentRendererIndex = INDEX_NONE;
 		}
 
 		// if the Line doesn't have any BreakCandidates
 		if ( LineModel.BreakCandidates.Num() == 0 )
 		{
 			//Then iterate over all of it's runs
-			CreateLineViewBlocks( LineModelIndex, INDEX_NONE, /*OUT*/CurrentRunIndex, /*OUT*/CurrentHighlightIndex, /*OUT*/PreviousBlockEnd, SoftLine );
+			CreateLineViewBlocks( LineModelIndex, INDEX_NONE, /*OUT*/CurrentRunIndex, /*OUT*/CurrentRendererIndex, /*OUT*/PreviousBlockEnd, SoftLine );
 			check( CurrentRunIndex == LineModel.Runs.Num() );
 			CurrentWidth = 0;
 			SoftLine.Empty();
@@ -340,7 +376,7 @@ void FTextLayout::FlowLayout()
 
 					const FBreakCandidate& FinalBreakOnSoftLine = ( !IsFirstBreak && !IsFirstBreakOnSoftLine && !BreakWithoutTrailingWhitespaceDoesFit ) ? LineModel.BreakCandidates[ --BreakIndex ] : Break;
 
-					CreateLineViewBlocks( LineModelIndex, FinalBreakOnSoftLine.ActualRange.EndIndex, /*OUT*/CurrentRunIndex, /*OUT*/CurrentHighlightIndex, /*OUT*/PreviousBlockEnd, SoftLine );
+					CreateLineViewBlocks( LineModelIndex, FinalBreakOnSoftLine.ActualRange.EndIndex, /*OUT*/CurrentRunIndex, /*OUT*/CurrentRendererIndex, /*OUT*/PreviousBlockEnd, SoftLine );
 
 					if ( CurrentRunIndex < LineModel.Runs.Num() && FinalBreakOnSoftLine.ActualRange.EndIndex == LineModel.Runs[ CurrentRunIndex ].GetTextRange().EndIndex )
 					{
@@ -365,16 +401,89 @@ void FTextLayout::FlowLayout()
 	DrawSize.Y += Margin.GetTotalSpaceAlong<Orient_Vertical>() * Scale;
 }
 
-void FTextLayout::CreateView()
+void FTextLayout::FlowHighlights()
 {
-	ClearView();
-	BeginLayout();
+	// FlowLayout must have been called first
+	check(!(DirtyFlags & EDirtyState::Layout));
 
-	FlowLayout();
-	JustifyLayout();
+	for (FLineView& LineView : LineViews)
+	{
+		LineView.UnderlayHighlights.Empty();
+		LineView.OverlayHighlights.Empty();
 
-	EndLayout();
-	HasChanged = false;
+		FLineModel& LineModel = LineModels[LineView.ModelIndex];
+		
+		// Insert each highlighter into every line view that's within its range, either as an underlay, or as an overlay
+		for (FTextLineHighlight& LineHighlight : LineModel.LineHighlights)
+		{
+			if (LineHighlight.LineIndex != LineView.ModelIndex)
+			{
+				continue;
+			}
+
+			const bool bIsHighlightInRange = LineView.Range.InclusiveContains(LineHighlight.Range.BeginIndex) && LineView.Range.InclusiveContains(LineHighlight.Range.EndIndex);
+			if(!bIsHighlightInRange)
+			{
+				continue;
+			}
+
+			FLineViewHighlight LineViewHighlight;
+			LineViewHighlight.OffsetX = 0.0f;
+			LineViewHighlight.Width = 0.0f;
+			LineViewHighlight.Highlighter = LineHighlight.Highlighter;
+
+			// Measure the blocks up to the start of this highlight to get the correct start offset
+			int32 CurrentBlockIndex = 0;
+			for (; CurrentBlockIndex < LineView.Blocks.Num(); ++CurrentBlockIndex)
+			{
+				const TSharedRef< ILayoutBlock >& Block = LineView.Blocks[CurrentBlockIndex];
+				const TSharedRef<IRun> Run = Block->GetRun();
+
+				const FTextRange& BlockTextRange = Block->GetTextRange();
+				if (LineHighlight.Range.BeginIndex > BlockTextRange.EndIndex)
+				{
+					// Highlight starts after this block, just include its entire size
+					LineViewHighlight.OffsetX += Run->Measure(BlockTextRange.BeginIndex, BlockTextRange.EndIndex, Scale).X;
+				}
+				else
+				{
+					// This block contains the start of this highlight, measure to that point and then we're done!
+					LineViewHighlight.OffsetX += Run->Measure(BlockTextRange.BeginIndex, LineHighlight.Range.BeginIndex, Scale).X;
+					break;
+				}
+			}
+
+			// Now measure the blocks under this highlight to get the correct size
+			for (; CurrentBlockIndex < LineView.Blocks.Num(); ++CurrentBlockIndex)
+			{
+				const TSharedRef< ILayoutBlock >& Block = LineView.Blocks[CurrentBlockIndex];
+				const TSharedRef<IRun> Run = Block->GetRun();
+
+				const FTextRange& BlockTextRange = Block->GetTextRange();
+				const FTextRange IntersectedRange = BlockTextRange.Intersect(LineHighlight.Range);
+				if (!IntersectedRange.IsEmpty())
+				{
+					// Measure the part of the run which intersects the highlight
+					LineViewHighlight.Width += Run->Measure(IntersectedRange.BeginIndex, IntersectedRange.EndIndex, Scale).X;
+				}
+					
+				if (BlockTextRange.EndIndex > LineHighlight.Range.EndIndex)
+				{
+					// We've measured all the runs under the highlight
+					break;
+				}
+			}
+
+			if (LineHighlight.ZOrder < 0)
+			{
+				LineView.UnderlayHighlights.Add(LineViewHighlight);
+			}
+			else
+			{
+				LineView.OverlayHighlights.Add(LineViewHighlight);
+			}
+		}
+	}
 }
 
 void FTextLayout::EndLayout()
@@ -449,7 +558,7 @@ void FTextLayout::ClearWrappingCache()
 
 FTextLayout::FTextLayout() : LineModels()
 	, LineViews()
-	, HasChanged( false )
+	, DirtyFlags( EDirtyState::None )
 	, Scale( 1.0f )
 	, WrappingWidth( 0 )
 	, Margin()
@@ -461,67 +570,143 @@ FTextLayout::FTextLayout() : LineModels()
 
 void FTextLayout::UpdateIfNeeded()
 {
-	if ( HasChanged )
+	const bool bHasChangedLayout = !!(DirtyFlags & EDirtyState::Layout);
+	const bool bHasChangedHighlights = !!(DirtyFlags & EDirtyState::Highlights);
+
+	if ( bHasChangedLayout )
 	{
 		// if something has changed then create a new View
-		Update();
+		UpdateLayout();
+	}
+
+	// If the layout has changed, we always need to update the highlights
+	if ( bHasChangedLayout || bHasChangedHighlights)
+	{
+		UpdateHighlights();
 	}
 }
 
-void FTextLayout::Update()
+void FTextLayout::UpdateLayout()
 {
-	CreateView();
+	ClearView();
+	BeginLayout();
+
+	FlowLayout();
+	JustifyLayout();
+
+	EndLayout();
+
+	DirtyFlags &= ~EDirtyState::Layout;
 }
 
-void FTextLayout::SetHighlights( const TArray< FTextHighlight >& Highlights )
+void FTextLayout::UpdateHighlights()
 {
-	ClearHighlights();
+	FlowHighlights();
+	JustifyHighlights();
 
-	for (int Index = 0; Index < Highlights.Num(); Index++)
+	DirtyFlags &= ~EDirtyState::Highlights;
+}
+
+void FTextLayout::ClearRunRenderers()
+{
+	for (int32 Index = 0; Index < LineModels.Num(); Index++)
 	{
-		AddHighlight( Highlights[ Index ] );
+		if (LineModels[ Index ].RunRenderers.Num() )
+		{
+			LineModels[ Index ].RunRenderers.Empty();
+			DirtyFlags |= EDirtyState::Layout;
+		}
 	}
 }
 
-void FTextLayout::AddHighlight( const FTextHighlight& Highlight )
+void FTextLayout::SetRunRenderers( const TArray< FTextRunRenderer >& Renderers )
 {
-	checkf( LineModels.IsValidIndex( Highlight.LineIndex ), TEXT("Highlights must be for a valid Line Index") );
+	ClearRunRenderers();
 
-	//Highlights needs to be in order and not overlap
-	bool InsertSuccessful = false;
-	FLineModel& LineModel = LineModels[ Highlight.LineIndex ];
-	for (int Index = 0; Index < LineModel.Highlights.Num(); Index++)
+	for (int Index = 0; Index < Renderers.Num(); Index++)
 	{
-		if ( LineModel.Highlights[ Index ].Range.BeginIndex > Highlight.Range.BeginIndex )
+		AddRunRenderer( Renderers[ Index ] );
+	}
+}
+
+void FTextLayout::AddRunRenderer( const FTextRunRenderer& Renderer )
+{
+	checkf( LineModels.IsValidIndex( Renderer.LineIndex ), TEXT("Renderers must be for a valid Line Index") );
+
+	FLineModel& LineModel = LineModels[ Renderer.LineIndex ];
+
+	// Renderers needs to be in order and not overlap
+	bool bWasInserted = false;
+	for (int Index = 0; Index < LineModel.RunRenderers.Num() && !bWasInserted; Index++)
+	{
+		if ( LineModel.RunRenderers[ Index ].Range.BeginIndex > Renderer.Range.BeginIndex )
 		{
-			checkf( Index == 0 || LineModel.Highlights[ Index - 1 ].Range.EndIndex <= Highlight.Range.BeginIndex, TEXT("Highlights cannot overlap") );
-			LineModel.Highlights.Insert( Highlight, Index - 1 );
-			InsertSuccessful = true;
+			checkf( Index == 0 || LineModel.RunRenderers[ Index - 1 ].Range.EndIndex <= Renderer.Range.BeginIndex, TEXT("Renderers cannot overlap") );
+			LineModel.RunRenderers.Insert( Renderer, Index - 1 );
+			bWasInserted = true;
 		}
-		else if ( LineModel.Highlights[ Index ].Range.EndIndex > Highlight.Range.EndIndex )
+		else if ( LineModel.RunRenderers[ Index ].Range.EndIndex > Renderer.Range.EndIndex )
 		{
-			checkf( LineModel.Highlights[ Index ].Range.BeginIndex >= Highlight.Range.EndIndex, TEXT("Highlights cannot overlap") );
-			LineModel.Highlights.Insert( Highlight, Index - 1 );
-			InsertSuccessful = true;
+			checkf( LineModel.RunRenderers[ Index ].Range.BeginIndex >= Renderer.Range.EndIndex, TEXT("Renderers cannot overlap") );
+			LineModel.RunRenderers.Insert( Renderer, Index - 1 );
+			bWasInserted = true;
 		}
 	}
 
-	if ( !InsertSuccessful )
+	if ( !bWasInserted )
 	{
-		LineModel.Highlights.Add( Highlight );
+		LineModel.RunRenderers.Add( Renderer );
 	}
 
-	HasChanged = true;
+	DirtyFlags |= EDirtyState::Layout;
 }
 
-void FTextLayout::ClearHighlights()
+void FTextLayout::ClearLineHighlights()
 {
 	for (int32 Index = 0; Index < LineModels.Num(); Index++)
 	{
-		LineModels[ Index ].Highlights.Empty();
+		if (LineModels[ Index ].LineHighlights.Num())
+		{
+			LineModels[ Index ].LineHighlights.Empty();
+			DirtyFlags |= EDirtyState::Highlights;
+		}
+	}
+}
+
+void FTextLayout::SetLineHighlights( const TArray< FTextLineHighlight >& Highlights )
+{
+	ClearLineHighlights();
+
+	for (int Index = 0; Index < Highlights.Num(); Index++)
+	{
+		AddLineHighlight( Highlights[ Index ] );
+	}
+}
+
+void FTextLayout::AddLineHighlight( const FTextLineHighlight& Highlight )
+{
+	checkf( LineModels.IsValidIndex( Highlight.LineIndex ), TEXT("Highlights must be for a valid Line Index") );
+	checkf( Highlight.ZOrder, TEXT("The highlight Z-order must be <0 to create an underlay, or >0 to create an overlay") );
+
+	FLineModel& LineModel = LineModels[ Highlight.LineIndex ];
+
+	// Try and maintain a stable sorted z-order - highlights with the same z-order should just render in the order they were added
+	bool bWasInserted = false;
+	for (int Index = 0; Index < LineModel.LineHighlights.Num() && !bWasInserted; Index++)
+	{
+		if ( LineModel.LineHighlights[ Index ].ZOrder > Highlight.ZOrder )
+		{
+			LineModel.LineHighlights.Insert( Highlight, Index - 1 );
+			bWasInserted = true;
+		}
+	}
+
+	if ( !bWasInserted )
+	{
+		LineModel.LineHighlights.Add( Highlight );
 	}
 
-	HasChanged = true;
+	DirtyFlags |= EDirtyState::Highlights;
 }
 
 FTextLocation FTextLayout::GetTextLocationAt( const FLineView& LineView, const FVector2D& Relative )
@@ -704,7 +889,7 @@ bool FTextLayout::InsertAt(const FTextLocation& Location, TCHAR Character)
 		}
 	}
 
-	HasChanged = true;
+	DirtyFlags |= EDirtyState::Layout;
 	return true;
 }
 
@@ -745,7 +930,7 @@ bool FTextLayout::InsertAt( const FTextLocation& Location, const FString& Text )
 		}
 	}
 
-	HasChanged = true;
+	DirtyFlags |= EDirtyState::Layout;
 	return true;
 }
 
@@ -792,7 +977,7 @@ bool FTextLayout::JoinLineWithNextLine(int32 LineIndex)
 	//Remove the next line from the list of line models
 	LineModels.RemoveAt(LineIndex + 1);
 
-	HasChanged = true;
+	DirtyFlags |= EDirtyState::Layout;
 	return true;
 }
 
@@ -854,7 +1039,7 @@ bool FTextLayout::SplitLineAt(const FTextLocation& Location)
 	LineModels.Insert(LeftLineModel, LineIndex);
 	LineModels.Insert(RightLineModel, LineIndex + 1);
 
-	HasChanged = true;
+	DirtyFlags |= EDirtyState::Layout;
 	return true;
 }
 
@@ -910,7 +1095,7 @@ bool FTextLayout::RemoveAt( const FTextLocation& Location, int32 Count )
 		}
 	}
 
-	HasChanged = true;
+	DirtyFlags |= EDirtyState::Layout;
 	return true;
 }
 
@@ -936,13 +1121,13 @@ void FTextLayout::AddLine( const TSharedRef< FString >& Text, const TArray< TSha
 	}
 
 	LineModels.Add( LineModel );
-	HasChanged = true;
+	DirtyFlags |= EDirtyState::Layout;
 }
 
 void FTextLayout::ClearLines()
 {
 	LineModels.Empty();
-	HasChanged = true;
+	DirtyFlags |= EDirtyState::Layout;
 }
 
 bool FTextLayout::IsEmpty() const
@@ -1141,7 +1326,7 @@ void FTextLayout::SetMargin( const FMargin& InMargin )
 	if ( WrappingWidth > 0 )
 	{
 		// Since we are wrapping our text we'll need to rebuild the view as the actual wrapping width includes the margin.
-		HasChanged = true;
+		DirtyFlags |= EDirtyState::Layout;
 	}
 	else
 	{
@@ -1176,7 +1361,7 @@ void FTextLayout::SetScale( float Value )
 	}
 
 	Scale = Value;
-	HasChanged = true;
+	DirtyFlags |= EDirtyState::Layout;
 
 	// Changing the scale will affect the wrapping information for *all lines*
 	// Clear out the entire cache so it gets regenerated on the text call to FlowLayout
@@ -1196,7 +1381,7 @@ void FTextLayout::SetJustification( ETextJustify::Type Value )
 	}
 
 	Justification = Value;
-	HasChanged = true;
+	DirtyFlags |= EDirtyState::Layout;
 }
 
 ETextJustify::Type FTextLayout::GetJustification() const
@@ -1209,7 +1394,7 @@ void FTextLayout::SetLineHeightPercentage( float Value )
 	if ( LineHeightPercentage != Value )
 	{
 		LineHeightPercentage = Value; 
-		HasChanged = true;
+		DirtyFlags |= EDirtyState::Layout;
 	}
 }
 
@@ -1223,7 +1408,7 @@ void FTextLayout::SetWrappingWidth( float Value )
 	if ( WrappingWidth != Value )
 	{
 		WrappingWidth = Value; 
-		HasChanged = true;
+		DirtyFlags |= EDirtyState::Layout;
 	}
 }
 
@@ -1261,7 +1446,7 @@ FTextLayout::FLineModel::FLineModel( const TSharedRef< FString >& InText )
 	: Text( InText )
 	, Runs()
 	, BreakCandidates()
-	, Highlights()
+	, RunRenderers()
 	, HasWrappingInformation( false )
 {
 
@@ -1290,7 +1475,7 @@ TSharedRef< ILayoutBlock > FTextLayout::FRunModel::CreateBlock( const FBlockDefi
 	if ( MeasuredRanges.Num() == 0 )
 	{
 		FTextRange RunRange = Run->GetTextRange();
-		return Run->CreateBlock( BlockDefine.ActualRange.BeginIndex, BlockDefine.ActualRange.EndIndex, Run->Measure( SizeRange.BeginIndex, SizeRange.EndIndex, Scale ), BlockDefine.Highlighter );
+		return Run->CreateBlock( BlockDefine.ActualRange.BeginIndex, BlockDefine.ActualRange.EndIndex, Run->Measure( SizeRange.BeginIndex, SizeRange.EndIndex, Scale ), BlockDefine.Renderer );
 	}
 
 	int32 StartRangeIndex = 0;
@@ -1388,7 +1573,7 @@ TSharedRef< ILayoutBlock > FTextLayout::FRunModel::CreateBlock( const FBlockDefi
 		}
 	}
 
-	return Run->CreateBlock( BlockDefine.ActualRange.BeginIndex, BlockDefine.ActualRange.EndIndex, BlockSize, BlockDefine.Highlighter );
+	return Run->CreateBlock( BlockDefine.ActualRange.BeginIndex, BlockDefine.ActualRange.EndIndex, BlockSize, BlockDefine.Renderer );
 }
 
 int FTextLayout::FRunModel::BinarySearchForEndIndex( const TArray< FTextRange >& Ranges, int32 RangeBeginIndex, int32 EndIndex )
diff --git a/Engine/Source/Runtime/Slate/Private/Widgets/Text/SMultiLineEditableText.cpp b/Engine/Source/Runtime/Slate/Private/Widgets/Text/SMultiLineEditableText.cpp
--- a/Engine/Source/Runtime/Slate/Private/Widgets/Text/SMultiLineEditableText.cpp
+++ b/Engine/Source/Runtime/Slate/Private/Widgets/Text/SMultiLineEditableText.cpp
@@ -68,34 +68,16 @@ void SMultiLineEditableText::FCursorInfo::RestoreFromUndo(const FCursorInfo& Und
 	SetCursorLocationAndAlignment(UndoData.CursorPosition, UndoData.CursorAlignment);
 }
 
-SMultiLineEditableText::FSlateCursorRunHighlighter::FSlateCursorRunHighlighter(const FCursorInfo* InCursorInfo)
+SMultiLineEditableText::FCursorLineHighlighter::FCursorLineHighlighter(const FCursorInfo* InCursorInfo)
 	: CursorInfo(InCursorInfo)
 {
 	check(CursorInfo);
 }
 
-void SMultiLineEditableText::FSlateCursorRunHighlighter::OnArrangeChildren( const TSharedRef< ILayoutBlock >& Block, const FGeometry& AllottedGeometry, FArrangedChildren& ArrangedChildren ) const 
+int32 SMultiLineEditableText::FCursorLineHighlighter::OnPaint( const FTextLayout::FLineView& Line, const float OffsetX, const float Width, const FTextBlockStyle& DefaultStyle, const FGeometry& AllottedGeometry, const FSlateRect& MyClippingRect, FSlateWindowElementList& OutDrawElements, int32 LayerId, const FWidgetStyle& InWidgetStyle, bool bParentEnabled ) const 
 {
-	//No Widgets
-}
-
-FChildren* SMultiLineEditableText::FSlateCursorRunHighlighter::GetChildren()
-{
-	static FNoChildren NoChildrenInstance;
-	return &NoChildrenInstance;
-}
-
-int32 SMultiLineEditableText::FSlateCursorRunHighlighter::OnPaint( const FTextLayout::FLineView& Line, const TSharedRef< ISlateRun >& Run, const TSharedRef< ILayoutBlock >& Block, const FTextBlockStyle& DefaultStyle, const FGeometry& AllottedGeometry, const FSlateRect& MyClippingRect, FSlateWindowElementList& OutDrawElements, int32 LayerId, const FWidgetStyle& InWidgetStyle, bool bParentEnabled ) const 
-{
-	LayerId = OnPaintCursor(Line, Run, Block, DefaultStyle, AllottedGeometry, MyClippingRect, OutDrawElements, LayerId, InWidgetStyle, bParentEnabled, 0.0f, CursorInfo->GetCursorAlignment());
-
-	return Run->OnPaint( Line, Block, DefaultStyle, AllottedGeometry, MyClippingRect, OutDrawElements, LayerId, InWidgetStyle, bParentEnabled );
-}
-
-int32 SMultiLineEditableText::FSlateCursorRunHighlighter::OnPaintCursor( const FTextLayout::FLineView& Line, const TSharedRef< ISlateRun >& Run, const TSharedRef< ILayoutBlock >& Block, const FTextBlockStyle& DefaultStyle, const FGeometry& AllottedGeometry, const FSlateRect& MyClippingRect, FSlateWindowElementList& OutDrawElements, int32 LayerId, const FWidgetStyle& InWidgetStyle, const bool bParentEnabled, const float InCursorOffset, const ECursorAlignment InCursorAlignment ) const
-{
-	FVector2D Location( Block->GetLocationOffset() );
-	Location.Y = Line.Offset.Y;
+	const FVector2D Location(Line.Offset.X + OffsetX, Line.Offset.Y);
+	const FVector2D Size(Width, Line.TextSize.Y);
 
 	FLinearColor CursorColorAndOpacity = InWidgetStyle.GetForegroundColor();
 
@@ -115,13 +97,12 @@ int32 SMultiLineEditableText::FSlateCursorRunHighlighter::OnPaintCursor( const F
 	// @todo: Slate Styles - make this brush part of the widget style
 	const FSlateBrush* CursorBrush = FCoreStyle::Get().GetBrush("EditableText.SelectionBackground");
 
-	FVector2D OptionalWidth = InCursorAlignment == ECursorAlignment::Right ? FVector2D(Block->GetSize().X, 0) : FVector2D::ZeroVector;
-	OptionalWidth += FVector2D(InCursorOffset, 0.0f);
+	const FVector2D OptionalWidth = CursorInfo->GetCursorAlignment() == ECursorAlignment::Right ? FVector2D(Size.X, 0) : FVector2D::ZeroVector;
 
 	FSlateDrawElement::MakeBox(
 		OutDrawElements,
 		LayerId,
-		FPaintGeometry(AllottedGeometry.AbsolutePosition + Location + OptionalWidth, FVector2D(CursorWidth * AllottedGeometry.Scale, Line.TextSize.Y), AllottedGeometry.Scale),
+		FPaintGeometry(AllottedGeometry.AbsolutePosition + Location + OptionalWidth, FVector2D(CursorWidth * AllottedGeometry.Scale, Size.Y), AllottedGeometry.Scale),
 		CursorBrush,
 		MyClippingRect,
 		bParentEnabled ? ESlateDrawEffect::None : ESlateDrawEffect::DisabledEffect,
@@ -130,123 +111,87 @@ int32 SMultiLineEditableText::FSlateCursorRunHighlighter::OnPaintCursor( const F
 	return LayerId;
 }
 
-TSharedRef< SMultiLineEditableText::FSlateCursorRunHighlighter > SMultiLineEditableText::FSlateCursorRunHighlighter::Create(const FCursorInfo* InCursorInfo)
+TSharedRef< SMultiLineEditableText::FCursorLineHighlighter > SMultiLineEditableText::FCursorLineHighlighter::Create(const FCursorInfo* InCursorInfo)
 {
-	return MakeShareable( new SMultiLineEditableText::FSlateCursorRunHighlighter(InCursorInfo) );
+	return MakeShareable( new SMultiLineEditableText::FCursorLineHighlighter(InCursorInfo) );
 }
 
-SMultiLineEditableText::FSlateSelectionRunHighlighter::FSlateSelectionRunHighlighter(const FCursorInfo* InCursorInfo)
-	: SMultiLineEditableText::FSlateCursorRunHighlighter(InCursorInfo)
-	, CursorAlignment(ECursorAlignment::Left)
+SMultiLineEditableText::FTextCompositionHighlighter::FTextCompositionHighlighter()
 {
 }
 
-int32 SMultiLineEditableText::FSlateSelectionRunHighlighter::OnPaint( const FTextLayout::FLineView& Line, const TSharedRef< ISlateRun >& Run, const TSharedRef< ILayoutBlock >& Block, const FTextBlockStyle& DefaultStyle, const FGeometry& AllottedGeometry, const FSlateRect& MyClippingRect, FSlateWindowElementList& OutDrawElements, int32 LayerId, const FWidgetStyle& InWidgetStyle, bool bParentEnabled ) const 
+int32 SMultiLineEditableText::FTextCompositionHighlighter::OnPaint( const FTextLayout::FLineView& Line, const float OffsetX, const float Width, const FTextBlockStyle& DefaultStyle, const FGeometry& AllottedGeometry, const FSlateRect& MyClippingRect, FSlateWindowElementList& OutDrawElements, int32 LayerId, const FWidgetStyle& InWidgetStyle, bool bParentEnabled ) const
 {
-	FVector2D Location( Block->GetLocationOffset() );
-	Location.Y = Line.Offset.Y;
+	const FVector2D Location(Line.Offset.X + OffsetX, Line.Offset.Y);
+	const FVector2D Size(Width, Line.TextSize.Y);
 
-	const FColor SelectionBackgroundColorAndOpacity = ((FLinearColor::White - DefaultStyle.ColorAndOpacity.GetColor(InWidgetStyle))*0.5f + FLinearColor(-0.2f, -0.05f, 0.15f)) * InWidgetStyle.GetColorAndOpacityTint();
-
-	const float HighlightWidth = Block->GetSize().X;
-	if (HighlightWidth)
+	if (Size.X)
 	{
-		// Draw the actual highlight rectangle
+		const FLinearColor LineColorAndOpacity = InWidgetStyle.GetForegroundColor();
+
+		// @todo: Slate Styles - make this brush part of the widget style
+		const FSlateBrush* CompositionBrush = FCoreStyle::Get().GetBrush("EditableText.CompositionBackground");
+
+		// Draw composition background
 		FSlateDrawElement::MakeBox(
 			OutDrawElements,
 			++LayerId,
-			FPaintGeometry(AllottedGeometry.AbsolutePosition + Location, FVector2D(HighlightWidth, Line.Size.Y), AllottedGeometry.Scale),
-			&DefaultStyle.HighlightShape,
+			FPaintGeometry(AllottedGeometry.AbsolutePosition + Location, Size, AllottedGeometry.Scale),
+			CompositionBrush,
 			MyClippingRect,
 			bParentEnabled ? ESlateDrawEffect::None : ESlateDrawEffect::DisabledEffect,
-			SelectionBackgroundColorAndOpacity
+			LineColorAndOpacity * InWidgetStyle.GetColorAndOpacityTint()
 		);
 	}
 
-	//FLinearColor InvertedForeground = FLinearColor::White - InWidgetStyle.GetForegroundColor();
-	//InvertedForeground.A = InWidgetStyle.GetForegroundColor().A;
-
-	//FWidgetStyle WidgetStyle( InWidgetStyle );
-	//WidgetStyle.SetForegroundColor( InvertedForeground );
-
-	// If the cursor is within the range of this block, we need to draw it
-	const FTextLocation CursorLocation = CursorInfo->GetCursorLocation();
-	if(Line.ModelIndex == CursorLocation.GetLineIndex() && Block->GetTextRange().InclusiveContains(CursorLocation.GetOffset()))
-	{
-		LayerId = OnPaintCursor(Line, Run, Block, DefaultStyle, AllottedGeometry, MyClippingRect, OutDrawElements, LayerId, InWidgetStyle, bParentEnabled, 0.0f, CursorAlignment);
-	}
-
-	return Run->OnPaint( Line, Block, DefaultStyle, AllottedGeometry, MyClippingRect, OutDrawElements, LayerId, InWidgetStyle, bParentEnabled );
+	return LayerId;
 }
 
-TSharedRef< SMultiLineEditableText::FSlateSelectionRunHighlighter > SMultiLineEditableText::FSlateSelectionRunHighlighter::Create(const FCursorInfo* InCursorInfo)
+TSharedRef< SMultiLineEditableText::FTextCompositionHighlighter > SMultiLineEditableText::FTextCompositionHighlighter::Create()
 {
-	return MakeShareable( new SMultiLineEditableText::FSlateSelectionRunHighlighter(InCursorInfo) );
+	return MakeShareable( new SMultiLineEditableText::FTextCompositionHighlighter() );
 }
 
-SMultiLineEditableText::FSlateCompositionRunHighlighter::FSlateCompositionRunHighlighter(const FCursorInfo* InCursorInfo)
-	: SMultiLineEditableText::FSlateCursorRunHighlighter(InCursorInfo)
+SMultiLineEditableText::FTextSelectionRunRenderer::FTextSelectionRunRenderer()
 {
 }
 
-int32 SMultiLineEditableText::FSlateCompositionRunHighlighter::OnPaint( const FTextLayout::FLineView& Line, const TSharedRef< ISlateRun >& Run, const TSharedRef< ILayoutBlock >& Block, const FTextBlockStyle& DefaultStyle, const FGeometry& AllottedGeometry, const FSlateRect& MyClippingRect, FSlateWindowElementList& OutDrawElements, int32 LayerId, const FWidgetStyle& InWidgetStyle, bool bParentEnabled ) const 
+int32 SMultiLineEditableText::FTextSelectionRunRenderer::OnPaint( const FTextLayout::FLineView& Line, const TSharedRef< ISlateRun >& Run, const TSharedRef< ILayoutBlock >& Block, const FTextBlockStyle& DefaultStyle, const FGeometry& AllottedGeometry, const FSlateRect& MyClippingRect, FSlateWindowElementList& OutDrawElements, int32 LayerId, const FWidgetStyle& InWidgetStyle, bool bParentEnabled ) const 
 {
 	FVector2D Location( Block->GetLocationOffset() );
 	Location.Y = Line.Offset.Y;
 
-	const float CompositionWidth = Block->GetSize().X;
-	if (CompositionWidth)
-	{
-		const FLinearColor LineColorAndOpacity = InWidgetStyle.GetForegroundColor();
-
-		// @todo: Slate Styles - make this brush part of the widget style
-		const FSlateBrush* CompositionBrush = FCoreStyle::Get().GetBrush("EditableText.CompositionBackground");
+	const FLinearColor FocusedSelectionColorAdjustment(-0.2f, -0.05f, 0.15f);
+	const FLinearColor UnfocusedSelectionColorAdjustment(-0.01f, -0.01f, 0.01f);
+	const FColor SelectionBackgroundColorAndOpacity = ((FLinearColor::White - DefaultStyle.ColorAndOpacity.GetColor(InWidgetStyle))*0.5f + (bHasKeyboardFocus ? FocusedSelectionColorAdjustment : UnfocusedSelectionColorAdjustment)) * InWidgetStyle.GetColorAndOpacityTint();
 
-		// Draw composition background
+	const float HighlightWidth = Block->GetSize().X;
+	if (HighlightWidth)
+	{
+		// Draw the actual highlight rectangle
 		FSlateDrawElement::MakeBox(
 			OutDrawElements,
 			++LayerId,
-			FPaintGeometry(AllottedGeometry.AbsolutePosition + Location, FVector2D(CompositionWidth, Line.TextSize.Y), AllottedGeometry.Scale),
-			CompositionBrush,
+			FPaintGeometry(AllottedGeometry.AbsolutePosition + Location, FVector2D(HighlightWidth, Line.Size.Y), AllottedGeometry.Scale),
+			&DefaultStyle.HighlightShape,
 			MyClippingRect,
 			bParentEnabled ? ESlateDrawEffect::None : ESlateDrawEffect::DisabledEffect,
-			LineColorAndOpacity * InWidgetStyle.GetColorAndOpacityTint()
+			SelectionBackgroundColorAndOpacity
 		);
 	}
 
-	// If the cursor is within the range of this block, we need to draw it
-	const FTextLocation CursorInteractionLocation = CursorInfo->GetCursorInteractionLocation();
-	if(Line.ModelIndex == CursorInteractionLocation.GetLineIndex() && Block->GetTextRange().InclusiveContains(CursorInteractionLocation.GetOffset()))
-	{
-		// Unlike selection where the cursor will always be at one end of the block, composition may place the cursor anywhere within the block,
-		// so we need to calculate the offset (in scaled screen space) required from the start of the block to the position of the cursor
-		// We always do this check using the interaction position of the cursor, as this ensures we draw the cursor at the correct place, 
-		// even when it's to the right of a character (which is why we also pass ECursorAlignment::Left into OnPaintCursor, otherwise we'd doubly compensate)
-
-		// Get the text for the current run
-		FString RunText;
-		Run->AppendText(RunText);
-
-		// Trim that text down to the text for this block
-		const FTextRange BlockTextRange = Block->GetTextRange();
-		const FString BlockText = RunText.Mid(BlockTextRange.BeginIndex, BlockTextRange.EndIndex - BlockTextRange.BeginIndex);
-
-		// Then trim that down so we have the text up-to the interaction point of the cursor
-		const FString TextUpToCursor = BlockText.Left(CursorInteractionLocation.GetOffset() - BlockTextRange.BeginIndex);
-
-		// Finally, measure that text so we can get an accurate offset for the cursor position within this block
-		const TSharedRef< FSlateFontMeasure > FontMeasureService = FSlateApplication::Get().GetRenderer()->GetFontMeasureService();
-		const FVector2D CursorOffset = FontMeasureService->Measure(TextUpToCursor, DefaultStyle.Font, AllottedGeometry.Scale);
+	//FLinearColor InvertedForeground = FLinearColor::White - InWidgetStyle.GetForegroundColor();
+	//InvertedForeground.A = InWidgetStyle.GetForegroundColor().A;
 
-		LayerId = OnPaintCursor(Line, Run, Block, DefaultStyle, AllottedGeometry, MyClippingRect, OutDrawElements, LayerId, InWidgetStyle, bParentEnabled, CursorOffset.X, ECursorAlignment::Left);
-	}
+	//FWidgetStyle WidgetStyle( InWidgetStyle );
+	//WidgetStyle.SetForegroundColor( InvertedForeground );
 
 	return Run->OnPaint( Line, Block, DefaultStyle, AllottedGeometry, MyClippingRect, OutDrawElements, LayerId, InWidgetStyle, bParentEnabled );
 }
 
-TSharedRef< SMultiLineEditableText::FSlateCompositionRunHighlighter > SMultiLineEditableText::FSlateCompositionRunHighlighter::Create(const FCursorInfo* InCursorInfo)
+TSharedRef< SMultiLineEditableText::FTextSelectionRunRenderer > SMultiLineEditableText::FTextSelectionRunRenderer::Create()
 {
-	return MakeShareable( new SMultiLineEditableText::FSlateCompositionRunHighlighter(InCursorInfo) );
+	return MakeShareable( new SMultiLineEditableText::FTextSelectionRunRenderer() );
 }
 
 SMultiLineEditableText::SMultiLineEditableText()
@@ -297,9 +242,9 @@ void SMultiLineEditableText::Construct( const FArguments& InArgs )
 	OnTextChanged = InArgs._OnTextChanged;
 	OnTextCommitted = InArgs._OnTextCommitted;
 
-	CursorRunHighlighter = FSlateCursorRunHighlighter::Create(&CursorInfo);
-	SelectionRunHighlighter = FSlateSelectionRunHighlighter::Create(&CursorInfo);
-	CompositionRunHighlighter = FSlateCompositionRunHighlighter::Create(&CursorInfo);
+	CursorLineHighlighter = FCursorLineHighlighter::Create(&CursorInfo);
+	TextCompositionHighlighter = FTextCompositionHighlighter::Create();
+	TextSelectionRunRenderer = FTextSelectionRunRenderer::Create();
 	TextLayout = FSlateTextLayout::Create();
 
 	BoundText = InArgs._Text;
@@ -395,6 +340,7 @@ bool SMultiLineEditableText::SetEditableText(const FText& TextToSet, const bool
 		const FString& TextToSetString = TextToSet.ToString();
 		const int32 TextToSetLength = TextToSetString.Len();
 
+		ClearSelection();
 		TextLayout->ClearLines();
 
 		if(TextToSetLength)
@@ -976,13 +922,15 @@ FReply SMultiLineEditableText::MoveCursor( ECursorMoveMethod::Type Method, const
 
 void SMultiLineEditableText::UpdateCursorHighlight()
 {
-	TextLayout->ClearHighlights();
+	RemoveCursorHighlight();
+
+	static const int32 CompositionRangeZOrder = 1; // draw above the text
+	static const int32 CursorZOrder = 2; // draw above the text and the composition
 
 	const FTextLocation CursorInteractionPosition = CursorInfo.GetCursorInteractionLocation();
 	const FTextLocation SelectionLocation = SelectionStart.Get( CursorInteractionPosition );
 
-	bool bDrawCursor = true;
-
+	const bool bHasKeyboardFocus = HasKeyboardFocus();
 	const bool bIsComposing = TextInputMethodContext->IsComposing;
 	const bool bHasSelection = SelectionLocation != CursorInteractionPosition;
 
@@ -1000,22 +948,15 @@ void SMultiLineEditableText::UpdateCursorHighlight()
 			const FTextRange Range(CompositionBeginLocation.GetOffset(), CompositionEndLocation.GetOffset());
 
 			// We only draw the composition highlight if the cursor is within the composition range
-			bDrawCursor = !(CompositionBeginLocation.GetLineIndex() == CursorInteractionPosition.GetLineIndex() && Range.InclusiveContains(CursorInteractionPosition.GetOffset()));
-			if (!Range.IsEmpty() && !bDrawCursor)
-			{
-				TextLayout->AddHighlight(FTextHighlight(CompositionBeginLocation.GetLineIndex(), Range, CompositionRunHighlighter.ToSharedRef()));
-			}
-			else
+			const bool bCursorInRange = (CompositionBeginLocation.GetLineIndex() == CursorInteractionPosition.GetLineIndex() && Range.InclusiveContains(CursorInteractionPosition.GetOffset()));
+			if (!Range.IsEmpty() && bCursorInRange)
 			{
-				// We have no composition range for this line, but we're supposed to draw the cursor on this line - let the normal cursor handling take care of that
-				bDrawCursor = true;
+				TextLayout->AddLineHighlight(FTextLineHighlight(CompositionBeginLocation.GetLineIndex(), Range, CompositionRangeZOrder, TextCompositionHighlighter.ToSharedRef()));
 			}
 		}
 	}
 	else if ( bHasSelection )
 	{
-		bDrawCursor = false;
-
 		const FTextSelection Selection( SelectionLocation, CursorInteractionPosition );
 
 		const int32 SelectionBeginningLineIndex = Selection.GetBeginning().GetLineIndex();
@@ -1024,20 +965,14 @@ void SMultiLineEditableText::UpdateCursorHighlight()
 		const int32 SelectionEndLineIndex = Selection.GetEnd().GetLineIndex();
 		const int32 SelectionEndLineOffset = Selection.GetEnd().GetOffset();
 
-		const bool bCursorIsBeforeSelection = CursorInteractionPosition < SelectionLocation;
-		SelectionRunHighlighter->SetCursorAlignment((bCursorIsBeforeSelection) ? ECursorAlignment::Left : ECursorAlignment::Right);
+		TextSelectionRunRenderer->SetHasKeyboardFocus(bHasKeyboardFocus);
 
 		if ( SelectionBeginningLineIndex == SelectionEndLineIndex )
 		{
 			const FTextRange Range(SelectionBeginningLineOffset, SelectionEndLineOffset);
 			if (!Range.IsEmpty())
 			{
-				TextLayout->AddHighlight(FTextHighlight(SelectionBeginningLineIndex, Range, SelectionRunHighlighter.ToSharedRef()));
-			}
-			else
-			{
-				// We have no selection range for this line, but we're supposed to draw the cursor on this line - let the normal cursor handling take care of that
-				bDrawCursor = true;
+				TextLayout->AddRunRenderer(FTextRunRenderer(SelectionBeginningLineIndex, Range, TextSelectionRunRenderer.ToSharedRef()));
 			}
 		}
 		else
@@ -1051,12 +986,7 @@ void SMultiLineEditableText::UpdateCursorHighlight()
 					const FTextRange Range(SelectionBeginningLineOffset, Lines[LineIndex].Text->Len());
 					if (!Range.IsEmpty())
 					{
-						TextLayout->AddHighlight(FTextHighlight(LineIndex, Range, SelectionRunHighlighter.ToSharedRef()));
-					}
-					else if (bCursorIsBeforeSelection)
-					{
-						// We have no selection range for this line, but we're supposed to draw the cursor on this line - let the normal cursor handling take care of that
-						bDrawCursor = true;
+						TextLayout->AddRunRenderer(FTextRunRenderer(LineIndex, Range, TextSelectionRunRenderer.ToSharedRef()));
 					}
 				}
 				else if ( LineIndex == SelectionEndLineIndex )
@@ -1064,12 +994,7 @@ void SMultiLineEditableText::UpdateCursorHighlight()
 					const FTextRange Range(0, SelectionEndLineOffset);
 					if (!Range.IsEmpty())
 					{
-						TextLayout->AddHighlight(FTextHighlight(LineIndex, Range, SelectionRunHighlighter.ToSharedRef()));
-					}
-					else if (!bCursorIsBeforeSelection)
-					{
-						// We have no selection range for this line, but we're supposed to draw the cursor on this line - let the normal cursor handling take care of that
-						bDrawCursor = true;
+						TextLayout->AddRunRenderer(FTextRunRenderer(LineIndex, Range, TextSelectionRunRenderer.ToSharedRef()));
 					}
 				}
 				else
@@ -1077,14 +1002,13 @@ void SMultiLineEditableText::UpdateCursorHighlight()
 					const FTextRange Range(0, Lines[LineIndex].Text->Len());
 					if (!Range.IsEmpty())
 					{
-						TextLayout->AddHighlight(FTextHighlight(LineIndex, Range, SelectionRunHighlighter.ToSharedRef()));
+						TextLayout->AddRunRenderer(FTextRunRenderer(LineIndex, Range, TextSelectionRunRenderer.ToSharedRef()));
 					}
 				}
 			}
 		}
 	}
 	
-	if ( bDrawCursor )
 	{
 		// The cursor mode uses the literal position rather than the interaction position
 		const FTextLocation CursorPosition = CursorInfo.GetCursorLocation();
@@ -1094,22 +1018,23 @@ void SMultiLineEditableText::UpdateCursorHighlight()
 
 		if (LineTextLength == 0)
 		{
-			TextLayout->AddHighlight(FTextHighlight(CursorPosition.GetLineIndex(), FTextRange(0, 0), CursorRunHighlighter.ToSharedRef()));
+			TextLayout->AddLineHighlight(FTextLineHighlight(CursorPosition.GetLineIndex(), FTextRange(0, 0), CursorZOrder, CursorLineHighlighter.ToSharedRef()));
 		}
 		else if (CursorPosition.GetOffset() == LineTextLength)
 		{
-			TextLayout->AddHighlight(FTextHighlight(CursorPosition.GetLineIndex(), FTextRange(LineTextLength - 1, LineTextLength), CursorRunHighlighter.ToSharedRef()));
+			TextLayout->AddLineHighlight(FTextLineHighlight(CursorPosition.GetLineIndex(), FTextRange(LineTextLength - 1, LineTextLength), CursorZOrder, CursorLineHighlighter.ToSharedRef()));
 		}
 		else
 		{
-			TextLayout->AddHighlight(FTextHighlight(CursorPosition.GetLineIndex(), FTextRange(CursorPosition.GetOffset(), CursorPosition.GetOffset() + 1), CursorRunHighlighter.ToSharedRef()));
+			TextLayout->AddLineHighlight(FTextLineHighlight(CursorPosition.GetLineIndex(), FTextRange(CursorPosition.GetOffset(), CursorPosition.GetOffset() + 1), CursorZOrder, CursorLineHighlighter.ToSharedRef()));
 		}
 	}
 }
 
 void SMultiLineEditableText::RemoveCursorHighlight()
 {
-	TextLayout->ClearHighlights();
+	TextLayout->ClearRunRenderers();
+	TextLayout->ClearLineHighlights();
 }
 
 void SMultiLineEditableText::JumpTo(ETextLocation::Type JumpLocation, ECursorAction::Type Action)
diff --git a/Engine/Source/Runtime/Slate/Private/Widgets/Text/SRichTextBlock.cpp b/Engine/Source/Runtime/Slate/Private/Widgets/Text/SRichTextBlock.cpp
--- a/Engine/Source/Runtime/Slate/Private/Widgets/Text/SRichTextBlock.cpp
+++ b/Engine/Source/Runtime/Slate/Private/Widgets/Text/SRichTextBlock.cpp
@@ -3,7 +3,7 @@
 #include "SlatePrivatePCH.h"
 #include "TextLayoutEngine.h"
 #include "SlateTextLayout.h"
-#include "SlateSimpleRunHighlighter.h"
+#include "SlateTextHighlightRunRenderer.h"
 #include "IRichTextMarkupParser.h"
 #include "RichTextMarkupProcessing.h"
 
@@ -26,7 +26,7 @@ void SRichTextBlock::Construct( const FArguments& InArgs )
 		Parser = FRichTextMarkupProcessing::Create();
 	}
 
-	TextHighlighter = FSlateSimpleRunHighlighter::Create();
+	TextHighlighter = FSlateTextHighlightRunRenderer::Create();
 
 	Decorators.Append( InArgs._Decorators );
 	Decorators.Append( InArgs.InlineDecorators );
@@ -217,7 +217,7 @@ void SRichTextBlock::SetHighlightText( const FText& InHighlightText )
 		return;
 	}
 
-	Highlights.Empty();
+	TextHighlights.Empty();
 	HighlightText = InHighlightText;
 
 	const FString HighlightTextString = HighlightText.ToString();
@@ -236,16 +236,16 @@ void SRichTextBlock::SetHighlightText( const FText& InHighlightText )
 		{
 			FindBegin = CurrentHighlightBegin + HighlightTextLength;
 
-			if ( Highlights.Num() > 0 && Highlights.Last().LineIndex == LineIndex && Highlights.Last().Range.EndIndex == CurrentHighlightBegin )
+			if ( TextHighlights.Num() > 0 && TextHighlights.Last().LineIndex == LineIndex && TextHighlights.Last().Range.EndIndex == CurrentHighlightBegin )
 			{
-				Highlights[ Highlights.Num() - 1 ] = FTextHighlight( LineIndex, FTextRange( Highlights.Last().Range.BeginIndex, FindBegin ), TextHighlighter.ToSharedRef() );
+				TextHighlights[ TextHighlights.Num() - 1 ] = FTextRunRenderer( LineIndex, FTextRange( TextHighlights.Last().Range.BeginIndex, FindBegin ), TextHighlighter.ToSharedRef() );
 			}
 			else
 			{
-				Highlights.Add( FTextHighlight( LineIndex, FTextRange( CurrentHighlightBegin, FindBegin ), TextHighlighter.ToSharedRef() ) );
+				TextHighlights.Add( FTextRunRenderer( LineIndex, FTextRange( CurrentHighlightBegin, FindBegin ), TextHighlighter.ToSharedRef() ) );
 			}
 		}
 	}
 
-	TextLayout->SetHighlights( Highlights );
+	TextLayout->SetRunRenderers( TextHighlights );
 }
diff --git a/Engine/Source/Runtime/Slate/Public/Framework/Text/DefaultLayoutBlock.h b/Engine/Source/Runtime/Slate/Public/Framework/Text/DefaultLayoutBlock.h
--- a/Engine/Source/Runtime/Slate/Public/Framework/Text/DefaultLayoutBlock.h
+++ b/Engine/Source/Runtime/Slate/Public/Framework/Text/DefaultLayoutBlock.h
@@ -5,9 +5,9 @@ class SLATE_API FDefaultLayoutBlock : public ILayoutBlock
 {
 public:
 
-	static TSharedRef< FDefaultLayoutBlock > Create( const TSharedRef< IRun >& InRun, const FTextRange& InRange, FVector2D InSize, const TSharedPtr< IRunHighlighter >& InHighlighter )
+	static TSharedRef< FDefaultLayoutBlock > Create( const TSharedRef< IRun >& InRun, const FTextRange& InRange, FVector2D InSize, const TSharedPtr< IRunRenderer >& InRenderer )
 	{
-		return MakeShareable( new FDefaultLayoutBlock( InRun, InRange, InSize, InHighlighter ) );
+		return MakeShareable( new FDefaultLayoutBlock( InRun, InRange, InSize, InRenderer ) );
 	}
 
 	virtual ~FDefaultLayoutBlock() {}
@@ -15,7 +15,7 @@ public:
 	virtual TSharedRef< IRun > GetRun() const override { return Run; }
 	virtual FTextRange GetTextRange() const override { return Range; }
 	virtual FVector2D GetSize() const override { return Size; }
-	virtual TSharedPtr< IRunHighlighter > GetHighlighter() const override { return Highlighter; }
+	virtual TSharedPtr< IRunRenderer > GetRenderer() const override { return Renderer; }
 
 	virtual void SetLocationOffset( const FVector2D& InLocationOffset ) override { LocationOffset = InLocationOffset; }
 	virtual FVector2D GetLocationOffset() const override { return LocationOffset; }
@@ -27,12 +27,12 @@ private:
 		return MakeShareable( new FDefaultLayoutBlock( Block ) );
 	}
 
-	FDefaultLayoutBlock( const TSharedRef< IRun >& InRun, const FTextRange& InRange, FVector2D InSize, const TSharedPtr< IRunHighlighter >& InHighlighter )
+	FDefaultLayoutBlock( const TSharedRef< IRun >& InRun, const FTextRange& InRange, FVector2D InSize, const TSharedPtr< IRunRenderer >& InRenderer )
 		: Run( InRun )
 		, Range( InRange )
 		, Size( InSize )
 		, LocationOffset( ForceInitToZero )
-		, Highlighter( InHighlighter )
+		, Renderer( InRenderer )
 	{
 
 	}
@@ -42,7 +42,7 @@ private:
 		, Range( Block.Range )
 		, Size( Block.Size )
 		, LocationOffset( ForceInitToZero )
-		, Highlighter( Block.Highlighter )
+		, Renderer( Block.Renderer )
 	{
 
 	}
@@ -54,5 +54,5 @@ private:
 	FTextRange Range;
 	FVector2D Size;
 	FVector2D LocationOffset;
-	TSharedPtr< IRunHighlighter > Highlighter;
+	TSharedPtr< IRunRenderer > Renderer;
 };
\ No newline at end of file
diff --git a/Engine/Source/Runtime/Slate/Public/Framework/Text/ILayoutBlock.h b/Engine/Source/Runtime/Slate/Public/Framework/Text/ILayoutBlock.h
--- a/Engine/Source/Runtime/Slate/Public/Framework/Text/ILayoutBlock.h
+++ b/Engine/Source/Runtime/Slate/Public/Framework/Text/ILayoutBlock.h
@@ -8,7 +8,7 @@ public:
 	virtual TSharedRef< class IRun > GetRun() const = 0;
 	virtual FTextRange GetTextRange() const = 0;
 	virtual FVector2D GetSize() const = 0;
-	virtual TSharedPtr< class IRunHighlighter > GetHighlighter() const = 0;
+	virtual TSharedPtr< class IRunRenderer > GetRenderer() const = 0;
 
 	virtual void SetLocationOffset( const FVector2D& InLocationOffset ) = 0;
 	virtual FVector2D GetLocationOffset() const = 0;
diff --git a/Engine/Source/Runtime/Slate/Public/Framework/Text/ILineHighlighter.h b/Engine/Source/Runtime/Slate/Public/Framework/Text/ILineHighlighter.h
new file mode 100644
--- /dev/null
+++ b/Engine/Source/Runtime/Slate/Public/Framework/Text/ILineHighlighter.h
@@ -0,0 +1,10 @@
+// Copyright 1998-2014 Epic Games, Inc. All Rights Reserved.
+#pragma once
+
+class SLATE_API ILineHighlighter
+{
+public:
+
+	virtual ~ILineHighlighter() {}
+
+};
diff --git a/Engine/Source/Runtime/Slate/Public/Framework/Text/IRun.h b/Engine/Source/Runtime/Slate/Public/Framework/Text/IRun.h
--- a/Engine/Source/Runtime/Slate/Public/Framework/Text/IRun.h
+++ b/Engine/Source/Runtime/Slate/Public/Framework/Text/IRun.h
@@ -14,7 +14,7 @@ public:
 	virtual FVector2D Measure( int32 StartIndex, int32 EndIndex, float Scale) const = 0;
 	virtual int8 GetKerning(int32 CurrentIndex, float Scale) const = 0;
 
-	virtual TSharedRef< class ILayoutBlock > CreateBlock( int32 StartIndex, int32 EndIndex, FVector2D Size, const TSharedPtr< class IRunHighlighter >& Highlighter ) = 0;
+	virtual TSharedRef< class ILayoutBlock > CreateBlock( int32 StartIndex, int32 EndIndex, FVector2D Size, const TSharedPtr< class IRunRenderer >& Renderer ) = 0;
 
 	virtual int32 GetTextIndexAt( const TSharedRef< ILayoutBlock >& Block, const FVector2D& Location, float Scale ) const = 0;
 	
diff --git a/Engine/Source/Runtime/Slate/Public/Framework/Text/IRunHighlighter.h b/Engine/Source/Runtime/Slate/Public/Framework/Text/IRunHighlighter.h
deleted file mode 100644
diff --git a/Engine/Source/Runtime/Slate/Public/Framework/Text/IRunRenderer.h b/Engine/Source/Runtime/Slate/Public/Framework/Text/IRunRenderer.h
new file mode 100644
--- /dev/null
+++ b/Engine/Source/Runtime/Slate/Public/Framework/Text/IRunRenderer.h
@@ -0,0 +1,10 @@
+// Copyright 1998-2014 Epic Games, Inc. All Rights Reserved.
+#pragma once
+
+class SLATE_API IRunRenderer
+{
+public:
+
+	virtual ~IRunRenderer() {}
+
+};
diff --git a/Engine/Source/Runtime/Slate/Public/Framework/Text/ISlateLineHighlighter.h b/Engine/Source/Runtime/Slate/Public/Framework/Text/ISlateLineHighlighter.h
new file mode 100644
--- /dev/null
+++ b/Engine/Source/Runtime/Slate/Public/Framework/Text/ISlateLineHighlighter.h
@@ -0,0 +1,10 @@
+// Copyright 1998-2014 Epic Games, Inc. All Rights Reserved.
+#pragma once
+
+class SLATE_API ISlateLineHighlighter : public ILineHighlighter
+{
+public:
+	virtual ~ISlateLineHighlighter() {}
+
+	virtual int32 OnPaint( const FTextLayout::FLineView& Line, const float OffsetX, const float Width, const FTextBlockStyle& DefaultStyle, const FGeometry& AllottedGeometry, const FSlateRect& MyClippingRect, FSlateWindowElementList& OutDrawElements, int32 LayerId, const FWidgetStyle& InWidgetStyle, bool bParentEnabled ) const = 0;
+};
diff --git a/Engine/Source/Runtime/Slate/Public/Framework/Text/ISlateRunHighlighter.h b/Engine/Source/Runtime/Slate/Public/Framework/Text/ISlateRunHighlighter.h
deleted file mode 100644
diff --git a/Engine/Source/Runtime/Slate/Public/Framework/Text/ISlateRunRenderer.h b/Engine/Source/Runtime/Slate/Public/Framework/Text/ISlateRunRenderer.h
new file mode 100644
--- /dev/null
+++ b/Engine/Source/Runtime/Slate/Public/Framework/Text/ISlateRunRenderer.h
@@ -0,0 +1,13 @@
+// Copyright 1998-2014 Epic Games, Inc. All Rights Reserved.
+#pragma once
+
+class ISlateRun;
+class ILayoutBlock;
+
+class SLATE_API ISlateRunRenderer : public IRunRenderer
+{
+public:
+	virtual ~ISlateRunRenderer() {}
+
+	virtual int32 OnPaint( const FTextLayout::FLineView& Line, const TSharedRef< ISlateRun >& Run, const TSharedRef< ILayoutBlock >& Block, const FTextBlockStyle& DefaultStyle, const FGeometry& AllottedGeometry, const FSlateRect& MyClippingRect, FSlateWindowElementList& OutDrawElements, int32 LayerId, const FWidgetStyle& InWidgetStyle, bool bParentEnabled ) const = 0;
+};
diff --git a/Engine/Source/Runtime/Slate/Public/Framework/Text/SlateHyperlinkRun.h b/Engine/Source/Runtime/Slate/Public/Framework/Text/SlateHyperlinkRun.h
--- a/Engine/Source/Runtime/Slate/Public/Framework/Text/SlateHyperlinkRun.h
+++ b/Engine/Source/Runtime/Slate/Public/Framework/Text/SlateHyperlinkRun.h
@@ -48,7 +48,7 @@ public:
 
 	virtual int8 GetKerning(int32 CurrentIndex, float Scale) const override;
 
-	virtual TSharedRef< ILayoutBlock > CreateBlock( int32 StartIndex, int32 EndIndex, FVector2D Size, const TSharedPtr< IRunHighlighter >& Highlighter ) override;
+	virtual TSharedRef< ILayoutBlock > CreateBlock( int32 StartIndex, int32 EndIndex, FVector2D Size, const TSharedPtr< IRunRenderer >& Renderer ) override;
 
 	virtual int32 OnPaint( const FTextLayout::FLineView& Line, const TSharedRef< ILayoutBlock >& Block, const FTextBlockStyle& DefaultStyle, const FGeometry& AllottedGeometry, const FSlateRect& MyClippingRect, FSlateWindowElementList& OutDrawElements, int32 LayerId, const FWidgetStyle& InWidgetStyle, bool bParentEnabled ) const override;
 
diff --git a/Engine/Source/Runtime/Slate/Public/Framework/Text/SlateImageRun.h b/Engine/Source/Runtime/Slate/Public/Framework/Text/SlateImageRun.h
--- a/Engine/Source/Runtime/Slate/Public/Framework/Text/SlateImageRun.h
+++ b/Engine/Source/Runtime/Slate/Public/Framework/Text/SlateImageRun.h
@@ -23,7 +23,7 @@ public:
 	virtual int8 GetKerning(int32 CurrentIndex, float Scale) const override;
 
 
-	virtual TSharedRef< ILayoutBlock > CreateBlock( int32 StartIndex, int32 EndIndex, FVector2D Size, const TSharedPtr< IRunHighlighter >& Highlighter ) override;
+	virtual TSharedRef< ILayoutBlock > CreateBlock( int32 StartIndex, int32 EndIndex, FVector2D Size, const TSharedPtr< IRunRenderer >& Renderer ) override;
 
 	virtual int32 OnPaint( const FTextLayout::FLineView& Line, const TSharedRef< ILayoutBlock >& Block, const FTextBlockStyle& DefaultStyle, const FGeometry& AllottedGeometry, const FSlateRect& MyClippingRect, FSlateWindowElementList& OutDrawElements, int32 LayerId, const FWidgetStyle& InWidgetStyle, bool bParentEnabled ) const override;
 
diff --git a/Engine/Source/Runtime/Slate/Public/Framework/Text/SlateSimpleRunHighlighter.h b/Engine/Source/Runtime/Slate/Public/Framework/Text/SlateSimpleRunHighlighter.h
deleted file mode 100644
diff --git a/Engine/Source/Runtime/Slate/Public/Framework/Text/SlateTextHighlightRunRenderer.h b/Engine/Source/Runtime/Slate/Public/Framework/Text/SlateTextHighlightRunRenderer.h
new file mode 100644
--- /dev/null
+++ b/Engine/Source/Runtime/Slate/Public/Framework/Text/SlateTextHighlightRunRenderer.h
@@ -0,0 +1,18 @@
+// Copyright 1998-2014 Epic Games, Inc. All Rights Reserved.
+#pragma once
+
+class SLATE_API FSlateTextHighlightRunRenderer : public ISlateRunRenderer
+{
+public:
+
+	static TSharedRef< FSlateTextHighlightRunRenderer > Create();
+
+	virtual ~FSlateTextHighlightRunRenderer() {}
+
+	virtual int32 OnPaint( const FTextLayout::FLineView& Line, const TSharedRef< ISlateRun >& Run, const TSharedRef< ILayoutBlock >& Block, const FTextBlockStyle& DefaultStyle, const FGeometry& AllottedGeometry, const FSlateRect& MyClippingRect, FSlateWindowElementList& OutDrawElements, int32 LayerId, const FWidgetStyle& InWidgetStyle, bool bParentEnabled ) const override;
+
+private:
+
+	FSlateTextHighlightRunRenderer();
+
+};
\ No newline at end of file
diff --git a/Engine/Source/Runtime/Slate/Public/Framework/Text/SlateTextLayout.h b/Engine/Source/Runtime/Slate/Public/Framework/Text/SlateTextLayout.h
--- a/Engine/Source/Runtime/Slate/Public/Framework/Text/SlateTextLayout.h
+++ b/Engine/Source/Runtime/Slate/Public/Framework/Text/SlateTextLayout.h
@@ -27,6 +27,8 @@ protected:
 
 	FSlateTextLayout();
 
+	int32 OnPaintHighlights( const FTextLayout::FLineView& LineView, const TArray<FLineViewHighlight>& Highlights, const FTextBlockStyle& DefaultTextStyle, const FGeometry& AllottedGeometry, const FSlateRect& ClippingRect, FSlateWindowElementList& OutDrawElements, const int32 LayerId, const FWidgetStyle& InWidgetStyle, bool bParentEnabled ) const;
+
 	void AggregateChildren();
 
 private:
diff --git a/Engine/Source/Runtime/Slate/Public/Framework/Text/SlateTextRun.h b/Engine/Source/Runtime/Slate/Public/Framework/Text/SlateTextRun.h
--- a/Engine/Source/Runtime/Slate/Public/Framework/Text/SlateTextRun.h
+++ b/Engine/Source/Runtime/Slate/Public/Framework/Text/SlateTextRun.h
@@ -21,7 +21,7 @@ public:
 	virtual FVector2D Measure( int32 StartIndex, int32 EndIndex, float Scale ) const override;
 	virtual int8 GetKerning(int32 CurrentIndex, float Scale) const override;
 
-	virtual TSharedRef< ILayoutBlock > CreateBlock( int32 StartIndex, int32 EndIndex, FVector2D Size, const TSharedPtr< IRunHighlighter >& Highlighter ) override;
+	virtual TSharedRef< ILayoutBlock > CreateBlock( int32 StartIndex, int32 EndIndex, FVector2D Size, const TSharedPtr< IRunRenderer >& Renderer ) override;
 
 	virtual int32 OnPaint( const FTextLayout::FLineView& Line, const TSharedRef< ILayoutBlock >& Block, const FTextBlockStyle& DefaultStyle, const FGeometry& AllottedGeometry, const FSlateRect& MyClippingRect, FSlateWindowElementList& OutDrawElements, int32 LayerId, const FWidgetStyle& InWidgetStyle, bool bParentEnabled ) const override;
 
diff --git a/Engine/Source/Runtime/Slate/Public/Framework/Text/SlateWidgetRun.h b/Engine/Source/Runtime/Slate/Public/Framework/Text/SlateWidgetRun.h
--- a/Engine/Source/Runtime/Slate/Public/Framework/Text/SlateWidgetRun.h
+++ b/Engine/Source/Runtime/Slate/Public/Framework/Text/SlateWidgetRun.h
@@ -43,7 +43,7 @@ public:
 	virtual FVector2D Measure( int32 StartIndex, int32 EndIndex, float Scale ) const override;
 	virtual int8 GetKerning(int32 CurrentIndex, float Scale) const override;
 
-	virtual TSharedRef< ILayoutBlock > CreateBlock( int32 StartIndex, int32 EndIndex, FVector2D Size, const TSharedPtr< IRunHighlighter >& Highlighter ) override;
+	virtual TSharedRef< ILayoutBlock > CreateBlock( int32 StartIndex, int32 EndIndex, FVector2D Size, const TSharedPtr< IRunRenderer >& Renderer ) override;
 
 	virtual int32 OnPaint( const FTextLayout::FLineView& Line, const TSharedRef< ILayoutBlock >& Block, const FTextBlockStyle& DefaultStyle, const FGeometry& AllottedGeometry, const FSlateRect& MyClippingRect, FSlateWindowElementList& OutDrawElements, int32 LayerId, const FWidgetStyle& InWidgetStyle, bool bParentEnabled ) const override;
 
diff --git a/Engine/Source/Runtime/Slate/Public/Framework/Text/TextHighlight.h b/Engine/Source/Runtime/Slate/Public/Framework/Text/TextHighlight.h
deleted file mode 100644
diff --git a/Engine/Source/Runtime/Slate/Public/Framework/Text/TextLayout.h b/Engine/Source/Runtime/Slate/Public/Framework/Text/TextLayout.h
--- a/Engine/Source/Runtime/Slate/Public/Framework/Text/TextLayout.h
+++ b/Engine/Source/Runtime/Slate/Public/Framework/Text/TextLayout.h
@@ -130,7 +130,7 @@ public:
 		/** Range inclusive of trailing whitespace, as used to visually display and interact with the text */
 		FTextRange ActualRange;
 
-		TSharedPtr< IRunHighlighter > Highlighter;
+		TSharedPtr< IRunRenderer > Renderer;
 	};
 
 	struct FBreakCandidate
@@ -204,13 +204,26 @@ public:
 		TSharedRef< FString > Text;
 		TArray< FRunModel > Runs;
 		TArray< FBreakCandidate > BreakCandidates;
-		TArray< FTextHighlight > Highlights;
+		TArray< FTextRunRenderer > RunRenderers;
+		TArray< FTextLineHighlight > LineHighlights;
 		mutable bool HasWrappingInformation;
 	};
 
+	struct FLineViewHighlight
+	{
+		/** Offset in X for this highlight, relative to the FLineView::Offset that contains it */
+		float OffsetX;
+		/** Width for this highlight, the height will be either FLineView::Size.Y or FLineView::TextSize.Y depending on whether you want to highlight the entire line, or just the text within the line */
+		float Width;
+		/** Custom highlighter implementation used to do the painting */
+		TSharedPtr< ILineHighlighter > Highlighter;
+	};
+
 	struct FLineView
 	{
 		TArray< TSharedRef< ILayoutBlock > > Blocks;
+		TArray< FLineViewHighlight > UnderlayHighlights;
+		TArray< FLineViewHighlight > OverlayHighlights;
 		FVector2D Offset;
 		FVector2D Size;
 		FVector2D TextSize;
@@ -280,26 +293,43 @@ public:
 	void AddLine( const TSharedRef< FString >& Text, const TArray< TSharedRef< IRun > >& Runs );
 
 	/**
-	* Clears all highlights
+	* Clears all run renderers
+	*/
+	void ClearRunRenderers();
+
+	/**
+	* Replaces the current set of run renderers with the provided renderers.
+	*/
+	void SetRunRenderers( const TArray< FTextRunRenderer >& Renderers );
+
+	/**
+	* Adds a single run renderer to the existing set of renderers.
+	*/
+	void AddRunRenderer( const FTextRunRenderer& Renderer );
+
+	/**
+	* Clears all line highlights
 	*/
-	void ClearHighlights();
+	void ClearLineHighlights();
 
 	/**
-	* Adds a single highlight to the existing set of highlights.
+	* Replaces the current set of line highlights with the provided highlights.
 	*/
-	void AddHighlight( const FTextHighlight& Highlight );
+	void SetLineHighlights( const TArray< FTextLineHighlight >& Highlights );
 
 	/**
-	* Replaces the current set of highlights with the provided highlights.
+	* Adds a single line highlight to the existing set of highlights.
 	*/
-	void SetHighlights( const TArray< FTextHighlight >& Highlights );
+	void AddLineHighlight( const FTextLineHighlight& Highlight );
 
 	/**
 	* Updates the TextLayout's if any changes have occurred since the last update.
 	*/
 	virtual void UpdateIfNeeded();
 
-	virtual void Update();
+	virtual void UpdateLayout();
+
+	virtual void UpdateHighlights();
 
 	int32 GetLineViewIndexForTextLocation(const TArray< FTextLayout::FLineView >& LineViews, const FTextLocation& Location);
 
@@ -368,25 +398,34 @@ protected:
 	*/
 	virtual void EndLayout();
 
-	/**
-	* Creates a new View for the layout, replacing the current one.
-	*/
-	void CreateView();
-
-
 private:
 
 	void FlowLayout();
 
+	void FlowHighlights();
+
 	void JustifyLayout();
 
-	void CreateLineViewBlocks( int32 LineModelIndex, const int32 StopIndex, int32& OutRunIndex, int32& OutHighlightIndex, int32& OutPreviousBlockEnd, TArray< TSharedRef< ILayoutBlock > >& OutSoftLine );
+	void JustifyHighlights();
+
+	void CreateLineViewBlocks( int32 LineModelIndex, const int32 StopIndex, int32& OutRunIndex, int32& OutRendererIndex, int32& OutPreviousBlockEnd, TArray< TSharedRef< ILayoutBlock > >& OutSoftLine );
 
 	FBreakCandidate CreateBreakCandidate( int32& OutRunIndex, FLineModel& Line, int32 PreviousBreak, int32 CurrentBreak );
 
 	void GetAsTextAndOffsets(FString* const OutDisplayText, FTextOffsetLocations* const OutTextOffsetLocations) const;
 
 protected:
+	
+	struct EDirtyState
+	{
+		typedef uint8 Flags;
+		enum Enum
+		{
+			None = 0,
+			Layout = 1<<0,
+			Highlights = 1<<1, 
+		};
+	};
 
 	/** The models for the lines of text. A LineModel represents a single string with no manual breaks. */
 	TArray< FLineModel > LineModels;
@@ -395,7 +434,7 @@ protected:
 	TArray< FLineView > LineViews;
 
 	/** Whether parameters on the layout have changed which requires the view be updated. */
-	bool HasChanged;
+	EDirtyState::Flags DirtyFlags;
 
 	/** The scale to draw the text at */
 	float Scale;
diff --git a/Engine/Source/Runtime/Slate/Public/Framework/Text/TextLayoutEngine.h b/Engine/Source/Runtime/Slate/Public/Framework/Text/TextLayoutEngine.h
--- a/Engine/Source/Runtime/Slate/Public/Framework/Text/TextLayoutEngine.h
+++ b/Engine/Source/Runtime/Slate/Public/Framework/Text/TextLayoutEngine.h
@@ -3,10 +3,12 @@
 
 // Generic Text Layout
 #include "TextRange.h"
-#include "TextHighlight.h"
+#include "TextRunRenderer.h"
+#include "TextLineHighlight.h"
 
 #include "IRun.h"
-#include "IRunHighlighter.h"
+#include "IRunRenderer.h"
+#include "ILineHighlighter.h"
 #include "ILayoutBlock.h"
 
 #include "TextLayout.h"
@@ -16,7 +18,8 @@
 
 // Slate Specific Text Layout
 #include "ISlateRun.h"
-#include "ISlateRunHighlighter.h"
+#include "ISlateRunRenderer.h"
+#include "ISlateLineHighlighter.h"
 
 #include "SlateTextLayout.h"
 
@@ -24,4 +27,3 @@
 #include "SlateHyperlinkRun.h"
 #include "SlateImageRun.h"
 #include "SlateWidgetRun.h"
-#include "SlateSimpleRunHighlighter.h"
diff --git a/Engine/Source/Runtime/Slate/Public/Framework/Text/TextLineHighlight.h b/Engine/Source/Runtime/Slate/Public/Framework/Text/TextLineHighlight.h
new file mode 100644
--- /dev/null
+++ b/Engine/Source/Runtime/Slate/Public/Framework/Text/TextLineHighlight.h
@@ -0,0 +1,21 @@
+// Copyright 1998-2014 Epic Games, Inc. All Rights Reserved.
+#pragma once
+
+class ILineHighlighter;
+
+struct SLATE_API FTextLineHighlight
+{
+	FTextLineHighlight( int32 InLineIndex, const FTextRange& InRange, int32 InZOrder, const TSharedRef< ILineHighlighter >& InHighlighter )
+		: LineIndex( InLineIndex )
+		, Range( InRange )
+		, ZOrder( InZOrder )
+		, Highlighter( InHighlighter )
+	{
+
+	}
+
+	int32 LineIndex;
+	FTextRange Range;
+	int32 ZOrder;
+	TSharedRef< ILineHighlighter > Highlighter;
+};
diff --git a/Engine/Source/Runtime/Slate/Public/Framework/Text/TextRunRenderer.h b/Engine/Source/Runtime/Slate/Public/Framework/Text/TextRunRenderer.h
new file mode 100644
--- /dev/null
+++ b/Engine/Source/Runtime/Slate/Public/Framework/Text/TextRunRenderer.h
@@ -0,0 +1,19 @@
+// Copyright 1998-2014 Epic Games, Inc. All Rights Reserved.
+#pragma once
+
+class IRunRenderer;
+
+struct SLATE_API FTextRunRenderer
+{
+	FTextRunRenderer( int32 InLineIndex, const FTextRange& InRange, const TSharedRef< IRunRenderer >& InRenderer )
+		: LineIndex( InLineIndex )
+		, Range( InRange )
+		, Renderer( InRenderer )
+	{
+
+	}
+
+	int32 LineIndex;
+	FTextRange Range;
+	TSharedRef< IRunRenderer > Renderer;
+};
diff --git a/Engine/Source/Runtime/Slate/Public/Framework/Text/WidgetLayoutBlock.h b/Engine/Source/Runtime/Slate/Public/Framework/Text/WidgetLayoutBlock.h
--- a/Engine/Source/Runtime/Slate/Public/Framework/Text/WidgetLayoutBlock.h
+++ b/Engine/Source/Runtime/Slate/Public/Framework/Text/WidgetLayoutBlock.h
@@ -5,9 +5,9 @@ class SLATE_API FWidgetLayoutBlock : public ILayoutBlock
 {
 public:
 
-	static TSharedRef< FWidgetLayoutBlock > Create( const TSharedRef< IRun >& InRun, const TSharedRef< SWidget >& InWidget, const FTextRange& InRange, FVector2D InSize, const TSharedPtr< IRunHighlighter >& InHighlighter )
+	static TSharedRef< FWidgetLayoutBlock > Create( const TSharedRef< IRun >& InRun, const TSharedRef< SWidget >& InWidget, const FTextRange& InRange, FVector2D InSize, const TSharedPtr< IRunRenderer >& InRenderer )
 	{
-		return MakeShareable( new FWidgetLayoutBlock( InRun, InWidget, InRange, InSize, InHighlighter ) );
+		return MakeShareable( new FWidgetLayoutBlock( InRun, InWidget, InRange, InSize, InRenderer ) );
 	}
 
 	virtual ~FWidgetLayoutBlock() {}
@@ -18,7 +18,7 @@ public:
 
 	virtual FVector2D GetSize() const override { return Size; }
 
-	virtual TSharedPtr< IRunHighlighter > GetHighlighter() const override { return Highlighter; }
+	virtual TSharedPtr< IRunRenderer > GetRenderer() const override { return Renderer; }
 
 	TSharedRef< SWidget > GetWidget() const { return Widget; }
 
@@ -32,13 +32,13 @@ private:
 		return MakeShareable( new FWidgetLayoutBlock( Block ) );
 	}
 
-	FWidgetLayoutBlock( const TSharedRef< IRun >& InRun, const TSharedRef< SWidget >& InWidget, const FTextRange& InRange, FVector2D InSize, const TSharedPtr< IRunHighlighter >& InHighlighter )
+	FWidgetLayoutBlock( const TSharedRef< IRun >& InRun, const TSharedRef< SWidget >& InWidget, const FTextRange& InRange, FVector2D InSize, const TSharedPtr< IRunRenderer >& InRenderer )
 		: Run( InRun )
 		, Widget( InWidget )
 		, Range( InRange )
 		, Size( InSize )
 		, LocationOffset( ForceInitToZero )
-		, Highlighter( InHighlighter )
+		, Renderer( InRenderer )
 	{
 
 	}
@@ -49,7 +49,7 @@ private:
 		, Range( Block.Range )
 		, Size( Block.Size )
 		, LocationOffset( ForceInitToZero )
-		, Highlighter( Block.Highlighter )
+		, Renderer( Block.Renderer )
 	{
 
 	}
@@ -62,5 +62,5 @@ private:
 	FTextRange Range;
 	FVector2D Size;
 	FVector2D LocationOffset;
-	TSharedPtr< IRunHighlighter > Highlighter;
+	TSharedPtr< IRunRenderer > Renderer;
 };
\ No newline at end of file
diff --git a/Engine/Source/Runtime/Slate/Public/Widgets/Text/SMultiLineEditableText.h b/Engine/Source/Runtime/Slate/Public/Widgets/Text/SMultiLineEditableText.h
--- a/Engine/Source/Runtime/Slate/Public/Widgets/Text/SMultiLineEditableText.h
+++ b/Engine/Source/Runtime/Slate/Public/Widgets/Text/SMultiLineEditableText.h
@@ -232,67 +232,62 @@ private:
 	};
 
 	/** Run highlighter used to draw the cursor */
-	class FSlateCursorRunHighlighter : public ISlateRunHighlighter
+	class FCursorLineHighlighter : public ISlateLineHighlighter
 	{
 	public:
 
-		static TSharedRef< FSlateCursorRunHighlighter > Create(const FCursorInfo* InCursorInfo);
+		static TSharedRef< FCursorLineHighlighter > Create(const FCursorInfo* InCursorInfo);
 
-		virtual ~FSlateCursorRunHighlighter() {}
+		virtual ~FCursorLineHighlighter() {}
 
-		virtual int32 OnPaint( const FTextLayout::FLineView& Line, const TSharedRef< ISlateRun >& Run, const TSharedRef< ILayoutBlock >& Block, const FTextBlockStyle& DefaultStyle, const FGeometry& AllottedGeometry, const FSlateRect& MyClippingRect, FSlateWindowElementList& OutDrawElements, int32 LayerId, const FWidgetStyle& InWidgetStyle, bool bParentEnabled ) const override;
-
-		virtual FChildren* GetChildren() override;
-
-		virtual void OnArrangeChildren( const TSharedRef< ILayoutBlock >& Block, const FGeometry& AllottedGeometry, FArrangedChildren& ArrangedChildren ) const override;
+		virtual int32 OnPaint( const FTextLayout::FLineView& Line, const float OffsetX, const float Width, const FTextBlockStyle& DefaultStyle, const FGeometry& AllottedGeometry, const FSlateRect& MyClippingRect, FSlateWindowElementList& OutDrawElements, int32 LayerId, const FWidgetStyle& InWidgetStyle, bool bParentEnabled ) const override;
 
 	protected:
 
-		int32 OnPaintCursor( const FTextLayout::FLineView& Line, const TSharedRef< ISlateRun >& Run, const TSharedRef< ILayoutBlock >& Block, const FTextBlockStyle& DefaultStyle, const FGeometry& AllottedGeometry, const FSlateRect& MyClippingRect, FSlateWindowElementList& OutDrawElements, int32 LayerId, const FWidgetStyle& InWidgetStyle, const bool bParentEnabled, const float InCursorOffset, const ECursorAlignment InCursorAlignment ) const;
-
-		FSlateCursorRunHighlighter(const FCursorInfo* InCursorInfo);
+		FCursorLineHighlighter(const FCursorInfo* InCursorInfo);
 
 		/** Cursor data that this highlighter is tracking */
 		const FCursorInfo* CursorInfo;
 	};
 
-	/** Run highlighter used to draw selection ranges */
-	class FSlateSelectionRunHighlighter : public FSlateCursorRunHighlighter
+	/** Run highlighter used to draw the composition range */
+	class FTextCompositionHighlighter : public ISlateLineHighlighter
 	{
 	public:
 
-		static TSharedRef< FSlateSelectionRunHighlighter > Create(const FCursorInfo* InCursorInfo);
+		static TSharedRef< FTextCompositionHighlighter > Create();
 
-		virtual ~FSlateSelectionRunHighlighter() {}
+		virtual ~FTextCompositionHighlighter() {}
 
-		virtual int32 OnPaint( const FTextLayout::FLineView& Line, const TSharedRef< ISlateRun >& Run, const TSharedRef< ILayoutBlock >& Block, const FTextBlockStyle& DefaultStyle, const FGeometry& AllottedGeometry, const FSlateRect& MyClippingRect, FSlateWindowElementList& OutDrawElements, int32 LayerId, const FWidgetStyle& InWidgetStyle, bool bParentEnabled ) const override;
-
-		void SetCursorAlignment(const ECursorAlignment InCursorAlignment)
-		{
-			CursorAlignment = InCursorAlignment;
-		}
+		virtual int32 OnPaint( const FTextLayout::FLineView& Line, const float OffsetX, const float Width, const FTextBlockStyle& DefaultStyle, const FGeometry& AllottedGeometry, const FSlateRect& MyClippingRect, FSlateWindowElementList& OutDrawElements, int32 LayerId, const FWidgetStyle& InWidgetStyle, bool bParentEnabled ) const override;
 
 	protected:
 
-		FSlateSelectionRunHighlighter(const FCursorInfo* InCursorInfo);
-
-		ECursorAlignment CursorAlignment;
+		FTextCompositionHighlighter();
 	};
 
-	/** Run highlighter used to draw the composition range */
-	class FSlateCompositionRunHighlighter : public FSlateCursorRunHighlighter
+	/** Run renderer used to draw selection ranges */
+	class FTextSelectionRunRenderer : public ISlateRunRenderer
 	{
 	public:
 
-		static TSharedRef< FSlateCompositionRunHighlighter > Create(const FCursorInfo* InCursorInfo);
+		static TSharedRef< FTextSelectionRunRenderer > Create();
 
-		virtual ~FSlateCompositionRunHighlighter() {}
+		virtual ~FTextSelectionRunRenderer() {}
 
 		virtual int32 OnPaint( const FTextLayout::FLineView& Line, const TSharedRef< ISlateRun >& Run, const TSharedRef< ILayoutBlock >& Block, const FTextBlockStyle& DefaultStyle, const FGeometry& AllottedGeometry, const FSlateRect& MyClippingRect, FSlateWindowElementList& OutDrawElements, int32 LayerId, const FWidgetStyle& InWidgetStyle, bool bParentEnabled ) const override;
 
+		void SetHasKeyboardFocus(const bool bInHasKeyboardFocus)
+		{
+			bHasKeyboardFocus = bInHasKeyboardFocus;
+		}
+
 	protected:
 
-		FSlateCompositionRunHighlighter(const FCursorInfo* InCursorInfo);
+		FTextSelectionRunRenderer();
+
+		/** true if the parent widget has keyboard focus, false otherwise */
+		bool bHasKeyboardFocus;
 	};
 
 private:
@@ -468,11 +463,9 @@ private:
 	TAttribute< ETextJustify::Type > Justification; 
 	TAttribute< float > LineHeightPercentage;
 
-	TSharedPtr< FSlateCursorRunHighlighter > CursorRunHighlighter;
-
-	TSharedPtr< FSlateSelectionRunHighlighter > SelectionRunHighlighter;
-
-	TSharedPtr< FSlateCompositionRunHighlighter > CompositionRunHighlighter;
+	TSharedPtr< FCursorLineHighlighter > CursorLineHighlighter;
+	TSharedPtr< FTextCompositionHighlighter > TextCompositionHighlighter;
+	TSharedPtr< FTextSelectionRunRenderer > TextSelectionRunRenderer;
 
 	/** That start of the selection when there is a selection. The end is implicitly wherever the cursor happens to be. */
 	TOptional<FTextLocation> SelectionStart;
diff --git a/Engine/Source/Runtime/Slate/Public/Widgets/Text/SRichTextBlock.h b/Engine/Source/Runtime/Slate/Public/Widgets/Text/SRichTextBlock.h
--- a/Engine/Source/Runtime/Slate/Public/Widgets/Text/SRichTextBlock.h
+++ b/Engine/Source/Runtime/Slate/Public/Widgets/Text/SRichTextBlock.h
@@ -138,8 +138,8 @@ private:
 	TSharedPtr< FSlateTextLayout > TextLayout;
 
 	FText HighlightText;
-	TArray< FTextHighlight > Highlights;
-	TSharedPtr< ISlateRunHighlighter> TextHighlighter;
+	TArray< FTextRunRenderer > TextHighlights;
+	TSharedPtr< ISlateRunRenderer> TextHighlighter;
 
 	const ISlateStyle* TagStyleSet;
 	FTextBlockStyle TextStyle;
