UMG - Add design time property updates an blueprint support for the Throbber and CircularThrobber.

diff --git a/Engine/Source/Runtime/Slate/Private/Widgets/Images/SThrobber.cpp b/Engine/Source/Runtime/Slate/Private/Widgets/Images/SThrobber.cpp
--- a/Engine/Source/Runtime/Slate/Private/Widgets/Images/SThrobber.cpp
+++ b/Engine/Source/Runtime/Slate/Private/Widgets/Images/SThrobber.cpp
@@ -5,40 +5,66 @@
 
 void SThrobber::Construct(const FArguments& InArgs)
 {
+	PieceImage = InArgs._PieceImage;
+	NumPieces = InArgs._NumPieces;
 	Animate = InArgs._Animate;
+
+	HBox = SNew(SHorizontalBox);
+
+	this->ChildSlot.Widget = HBox.ToSharedRef();
+
+	ConstructPieces();
+}
+
+void SThrobber::ConstructPieces()
+{
+	ThrobberCurve.Empty();
 	AnimCurves = FCurveSequence();
+	for (int32 PieceIndex = 0; PieceIndex < NumPieces; ++PieceIndex)
 	{
-		for (int32 PieceIndex=0; PieceIndex < InArgs._NumPieces; ++PieceIndex)
-		{
-			ThrobberCurve.Add( AnimCurves.AddCurve(PieceIndex*0.05f,1.5f) );
-		}
-		AnimCurves.Play();
+		ThrobberCurve.Add(AnimCurves.AddCurve(PieceIndex*0.05f, 1.5f));
 	}
+	AnimCurves.Play();
 
-	TSharedRef<SHorizontalBox> HBox = SNew(SHorizontalBox);
-	
-	for (int32 PieceIndex=0; PieceIndex < InArgs._NumPieces; ++PieceIndex)
+	HBox->ClearChildren();
+	for (int32 PieceIndex = 0; PieceIndex < NumPieces; ++PieceIndex)
 	{
 		HBox->AddSlot()
-		.AutoWidth()
-		[
-			SNew(SBorder)
-			.BorderImage( FStyleDefaults::GetNoBrush() )
-			.ContentScale( this, &SThrobber::GetPieceScale, PieceIndex )
-			.ColorAndOpacity( this, &SThrobber::GetPieceColor, PieceIndex )
-			.HAlign(HAlign_Center)
-			.VAlign(VAlign_Center)
+			.AutoWidth()
 			[
-				SNew(SImage)
-				.Image( InArgs._PieceImage )
-			]
-		];
+				SNew(SBorder)
+				.BorderImage(FStyleDefaults::GetNoBrush())
+				.ContentScale(this, &SThrobber::GetPieceScale, PieceIndex)
+				.ColorAndOpacity(this, &SThrobber::GetPieceColor, PieceIndex)
+				.HAlign(HAlign_Center)
+				.VAlign(VAlign_Center)
+				[
+					SNew(SImage)
+					.Image(this, &SThrobber::GetPieceBrush)
+				]
+			];
 	}
+}
 
-	this->ChildSlot
-	[
-		HBox
-	];
+const FSlateBrush* SThrobber::GetPieceBrush() const
+{
+	return PieceImage;
+}
+
+void SThrobber::SetPieceImage(const FSlateBrush* InPieceImage)
+{
+	PieceImage = InPieceImage;
+}
+
+void SThrobber::SetNumPieces(int InNumPieces)
+{
+	NumPieces = InNumPieces;
+	ConstructPieces();
+}
+
+void SThrobber::SetAnimate(EAnimation InAnimate)
+{
+	Animate = InAnimate;
 }
 
 FVector2D SThrobber::GetPieceScale(int32 PieceIndex) const
@@ -68,17 +94,43 @@ FLinearColor SThrobber::GetPieceColor(int32 PieceIndex) const
 	}
 }
 
-
 // SCircularThrobber
 void SCircularThrobber::Construct(const FArguments& InArgs)
 {
-	Sequence = FCurveSequence();
-	Curve = Sequence.AddCurve(0.f, InArgs._Period);
-	Sequence.Play();
-	
 	PieceImage = InArgs._PieceImage;
 	NumPieces = InArgs._NumPieces;
+	Period = InArgs._Period;
 	Radius = InArgs._Radius;
+
+	ConstructSequence();
+}
+
+void SCircularThrobber::SetPieceImage(const FSlateBrush* InPieceImage)
+{
+	PieceImage = InPieceImage;
+}
+
+void SCircularThrobber::SetNumPieces(const int32 InNumPieces)
+{
+	NumPieces = InNumPieces;
+}
+
+void SCircularThrobber::SetPeriod(const float InPeriod)
+{
+	Period = InPeriod;
+	ConstructSequence();
+}
+
+void SCircularThrobber::SetRadius(const float InRadius)
+{
+	Radius = InRadius;
+}
+
+void SCircularThrobber::ConstructSequence()
+{
+	Sequence = FCurveSequence();
+	Curve = Sequence.AddCurve(0.f, Period);
+	Sequence.Play();
 }
 
 int32 SCircularThrobber::OnPaint( const FGeometry& AllottedGeometry, const FSlateRect& MyClippingRect, FSlateWindowElementList& OutDrawElements, int32 LayerId, const FWidgetStyle& InWidgetStyle, bool bParentEnabled ) const
diff --git a/Engine/Source/Runtime/Slate/Public/Widgets/Images/SThrobber.h b/Engine/Source/Runtime/Slate/Public/Widgets/Images/SThrobber.h
--- a/Engine/Source/Runtime/Slate/Public/Widgets/Images/SThrobber.h
+++ b/Engine/Source/Runtime/Slate/Public/Widgets/Images/SThrobber.h
@@ -36,13 +36,38 @@ public:
 
 	void Construct(const FArguments& InArgs);
 
+	/** Sets what each segment of the throbber looks like. */
+	void SetPieceImage(const FSlateBrush* InPieceImage);
+
+	/** Sets how many pieces there are */
+	void SetNumPieces(int InNumPieces);
+
+	/** Sets which aspects of the throbber to animate */
+	void SetAnimate(EAnimation InAnimate);
 
 private:
 	FVector2D GetPieceScale(int32 PieceIndex) const;
 	FLinearColor GetPieceColor(int32 PieceIndex) const;
+
+	/** Constructs the curves and widgets for the pieces which make up the throbber. */
+	void ConstructPieces();
+
+	/** Gets the brush used to draw each piece of the thobber. */
+	const FSlateBrush* GetPieceBrush() const;
+
+private:
 	
 	FCurveSequence AnimCurves;
 	TArray<FCurveHandle, TInlineAllocator<DefaultNumPieces> > ThrobberCurve;
+
+	/** The horizontal box which contains the widgets for the trobber pieces. */
+	TSharedPtr<SHorizontalBox> HBox;
+
+	/** The image used to draw each piece of the throbber. */
+	const FSlateBrush* PieceImage;
+	/** The number of pieces to display. */
+	int NumPieces;
+	/** Flags for which apsects of the throbber to animate. */
 	EAnimation Animate;
 };
 
@@ -71,9 +96,26 @@ public:
 	/** Constructs the widget */
 	void Construct(const FArguments& InArgs);
 
+	/** Sets what each segment of the throbber looks like */
+	void SetPieceImage(const FSlateBrush* InPieceImage);
+
+	/** Sets how many pieces there are */
+	void SetNumPieces(int32 InNumPieces);
+
+	/** Sets the amount of time in seconds for a full circle */
+	void SetPeriod(float InPeriod);
+
+	/** Sets the radius of the circle */
+	void SetRadius(float InRadius);
+
 	virtual int32 OnPaint( const FGeometry& AllottedGeometry, const FSlateRect& MyClippingRect, FSlateWindowElementList& OutDrawElements, int32 LayerId, const FWidgetStyle& InWidgetStyle, bool bParentEnabled ) const override;
 	virtual FVector2D ComputeDesiredSize() const override;
 
+private:
+
+	/** Constructs the sequence used to animate the trobber. */
+	void ConstructSequence();
+
 private:
 	/** The sequence to drive the spinning animation */
 	FCurveSequence Sequence;
@@ -83,6 +125,8 @@ private:
 	const FSlateBrush* PieceImage;
 	/** How many pieces there are */
 	int32 NumPieces;
+	/** The amount of time in seconds for a full circle */
+	float Period;
 	/** The radius of the circle */
 	float Radius;
 };
diff --git a/Engine/Source/Runtime/UMG/Private/Components/CircularThrobber.cpp b/Engine/Source/Runtime/UMG/Private/Components/CircularThrobber.cpp
--- a/Engine/Source/Runtime/UMG/Private/Components/CircularThrobber.cpp
+++ b/Engine/Source/Runtime/UMG/Private/Components/CircularThrobber.cpp
@@ -22,15 +22,72 @@ TSharedRef<SWidget> UCircularThrobber::RebuildWidget()
 {
 	SCircularThrobber::FArguments DefaultArgs;
 
-	const FSlateBrush* Image = PieceImage ? &PieceImage->Brush : DefaultArgs._PieceImage;
-	
-	return SNew(SCircularThrobber)
-		.PieceImage(Image)
+	MyCircularThrobber = SNew(SCircularThrobber)
+		.PieceImage(GetPieceBrush())
 		.NumPieces(NumberOfPieces)
 		.Period(Period)
 		.Radius(Radius);
+
+	return MyCircularThrobber.ToSharedRef();
+}
+
+void UCircularThrobber::SyncronizeProperties()
+{
+	Super::SyncronizeProperties();
+
+	MyCircularThrobber->SetPieceImage(GetPieceBrush());
+	MyCircularThrobber->SetNumPieces(NumberOfPieces);
+	MyCircularThrobber->SetPeriod(Period);
+	MyCircularThrobber->SetRadius(Radius);
+}
+
+const FSlateBrush* UCircularThrobber::GetPieceBrush() const
+{
+	if (PieceImage == NULL)
+	{
+		SCircularThrobber::FArguments DefaultArgs;
+		return DefaultArgs._PieceImage;
+	}
+	return &PieceImage->Brush;
+}
+
+void UCircularThrobber::SetNumberOfPieces(int32 InNumberOfPieces)
+{
+	NumberOfPieces = NumberOfPieces;
+	if (MyCircularThrobber.IsValid())
+	{
+		MyCircularThrobber->SetNumPieces(InNumberOfPieces);
+	}
 }
 
+void UCircularThrobber::SetPeriod(float InPeriod)
+{
+	Period = InPeriod;
+	if (MyCircularThrobber.IsValid())
+	{
+		MyCircularThrobber->SetPeriod(InPeriod);
+	}
+}
+
+void UCircularThrobber::SetRadius(float InRadius)
+{
+	Radius = InRadius;
+	if (MyCircularThrobber.IsValid())
+	{
+		MyCircularThrobber->SetRadius(InRadius);
+	}
+}
+
+void UCircularThrobber::SetPieceImage(USlateBrushAsset* InPieceImage)
+{
+	PieceImage = InPieceImage;
+	if (MyCircularThrobber.IsValid())
+	{
+		MyCircularThrobber->SetPieceImage(GetPieceBrush());
+	}
+}
+
+
 /////////////////////////////////////////////////////
 
 #undef LOCTEXT_NAMESPACE
diff --git a/Engine/Source/Runtime/UMG/Private/Components/Throbber.cpp b/Engine/Source/Runtime/UMG/Private/Components/Throbber.cpp
--- a/Engine/Source/Runtime/UMG/Private/Components/Throbber.cpp
+++ b/Engine/Source/Runtime/UMG/Private/Components/Throbber.cpp
@@ -24,16 +24,85 @@ TSharedRef<SWidget> UThrobber::RebuildWidget()
 {
 	SThrobber::FArguments DefaultArgs;
 
-	const FSlateBrush* Image = PieceImage ? &PieceImage->Brush : DefaultArgs._PieceImage;
+	MyThrobber = SNew(SThrobber)
+		.PieceImage(GetPieceBrush())
+		.NumPieces(NumberOfPieces)
+		.Animate(GetAnimation());
+
+	return MyThrobber.ToSharedRef();
+}
+
+void UThrobber::SyncronizeProperties()
+{
+	Super::SyncronizeProperties();
+
+	MyThrobber->SetPieceImage(GetPieceBrush());
+	MyThrobber->SetNumPieces(NumberOfPieces);
+	MyThrobber->SetAnimate(GetAnimation());
+}
 
+SThrobber::EAnimation UThrobber::GetAnimation() const
+{
 	const int32 AnimationParams = (bAnimateVertically ? SThrobber::Vertical : 0) |
 		(bAnimateHorizontally ? SThrobber::Horizontal : 0) |
 		(bAnimateOpacity ? SThrobber::Opacity : 0);
 
-	return SNew(SThrobber)
-		.PieceImage(Image)
-		.NumPieces(NumberOfPieces)
-		.Animate(static_cast<SThrobber::EAnimation>(AnimationParams));
+	return static_cast<SThrobber::EAnimation>(AnimationParams);
+}
+
+const FSlateBrush* UThrobber::GetPieceBrush() const
+{
+	if (PieceImage == NULL)
+	{
+		SThrobber::FArguments DefaultArgs;
+		return DefaultArgs._PieceImage;
+	}
+	return &PieceImage->Brush;
+}
+
+void UThrobber::SetNumberOfPieces(int32 InNumberOfPieces)
+{
+	NumberOfPieces = NumberOfPieces;
+	if (MyThrobber.IsValid())
+	{
+		MyThrobber->SetNumPieces(InNumberOfPieces);
+	}
+}
+
+void UThrobber:: SetAnimateHorizontally(bool bInAnimateHorizontally)
+{
+	bInAnimateHorizontally = bInAnimateHorizontally;
+	if (MyThrobber.IsValid())
+	{
+		MyThrobber->SetAnimate(GetAnimation());
+	}
+}
+
+void UThrobber:: SetAnimateVertically(bool bInAnimateVertically)
+{
+	bInAnimateVertically = bInAnimateVertically;
+	if (MyThrobber.IsValid())
+	{
+		MyThrobber->SetAnimate(GetAnimation());
+	}
+}
+
+void UThrobber:: SetAnimateOpacity(bool bInAnimateOpacity)
+{
+	bInAnimateOpacity = bInAnimateOpacity;
+	if (MyThrobber.IsValid())
+	{
+		MyThrobber->SetAnimate(GetAnimation());
+	}
+}
+
+void UThrobber::SetPieceImage(USlateBrushAsset* InPieceImage)
+{
+	PieceImage = InPieceImage;
+	if (MyThrobber.IsValid())
+	{
+		MyThrobber->SetPieceImage(GetPieceBrush());
+	}
 }
 
 /////////////////////////////////////////////////////
diff --git a/Engine/Source/Runtime/UMG/Public/Components/CircularThrobber.h b/Engine/Source/Runtime/UMG/Public/Components/CircularThrobber.h
--- a/Engine/Source/Runtime/UMG/Public/Components/CircularThrobber.h
+++ b/Engine/Source/Runtime/UMG/Public/Components/CircularThrobber.h
@@ -12,26 +12,57 @@ class UMG_API UCircularThrobber : public UWidget
 {
 	GENERATED_UCLASS_BODY()
 
+public:
+
+	/** Sets how many pieces there are. */
+	UFUNCTION(BlueprintCallable, Category = "Appearance")
+	void SetNumberOfPieces(int32 InNumberOfPieces);
+
+	/** Sets the amount of time for a full circle (in seconds). */
+	UFUNCTION(BlueprintCallable, Category = "Appearance")
+	void SetPeriod(float InPeriod);
+
+	/** Sets the radius of the circle. */
+	UFUNCTION(BlueprintCallable, Category = "Appearance")
+	void SetRadius(float InRadius);
+
+	/** Sets image to use for each segment of the throbber. */
+	UFUNCTION(BlueprintCallable, Category = "Appearance")
+	void SetPieceImage(USlateBrushAsset* InPieceImage);
+
+	// UWidget interface
+	virtual void SyncronizeProperties() override;
+	// End of UWidget interface
+
 protected:
 
 	/** How many pieces there are */
-	UPROPERTY(EditAnywhere, BlueprintReadOnly, Category=Appearance, meta=(ClampMin = "1", UIMin = "1", UIMax = "8"))
+	UPROPERTY(EditDefaultsOnly, Category=Appearance, meta=(ClampMin = "1", UIMin = "1", UIMax = "8"))
 	int32 NumberOfPieces;
 
 	/** The amount of time for a full circle (in seconds) */
-	UPROPERTY(EditAnywhere, BlueprintReadOnly, Category=Appearance)
+	UPROPERTY(EditDefaultsOnly, Category=Appearance)
 	float Period;
 
 	/** The radius of the circle */
-	UPROPERTY(EditAnywhere, BlueprintReadOnly, Category=Appearance)
+	UPROPERTY(EditDefaultsOnly, Category=Appearance)
 	float Radius;
 
 	/** Image to use for each segment of the throbber */
-	UPROPERTY(EditAnywhere, BlueprintReadWrite, Category=Appearance, meta=( DisplayThumbnail = "true" ))
+	UPROPERTY(EditDefaultsOnly, Category=Appearance, meta=( DisplayThumbnail = "true" ))
 	USlateBrushAsset* PieceImage;
 
 protected:
 	// UWidget interface
 	virtual TSharedRef<SWidget> RebuildWidget() override;
 	// End of UWidget interface
+
+private:
+
+	/** Gets the brush used to draw the pieces of the throbber. */
+	const FSlateBrush* GetPieceBrush() const;
+
+private:
+	/** The CircularThrobber widget managed by this object. */
+	TSharedPtr<SCircularThrobber> MyCircularThrobber;
 };
diff --git a/Engine/Source/Runtime/UMG/Public/Components/Throbber.h b/Engine/Source/Runtime/UMG/Public/Components/Throbber.h
--- a/Engine/Source/Runtime/UMG/Public/Components/Throbber.h
+++ b/Engine/Source/Runtime/UMG/Public/Components/Throbber.h
@@ -12,28 +12,68 @@ class UMG_API UThrobber : public UWidget
 {
 	GENERATED_UCLASS_BODY()
 
+public:
+
+	/** Sets how many pieces there are */
+	UFUNCTION(BlueprintCallable, Category = "Appearance")
+	void SetNumberOfPieces(int32 InNumberOfPieces);
+
+	/** Sets whether the pieces animate horizontally. */
+	UFUNCTION(BlueprintCallable, Category = "Appearance")
+	void SetAnimateHorizontally(bool bInAnimateHorizontally);
+
+	/** Sets whether the pieces animate vertically. */
+	UFUNCTION(BlueprintCallable, Category = "Appearance")
+	void SetAnimateVertically(bool bInAnimateVertically);
+
+	/** Sets whether the pieces animate their opacity. */
+	UFUNCTION(BlueprintCallable, Category = "Appearance")
+	void SetAnimateOpacity(bool bInAnimateOpacity);
+
+	/** Sets the image to use for each segment of the throbber. */
+	UFUNCTION(BlueprintCallable, Category = "Appearance")
+	void SetPieceImage(USlateBrushAsset* InPieceImage);
+
+	// UWidget interface
+	virtual void SyncronizeProperties() override;
+	// End of UWidget interface
+
+protected:
+
 	/** How many pieces there are */
-	UPROPERTY(EditAnywhere, BlueprintReadOnly, Category=Appearance, meta=(ClampMin = "1", UIMin = "1", UIMax = "8"))
+	UPROPERTY(EditDefaultsOnly, Category=Appearance, meta=(ClampMin = "1", UIMin = "1", UIMax = "8"))
 	int32 NumberOfPieces;
 
 	/** Should the pieces animate horizontally? */
-	UPROPERTY(EditAnywhere, BlueprintReadWrite, Category=Appearance)
+	UPROPERTY(EditDefaultsOnly, Category=Appearance)
 	bool bAnimateHorizontally;
 
 	/** Should the pieces animate vertically? */
-	UPROPERTY(EditAnywhere, BlueprintReadWrite, Category=Appearance)
+	UPROPERTY(EditDefaultsOnly, Category=Appearance)
 	bool bAnimateVertically;
 
 	/** Should the pieces animate their opacity? */
-	UPROPERTY(EditAnywhere, BlueprintReadWrite, Category=Appearance)
+	UPROPERTY(EditDefaultsOnly, Category=Appearance)
 	bool bAnimateOpacity;
 
 	/** Image to use for each segment of the throbber */
-	UPROPERTY(EditAnywhere, BlueprintReadWrite, Category=Appearance, meta=( DisplayThumbnail = "true" ))
+	UPROPERTY(EditDefaultsOnly, Category=Appearance, meta=( DisplayThumbnail = "true" ))
 	USlateBrushAsset* PieceImage;
 
 protected:
 	// UWidget interface
 	virtual TSharedRef<SWidget> RebuildWidget() override;
 	// End of UWidget interface
+
+private:
+
+	/** Gets the combined value of the animation properties as a single SThrobber::EAnimation value. */
+	SThrobber::EAnimation GetAnimation() const;
+
+	/** Gets the brush which should be used to draw the trobber pieces. */
+	const FSlateBrush* GetPieceBrush() const;
+
+private:
+	/** The Throbber widget managed by this object. */
+	TSharedPtr<SThrobber> MyThrobber;
 };
