SMenuAnchor calls delegate when window destroy is requested #codereview Nick.Atamas

diff --git a/Engine/Source/Runtime/Slate/Private/Widgets/Input/SMenuAnchor.cpp b/Engine/Source/Runtime/Slate/Private/Widgets/Input/SMenuAnchor.cpp
--- a/Engine/Source/Runtime/Slate/Private/Widgets/Input/SMenuAnchor.cpp
+++ b/Engine/Source/Runtime/Slate/Private/Widgets/Input/SMenuAnchor.cpp
@@ -47,7 +47,14 @@ void SMenuAnchor::RequestClosePopupWindow( const TSharedRef<SWindow>& PopupWindo
 	if (ensure(Method == CreateNewWindow))
 	{
 		FSlateApplication::Get().RequestDestroyWindow(PopupWindow);
+		
+		if (OnMenuClosed.IsBound())
+		{
+			OnMenuClosed.Execute();
+		}
 	}
+
+
 }
 
 void SMenuAnchor::OnClickedOutsidePopup()
