Fixed some cursor movement/placement issues in the multi-line editable text

TTP# 336464 - Editor: Finish the Multiline Editable Text Block

The fact that the end of a soft-wrapped line has the same FTextLocation as the start of the next soft-wrapped line leads to ambiguity when moving the cursor, as the cursor is always positioned based upon the line model (FTextLocation) rather than the line view.

I'd previously attempted to fix this when placing the cursor with the mouse by moving the cursor back one place (and aligning the cursor to the right) if it matched the end of a soft-line boundary. This worked for the end of the lines, but caused another issue where the clicking to the left of the start of a soft-line would move the cursor back to the previous line.

I've removed that code, and instead added code which allows the text layout (and text run) to report where on the line the given position has hit (either within the line itself, or inside the left or right gutter). This lets me adjust the cursor only if the click landed in the right gutter.

I also had to fix some issues with PreferredCursorScreenOffsetInLine as it was sometimes being set to a different position to the cursor, which could lead to it jumping around.

I also changed the code which converts a position into a line view location (GetLineViewIndexForTextLocation) to perform a non-exclusive range check - this ensures that the FTextLocation will match the start of a line view, rather than the end of a line view (which is correct due to the way I move the cursor back). The only exception to this is when the cursor has been placed to the right of a character, as this is assumed to mean that the cursor is at the end of a line view, and therefore we want to return the line view index for end of that line.

diff --git a/Engine/Source/Runtime/Slate/Private/Framework/Text/SlateHyperlinkRun.cpp b/Engine/Source/Runtime/Slate/Private/Framework/Text/SlateHyperlinkRun.cpp
--- a/Engine/Source/Runtime/Slate/Private/Framework/Text/SlateHyperlinkRun.cpp
+++ b/Engine/Source/Runtime/Slate/Private/Framework/Text/SlateHyperlinkRun.cpp
@@ -87,7 +87,7 @@ void FSlateHyperlinkRun::ArrangeChildren( const TSharedRef< ILayoutBlock >& Bloc
 	ArrangedChildren.AddWidget( AllottedGeometry.MakeChild( WidgetBlock->GetWidget(), Block->GetLocationOffset() * ( 1 / AllottedGeometry.Scale ), Block->GetSize() * ( 1 / AllottedGeometry.Scale ), 1.0f ) );
 }
 
-int32 FSlateHyperlinkRun::GetTextIndexAt( const TSharedRef< ILayoutBlock >& Block, const FVector2D& Location, float Scale ) const
+int32 FSlateHyperlinkRun::GetTextIndexAt( const TSharedRef< ILayoutBlock >& Block, const FVector2D& Location, float Scale, ETextHitPoint* const OutHitPoint ) const
 {
 	return INDEX_NONE;
 }
diff --git a/Engine/Source/Runtime/Slate/Private/Framework/Text/SlateImageRun.cpp b/Engine/Source/Runtime/Slate/Private/Framework/Text/SlateImageRun.cpp
--- a/Engine/Source/Runtime/Slate/Private/Framework/Text/SlateImageRun.cpp
+++ b/Engine/Source/Runtime/Slate/Private/Framework/Text/SlateImageRun.cpp
@@ -53,7 +53,7 @@ void FSlateImageRun::ArrangeChildren( const TSharedRef< ILayoutBlock >& Block, c
 	// no widgets
 }
 
-int32 FSlateImageRun::GetTextIndexAt( const TSharedRef< ILayoutBlock >& Block, const FVector2D& Location, float Scale ) const
+int32 FSlateImageRun::GetTextIndexAt( const TSharedRef< ILayoutBlock >& Block, const FVector2D& Location, float Scale, ETextHitPoint* const OutHitPoint ) const
 {
 	return INDEX_NONE;
 }
diff --git a/Engine/Source/Runtime/Slate/Private/Framework/Text/SlateTextRun.cpp b/Engine/Source/Runtime/Slate/Private/Framework/Text/SlateTextRun.cpp
--- a/Engine/Source/Runtime/Slate/Private/Framework/Text/SlateTextRun.cpp
+++ b/Engine/Source/Runtime/Slate/Private/Framework/Text/SlateTextRun.cpp
@@ -118,7 +118,7 @@ void FSlateTextRun::ArrangeChildren( const TSharedRef< ILayoutBlock >& Block, co
 	// no widgets
 }
 
-int32 FSlateTextRun::GetTextIndexAt( const TSharedRef< ILayoutBlock >& Block, const FVector2D& Location, float Scale ) const
+int32 FSlateTextRun::GetTextIndexAt( const TSharedRef< ILayoutBlock >& Block, const FVector2D& Location, float Scale, ETextHitPoint* const OutHitPoint ) const
 {
 	const FVector2D& BlockOffset = Block->GetLocationOffset();
 	const FVector2D& BlockSize = Block->GetSize();
@@ -138,7 +138,13 @@ int32 FSlateTextRun::GetTextIndexAt( const TSharedRef< ILayoutBlock >& Block, co
 	const FTextRange BlockRange = Block->GetTextRange();
 	const TSharedRef< FSlateFontMeasure > FontMeasure = FSlateApplication::Get().GetRenderer()->GetFontMeasureService();
 
-	return FontMeasure->FindCharacterIndexAtOffset( Text.Get(), BlockRange.BeginIndex, BlockRange.EndIndex, Style.Font, Location.X - BlockOffset.X, true, Scale );
+	const int32 Index = FontMeasure->FindCharacterIndexAtOffset( Text.Get(), BlockRange.BeginIndex, BlockRange.EndIndex, Style.Font, Location.X - BlockOffset.X, true, Scale );
+	if (OutHitPoint)
+	{
+		*OutHitPoint = (Index == BlockRange.EndIndex) ? ETextHitPoint::RightGutter : ETextHitPoint::WithinText;
+	}
+
+	return Index;
 }
 
 FVector2D FSlateTextRun::GetLocationAt( const TSharedRef< ILayoutBlock >& Block, int32 Offset, float Scale ) const
diff --git a/Engine/Source/Runtime/Slate/Private/Framework/Text/SlateWidgetRun.cpp b/Engine/Source/Runtime/Slate/Private/Framework/Text/SlateWidgetRun.cpp
--- a/Engine/Source/Runtime/Slate/Private/Framework/Text/SlateWidgetRun.cpp
+++ b/Engine/Source/Runtime/Slate/Private/Framework/Text/SlateWidgetRun.cpp
@@ -64,7 +64,7 @@ void FSlateWidgetRun::ArrangeChildren( const TSharedRef< ILayoutBlock >& Block,
 	ArrangedChildren.AddWidget( AllottedGeometry.MakeChild( Info.Widget, Block->GetLocationOffset() * ( 1 / AllottedGeometry.Scale ), Block->GetSize() * ( 1 / AllottedGeometry.Scale ), 1.0f ) );
 }
 
-int32 FSlateWidgetRun::GetTextIndexAt( const TSharedRef< ILayoutBlock >& Block, const FVector2D& Location, float Scale ) const
+int32 FSlateWidgetRun::GetTextIndexAt( const TSharedRef< ILayoutBlock >& Block, const FVector2D& Location, float Scale, ETextHitPoint* const OutHitPoint ) const
 {
 	return INDEX_NONE;
 }
diff --git a/Engine/Source/Runtime/Slate/Private/Framework/Text/TextLayout.cpp b/Engine/Source/Runtime/Slate/Private/Framework/Text/TextLayout.cpp
--- a/Engine/Source/Runtime/Slate/Private/Framework/Text/TextLayout.cpp
+++ b/Engine/Source/Runtime/Slate/Private/Framework/Text/TextLayout.cpp
@@ -709,12 +709,12 @@ void FTextLayout::AddLineHighlight( const FTextLineHighlight& Highlight )
 	DirtyFlags |= EDirtyState::Highlights;
 }
 
-FTextLocation FTextLayout::GetTextLocationAt( const FLineView& LineView, const FVector2D& Relative )
+FTextLocation FTextLayout::GetTextLocationAt( const FLineView& LineView, const FVector2D& Relative, ETextHitPoint* const OutHitPoint )
 {
 	for (int32 BlockIndex = 0; BlockIndex < LineView.Blocks.Num(); BlockIndex++)
 	{
 		const TSharedRef< ILayoutBlock >& Block = LineView.Blocks[ BlockIndex ];
-		const int32 TextIndex = Block->GetRun()->GetTextIndexAt(Block, FVector2D(Relative.X, Block->GetLocationOffset().Y), Scale);
+		const int32 TextIndex = Block->GetRun()->GetTextIndexAt(Block, FVector2D(Relative.X, Block->GetLocationOffset().Y), Scale, OutHitPoint);
 
 		if ( TextIndex == INDEX_NONE )
 		{
@@ -727,36 +727,59 @@ FTextLocation FTextLayout::GetTextLocationAt( const FLineView& LineView, const F
 	const int32 LineTextLength = LineModels[ LineView.ModelIndex ].Text->Len();
 	if ( LineTextLength == 0 || !LineView.Blocks.Num() )
 	{
+		if(OutHitPoint)
+		{
+			*OutHitPoint = ETextHitPoint::WithinText;
+		}
 		return FTextLocation( LineView.ModelIndex, 0 );
 	}
 	else if (Relative.X < LineView.Blocks[0]->GetLocationOffset().X)
 	{
+		if(OutHitPoint)
+		{
+			*OutHitPoint = ETextHitPoint::LeftGutter;
+		}
 		return FTextLocation(LineView.ModelIndex, LineView.Range.BeginIndex);
 	}
 
+	if(OutHitPoint)
+	{
+		*OutHitPoint = ETextHitPoint::RightGutter;
+	}
 	return FTextLocation( LineView.ModelIndex, LineView.Range.EndIndex );
 }
 
-int32 FTextLayout::GetLineViewIndexForTextLocation(const TArray< FTextLayout::FLineView >& LineViews, const FTextLocation& Location)
+int32 FTextLayout::GetLineViewIndexForTextLocation(const TArray< FTextLayout::FLineView >& LineViews, const FTextLocation& Location, const bool bPerformInclusiveBoundsCheck)
 {
 	const int32 LineModelIndex = Location.GetLineIndex();
 	const int32 Offset = Location.GetOffset();
 
 	const FLineModel& LineModel = LineModels[LineModelIndex];
-	for (int Index = 0; Index < LineViews.Num(); Index++)
+	for(int32 Index = 0; Index < LineViews.Num(); Index++)
 	{
 		const FTextLayout::FLineView& LineView = LineViews[Index];
 
-		if (LineView.ModelIndex == LineModelIndex && (Offset == 0 || LineView.Range.InclusiveContains(Offset) || LineModel.Text->IsEmpty()))
+		if(LineView.ModelIndex == LineModelIndex)
 		{
-			return Index;
+			// Simple case where we're either the start of, or are contained within, the line view
+			if(Offset == 0 || LineModel.Text->IsEmpty() || LineView.Range.Contains(Offset))
+			{
+				return Index;
+			}
+
+			// If we're the last line, then we also need to test for the end index being part of the range
+			const bool bIsLastLineForModel = Index == (LineViews.Num() - 1) || LineViews[Index + 1].ModelIndex != LineModelIndex;
+			if((bIsLastLineForModel || bPerformInclusiveBoundsCheck) && LineView.Range.EndIndex == Offset)
+			{
+				return Index;
+			}
 		}
 	}
 
 	return INDEX_NONE;
 }
 
-FTextLocation FTextLayout::GetTextLocationAt( const FVector2D& Relative )
+FTextLocation FTextLayout::GetTextLocationAt( const FVector2D& Relative, ETextHitPoint* const OutHitPoint )
 {
 	// Early out if we have no LineViews
 	if (LineViews.Num() == 0)
@@ -784,7 +807,7 @@ FTextLocation FTextLayout::GetTextLocationAt( const FVector2D& Relative )
 		// just use the very last line
 		ViewIndex = LineViews.Num() - 1;
 		const FLineView& LineView = LineViews[ ViewIndex ];
-		return GetTextLocationAt( LineView, Relative );
+		return GetTextLocationAt( LineView, Relative, OutHitPoint );
 	}
 	else 
 	{
@@ -798,25 +821,16 @@ FTextLocation FTextLayout::GetTextLocationAt( const FVector2D& Relative )
 	}
 
 	const FLineView& LineView = LineViews[ ViewIndex ];
-	return GetTextLocationAt(LineView, FVector2D(Relative.X, LineView.Offset.Y));
+	return GetTextLocationAt( LineView, FVector2D(Relative.X, LineView.Offset.Y), OutHitPoint );
 }
 
-FVector2D FTextLayout::GetLocationAt( const FTextLocation& Location )
+FVector2D FTextLayout::GetLocationAt( const FTextLocation& Location, const bool bPerformInclusiveBoundsCheck )
 {
 	const int32 LineModelIndex = Location.GetLineIndex();
 	const int32 Offset = Location.GetOffset();
 
 	// Find the LineView which encapsulates the location's offset
-	int LineViewIndex;
-	for (LineViewIndex = 0; LineViewIndex < LineViews.Num(); LineViewIndex++)
-	{
-		const FTextLayout::FLineView& LineView = LineViews[ LineViewIndex ];
-
-		if ( LineView.ModelIndex == LineModelIndex && ( Offset == 0 || LineView.Range.InclusiveContains(Offset) ) )
-		{
-			break;
-		}
-	}
+	const int32 LineViewIndex = GetLineViewIndexForTextLocation( LineViews, Location, bPerformInclusiveBoundsCheck );
 
 	// if we failed to find a LineView for the location, early out
 	if ( !LineViews.IsValidIndex( LineViewIndex ) )
diff --git a/Engine/Source/Runtime/Slate/Private/Widgets/Text/SMultiLineEditableText.cpp b/Engine/Source/Runtime/Slate/Private/Widgets/Text/SMultiLineEditableText.cpp
--- a/Engine/Source/Runtime/Slate/Private/Widgets/Text/SMultiLineEditableText.cpp
+++ b/Engine/Source/Runtime/Slate/Private/Widgets/Text/SMultiLineEditableText.cpp
@@ -4,7 +4,7 @@
 #include "TextEditHelper.h"
 #include "SlateWordWrapper.h"
 
-void SMultiLineEditableText::FCursorInfo::SetCursorLocationAndCalculateAlignment(const TSharedPtr<FTextLayout>& TextLayout, const FTextLocation& InCursorPosition, const ECursorEOLMode CursorEOLMode)
+void SMultiLineEditableText::FCursorInfo::SetCursorLocationAndCalculateAlignment(const TSharedPtr<FTextLayout>& TextLayout, const FTextLocation& InCursorPosition)
 {
 	FTextLocation NewCursorPosition = InCursorPosition;
 	ECursorAlignment NewAlignment = ECursorAlignment::Left;
@@ -15,27 +15,9 @@ void SMultiLineEditableText::FCursorInfo::SetCursorLocationAndCalculateAlignment
 	// A CursorOffset of zero could mark the end of an empty line, but we don't need to adjust the cursor for an empty line
 	if (TextLayout.IsValid() && CursorOffset > 0)
 	{
-		bool bHasFoundEOL = false;
-
-		if (CursorEOLMode == ECursorEOLMode::HardLines)
-		{
-			const TArray< FTextLayout::FLineModel >& Lines = TextLayout->GetLineModels();
-			const FTextLayout::FLineModel& Line = Lines[CursorLineIndex];
-			bHasFoundEOL = (Line.Text->Len() == CursorOffset);
-		}
-		else
-		{
-			const TArray< FTextLayout::FLineView >& LineViews = TextLayout->GetLineViews();
-			const int32 LineViewIndex = TextLayout->GetLineViewIndexForTextLocation(LineViews, InCursorPosition);
-
-			if (LineViews.IsValidIndex(LineViewIndex))
-			{
-				const FTextLayout::FLineView& CurrentLineView = LineViews[LineViewIndex];
-				bHasFoundEOL = (CurrentLineView.Range.EndIndex == CursorOffset);
-			}
-		}
-
-		if (bHasFoundEOL)
+		const TArray< FTextLayout::FLineModel >& Lines = TextLayout->GetLineModels();
+		const FTextLayout::FLineModel& Line = Lines[CursorLineIndex];
+		if (Line.Text->Len() == CursorOffset)
 		{
 			// We need to move the cursor back one from where it currently is; this keeps the interaction point the same 
 			// (since the cursor is aligned to the right), but visually draws the cursor in the correct place
@@ -520,6 +502,9 @@ void SMultiLineEditableText::FinishChangingText()
 
 		// Let outsiders know that the text content has been changed
 		OnTextChanged.ExecuteIfBound(EditedText);
+
+		// Update the desired cursor position, since typing will have moved it
+		UpdatePreferredCursorScreenOffsetInLine();
 	}
 }
 
@@ -856,29 +841,41 @@ FReply SMultiLineEditableText::MoveCursor( ECursorMoveMethod::Type Method, const
 		}
 	}
 
-	ECursorEOLMode CursorEOLMode = ECursorEOLMode::HardLines;
+	TOptional<ECursorAlignment> NewCursorAlignment;
+	bool bUpdatePreferredCursorScreenOffsetInLine = false;
 	if (bAllowMoveCursor)
 	{
 		switch (Method)
 		{
 		case ECursorMoveMethod::CharacterHorizontal:
 			NewCursorPosition = TranslatedLocation( CursorPosition, Direction.X );
-			PreferredCursorScreenOffsetInLine = TextLayout->GetLocationAt( NewCursorPosition ).X;
+			bUpdatePreferredCursorScreenOffsetInLine = true;
 			break;
 
 		case ECursorMoveMethod::Word:
 			NewCursorPosition = ScanForWordBoundary( CursorPosition, Direction.X );
-			PreferredCursorScreenOffsetInLine = TextLayout->GetLocationAt( NewCursorPosition ).X;
+			bUpdatePreferredCursorScreenOffsetInLine = true;
 			break;
 	
 		case ECursorMoveMethod::CharacterVertical:
-			NewCursorPosition = TranslateLocationVertical( CursorPosition, Direction.Y );
+			TranslateLocationVertical( CursorPosition, Direction.Y, NewCursorPosition, NewCursorAlignment );
 			break;
 	
 		case ECursorMoveMethod::ScreenPosition:
-			NewCursorPosition = TextLayout->GetTextLocationAt( Direction * TextLayout->GetScale() );
-			CursorEOLMode = ECursorEOLMode::SoftLines;
-			PreferredCursorScreenOffsetInLine = Direction.X;
+			{
+				ETextHitPoint HitPoint = ETextHitPoint::WithinText;
+				NewCursorPosition = TextLayout->GetTextLocationAt( Direction * TextLayout->GetScale(), &HitPoint );
+				bUpdatePreferredCursorScreenOffsetInLine = true;
+
+				// Moving with the mouse behaves a bit differently to moving with the keyboard, as clicking at the end of a wrapped line needs to place the cursor there
+				// rather than at the start of the next line (which is tricky since they have the same index according to GetTextLocationAt!).
+				// We use the HitPoint to work this out and then adjust the cursor position accordingly
+				if (HitPoint == ETextHitPoint::RightGutter)
+				{
+					NewCursorPosition = FTextLocation(NewCursorPosition, -1);
+					NewCursorAlignment = ECursorAlignment::Right;
+				}
+			}
 			break;
 
 		default:
@@ -902,7 +899,20 @@ FReply SMultiLineEditableText::MoveCursor( ECursorMoveMethod::Type Method, const
 		this->ClearSelection();
 	}
 
-	CursorInfo.SetCursorLocationAndCalculateAlignment(TextLayout, NewCursorPosition, CursorEOLMode);
+	if (NewCursorAlignment.IsSet())
+	{
+		CursorInfo.SetCursorLocationAndAlignment(NewCursorPosition, NewCursorAlignment.GetValue());
+	}
+	else
+	{
+		CursorInfo.SetCursorLocationAndCalculateAlignment(TextLayout, NewCursorPosition);
+	}
+
+	if (bUpdatePreferredCursorScreenOffsetInLine)
+	{
+		UpdatePreferredCursorScreenOffsetInLine();
+	}
+
 	UpdateCursorHighlight();
 
 	// If we've moved the cursor while composing, we need to end the current composition session
@@ -1037,6 +1047,11 @@ void SMultiLineEditableText::RemoveCursorHighlight()
 	TextLayout->ClearLineHighlights();
 }
 
+void SMultiLineEditableText::UpdatePreferredCursorScreenOffsetInLine()
+{
+	PreferredCursorScreenOffsetInLine = TextLayout->GetLocationAt(CursorInfo.GetCursorInteractionLocation(), CursorInfo.GetCursorAlignment() == ECursorAlignment::Right).X;
+}
+
 void SMultiLineEditableText::JumpTo(ETextLocation::Type JumpLocation, ECursorAction::Type Action)
 {
 	switch (JumpLocation)
@@ -1045,7 +1060,7 @@ void SMultiLineEditableText::JumpTo(ETextLocation::Type JumpLocation, ECursorAct
 		{
 			const FTextLocation CursorInteractionPosition = CursorInfo.GetCursorInteractionLocation();
 			const TArray< FTextLayout::FLineView >& LineViews = TextLayout->GetLineViews();
-			const int32 CurrentLineViewIndex = TextLayout->GetLineViewIndexForTextLocation(LineViews, CursorInteractionPosition );
+			const int32 CurrentLineViewIndex = TextLayout->GetLineViewIndexForTextLocation(LineViews, CursorInteractionPosition, CursorInfo.GetCursorAlignment() == ECursorAlignment::Right );
 
 			if (LineViews.IsValidIndex(CurrentLineViewIndex))
 			{
@@ -1053,7 +1068,7 @@ void SMultiLineEditableText::JumpTo(ETextLocation::Type JumpLocation, ECursorAct
 				
 				const FTextLocation OldCursorPosition = CursorInteractionPosition;
 				const FTextLocation NewCursorPosition = FTextLocation(OldCursorPosition.GetLineIndex(), CurrentLineView.Range.BeginIndex);
-				
+
 				if (Action == ECursorAction::SelectText)
 				{
 					this->SelectionStart = OldCursorPosition;
@@ -1063,9 +1078,8 @@ void SMultiLineEditableText::JumpTo(ETextLocation::Type JumpLocation, ECursorAct
 					ClearSelection();
 				}
 
-				PreferredCursorScreenOffsetInLine = TextLayout->GetLocationAt(NewCursorPosition).X;
-
 				CursorInfo.SetCursorLocationAndCalculateAlignment(TextLayout, NewCursorPosition);
+				UpdatePreferredCursorScreenOffsetInLine();
 				UpdateCursorHighlight();
 			}
 		}
@@ -1085,9 +1099,8 @@ void SMultiLineEditableText::JumpTo(ETextLocation::Type JumpLocation, ECursorAct
 				ClearSelection();
 			}
 
-			PreferredCursorScreenOffsetInLine = TextLayout->GetLocationAt(NewCursorPosition).X;
-
 			CursorInfo.SetCursorLocationAndCalculateAlignment(TextLayout, NewCursorPosition);
+			UpdatePreferredCursorScreenOffsetInLine();
 			UpdateCursorHighlight();
 		}
 		break;
@@ -1096,14 +1109,14 @@ void SMultiLineEditableText::JumpTo(ETextLocation::Type JumpLocation, ECursorAct
 		{
 			const FTextLocation CursorInteractionPosition = CursorInfo.GetCursorInteractionLocation();
 			const TArray< FTextLayout::FLineView >& LineViews = TextLayout->GetLineViews();
-			const int32 CurrentLineViewIndex = TextLayout->GetLineViewIndexForTextLocation(LineViews, CursorInteractionPosition);
+			const int32 CurrentLineViewIndex = TextLayout->GetLineViewIndexForTextLocation(LineViews, CursorInteractionPosition, CursorInfo.GetCursorAlignment() == ECursorAlignment::Right);
 
 			if (LineViews.IsValidIndex(CurrentLineViewIndex))
 			{
 				const FTextLayout::FLineView& CurrentLineView = LineViews[CurrentLineViewIndex];
 
 				const FTextLocation OldCursorPosition = CursorInteractionPosition;
-				const FTextLocation NewCursorPosition = FTextLocation(OldCursorPosition.GetLineIndex(), CurrentLineView.Range.EndIndex);
+				const FTextLocation NewCursorPosition = FTextLocation(OldCursorPosition.GetLineIndex(), FMath::Max(0, CurrentLineView.Range.EndIndex - 1));
 
 				if (Action == ECursorAction::SelectText)
 				{
@@ -1114,9 +1127,8 @@ void SMultiLineEditableText::JumpTo(ETextLocation::Type JumpLocation, ECursorAct
 					ClearSelection();
 				}
 
-				PreferredCursorScreenOffsetInLine = TextLayout->GetLocationAt(NewCursorPosition).X;
-
-				CursorInfo.SetCursorLocationAndCalculateAlignment(TextLayout, NewCursorPosition);
+				CursorInfo.SetCursorLocationAndAlignment(NewCursorPosition, ECursorAlignment::Right);
+				UpdatePreferredCursorScreenOffsetInLine();
 				UpdateCursorHighlight();
 			}
 		}
@@ -1140,9 +1152,8 @@ void SMultiLineEditableText::JumpTo(ETextLocation::Type JumpLocation, ECursorAct
 					ClearSelection();
 				}
 
-				PreferredCursorScreenOffsetInLine = TextLayout->GetLocationAt(NewCursorPosition).X;
-
 				CursorInfo.SetCursorLocationAndCalculateAlignment(TextLayout, NewCursorPosition);
+				UpdatePreferredCursorScreenOffsetInLine();
 				UpdateCursorHighlight();
 			}
 		}
@@ -1900,12 +1911,12 @@ FTextLocation SMultiLineEditableText::TranslatedLocation( const FTextLocation& L
 	}
 }
 
-FTextLocation SMultiLineEditableText::TranslateLocationVertical( const FTextLocation& Location, int8 Direction ) const
+void SMultiLineEditableText::TranslateLocationVertical( const FTextLocation& Location, int8 Direction, FTextLocation& OutCursorPosition, TOptional<ECursorAlignment>& OutCursorAlignment ) const
 {
 	const TArray< FTextLayout::FLineView >& LineViews = TextLayout->GetLineViews();
 	const int32 NumberOfLineViews = LineViews.Num();
 
-	const int32 CurrentLineViewIndex = TextLayout->GetLineViewIndexForTextLocation(LineViews, Location);
+	const int32 CurrentLineViewIndex = TextLayout->GetLineViewIndexForTextLocation(LineViews, Location, CursorInfo.GetCursorAlignment() == ECursorAlignment::Right);
 	ensure(CurrentLineViewIndex != INDEX_NONE);
 	const FTextLayout::FLineView& CurrentLineView = LineViews[ CurrentLineViewIndex ];
 
@@ -1913,7 +1924,17 @@ FTextLocation SMultiLineEditableText::TranslateLocationVertical( const FTextLoca
 	const FTextLayout::FLineView& NewLineView = LineViews[ NewLineViewIndex ];
 	
 	// Our horizontal position is the clamped version of whatever the user explicitly set with horizontal movement.
-	return TextLayout->GetTextLocationAt( NewLineView, FVector2D( PreferredCursorScreenOffsetInLine, NewLineView.Offset.Y ) * TextLayout->GetScale() );
+	ETextHitPoint HitPoint = ETextHitPoint::WithinText;
+	OutCursorPosition = TextLayout->GetTextLocationAt( NewLineView, FVector2D( PreferredCursorScreenOffsetInLine, NewLineView.Offset.Y ) * TextLayout->GetScale(), &HitPoint );
+
+	// PreferredCursorScreenOffsetInLine can cause the cursor to move to the right hand gutter, and it needs to be placed there
+	// rather than at the start of the next line (which is tricky since they have the same index according to GetTextLocationAt!).
+	// We use the HitPoint to work this out and then adjust the cursor position accordingly
+	if(HitPoint == ETextHitPoint::RightGutter)
+	{
+		OutCursorPosition = FTextLocation(OutCursorPosition, -1);
+		OutCursorAlignment = ECursorAlignment::Right;
+	}
 }
 
 FTextLocation SMultiLineEditableText::ScanForWordBoundary( const FTextLocation& CurrentLocation, int8 Direction ) const
@@ -2139,8 +2160,8 @@ bool SMultiLineEditableText::FTextInputMethodContext::GetTextBounds(const uint32
 		const FTextLocation BeginLocation = OffsetLocations.OffsetToTextLocation(BeginIndex);
 		const FTextLocation EndLocation = OffsetLocations.OffsetToTextLocation(BeginIndex + Length);
 
-		const FVector2D BeginPosition = OwningWidgetPtr->TextLayout->GetLocationAt(BeginLocation);
-		const FVector2D EndPosition = OwningWidgetPtr->TextLayout->GetLocationAt(EndLocation);
+		const FVector2D BeginPosition = OwningWidgetPtr->TextLayout->GetLocationAt(BeginLocation, false);
+		const FVector2D EndPosition = OwningWidgetPtr->TextLayout->GetLocationAt(EndLocation, false);
 
 		if(BeginPosition.Y == EndPosition.Y)
 		{
diff --git a/Engine/Source/Runtime/Slate/Public/Framework/Text/IRun.h b/Engine/Source/Runtime/Slate/Public/Framework/Text/IRun.h
--- a/Engine/Source/Runtime/Slate/Public/Framework/Text/IRun.h
+++ b/Engine/Source/Runtime/Slate/Public/Framework/Text/IRun.h
@@ -16,7 +16,7 @@ public:
 
 	virtual TSharedRef< class ILayoutBlock > CreateBlock( int32 StartIndex, int32 EndIndex, FVector2D Size, const TSharedPtr< class IRunRenderer >& Renderer ) = 0;
 
-	virtual int32 GetTextIndexAt( const TSharedRef< ILayoutBlock >& Block, const FVector2D& Location, float Scale ) const = 0;
+	virtual int32 GetTextIndexAt( const TSharedRef< ILayoutBlock >& Block, const FVector2D& Location, float Scale, ETextHitPoint* const OutHitPoint = nullptr ) const = 0;
 	
 	virtual FVector2D GetLocationAt(const TSharedRef< ILayoutBlock >& Block, int32 Offset, float Scale) const = 0;
 
diff --git a/Engine/Source/Runtime/Slate/Public/Framework/Text/SlateHyperlinkRun.h b/Engine/Source/Runtime/Slate/Public/Framework/Text/SlateHyperlinkRun.h
--- a/Engine/Source/Runtime/Slate/Public/Framework/Text/SlateHyperlinkRun.h
+++ b/Engine/Source/Runtime/Slate/Public/Framework/Text/SlateHyperlinkRun.h
@@ -56,7 +56,7 @@ public:
 
 	virtual void ArrangeChildren( const TSharedRef< ILayoutBlock >& Block, const FGeometry& AllottedGeometry, FArrangedChildren& ArrangedChildren ) const override;
 
-	virtual int32 GetTextIndexAt( const TSharedRef< ILayoutBlock >& Block, const FVector2D& Location, float Scale ) const override;
+	virtual int32 GetTextIndexAt( const TSharedRef< ILayoutBlock >& Block, const FVector2D& Location, float Scale, ETextHitPoint* const OutHitPoint = nullptr ) const override;
 
 	virtual FVector2D GetLocationAt( const TSharedRef< ILayoutBlock >& Block, int32 Offset, float Scale ) const override;
 
diff --git a/Engine/Source/Runtime/Slate/Public/Framework/Text/SlateImageRun.h b/Engine/Source/Runtime/Slate/Public/Framework/Text/SlateImageRun.h
--- a/Engine/Source/Runtime/Slate/Public/Framework/Text/SlateImageRun.h
+++ b/Engine/Source/Runtime/Slate/Public/Framework/Text/SlateImageRun.h
@@ -31,7 +31,7 @@ public:
 
 	virtual void ArrangeChildren( const TSharedRef< ILayoutBlock >& Block, const FGeometry& AllottedGeometry, FArrangedChildren& ArrangedChildren ) const override;
 
-	virtual int32 GetTextIndexAt( const TSharedRef< ILayoutBlock >& Block, const FVector2D& Location, float Scale ) const override;
+	virtual int32 GetTextIndexAt( const TSharedRef< ILayoutBlock >& Block, const FVector2D& Location, float Scale, ETextHitPoint* const OutHitPoint = nullptr ) const override;
 
 	virtual FVector2D GetLocationAt( const TSharedRef< ILayoutBlock >& Block, int32 Offset, float Scale ) const override;
 
diff --git a/Engine/Source/Runtime/Slate/Public/Framework/Text/SlateTextRun.h b/Engine/Source/Runtime/Slate/Public/Framework/Text/SlateTextRun.h
--- a/Engine/Source/Runtime/Slate/Public/Framework/Text/SlateTextRun.h
+++ b/Engine/Source/Runtime/Slate/Public/Framework/Text/SlateTextRun.h
@@ -32,7 +32,7 @@ public:
 	virtual void BeginLayout() override {}
 	virtual void EndLayout() override {}
 
-	virtual int32 GetTextIndexAt( const TSharedRef< ILayoutBlock >& Block, const FVector2D& Location, float Scale ) const override;
+	virtual int32 GetTextIndexAt( const TSharedRef< ILayoutBlock >& Block, const FVector2D& Location, float Scale, ETextHitPoint* const OutHitPoint = nullptr ) const override;
 
 	virtual FVector2D GetLocationAt(const TSharedRef< ILayoutBlock >& Block, int32 Offset, float Scale) const override;
 
diff --git a/Engine/Source/Runtime/Slate/Public/Framework/Text/SlateWidgetRun.h b/Engine/Source/Runtime/Slate/Public/Framework/Text/SlateWidgetRun.h
--- a/Engine/Source/Runtime/Slate/Public/Framework/Text/SlateWidgetRun.h
+++ b/Engine/Source/Runtime/Slate/Public/Framework/Text/SlateWidgetRun.h
@@ -51,7 +51,7 @@ public:
 
 	virtual void ArrangeChildren( const TSharedRef< ILayoutBlock >& Block, const FGeometry& AllottedGeometry, FArrangedChildren& ArrangedChildren ) const;
 
-	virtual int32 GetTextIndexAt( const TSharedRef< ILayoutBlock >& Block, const FVector2D& Location, float Scale ) const override;
+	virtual int32 GetTextIndexAt( const TSharedRef< ILayoutBlock >& Block, const FVector2D& Location, float Scale, ETextHitPoint* const OutHitPoint = nullptr ) const override;
 
 	virtual FVector2D GetLocationAt( const TSharedRef< ILayoutBlock >& Block, int32 Offset, float Scale ) const override;
 
diff --git a/Engine/Source/Runtime/Slate/Public/Framework/Text/TextHitPoint.h b/Engine/Source/Runtime/Slate/Public/Framework/Text/TextHitPoint.h
new file mode 100644
--- /dev/null
+++ b/Engine/Source/Runtime/Slate/Public/Framework/Text/TextHitPoint.h
@@ -0,0 +1,13 @@
+// Copyright 1998-2014 Epic Games, Inc. All Rights Reserved.
+#pragma once
+
+/** Describes how a line was hit when performing a screen-space -> view model conversion */
+enum class ETextHitPoint : uint8
+{
+	/** The hit-point was within the text range */
+	WithinText,
+	/** The hit-point was to left of the text range (the gutter) */
+	LeftGutter,
+	/** The hit-point was to right of the text range (the gutter) */
+	RightGutter,
+};
diff --git a/Engine/Source/Runtime/Slate/Public/Framework/Text/TextLayout.h b/Engine/Source/Runtime/Slate/Public/Framework/Text/TextLayout.h
--- a/Engine/Source/Runtime/Slate/Public/Framework/Text/TextLayout.h
+++ b/Engine/Source/Runtime/Slate/Public/Framework/Text/TextLayout.h
@@ -331,16 +331,16 @@ public:
 
 	virtual void UpdateHighlights();
 
-	int32 GetLineViewIndexForTextLocation(const TArray< FTextLayout::FLineView >& LineViews, const FTextLocation& Location);
+	int32 GetLineViewIndexForTextLocation(const TArray< FTextLayout::FLineView >& LineViews, const FTextLocation& Location, const bool bPerformInclusiveBoundsCheck);
 
 	/**
 	 * 
 	 */
-	FTextLocation GetTextLocationAt( const FVector2D& Relative );
+	FTextLocation GetTextLocationAt( const FVector2D& Relative, ETextHitPoint* const OutHitPoint = nullptr );
 
-	FTextLocation GetTextLocationAt( const FLineView& LineView, const FVector2D& Relative );
+	FTextLocation GetTextLocationAt( const FLineView& LineView, const FVector2D& Relative, ETextHitPoint* const OutHitPoint = nullptr );
 
-	FVector2D GetLocationAt( const FTextLocation& Location );
+	FVector2D GetLocationAt( const FTextLocation& Location, const bool bPerformInclusiveBoundsCheck );
 
 	bool SplitLineAt(const FTextLocation& Location);
 
diff --git a/Engine/Source/Runtime/Slate/Public/Framework/Text/TextLayoutEngine.h b/Engine/Source/Runtime/Slate/Public/Framework/Text/TextLayoutEngine.h
--- a/Engine/Source/Runtime/Slate/Public/Framework/Text/TextLayoutEngine.h
+++ b/Engine/Source/Runtime/Slate/Public/Framework/Text/TextLayoutEngine.h
@@ -6,6 +6,7 @@
 #include "TextRunRenderer.h"
 #include "TextLineHighlight.h"
 
+#include "TextHitPoint.h"
 #include "IRun.h"
 #include "IRunRenderer.h"
 #include "ILineHighlighter.h"
diff --git a/Engine/Source/Runtime/Slate/Public/Widgets/Text/SMultiLineEditableText.h b/Engine/Source/Runtime/Slate/Public/Widgets/Text/SMultiLineEditableText.h
--- a/Engine/Source/Runtime/Slate/Public/Widgets/Text/SMultiLineEditableText.h
+++ b/Engine/Source/Runtime/Slate/Public/Widgets/Text/SMultiLineEditableText.h
@@ -137,15 +137,6 @@ private:
 		Right,
 	};
 
-	enum class ECursorEOLMode : uint8
-	{
-		/** The cursor considers hard-line endings when looking for an EOL to align the cursor to (ignores line wrapping) */
-		HardLines,
-
-		/** The cursor considers soft-line endings when looking for an EOL to align the cursor to (checks for line wrapping) */
-		SoftLines,
-	};
-
 	/** Store the information about the current cursor position */
 	class FCursorInfo
 	{
@@ -182,7 +173,7 @@ private:
 		}
 
 		/** Set the position of the cursor, and then work out the correct alignment based on the current text layout */
-		void SetCursorLocationAndCalculateAlignment(const TSharedPtr<FTextLayout>& TextLayout, const FTextLocation& InCursorPosition, const ECursorEOLMode CursorEOLMode = ECursorEOLMode::HardLines);
+		void SetCursorLocationAndCalculateAlignment(const TSharedPtr<FTextLayout>& TextLayout, const FTextLocation& InCursorPosition);
 
 		/** Set the literal position and alignment of the cursor */
 		void SetCursorLocationAndAlignment(const FTextLocation& InCursorPosition, const ECursorAlignment InCursorAlignment);
@@ -381,10 +372,12 @@ private:
 	/**
 	 * Given a location and a Direction to offset, return a new location.
 	 *
-	 * @param Location    Cursor location from which to offset
-	 * @param Direction   Positive means down, negative means up.
+	 * @param Location				Cursor location from which to offset
+	 * @param Direction				Positive means down, negative means up.
+	 * @param OutCursorPosition		Fill with the updated cursor position.
+	 * @param OutCursorAlignment	Optionally fill with a new cursor alignment (will be auto-calculated if not set).
 	 */
-	FTextLocation TranslateLocationVertical( const FTextLocation& Location, int8 Direction ) const;
+	void TranslateLocationVertical( const FTextLocation& Location, int8 Direction, FTextLocation& OutCursorPosition, TOptional<ECursorAlignment>& OutCursorAlignment ) const;
 
 	/** Find the closest word boundary */
 	FTextLocation ScanForWordBoundary( const FTextLocation& Location, int8 Direction ) const; 
@@ -411,6 +404,8 @@ private:
 
 	void RemoveCursorHighlight();
 
+	void UpdatePreferredCursorScreenOffsetInLine();
+
 	void PushUndoState(const SMultiLineEditableText::FUndoState& InUndoState);
 
 	void ClearUndoStates();
