Color picker can now be summoned on all desktop platforms

diff --git a/Engine/Source/Runtime/Slate/Private/Widgets/Colors/SColorPicker.cpp b/Engine/Source/Runtime/Slate/Private/Widgets/Colors/SColorPicker.cpp
--- a/Engine/Source/Runtime/Slate/Private/Widgets/Colors/SColorPicker.cpp
+++ b/Engine/Source/Runtime/Slate/Private/Widgets/Colors/SColorPicker.cpp
@@ -2139,100 +2139,100 @@ bool OpenColorPicker(const FColorPickerArgs& Args)
 	DestroyColorPicker();
 
 	bool Result = false;
-	
-	//guard to verify that we're not in game using edit actor
-	if (GIsEditor)
-	{
-		FLinearColor OldColor = Args.InitialColorOverride;
 
-		if (Args.ColorArray && Args.ColorArray->Num() > 0)
-		{
-			OldColor = (*Args.ColorArray)[0]->ReinterpretAsLinear();
-		}
-		else if (Args.LinearColorArray && Args.LinearColorArray->Num() > 0)
-		{
-			OldColor = *(*Args.LinearColorArray)[0];
-		}
-		else if (Args.ColorChannelsArray && Args.ColorChannelsArray->Num() > 0)
-		{
-			OldColor.R = (*Args.ColorChannelsArray)[0].Red ? *(*Args.ColorChannelsArray)[0].Red : 0.0f;
-			OldColor.G = (*Args.ColorChannelsArray)[0].Green ? *(*Args.ColorChannelsArray)[0].Green : 0.0f;
-			OldColor.B = (*Args.ColorChannelsArray)[0].Blue ? *(*Args.ColorChannelsArray)[0].Blue : 0.0f;
-			OldColor.A = (*Args.ColorChannelsArray)[0].Alpha ? *(*Args.ColorChannelsArray)[0].Alpha : 0.0f;
-		}
-		else
-		{
-			check(Args.OnColorCommitted.IsBound());
-		}
+	// Consoles do not support opening new windows
+#if PLATFORM_DESKTOP
+	FLinearColor OldColor = Args.InitialColorOverride;
+
+	if (Args.ColorArray && Args.ColorArray->Num() > 0)
+	{
+		OldColor = (*Args.ColorArray)[0]->ReinterpretAsLinear();
+	}
+	else if (Args.LinearColorArray && Args.LinearColorArray->Num() > 0)
+	{
+		OldColor = *(*Args.LinearColorArray)[0];
+	}
+	else if (Args.ColorChannelsArray && Args.ColorChannelsArray->Num() > 0)
+	{
+		OldColor.R = (*Args.ColorChannelsArray)[0].Red ? *(*Args.ColorChannelsArray)[0].Red : 0.0f;
+		OldColor.G = (*Args.ColorChannelsArray)[0].Green ? *(*Args.ColorChannelsArray)[0].Green : 0.0f;
+		OldColor.B = (*Args.ColorChannelsArray)[0].Blue ? *(*Args.ColorChannelsArray)[0].Blue : 0.0f;
+		OldColor.A = (*Args.ColorChannelsArray)[0].Alpha ? *(*Args.ColorChannelsArray)[0].Alpha : 0.0f;
+	}
+	else
+	{
+		check(Args.OnColorCommitted.IsBound());
+	}
 		
-		// Determine the position of the window so that it will spawn near the mouse, but not go off the screen.
-		const FVector2D CursorPos = FSlateApplication::Get().GetCursorPos();
-		FSlateRect Anchor(CursorPos.X, CursorPos.Y, CursorPos.X, CursorPos.Y);
-
-		FVector2D AdjustedSummonLocation = FSlateApplication::Get().CalculatePopupWindowPosition( Anchor, SColorPicker::DEFAULT_WINDOW_SIZE, Orient_Horizontal );
-
-		TSharedPtr<SWindow> Window = SNew(SWindow)
-			.AutoCenter(EAutoCenter::None)
-			.ScreenPosition(AdjustedSummonLocation)
-			.SupportsMaximize(false)
-			.SupportsMinimize(false)
-			.SizingRule(ESizingRule::Autosized)
-			.Title(LOCTEXT("WindowHeader", "Color Picker"));
-
-		TSharedRef<SColorPicker> ColorPicker = SNew(SColorPicker)
-			.TargetColorAttribute(OldColor)
-			.TargetFColors(Args.ColorArray ? *Args.ColorArray : TArray<FColor*>())
-			.TargetLinearColors(Args.LinearColorArray ? *Args.LinearColorArray : TArray<FLinearColor*>())
-			.TargetColorChannels(Args.ColorChannelsArray ? *Args.ColorChannelsArray : TArray<FColorChannels>())
-			.UseAlpha(Args.bUseAlpha)
-			.OnlyRefreshOnMouseUp(Args.bOnlyRefreshOnMouseUp && !Args.bIsModal)
-			.OnlyRefreshOnOk(Args.bOnlyRefreshOnOk || Args.bIsModal)
-			.OnColorCommitted(Args.OnColorCommitted)
-			.PreColorCommitted(Args.PreColorCommitted)
-			.OnColorPickerCancelled(Args.OnColorPickerCancelled)
-			.OnInteractivePickBegin(Args.OnInteractivePickBegin)
-			.OnInteractivePickEnd(Args.OnInteractivePickEnd)
-			.OnColorPickerWindowClosed(Args.OnColorPickerWindowClosed)
-			.ParentWindow(Window)
-			.DisplayGamma(Args.DisplayGamma);
+	// Determine the position of the window so that it will spawn near the mouse, but not go off the screen.
+	const FVector2D CursorPos = FSlateApplication::Get().GetCursorPos();
+	FSlateRect Anchor(CursorPos.X, CursorPos.Y, CursorPos.X, CursorPos.Y);
+
+	FVector2D AdjustedSummonLocation = FSlateApplication::Get().CalculatePopupWindowPosition( Anchor, SColorPicker::DEFAULT_WINDOW_SIZE, Orient_Horizontal );
+
+	TSharedPtr<SWindow> Window = SNew(SWindow)
+		.AutoCenter(EAutoCenter::None)
+		.ScreenPosition(AdjustedSummonLocation)
+		.SupportsMaximize(false)
+		.SupportsMinimize(false)
+		.SizingRule(ESizingRule::Autosized)
+		.Title(LOCTEXT("WindowHeader", "Color Picker"));
+
+	TSharedRef<SColorPicker> ColorPicker = SNew(SColorPicker)
+		.TargetColorAttribute(OldColor)
+		.TargetFColors(Args.ColorArray ? *Args.ColorArray : TArray<FColor*>())
+		.TargetLinearColors(Args.LinearColorArray ? *Args.LinearColorArray : TArray<FLinearColor*>())
+		.TargetColorChannels(Args.ColorChannelsArray ? *Args.ColorChannelsArray : TArray<FColorChannels>())
+		.UseAlpha(Args.bUseAlpha)
+		.OnlyRefreshOnMouseUp(Args.bOnlyRefreshOnMouseUp && !Args.bIsModal)
+		.OnlyRefreshOnOk(Args.bOnlyRefreshOnOk || Args.bIsModal)
+		.OnColorCommitted(Args.OnColorCommitted)
+		.PreColorCommitted(Args.PreColorCommitted)
+		.OnColorPickerCancelled(Args.OnColorPickerCancelled)
+		.OnInteractivePickBegin(Args.OnInteractivePickBegin)
+		.OnInteractivePickEnd(Args.OnInteractivePickEnd)
+		.OnColorPickerWindowClosed(Args.OnColorPickerWindowClosed)
+		.ParentWindow(Window)
+		.DisplayGamma(Args.DisplayGamma);
 		
-		Window->SetContent(
-			SNew(SBox)
-				[
-					SNew(SBorder)
-						.BorderImage(FCoreStyle::Get().GetBrush("ToolPanel.GroupBorder"))
-						.Padding(FMargin(8.0f, 8.0f))
-						[
-							ColorPicker
-						]
-				]
-		);
+	Window->SetContent(
+		SNew(SBox)
+			[
+				SNew(SBorder)
+					.BorderImage(FCoreStyle::Get().GetBrush("ToolPanel.GroupBorder"))
+					.Padding(FMargin(8.0f, 8.0f))
+					[
+						ColorPicker
+					]
+			]
+	);
 
-		if (Args.bIsModal)
+	if (Args.bIsModal)
+	{
+		FSlateApplication::Get().AddModalWindow(Window.ToSharedRef(), Args.ParentWidget);
+	}
+	else
+	{
+		if ( Args.ParentWidget.IsValid() )
 		{
-			FSlateApplication::Get().AddModalWindow(Window.ToSharedRef(), Args.ParentWidget);
+			// Find the window of the parent widget
+			FWidgetPath WidgetPath;
+			FSlateApplication::Get().GeneratePathToWidgetChecked( Args.ParentWidget.ToSharedRef(), WidgetPath );
+			Window = FSlateApplication::Get().AddWindowAsNativeChild( Window.ToSharedRef(), WidgetPath.GetWindow() );
 		}
 		else
 		{
-			if ( Args.ParentWidget.IsValid() )
-			{
-				// Find the window of the parent widget
-				FWidgetPath WidgetPath;
-				FSlateApplication::Get().GeneratePathToWidgetChecked( Args.ParentWidget.ToSharedRef(), WidgetPath );
-				Window = FSlateApplication::Get().AddWindowAsNativeChild( Window.ToSharedRef(), WidgetPath.GetWindow() );
-			}
-			else
-			{
-				Window = FSlateApplication::Get().AddWindow(Window.ToSharedRef());
-			}
-			
+			Window = FSlateApplication::Get().AddWindow(Window.ToSharedRef());
 		}
+			
+	}
 
-		Result = true;
+	Result = true;
 
-		//hold on to the window created for external use...
-		ColorPickerWindow = Window;
-	}
+	//hold on to the window created for external use...
+	ColorPickerWindow = Window;
+	
+#endif
 
 	return Result;
 }
