Fixed crash when deleting a line from a FTextLayout

FTextLayout::RemoveLine wasn't dirtying the layout, so the lightweight highlighters crashed the next time they were flowed due to trying to access line view entries for the deleted line model.

diff --git a/Engine/Source/Runtime/Slate/Private/Framework/Text/TextLayout.cpp b/Engine/Source/Runtime/Slate/Private/Framework/Text/TextLayout.cpp
--- a/Engine/Source/Runtime/Slate/Private/Framework/Text/TextLayout.cpp
+++ b/Engine/Source/Runtime/Slate/Private/Framework/Text/TextLayout.cpp
@@ -1122,7 +1122,8 @@ bool FTextLayout::RemoveLine(int32 LineIndex)
 
 	LineModels.RemoveAt(LineIndex);
 
-	return true;;
+	DirtyFlags |= EDirtyState::Layout;
+	return true;
 }
 
 void FTextLayout::AddLine( const TSharedRef< FString >& Text, const TArray< TSharedRef< IRun > >& Runs )
