- Fixed the IOSPlatformTextField to use proper FString conversion from TCHAR

diff --git a/Engine/Source/Runtime/Slate/Private/Framework/Text/IOS/IOSPlatformTextField.cpp b/Engine/Source/Runtime/Slate/Private/Framework/Text/IOS/IOSPlatformTextField.cpp
--- a/Engine/Source/Runtime/Slate/Private/Framework/Text/IOS/IOSPlatformTextField.cpp
+++ b/Engine/Source/Runtime/Slate/Private/Framework/Text/IOS/IOSPlatformTextField.cpp
@@ -66,8 +66,8 @@ void FIOSPlatformTextField::ShowKeyboard(bool bShow, TSharedPtr<SVirtualKeyboard
 	TextField.clearsOnInsertion = NO;
 	TextField.autocorrectionType = UITextAutocorrectionTypeNo;
 	TextField.autocapitalizationType = UITextAutocapitalizationTypeNone;
-	TextField.text = [NSString stringWithFormat:@"%s" , TCHAR_TO_ANSI(*TextWidget->GetText().ToString())];
-	TextField.placeholder = [NSString stringWithFormat:@"%s" , TCHAR_TO_ANSI(*TextWidget->GetHintText().ToString())];
+	TextField.text = [NSString stringWithFString:TextWidget->GetText().ToString()];
+	TextField.placeholder = [NSString stringWithFString:TextWidget->GetText().ToString()];
 
 	// set up the keyboard styles not supported in the AlertViewStyle styles
 	switch(TextWidget->GetKeyboardType())
