Slate Hittesting 2.0 : Hittesting no longer bound to logical widget hierarchy.

diff --git a/Engine/Source/Runtime/Slate/Private/Framework/Application/SlateApplication.cpp b/Engine/Source/Runtime/Slate/Private/Framework/Application/SlateApplication.cpp
--- a/Engine/Source/Runtime/Slate/Private/Framework/Application/SlateApplication.cpp
+++ b/Engine/Source/Runtime/Slate/Private/Framework/Application/SlateApplication.cpp
@@ -6,6 +6,7 @@
 
 #include "SlatePrivatePCH.h"
 #include "SWindowTitleBar.h"
+#include "HittestGrid.h"
 
 
 DECLARE_CYCLE_STAT( TEXT("Message Tick Time"), STAT_SlateMessageTick, STATGROUP_Slate );
@@ -306,6 +307,23 @@ void FSlateApplication::Create()
 
 	PlatformApplication = MakeShareable( FPlatformMisc::CreateApplication() );
 	PlatformApplication->SetMessageHandler( CurrentApplication.ToSharedRef() );
+
+	// The grid needs to know the size and coordinate system of the desktop.
+	// Some monitor setups have a primary monitor on the right and below the
+	// left one, so the leftmost upper right monitor can be something like (-1280, -200)
+	{
+		// Get an initial value for the VirtualDesktop geometry
+		CurrentApplication->VirtualDesktopRect = []()
+		{
+			FDisplayMetrics DisplayMetrics;
+			FSlateApplicationBase::Get().GetDisplayMetrics(DisplayMetrics);
+			const FPlatformRect& VirtualDisplayRect = DisplayMetrics.VirtualDisplayRect;
+			return FSlateRect(VirtualDisplayRect.Left, VirtualDisplayRect.Top, VirtualDisplayRect.Right, VirtualDisplayRect.Bottom);
+		}();
+
+		// Sign up for updates from the OS. Polling this every frame is too expensive on at least some OSs.
+		PlatformApplication->OnDisplayMetricsChanged().AddSP(CurrentApplication.ToSharedRef(), &FSlateApplication::OnVirtualDesktopSizeChanged);
+	}
 }
 
 void FSlateApplication::Shutdown()
@@ -366,6 +384,8 @@ FSlateApplication::FSlateApplication()
 	, bTouchFallbackToMouse( true )
 	, bMenuAnimationsEnabled( true )
 	, AppIcon( FCoreStyle::Get().GetBrush("DefaultAppIcon") )
+	, VirtualDesktopRect( 0,0,0,0 )
+	, HittestGrid( MakeShareable( new FHittestGrid() ) )
 {
 #if WITH_UNREAL_DEVELOPER_TOOLS
 	FModuleManager::Get().LoadModule(TEXT("Settings"));
@@ -475,15 +495,23 @@ FWidgetPath FSlateApplication::LocateWindowUnderMouse( FVector2D ScreenspaceMous
 
 		// Only accept input if the current window accepts input and the current window is not under a modal window or an interactive tooltip
 		const TSharedPtr<IToolTip> ActiveToolTipPtr = ActiveToolTip.Pin();
-		const TSharedPtr< SWindow > ToolTipWindowPtr = ToolTipWindow.Pin();
+		const TSharedPtr<SWindow> ToolTipWindowPtr = ToolTipWindow.Pin();
 		const bool AcceptsInput = Window->AcceptsInput() || ( Window == ToolTipWindowPtr && ActiveToolTipPtr.IsValid() && ActiveToolTipPtr->IsInteractive() );
 
 		if ( Window->IsVisible() && AcceptsInput && Window->IsScreenspaceMouseWithin(ScreenspaceMouseCoordinate) && !bPrevWindowWasModal )
 		{
-			FArrangedWidget WindowWidgetGeometry( Window, Window->GetWindowGeometryInScreen() );
-			if ( Window->IsEnabled() && LocateWidgetsUnderCursor_Helper(WindowWidgetGeometry, ScreenspaceMouseCoordinate, OutWidgetPath, bIgnoreEnabledStatus) )		
-			{			
-				return FWidgetPath( Windows[WindowIndex], OutWidgetPath );
+			if ( !SWidget::UseLegacyHittest() )
+			{
+				const TArray<FArrangedWidget> ArrangedWidgets = HittestGrid->GetBubblePath( ScreenspaceMouseCoordinate );
+				return FWidgetPath( ArrangedWidgets );
+			}
+			else
+			{
+				FArrangedWidget WindowWidgetGeometry( Window, Window->GetWindowGeometryInScreen() );
+				if ( Window->IsEnabled() && LocateWidgetsUnderCursor_Helper(WindowWidgetGeometry, ScreenspaceMouseCoordinate, OutWidgetPath, bIgnoreEnabledStatus) )		
+				{			
+					return FWidgetPath( Windows[WindowIndex], OutWidgetPath );
+				}
 			}
 		}
 	}
@@ -551,6 +579,7 @@ struct FDrawWindowArgs
 	const FWidgetPath& WidgetsUnderCursor;
 };
 
+
 void FSlateApplication::DrawWindowAndChildren( const TSharedRef<SWindow>& WindowToDraw, FDrawWindowArgs& DrawWindowArgs )
 {
 	// Only draw visible windows
@@ -565,7 +594,13 @@ void FSlateApplication::DrawWindowAndChildren( const TSharedRef<SWindow>& Window
 		FGeometry WindowGeometry = WindowToDraw->GetWindowGeometryInWindow();
 		int32 MaxLayerId = 0;
 		{
-			MaxLayerId = WindowToDraw->OnPaint( WindowGeometry, WindowToDraw->GetClippingRectangleInWindow(), WindowElementList, 0, FWidgetStyle(), WindowToDraw->IsEnabled() );
+			MaxLayerId = WindowToDraw->Paint(
+				FPaintArgs(WindowToDraw, *HittestGrid, WindowToDraw->GetPositionInScreen()),
+				WindowGeometry, WindowToDraw->GetClippingRectangleInWindow(),
+				WindowElementList,
+				0,
+				FWidgetStyle(),
+				WindowToDraw->IsEnabled() );
 		}
 
 		if (DrawWindowArgs.FocusedPath.IsValid() && DrawWindowArgs.FocusedPath.GetWindow() == WindowToDraw)
@@ -663,11 +698,9 @@ void FSlateApplication::PrivateDrawWindows( TSharedPtr<SWindow> DrawOnlyThisWind
 	// Is user expecting visual feedback from the Widget Reflector?
 	const bool bVisualizeLayoutUnderCursor = WidgetReflectorPtr.IsValid() && WidgetReflectorPtr.Pin()->IsVisualizingLayoutUnderCursor();
 
-	FWidgetPath WidgetsUnderCursor;
-	if (bVisualizeLayoutUnderCursor)
-	{
-		WidgetsUnderCursor = LocateWindowUnderMouse( GetCursorPos(), SlateWindows );
-	}
+	FWidgetPath WidgetsUnderCursor = bVisualizeLayoutUnderCursor
+		? WidgetsUnderCursorLastEvent.ToWidgetPath()
+		: FWidgetPath();
 
 	FWidgetPath FocusPath = FocusedWidgetPath.ToWidgetPath();
 
@@ -819,6 +852,11 @@ void FSlateApplication::Tick()
 	{
 		SCOPE_CYCLE_COUNTER( STAT_SlateTickWindowAndChildren );
 
+		if ( !SWidget::UseLegacyHittest() )
+		{
+			HittestGrid->BeginFrame( VirtualDesktopRect );
+		}
+
 		if ( ActiveModalWindow.IsValid() )
 		{
 			// There is a modal window, and we just need to tick it.
@@ -4487,6 +4525,16 @@ bool FSlateApplication::OnWindowAction( const TSharedRef< FGenericWindow >& Plat
 	return !IsExternalUIOpened();
 }
 
+void FSlateApplication::OnVirtualDesktopSizeChanged(const FDisplayMetrics& NewDisplayMetric)
+{
+	const FPlatformRect& VirtualDisplayRect = NewDisplayMetric.VirtualDisplayRect;
+	VirtualDesktopRect = FSlateRect(
+		VirtualDisplayRect.Left,
+		VirtualDisplayRect.Top,
+		VirtualDisplayRect.Right,
+		VirtualDisplayRect.Bottom);
+}
+
 
 /* 
  *****************************************************************************/
diff --git a/Engine/Source/Runtime/Slate/Private/Framework/Docking/SDockingCross.cpp b/Engine/Source/Runtime/Slate/Private/Framework/Docking/SDockingCross.cpp
--- a/Engine/Source/Runtime/Slate/Private/Framework/Docking/SDockingCross.cpp
+++ b/Engine/Source/Runtime/Slate/Private/Framework/Docking/SDockingCross.cpp
@@ -19,7 +19,7 @@ void SDockingCross::Construct( const FArguments& InArgs, const TSharedPtr<class
 
 
 
-int32 SDockingCross::OnPaint( const FGeometry& AllottedGeometry, const FSlateRect& MyClippingRect, FSlateWindowElementList& OutDrawElements, int32 LayerId, const FWidgetStyle& InWidgetStyle, bool bParentEnabled ) const
+int32 SDockingCross::OnPaint( const FPaintArgs& Args, const FGeometry& AllottedGeometry, const FSlateRect& MyClippingRect, FSlateWindowElementList& OutDrawElements, int32 LayerId, const FWidgetStyle& InWidgetStyle, bool bParentEnabled ) const
 {
 	using namespace DockingConstants;
 
diff --git a/Engine/Source/Runtime/Slate/Private/Framework/Docking/SDockingTabWell.cpp b/Engine/Source/Runtime/Slate/Private/Framework/Docking/SDockingTabWell.cpp
--- a/Engine/Source/Runtime/Slate/Private/Framework/Docking/SDockingTabWell.cpp
+++ b/Engine/Source/Runtime/Slate/Private/Framework/Docking/SDockingTabWell.cpp
@@ -108,7 +108,7 @@ void SDockingTabWell::OnArrangeChildren( const FGeometry& AllottedGeometry, FArr
 }
 	
 
-int32 SDockingTabWell::OnPaint( const FGeometry& AllottedGeometry, const FSlateRect& MyClippingRect, FSlateWindowElementList& OutDrawElements, int32 LayerId, const FWidgetStyle& InWidgetStyle, bool bParentEnabled ) const
+int32 SDockingTabWell::OnPaint( const FPaintArgs& Args, const FGeometry& AllottedGeometry, const FSlateRect& MyClippingRect, FSlateWindowElementList& OutDrawElements, int32 LayerId, const FWidgetStyle& InWidgetStyle, bool bParentEnabled ) const
 {
 	// When we are dragging a tab, it must be painted on top of the other tabs, so we cannot
 	// just reuse the Panel's default OnPaint.
@@ -136,7 +136,7 @@ int32 SDockingTabWell::OnPaint( const FGeometry& AllottedGeometry, const FSlateR
 		else
 		{
 			FSlateRect ChildClipRect = MyClippingRect.IntersectionWith( CurWidget.Geometry.GetClippingRect() );
-			const int32 CurWidgetsMaxLayerId = CurWidget.Widget->Paint( CurWidget.Geometry, ChildClipRect, OutDrawElements, MaxLayerId, InWidgetStyle, ShouldBeEnabled( bParentEnabled ) );
+			const int32 CurWidgetsMaxLayerId = CurWidget.Widget->Paint( Args.WithNewParent(this), CurWidget.Geometry, ChildClipRect, OutDrawElements, MaxLayerId, InWidgetStyle, ShouldBeEnabled( bParentEnabled ) );
 			MaxLayerId = FMath::Max( MaxLayerId, CurWidgetsMaxLayerId );
 		}
 	}
@@ -145,7 +145,7 @@ int32 SDockingTabWell::OnPaint( const FGeometry& AllottedGeometry, const FSlateR
 	if (ForegroundTab != TSharedPtr<SDockTab>())
 	{
 		FSlateRect ChildClipRect = MyClippingRect.IntersectionWith( ForegroundTabGeometry->Geometry.GetClippingRect() );
-		const int32 CurWidgetsMaxLayerId = ForegroundTabGeometry->Widget->Paint( ForegroundTabGeometry->Geometry, ChildClipRect, OutDrawElements, MaxLayerId, InWidgetStyle, ShouldBeEnabled( bParentEnabled ) );
+		const int32 CurWidgetsMaxLayerId = ForegroundTabGeometry->Widget->Paint( Args.WithNewParent(this), ForegroundTabGeometry->Geometry, ChildClipRect, OutDrawElements, MaxLayerId, InWidgetStyle, ShouldBeEnabled( bParentEnabled ) );
 		MaxLayerId = FMath::Max( MaxLayerId, CurWidgetsMaxLayerId );
 	}
 
diff --git a/Engine/Source/Runtime/Slate/Private/Framework/MultiBox/SClippingHorizontalBox.cpp b/Engine/Source/Runtime/Slate/Private/Framework/MultiBox/SClippingHorizontalBox.cpp
--- a/Engine/Source/Runtime/Slate/Private/Framework/MultiBox/SClippingHorizontalBox.cpp
+++ b/Engine/Source/Runtime/Slate/Private/Framework/MultiBox/SClippingHorizontalBox.cpp
@@ -55,7 +55,7 @@ void SClippingHorizontalBox::OnArrangeChildren( const FGeometry& AllottedGeometr
 	}
 }
 
-int32 SClippingHorizontalBox::OnPaint( const FGeometry& AllottedGeometry, const FSlateRect& MyClippingRect, FSlateWindowElementList& OutDrawElements, int32 LayerId, const FWidgetStyle& InWidgetStyle, bool bParentEnabled ) const
+int32 SClippingHorizontalBox::OnPaint( const FPaintArgs& Args, const FGeometry& AllottedGeometry, const FSlateRect& MyClippingRect, FSlateWindowElementList& OutDrawElements, int32 LayerId, const FWidgetStyle& InWidgetStyle, bool bParentEnabled ) const
 {
 	// Get the clipped children info
 	FArrangedChildren ClippedArrangedChildren(EVisibility::Visible);
@@ -107,7 +107,7 @@ int32 SClippingHorizontalBox::OnPaint( const FGeometry& AllottedGeometry, const
 			);
 	}
 
-	return SHorizontalBox::OnPaint(AllottedGeometry, MyClippingRect, OutDrawElements, LayerId, InWidgetStyle, bParentEnabled);
+	return SHorizontalBox::OnPaint(Args, AllottedGeometry, MyClippingRect, OutDrawElements, LayerId, InWidgetStyle, bParentEnabled);
 }
 
 FVector2D SClippingHorizontalBox::ComputeDesiredSize() const
diff --git a/Engine/Source/Runtime/Slate/Private/Framework/Text/RichTextMarkupProcessing.cpp b/Engine/Source/Runtime/Slate/Private/Framework/Text/RichTextMarkupProcessing.cpp
--- a/Engine/Source/Runtime/Slate/Private/Framework/Text/RichTextMarkupProcessing.cpp
+++ b/Engine/Source/Runtime/Slate/Private/Framework/Text/RichTextMarkupProcessing.cpp
@@ -1,6 +1,9 @@
 // Copyright 1998-2014 Epic Games, Inc. All Rights Reserved.
 
 #include "SlatePrivatePCH.h"
+
+#if WITH_FANCY_TEXT
+
 #include "RichTextMarkupProcessing.h"
 
 
@@ -290,4 +293,6 @@ void FRichTextMarkupProcessing::HandleEscapeSequences(const FString& Input, TArr
 		// Adjust end indices for previous substitutions.
 		LineParseResults.Range.EndIndex = ConcatenatedUnescapedLines.Len();
 	}
-}
\ No newline at end of file
+}
+
+#endif //WITH_FANCY_TEXT
\ No newline at end of file
diff --git a/Engine/Source/Runtime/Slate/Private/Framework/Text/SlateHyperlinkRun.cpp b/Engine/Source/Runtime/Slate/Private/Framework/Text/SlateHyperlinkRun.cpp
--- a/Engine/Source/Runtime/Slate/Private/Framework/Text/SlateHyperlinkRun.cpp
+++ b/Engine/Source/Runtime/Slate/Private/Framework/Text/SlateHyperlinkRun.cpp
@@ -1,6 +1,9 @@
 // Copyright 1998-2014 Epic Games, Inc. All Rights Reserved.
 
 #include "SlatePrivatePCH.h"
+
+#if WITH_FANCY_TEXT
+
 #include "SlateHyperlinkRun.h"
 
 TSharedRef< FSlateHyperlinkRun > FSlateHyperlinkRun::Create( const TSharedRef< const FString >& InText, const FHyperlinkStyle& InStyle, const TMap<FString,FString>& Metadata, FOnClick NavigateDelegate )
@@ -68,12 +71,12 @@ void FSlateHyperlinkRun::OnNavigate()
 	NavigateDelegate.Execute( Metadata );
 }
 
-int32 FSlateHyperlinkRun::OnPaint( const FTextLayout::FLineView& Line, const TSharedRef< ILayoutBlock >& Block, const FTextBlockStyle& DefaultStyle, const FGeometry& AllottedGeometry, const FSlateRect& MyClippingRect, FSlateWindowElementList& OutDrawElements, int32 LayerId, const FWidgetStyle& InWidgetStyle, bool bParentEnabled ) const 
+int32 FSlateHyperlinkRun::OnPaint( const FPaintArgs& Args, const FTextLayout::FLineView& Line, const TSharedRef< ILayoutBlock >& Block, const FTextBlockStyle& DefaultStyle, const FGeometry& AllottedGeometry, const FSlateRect& MyClippingRect, FSlateWindowElementList& OutDrawElements, int32 LayerId, const FWidgetStyle& InWidgetStyle, bool bParentEnabled ) const 
 {
 	const TSharedRef< FWidgetLayoutBlock > WidgetBlock = StaticCastSharedRef< FWidgetLayoutBlock >( Block );
 
 	FGeometry WidgetGeometry = AllottedGeometry.MakeChild( Block->GetLocationOffset() * ( 1 / AllottedGeometry.Scale ), Block->GetSize() * ( 1 / AllottedGeometry.Scale ), 1.0f );
-	return WidgetBlock->GetWidget()->Paint( WidgetGeometry, MyClippingRect, OutDrawElements, LayerId, InWidgetStyle, bParentEnabled );
+	return WidgetBlock->GetWidget()->Paint( Args, WidgetGeometry, MyClippingRect, OutDrawElements, LayerId, InWidgetStyle, bParentEnabled );
 }
 
 FChildren* FSlateHyperlinkRun::GetChildren()
@@ -157,3 +160,6 @@ FSlateHyperlinkRun::FSlateHyperlinkRun( const FSlateHyperlinkRun& Run )
 
 }
 
+
+
+#endif //WITH_FANCY_TEXT
\ No newline at end of file
diff --git a/Engine/Source/Runtime/Slate/Private/Framework/Text/SlateImageRun.cpp b/Engine/Source/Runtime/Slate/Private/Framework/Text/SlateImageRun.cpp
--- a/Engine/Source/Runtime/Slate/Private/Framework/Text/SlateImageRun.cpp
+++ b/Engine/Source/Runtime/Slate/Private/Framework/Text/SlateImageRun.cpp
@@ -1,6 +1,9 @@
 // Copyright 1998-2014 Epic Games, Inc. All Rights Reserved.
 
 #include "SlatePrivatePCH.h"
+
+#if WITH_FANCY_TEXT
+
 #include "SlateImageRun.h"
 
 FNoChildren FSlateImageRun::NoChildrenInstance;
@@ -63,7 +66,7 @@ FVector2D FSlateImageRun::GetLocationAt( const TSharedRef< ILayoutBlock >& Block
 	return Block->GetLocationOffset();
 }
 
-int32 FSlateImageRun::OnPaint( const FTextLayout::FLineView& Line, const TSharedRef< ILayoutBlock >& Block, const FTextBlockStyle& DefaultStyle, const FGeometry& AllottedGeometry, const FSlateRect& MyClippingRect, FSlateWindowElementList& OutDrawElements, int32 LayerId, const FWidgetStyle& InWidgetStyle, bool bParentEnabled ) const 
+int32 FSlateImageRun::OnPaint( const FPaintArgs& Args, const FTextLayout::FLineView& Line, const TSharedRef< ILayoutBlock >& Block, const FTextBlockStyle& DefaultStyle, const FGeometry& AllottedGeometry, const FSlateRect& MyClippingRect, FSlateWindowElementList& OutDrawElements, int32 LayerId, const FWidgetStyle& InWidgetStyle, bool bParentEnabled ) const 
 {
 	if ( Image->DrawAs != ESlateBrushDrawType::NoDrawType )
 	{
@@ -131,4 +134,6 @@ void FSlateImageRun::AppendText(FString& AppendToText, const FTextRange& Partial
 {
 	check(Range.BeginIndex <= PartialRange.BeginIndex);
 	check(Range.EndIndex >= PartialRange.EndIndex);
-}
\ No newline at end of file
+}
+
+#endif //WITH_FANCY_TEXT
\ No newline at end of file
diff --git a/Engine/Source/Runtime/Slate/Private/Framework/Text/SlateTextHighlightRunRenderer.cpp b/Engine/Source/Runtime/Slate/Private/Framework/Text/SlateTextHighlightRunRenderer.cpp
--- a/Engine/Source/Runtime/Slate/Private/Framework/Text/SlateTextHighlightRunRenderer.cpp
+++ b/Engine/Source/Runtime/Slate/Private/Framework/Text/SlateTextHighlightRunRenderer.cpp
@@ -1,6 +1,9 @@
 // Copyright 1998-2014 Epic Games, Inc. All Rights Reserved.
 
 #include "SlatePrivatePCH.h"
+
+#if WITH_FANCY_TEXT
+
 #include "SlateTextHighlightRunRenderer.h"
 
 FSlateTextHighlightRunRenderer::FSlateTextHighlightRunRenderer()
@@ -8,7 +11,7 @@ FSlateTextHighlightRunRenderer::FSlateTextHighlightRunRenderer()
 
 }
 
-int32 FSlateTextHighlightRunRenderer::OnPaint( const FTextLayout::FLineView& Line, const TSharedRef< ISlateRun >& Run, const TSharedRef< ILayoutBlock >& Block, const FTextBlockStyle& DefaultStyle, const FGeometry& AllottedGeometry, const FSlateRect& MyClippingRect, FSlateWindowElementList& OutDrawElements, int32 LayerId, const FWidgetStyle& InWidgetStyle, bool bParentEnabled ) const 
+int32 FSlateTextHighlightRunRenderer::OnPaint( const FPaintArgs& Args, const FTextLayout::FLineView& Line, const TSharedRef< ISlateRun >& Run, const TSharedRef< ILayoutBlock >& Block, const FTextBlockStyle& DefaultStyle, const FGeometry& AllottedGeometry, const FSlateRect& MyClippingRect, FSlateWindowElementList& OutDrawElements, int32 LayerId, const FWidgetStyle& InWidgetStyle, bool bParentEnabled ) const 
 {
 	FVector2D Location( Block->GetLocationOffset() );
 	Location.Y = Line.Offset.Y;
@@ -30,7 +33,7 @@ int32 FSlateTextHighlightRunRenderer::OnPaint( const FTextLayout::FLineView& Lin
 	FWidgetStyle WidgetStyle( InWidgetStyle );
 	WidgetStyle.SetForegroundColor( InvertedForeground );
 
-	return Run->OnPaint( Line, Block, DefaultStyle, AllottedGeometry, MyClippingRect, OutDrawElements, LayerId, WidgetStyle, bParentEnabled );
+	return Run->OnPaint( Args, Line, Block, DefaultStyle, AllottedGeometry, MyClippingRect, OutDrawElements, LayerId, WidgetStyle, bParentEnabled );
 }
 
 TSharedRef< FSlateTextHighlightRunRenderer > FSlateTextHighlightRunRenderer::Create()
@@ -38,3 +41,6 @@ TSharedRef< FSlateTextHighlightRunRenderer > FSlateTextHighlightRunRenderer::Cre
 	return MakeShareable( new FSlateTextHighlightRunRenderer() );
 }
 
+
+
+#endif //WITH_FANCY_TEXT
\ No newline at end of file
diff --git a/Engine/Source/Runtime/Slate/Private/Framework/Text/SlateTextLayout.cpp b/Engine/Source/Runtime/Slate/Private/Framework/Text/SlateTextLayout.cpp
--- a/Engine/Source/Runtime/Slate/Private/Framework/Text/SlateTextLayout.cpp
+++ b/Engine/Source/Runtime/Slate/Private/Framework/Text/SlateTextLayout.cpp
@@ -1,6 +1,9 @@
 // Copyright 1998-2014 Epic Games, Inc. All Rights Reserved.
 
 #include "SlatePrivatePCH.h"
+
+#if WITH_FANCY_TEXT
+
 #include "SlateTextLayout.h"
 
 TSharedRef< FSlateTextLayout > FSlateTextLayout::Create()
@@ -37,7 +40,7 @@ void FSlateTextLayout::ArrangeChildren( const FGeometry& AllottedGeometry, FArra
 	}
 }
 
-int32 FSlateTextLayout::OnPaint( const FTextBlockStyle& DefaultTextStyle, const FGeometry& AllottedGeometry, const FSlateRect& MyClippingRect, FSlateWindowElementList& OutDrawElements, int32 LayerId, const FWidgetStyle& InWidgetStyle, bool bParentEnabled ) const
+int32 FSlateTextLayout::OnPaint( const FPaintArgs& Args, const FTextBlockStyle& DefaultTextStyle, const FGeometry& AllottedGeometry, const FSlateRect& MyClippingRect, FSlateWindowElementList& OutDrawElements, int32 LayerId, const FWidgetStyle& InWidgetStyle, bool bParentEnabled ) const
 {
 	const FSlateRect ClippingRect = AllottedGeometry.GetClippingRect().IntersectionWith(MyClippingRect);
 	const ESlateDrawEffect::Type DrawEffects = bParentEnabled ? ESlateDrawEffect::None : ESlateDrawEffect::DisabledEffect;
@@ -49,7 +52,7 @@ int32 FSlateTextLayout::OnPaint( const FTextBlockStyle& DefaultTextStyle, const
 	for (const FTextLayout::FLineView& LineView : LineViews)
 	{
 		// Render any underlays for this line
-		const int32 HighestUnderlayLayerId = OnPaintHighlights( LineView, LineView.UnderlayHighlights, DefaultTextStyle, AllottedGeometry, ClippingRect, OutDrawElements, LayerId, InWidgetStyle, bParentEnabled );
+		const int32 HighestUnderlayLayerId = OnPaintHighlights( Args, LineView, LineView.UnderlayHighlights, DefaultTextStyle, AllottedGeometry, ClippingRect, OutDrawElements, LayerId, InWidgetStyle, bParentEnabled );
 
 		const int32 BlockDebugLayer = HighestUnderlayLayerId;
 		const int32 TextLayer = BlockDebugLayer + 1;
@@ -79,25 +82,25 @@ int32 FSlateTextLayout::OnPaint( const FTextBlockStyle& DefaultTextStyle, const
 			const TSharedPtr< ISlateRunRenderer > RunRenderer = StaticCastSharedPtr< ISlateRunRenderer >( Block->GetRenderer() );
 			if ( RunRenderer.IsValid() )
 			{
-				HighestRunLayerId = RunRenderer->OnPaint( LineView, Run, Block, DefaultTextStyle, AllottedGeometry, ClippingRect, OutDrawElements, TextLayer, InWidgetStyle, bParentEnabled );
+				HighestRunLayerId = RunRenderer->OnPaint( Args, LineView, Run, Block, DefaultTextStyle, AllottedGeometry, ClippingRect, OutDrawElements, TextLayer, InWidgetStyle, bParentEnabled );
 			}
 			else
 			{
-				HighestRunLayerId = Run->OnPaint( LineView, Block, DefaultTextStyle, AllottedGeometry, ClippingRect, OutDrawElements, TextLayer, InWidgetStyle, bParentEnabled );
+				HighestRunLayerId = Run->OnPaint( Args, LineView, Block, DefaultTextStyle, AllottedGeometry, ClippingRect, OutDrawElements, TextLayer, InWidgetStyle, bParentEnabled );
 			}
 
 			HighestBlockLayerId = FMath::Max( HighestBlockLayerId, HighestRunLayerId );
 		}
 
 		// Render any overlays for this line
-		const int32 HighestOverlayLayerId = OnPaintHighlights( LineView, LineView.OverlayHighlights, DefaultTextStyle, AllottedGeometry, ClippingRect, OutDrawElements, HighestBlockLayerId, InWidgetStyle, bParentEnabled );
+		const int32 HighestOverlayLayerId = OnPaintHighlights( Args, LineView, LineView.OverlayHighlights, DefaultTextStyle, AllottedGeometry, ClippingRect, OutDrawElements, HighestBlockLayerId, InWidgetStyle, bParentEnabled );
 		HighestLayerId = FMath::Max( HighestLayerId, HighestOverlayLayerId );
 	}
 
 	return HighestLayerId;
 }
 
-int32 FSlateTextLayout::OnPaintHighlights( const FTextLayout::FLineView& LineView, const TArray<FLineViewHighlight>& Highlights, const FTextBlockStyle& DefaultTextStyle, const FGeometry& AllottedGeometry, const FSlateRect& ClippingRect, FSlateWindowElementList& OutDrawElements, const int32 LayerId, const FWidgetStyle& InWidgetStyle, bool bParentEnabled ) const
+int32 FSlateTextLayout::OnPaintHighlights( const FPaintArgs& Args, const FTextLayout::FLineView& LineView, const TArray<FLineViewHighlight>& Highlights, const FTextBlockStyle& DefaultTextStyle, const FGeometry& AllottedGeometry, const FSlateRect& ClippingRect, FSlateWindowElementList& OutDrawElements, const int32 LayerId, const FWidgetStyle& InWidgetStyle, bool bParentEnabled ) const
 {
 	int32 CurrentLayerId = LayerId;
 
@@ -106,7 +109,7 @@ int32 FSlateTextLayout::OnPaintHighlights( const FTextLayout::FLineView& LineVie
 		const TSharedPtr< ISlateLineHighlighter > LineHighlighter = StaticCastSharedPtr< ISlateLineHighlighter >( Highlight.Highlighter );
 		if (LineHighlighter.IsValid())
 		{
-			CurrentLayerId = LineHighlighter->OnPaint( LineView, Highlight.OffsetX, Highlight.Width, DefaultTextStyle, AllottedGeometry, ClippingRect, OutDrawElements, CurrentLayerId, InWidgetStyle, bParentEnabled );
+			CurrentLayerId = LineHighlighter->OnPaint( Args, LineView, Highlight.OffsetX, Highlight.Width, DefaultTextStyle, AllottedGeometry, ClippingRect, OutDrawElements, CurrentLayerId, InWidgetStyle, bParentEnabled );
 		}
 	}
 
@@ -140,3 +143,6 @@ void FSlateTextLayout::AggregateChildren()
 		}
 	}
 }
+
+
+#endif //WITH_FANCY_TEXT
\ No newline at end of file
diff --git a/Engine/Source/Runtime/Slate/Private/Framework/Text/SlateTextRun.cpp b/Engine/Source/Runtime/Slate/Private/Framework/Text/SlateTextRun.cpp
--- a/Engine/Source/Runtime/Slate/Private/Framework/Text/SlateTextRun.cpp
+++ b/Engine/Source/Runtime/Slate/Private/Framework/Text/SlateTextRun.cpp
@@ -1,6 +1,9 @@
 // Copyright 1998-2014 Epic Games, Inc. All Rights Reserved.
 
 #include "SlatePrivatePCH.h"
+
+#if WITH_FANCY_TEXT
+
 #include "SlateTextRun.h"
 
 FNoChildren FSlateTextRun::NoChildrenInstance;
@@ -65,7 +68,7 @@ TSharedRef< ILayoutBlock > FSlateTextRun::CreateBlock( int32 BeginIndex, int32 E
 	return FDefaultLayoutBlock::Create( SharedThis( this ), FTextRange( BeginIndex, EndIndex ), Size, Renderer );
 }
 
-int32 FSlateTextRun::OnPaint( const FTextLayout::FLineView& Line, const TSharedRef< ILayoutBlock >& Block, const FTextBlockStyle& DefaultStyle, const FGeometry& AllottedGeometry, const FSlateRect& MyClippingRect, FSlateWindowElementList& OutDrawElements, int32 LayerId, const FWidgetStyle& InWidgetStyle, bool bParentEnabled ) const 
+int32 FSlateTextRun::OnPaint( const FPaintArgs& Args, const FTextLayout::FLineView& Line, const TSharedRef< ILayoutBlock >& Block, const FTextBlockStyle& DefaultStyle, const FGeometry& AllottedGeometry, const FSlateRect& MyClippingRect, FSlateWindowElementList& OutDrawElements, int32 LayerId, const FWidgetStyle& InWidgetStyle, bool bParentEnabled ) const 
 {
 	const FSlateRect ClippingRect = AllottedGeometry.GetClippingRect().IntersectionWith(MyClippingRect);
 	const ESlateDrawEffect::Type DrawEffects = bParentEnabled ? ESlateDrawEffect::None : ESlateDrawEffect::DisabledEffect;
@@ -224,3 +227,4 @@ FSlateTextRun::FSlateTextRun( const FSlateTextRun& Run )
 
 }
 
+#endif //WITH_FANCY_TEXT
\ No newline at end of file
diff --git a/Engine/Source/Runtime/Slate/Private/Framework/Text/SlateWidgetRun.cpp b/Engine/Source/Runtime/Slate/Private/Framework/Text/SlateWidgetRun.cpp
--- a/Engine/Source/Runtime/Slate/Private/Framework/Text/SlateWidgetRun.cpp
+++ b/Engine/Source/Runtime/Slate/Private/Framework/Text/SlateWidgetRun.cpp
@@ -1,6 +1,9 @@
 // Copyright 1998-2014 Epic Games, Inc. All Rights Reserved.
 
 #include "SlatePrivatePCH.h"
+
+#if WITH_FANCY_TEXT
+
 #include "SlateWidgetRun.h"
 
 TSharedRef< FSlateWidgetRun > FSlateWidgetRun::Create( const TSharedRef< const FString >& InText, const FWidgetRunInfo& InWidgetInfo )
@@ -48,10 +51,10 @@ TSharedRef< ILayoutBlock > FSlateWidgetRun::CreateBlock( int32 StartIndex, int32
 	return FDefaultLayoutBlock::Create( SharedThis( this ), FTextRange( StartIndex, EndIndex ), Size, Renderer );
 }
 
-int32 FSlateWidgetRun::OnPaint( const FTextLayout::FLineView& Line, const TSharedRef< ILayoutBlock >& Block, const FTextBlockStyle& DefaultStyle, const FGeometry& AllottedGeometry, const FSlateRect& MyClippingRect, FSlateWindowElementList& OutDrawElements, int32 LayerId, const FWidgetStyle& InWidgetStyle, bool bParentEnabled ) const 
+int32 FSlateWidgetRun::OnPaint( const FPaintArgs& Args, const FTextLayout::FLineView& Line, const TSharedRef< ILayoutBlock >& Block, const FTextBlockStyle& DefaultStyle, const FGeometry& AllottedGeometry, const FSlateRect& MyClippingRect, FSlateWindowElementList& OutDrawElements, int32 LayerId, const FWidgetStyle& InWidgetStyle, bool bParentEnabled ) const 
 {
 	FGeometry WidgetGeometry = AllottedGeometry.MakeChild( Block->GetLocationOffset() * ( 1 / AllottedGeometry.Scale ), Block->GetSize() * ( 1 / AllottedGeometry.Scale ), 1.0f );
-	return Info.Widget->Paint( WidgetGeometry, MyClippingRect, OutDrawElements, LayerId, InWidgetStyle, bParentEnabled );
+	return Info.Widget->Paint( Args, WidgetGeometry, MyClippingRect, OutDrawElements, LayerId, InWidgetStyle, bParentEnabled );
 }
 
 FChildren* FSlateWidgetRun::GetChildren()
@@ -125,3 +128,5 @@ FSlateWidgetRun::FSlateWidgetRun( const FSlateWidgetRun& Run )
 
 }
 
+
+#endif //WITH_FANCY_TEXT
diff --git a/Engine/Source/Runtime/Slate/Private/Framework/Text/Tests/RichTextMarkupProcessingTest.cpp b/Engine/Source/Runtime/Slate/Private/Framework/Text/Tests/RichTextMarkupProcessingTest.cpp
--- a/Engine/Source/Runtime/Slate/Private/Framework/Text/Tests/RichTextMarkupProcessingTest.cpp
+++ b/Engine/Source/Runtime/Slate/Private/Framework/Text/Tests/RichTextMarkupProcessingTest.cpp
@@ -1,6 +1,9 @@
 // Copyright 1998-2014 Epic Games, Inc. All Rights Reserved.
 
 #include "SlatePrivatePCH.h"
+
+#if WITH_FANCY_TEXT
+
 #include "AutomationTest.h"
 
 #include "../Text/RichTextMarkupProcessing.h"
@@ -380,4 +383,6 @@ bool FRTFProcessingTest::RunTest (const FString& Parameters)
 
 #undef LOCTEXT_NAMESPACE 
 
-PRAGMA_ENABLE_OPTIMIZATION
\ No newline at end of file
+PRAGMA_ENABLE_OPTIMIZATION
+
+#endif //WITH_FANCY_TEXT
\ No newline at end of file
diff --git a/Engine/Source/Runtime/Slate/Private/Framework/Text/TextDecorators.cpp b/Engine/Source/Runtime/Slate/Private/Framework/Text/TextDecorators.cpp
--- a/Engine/Source/Runtime/Slate/Private/Framework/Text/TextDecorators.cpp
+++ b/Engine/Source/Runtime/Slate/Private/Framework/Text/TextDecorators.cpp
@@ -1,6 +1,9 @@
 // Copyright 1998-2014 Epic Games, Inc. All Rights Reserved.
 
 #include "SlatePrivatePCH.h"
+
+#if WITH_FANCY_TEXT
+
 #include "TextDecorators.h"
 
 TSharedRef< FWidgetDecorator > FWidgetDecorator::Create( FString InRunName, const FCreateWidget& InCreateWidgetDelegate )
@@ -137,4 +140,5 @@ TSharedRef< ISlateRun > FHyperlinkDecorator::Create( const FTextRunParseResults&
 	}
 
 	return FSlateHyperlinkRun::Create( InOutModelText, Style->GetWidgetStyle<FHyperlinkStyle>( FName( *StyleName ) ), MetaData, NavigateDelegate, ModelRange );
-}
\ No newline at end of file
+}
+#endif //WITH_FANCY_TEXT
\ No newline at end of file
diff --git a/Engine/Source/Runtime/Slate/Private/Widgets/Colors/SColorBlock.cpp b/Engine/Source/Runtime/Slate/Private/Widgets/Colors/SColorBlock.cpp
--- a/Engine/Source/Runtime/Slate/Private/Widgets/Colors/SColorBlock.cpp
+++ b/Engine/Source/Runtime/Slate/Private/Widgets/Colors/SColorBlock.cpp
@@ -19,7 +19,7 @@ void SColorBlock::Construct( const FArguments& InArgs )
 	ColorBlockSize = InArgs._Size;
 }
 
-int32 SColorBlock::OnPaint( const FGeometry& AllottedGeometry, const FSlateRect& MyClippingRect, FSlateWindowElementList& OutDrawElements, int32 LayerId, const FWidgetStyle& InWidgetStyle, bool bParentEnabled ) const
+int32 SColorBlock::OnPaint( const FPaintArgs& Args, const FGeometry& AllottedGeometry, const FSlateRect& MyClippingRect, FSlateWindowElementList& OutDrawElements, int32 LayerId, const FWidgetStyle& InWidgetStyle, bool bParentEnabled ) const
 {
 	const FSlateBrush* GenericBrush = FCoreStyle::Get().GetBrush( "GenericWhiteBox" );
 
diff --git a/Engine/Source/Runtime/Slate/Private/Widgets/Colors/SColorPicker.cpp b/Engine/Source/Runtime/Slate/Private/Widgets/Colors/SColorPicker.cpp
--- a/Engine/Source/Runtime/Slate/Private/Widgets/Colors/SColorPicker.cpp
+++ b/Engine/Source/Runtime/Slate/Private/Widgets/Colors/SColorPicker.cpp
@@ -45,20 +45,7 @@ public:
 	}
 
 private:
-	/**
-	 * The widget should respond by populating the OutDrawElements array with FDrawElements 
-	 * that represent it and any of its children.
-	 *
-	 * @param AllottedGeometry  The FGeometry that describes an area in which the widget should appear.
-	 * @param MyClippingRect    The clipping rectangle allocated for this widget and its children.
-	 * @param OutDrawElements   A list of FDrawElements to populate with the output.
-	 * @param LayerId           The Layer onto which this widget should be rendered.
-	 * @param InColorAndOpacity Color and Opacity to be applied to all the descendants of the widget being painted
- 	 * @param bParentEnabled	True if the parent of this widget is enabled.
-	 *
-	 * @return The maximum layer ID attained by this widget or any of its children.
-	 */
-	virtual int32 OnPaint( const FGeometry& AllottedGeometry, const FSlateRect& MyClippingRect, FSlateWindowElementList& OutDrawElements, int32 LayerId, const FWidgetStyle& InWidgetStyle, bool bParentEnabled ) const
+	virtual int32 OnPaint( const FPaintArgs& Args, const FGeometry& AllottedGeometry, const FSlateRect& MyClippingRect, FSlateWindowElementList& OutDrawElements, int32 LayerId, const FWidgetStyle& InWidgetStyle, bool bParentEnabled ) const
 	{
 		const ESlateDrawEffect::Type DrawEffects = this->ShouldBeEnabled(bParentEnabled) ? ESlateDrawEffect::None : ESlateDrawEffect::DisabledEffect;
 
@@ -145,20 +132,7 @@ public:
 	}
 
 private:
-	/**
-	 * The widget should respond by populating the OutDrawElements array with FDrawElements 
-	 * that represent it and any of its children.
-	 *
-	 * @param AllottedGeometry  The FGeometry that describes an area in which the widget should appear.
-	 * @param MyClippingRect    The clipping rectangle allocated for this widget and its children.
-	 * @param OutDrawElements   A list of FDrawElements to populate with the output.
-	 * @param LayerId           The Layer onto which this widget should be rendered.
-	 * @param InColorAndOpacity Color and Opacity to be applied to all the descendants of the widget being painted
- 	 * @param bParentEnabled	True if the parent of this widget is enabled.
-	 *
-	 * @return The maximum layer ID attained by this widget or any of its children.
-	 */
-	virtual int32 OnPaint( const FGeometry& AllottedGeometry, const FSlateRect& MyClippingRect, FSlateWindowElementList& OutDrawElements, int32 LayerId, const FWidgetStyle& InWidgetStyle, bool bParentEnabled ) const
+	virtual int32 OnPaint( const FPaintArgs& Args, const FGeometry& AllottedGeometry, const FSlateRect& MyClippingRect, FSlateWindowElementList& OutDrawElements, int32 LayerId, const FWidgetStyle& InWidgetStyle, bool bParentEnabled ) const
 	{
 		ESlateDrawEffect::Type DrawEffects = (bParentEnabled && IsEnabled()) ? ESlateDrawEffect::None : ESlateDrawEffect::DisabledEffect;
 
@@ -245,7 +219,7 @@ public:
 		SelectedColor = InArgs._SelectedColor;
 	}
 
-	virtual int32 OnPaint( const FGeometry& AllottedGeometry, const FSlateRect& MyClippingRect, FSlateWindowElementList& OutDrawElements, int32 LayerId, const FWidgetStyle& InWidgetStyle, bool bParentEnabled ) const
+	virtual int32 OnPaint( const FPaintArgs& Args, const FGeometry& AllottedGeometry, const FSlateRect& MyClippingRect, FSlateWindowElementList& OutDrawElements, int32 LayerId, const FWidgetStyle& InWidgetStyle, bool bParentEnabled ) const
 	{
 		const bool bIsEnabled = ShouldBeEnabled(bParentEnabled);
 		const uint32 DrawEffects = bIsEnabled ? ESlateDrawEffect::None : ESlateDrawEffect::DisabledEffect;
diff --git a/Engine/Source/Runtime/Slate/Private/Widgets/Colors/SColorSpectrum.cpp b/Engine/Source/Runtime/Slate/Private/Widgets/Colors/SColorSpectrum.cpp
--- a/Engine/Source/Runtime/Slate/Private/Widgets/Colors/SColorSpectrum.cpp
+++ b/Engine/Source/Runtime/Slate/Private/Widgets/Colors/SColorSpectrum.cpp
@@ -87,7 +87,7 @@ FReply SColorSpectrum::OnMouseMove( const FGeometry& MyGeometry, const FPointerE
 }
 
 
-int32 SColorSpectrum::OnPaint( const FGeometry& AllottedGeometry, const FSlateRect& MyClippingRect, FSlateWindowElementList& OutDrawElements, int32 LayerId, const FWidgetStyle& InWidgetStyle, bool bParentEnabled ) const
+int32 SColorSpectrum::OnPaint( const FPaintArgs& Args, const FGeometry& AllottedGeometry, const FSlateRect& MyClippingRect, FSlateWindowElementList& OutDrawElements, int32 LayerId, const FWidgetStyle& InWidgetStyle, bool bParentEnabled ) const
 {
 	const bool bIsEnabled = ShouldBeEnabled(bParentEnabled);
 	const uint32 DrawEffects = bIsEnabled ? ESlateDrawEffect::None : ESlateDrawEffect::DisabledEffect;
diff --git a/Engine/Source/Runtime/Slate/Private/Widgets/Colors/SColorWheel.cpp b/Engine/Source/Runtime/Slate/Private/Widgets/Colors/SColorWheel.cpp
--- a/Engine/Source/Runtime/Slate/Private/Widgets/Colors/SColorWheel.cpp
+++ b/Engine/Source/Runtime/Slate/Private/Widgets/Colors/SColorWheel.cpp
@@ -97,7 +97,7 @@ FReply SColorWheel::OnMouseMove( const FGeometry& MyGeometry, const FPointerEven
 }
 
 
-int32 SColorWheel::OnPaint( const FGeometry& AllottedGeometry, const FSlateRect& MyClippingRect, FSlateWindowElementList& OutDrawElements, int32 LayerId, const FWidgetStyle& InWidgetStyle, bool bParentEnabled ) const
+int32 SColorWheel::OnPaint( const FPaintArgs& Args, const FGeometry& AllottedGeometry, const FSlateRect& MyClippingRect, FSlateWindowElementList& OutDrawElements, int32 LayerId, const FWidgetStyle& InWidgetStyle, bool bParentEnabled ) const
 {
 	const bool bIsEnabled = ShouldBeEnabled(bParentEnabled);
 	const uint32 DrawEffects = bIsEnabled ? ESlateDrawEffect::None : ESlateDrawEffect::DisabledEffect;
diff --git a/Engine/Source/Runtime/Slate/Private/Widgets/Images/SImage.cpp b/Engine/Source/Runtime/Slate/Private/Widgets/Images/SImage.cpp
--- a/Engine/Source/Runtime/Slate/Private/Widgets/Images/SImage.cpp
+++ b/Engine/Source/Runtime/Slate/Private/Widgets/Images/SImage.cpp
@@ -16,20 +16,7 @@ void SImage::Construct( const FArguments& InArgs )
 }
 
 
-/**
- * The widget should respond by populating the OutDrawElements array with FDrawElements 
- * that represent it and any of its children.
- *
- * @param AllottedGeometry  The FGeometry that describes an area in which the widget should appear.
- * @param MyClippingRect    The clipping rectangle allocated for this widget and its children.
- * @param OutDrawElements   A list of FDrawElements to populate with the output.
- * @param LayerId           The Layer onto which this widget should be rendered.
- * @param InColorAndOpacity Color and Opacity to be applied to all the descendants of the widget being painted
- * @param bParentEnabled	True if the parent of this widget is enabled.
- *
- * @return The maximum layer ID attained by this widget or any of its children.
- */
-int32 SImage::OnPaint( const FGeometry& AllottedGeometry, const FSlateRect& MyClippingRect, FSlateWindowElementList& OutDrawElements, int32 LayerId, const FWidgetStyle& InWidgetStyle, bool bParentEnabled ) const
+int32 SImage::OnPaint( const FPaintArgs& Args, const FGeometry& AllottedGeometry, const FSlateRect& MyClippingRect, FSlateWindowElementList& OutDrawElements, int32 LayerId, const FWidgetStyle& InWidgetStyle, bool bParentEnabled ) const
 {
 #if SLATE_HD_STATS
 	SCOPE_CYCLE_COUNTER( STAT_SlateOnPaint_SImage );
diff --git a/Engine/Source/Runtime/Slate/Private/Widgets/Images/SSpinningImage.cpp b/Engine/Source/Runtime/Slate/Private/Widgets/Images/SSpinningImage.cpp
--- a/Engine/Source/Runtime/Slate/Private/Widgets/Images/SSpinningImage.cpp
+++ b/Engine/Source/Runtime/Slate/Private/Widgets/Images/SSpinningImage.cpp
@@ -14,7 +14,8 @@ void SSpinningImage::Construct(const FArguments& InArgs)
 	Sequence.Play();
 }
 
-int32 SSpinningImage::OnPaint( const FGeometry& AllottedGeometry, const FSlateRect& MyClippingRect, FSlateWindowElementList& OutDrawElements, int32 LayerId, const FWidgetStyle& InWidgetStyle, bool bParentEnabled ) const
+// Override SImage's OnPaint to draw rotated
+int32 SSpinningImage::OnPaint( const FPaintArgs& Args, const FGeometry& AllottedGeometry, const FSlateRect& MyClippingRect, FSlateWindowElementList& OutDrawElements, int32 LayerId, const FWidgetStyle& InWidgetStyle, bool bParentEnabled ) const
 {
 	const FSlateBrush* ImageBrush = Image.Get();
 
diff --git a/Engine/Source/Runtime/Slate/Private/Widgets/Images/SThrobber.cpp b/Engine/Source/Runtime/Slate/Private/Widgets/Images/SThrobber.cpp
--- a/Engine/Source/Runtime/Slate/Private/Widgets/Images/SThrobber.cpp
+++ b/Engine/Source/Runtime/Slate/Private/Widgets/Images/SThrobber.cpp
@@ -30,19 +30,19 @@ void SThrobber::ConstructPieces()
 	for (int32 PieceIndex = 0; PieceIndex < NumPieces; ++PieceIndex)
 	{
 		HBox->AddSlot()
-			.AutoWidth()
+		.AutoWidth()
+		[
+			SNew(SBorder)
+			.BorderImage(FStyleDefaults::GetNoBrush())
+			.ContentScale(this, &SThrobber::GetPieceScale, PieceIndex)
+			.ColorAndOpacity(this, &SThrobber::GetPieceColor, PieceIndex)
+			.HAlign(HAlign_Center)
+			.VAlign(VAlign_Center)
 			[
-				SNew(SBorder)
-				.BorderImage(FStyleDefaults::GetNoBrush())
-				.ContentScale(this, &SThrobber::GetPieceScale, PieceIndex)
-				.ColorAndOpacity(this, &SThrobber::GetPieceColor, PieceIndex)
-				.HAlign(HAlign_Center)
-				.VAlign(VAlign_Center)
-				[
-					SNew(SImage)
-					.Image(this, &SThrobber::GetPieceBrush)
-				]
-			];
+				SNew(SImage)
+				.Image(this, &SThrobber::GetPieceBrush)
+			]
+		];
 	}
 }
 
@@ -133,7 +133,7 @@ void SCircularThrobber::ConstructSequence()
 	Sequence.Play();
 }
 
-int32 SCircularThrobber::OnPaint( const FGeometry& AllottedGeometry, const FSlateRect& MyClippingRect, FSlateWindowElementList& OutDrawElements, int32 LayerId, const FWidgetStyle& InWidgetStyle, bool bParentEnabled ) const
+int32 SCircularThrobber::OnPaint( const FPaintArgs& Args, const FGeometry& AllottedGeometry, const FSlateRect& MyClippingRect, FSlateWindowElementList& OutDrawElements, int32 LayerId, const FWidgetStyle& InWidgetStyle, bool bParentEnabled ) const
 {
 	const FColor FinalColorAndOpacity( InWidgetStyle.GetColorAndOpacityTint() * PieceImage->GetTint( InWidgetStyle ) );
 	
diff --git a/Engine/Source/Runtime/Slate/Private/Widgets/Input/SEditableText.cpp b/Engine/Source/Runtime/Slate/Private/Widgets/Input/SEditableText.cpp
--- a/Engine/Source/Runtime/Slate/Private/Widgets/Input/SEditableText.cpp
+++ b/Engine/Source/Runtime/Slate/Private/Widgets/Input/SEditableText.cpp
@@ -1210,7 +1210,7 @@ FVector2D SEditableText::ComputeDesiredSize() const
 }
 
 
-int32 SEditableText::OnPaint( const FGeometry& AllottedGeometry, const FSlateRect& MyClippingRect, FSlateWindowElementList& OutDrawElements, int32 LayerId, const FWidgetStyle& InWidgetStyle, bool bParentEnabled ) const
+int32 SEditableText::OnPaint( const FPaintArgs& Args, const FGeometry& AllottedGeometry, const FSlateRect& MyClippingRect, FSlateWindowElementList& OutDrawElements, int32 LayerId, const FWidgetStyle& InWidgetStyle, bool bParentEnabled ) const
 {
 #if SLATE_HD_STATS
 	SCOPE_CYCLE_COUNTER( STAT_SlateOnPaint_SEditableText );
diff --git a/Engine/Source/Runtime/Slate/Private/Widgets/Input/SMultiLineEditableTextBox.cpp b/Engine/Source/Runtime/Slate/Private/Widgets/Input/SMultiLineEditableTextBox.cpp
--- a/Engine/Source/Runtime/Slate/Private/Widgets/Input/SMultiLineEditableTextBox.cpp
+++ b/Engine/Source/Runtime/Slate/Private/Widgets/Input/SMultiLineEditableTextBox.cpp
@@ -2,6 +2,8 @@
 
 #include "SlatePrivatePCH.h"
 
+#if WITH_FANCY_TEXT
+
 
 /**
  * Construct this widget
@@ -148,3 +150,6 @@ FReply SMultiLineEditableTextBox::OnKeyboardFocusReceived( const FGeometry& MyGe
 
 	return Reply;
 }
+
+
+#endif //WITH_FANCY_TEXT
\ No newline at end of file
diff --git a/Engine/Source/Runtime/Slate/Private/Widgets/Input/SSlider.cpp b/Engine/Source/Runtime/Slate/Private/Widgets/Input/SSlider.cpp
--- a/Engine/Source/Runtime/Slate/Private/Widgets/Input/SSlider.cpp
+++ b/Engine/Source/Runtime/Slate/Private/Widgets/Input/SSlider.cpp
@@ -22,7 +22,7 @@ void SSlider::Construct( const SSlider::FArguments& InDeclaration )
 }
 
 
-int32 SSlider::OnPaint( const FGeometry& AllottedGeometry, const FSlateRect& MyClippingRect, FSlateWindowElementList& OutDrawElements, int32 LayerId, const FWidgetStyle& InWidgetStyle, bool bParentEnabled ) const
+int32 SSlider::OnPaint( const FPaintArgs& Args, const FGeometry& AllottedGeometry, const FSlateRect& MyClippingRect, FSlateWindowElementList& OutDrawElements, int32 LayerId, const FWidgetStyle& InWidgetStyle, bool bParentEnabled ) const
 {
 #if SLATE_HD_STATS
 	SCOPE_CYCLE_COUNTER(STAT_SlateOnPaint_SSlider);
diff --git a/Engine/Source/Runtime/Slate/Private/Widgets/Input/SVirtualJoystick.cpp b/Engine/Source/Runtime/Slate/Private/Widgets/Input/SVirtualJoystick.cpp
--- a/Engine/Source/Runtime/Slate/Private/Widgets/Input/SVirtualJoystick.cpp
+++ b/Engine/Source/Runtime/Slate/Private/Widgets/Input/SVirtualJoystick.cpp
@@ -99,7 +99,7 @@ static bool PositionIsInside(const SVirtualJoystick::FControlInfo& Control, cons
 		Position.Y <= Control.Center.Y + BoxSize.Y * 0.5f;
 }
 
-int32 SVirtualJoystick::OnPaint( const FGeometry& AllottedGeometry, const FSlateRect& MyClippingRect, FSlateWindowElementList& OutDrawElements, int32 LayerId, const FWidgetStyle& InWidgetStyle, bool bParentEnabled ) const
+int32 SVirtualJoystick::OnPaint( const FPaintArgs& Args, const FGeometry& AllottedGeometry, const FSlateRect& MyClippingRect, FSlateWindowElementList& OutDrawElements, int32 LayerId, const FWidgetStyle& InWidgetStyle, bool bParentEnabled ) const
 {
 	int32 RetLayerId = LayerId;
 
diff --git a/Engine/Source/Runtime/Slate/Private/Widgets/Input/SVirtualKeyboardEntry.cpp b/Engine/Source/Runtime/Slate/Private/Widgets/Input/SVirtualKeyboardEntry.cpp
--- a/Engine/Source/Runtime/Slate/Private/Widgets/Input/SVirtualKeyboardEntry.cpp
+++ b/Engine/Source/Runtime/Slate/Private/Widgets/Input/SVirtualKeyboardEntry.cpp
@@ -149,20 +149,7 @@ FVector2D SVirtualKeyboardEntry::ComputeDesiredSize() const
 	return FVector2D( FMath::Max(TextSize.X, MinDesiredWidth.Get()), FMath::Max(TextSize.Y, FontMaxCharHeight) );
 }
 
-/**
- * The widget should respond by populating the OutDrawElements array with FDrawElements 
- * that represent it and any of its children.
- *
- * @param AllottedGeometry  The FGeometry that describes an area in which the widget should appear.
- * @param MyClippingRect    The clipping rectangle allocated for this widget and its children.
- * @param OutDrawElements   A list of FDrawElements to populate with the output.
- * @param LayerId           The Layer onto which this widget should be rendered.
- * @param InColorAndOpacity Color and Opacity to be applied to all the descendants of the widget being painted
- * @param bParentEnabled	True if the parent of this widget is enabled.
- *
- * @return The maximum layer ID attained by this widget or any of its children.
- */
-int32 SVirtualKeyboardEntry::OnPaint( const FGeometry& AllottedGeometry, const FSlateRect& MyClippingRect, FSlateWindowElementList& OutDrawElements, int32 LayerId, const FWidgetStyle& InWidgetStyle, bool bParentEnabled ) const
+int32 SVirtualKeyboardEntry::OnPaint( const FPaintArgs& Args, const FGeometry& AllottedGeometry, const FSlateRect& MyClippingRect, FSlateWindowElementList& OutDrawElements, int32 LayerId, const FWidgetStyle& InWidgetStyle, bool bParentEnabled ) const
 {
 	const TSharedRef< FSlateFontMeasure > FontMeasureService = FSlateApplication::Get().GetRenderer()->GetFontMeasureService();
 
diff --git a/Engine/Source/Runtime/Slate/Private/Widgets/Layout/SBorder.cpp b/Engine/Source/Runtime/Slate/Private/Widgets/Layout/SBorder.cpp
--- a/Engine/Source/Runtime/Slate/Private/Widgets/Layout/SBorder.cpp
+++ b/Engine/Source/Runtime/Slate/Private/Widgets/Layout/SBorder.cpp
@@ -65,20 +65,7 @@ void SBorder::ClearContent()
 	ChildSlot = FSimpleSlot();
 }
 
-/**
- * The widget should respond by populating the OutDrawElements array with FDrawElements 
- * that represent it and any of its children.
- *
- * @param AllottedGeometry  The FGeometry that describes an area in which the widget should appear.
- * @param MyClippingRect    The clipping rectangle allocated for this widget and its children.
- * @param OutDrawElements   A list of FDrawElements to populate with the output.
- * @param LayerId           The Layer onto which this widget should be rendered.
- * @param InColorAndOpacity Color and Opacity to be applied to all the descendants of the widget being painted
- * @param bParentEnabled	True if the parent of this widget is enabled.
- *
- * @return The maximum layer ID attained by this widget or any of its children.
- */
-int32 SBorder::OnPaint( const FGeometry& AllottedGeometry, const FSlateRect& MyClippingRect, FSlateWindowElementList& OutDrawElements, int32 LayerId, const FWidgetStyle& InWidgetStyle, bool bParentEnabled ) const
+int32 SBorder::OnPaint( const FPaintArgs& Args, const FGeometry& AllottedGeometry, const FSlateRect& MyClippingRect, FSlateWindowElementList& OutDrawElements, int32 LayerId, const FWidgetStyle& InWidgetStyle, bool bParentEnabled ) const
 {
 	bool bEnabled = false;
 	{
@@ -109,7 +96,7 @@ int32 SBorder::OnPaint( const FGeometry& AllottedGeometry, const FSlateRect& MyC
 		.BlendColorAndOpacityTint(ColorAndOpacity.Get())
 		.SetForegroundColor( ForegroundColor.Get() );
 
-	return SCompoundWidget::OnPaint(AllottedGeometry, MyClippingRect.IntersectionWith( AllottedGeometry.GetClippingRect() ), OutDrawElements, LayerId, CompoundedWidgetStyle, bEnabled );
+	return SCompoundWidget::OnPaint(Args, AllottedGeometry, MyClippingRect.IntersectionWith( AllottedGeometry.GetClippingRect() ), OutDrawElements, LayerId, CompoundedWidgetStyle, bEnabled );
 }
 
 
diff --git a/Engine/Source/Runtime/Slate/Private/Widgets/Layout/SBox.cpp b/Engine/Source/Runtime/Slate/Private/Widgets/Layout/SBox.cpp
--- a/Engine/Source/Runtime/Slate/Private/Widgets/Layout/SBox.cpp
+++ b/Engine/Source/Runtime/Slate/Private/Widgets/Layout/SBox.cpp
@@ -62,7 +62,7 @@ FChildren* SBox::GetChildren()
 	return &ChildSlot;
 }
 
-int32 SBox::OnPaint( const FGeometry& AllottedGeometry, const FSlateRect& MyClippingRect, FSlateWindowElementList& OutDrawElements, int32 LayerId, const FWidgetStyle& InWidgetStyle, bool bParentEnabled ) const
+int32 SBox::OnPaint( const FPaintArgs& Args, const FGeometry& AllottedGeometry, const FSlateRect& MyClippingRect, FSlateWindowElementList& OutDrawElements, int32 LayerId, const FWidgetStyle& InWidgetStyle, bool bParentEnabled ) const
 {
 	// An SBox just draws its only child
 	FArrangedChildren ArrangedChildren(EVisibility::Visible);
@@ -81,7 +81,7 @@ int32 SBox::OnPaint( const FGeometry& AllottedGeometry, const FSlateRect& MyClip
 
 		const FSlateRect ChildClippingRect = AllottedGeometry.GetClippingRect().InsetBy( ChildSlot.SlotPadding.Get() * AllottedGeometry.Scale ).IntersectionWith(MyClippingRect);
 
-		return TheChild.Widget->Paint( TheChild.Geometry, ChildClippingRect, OutDrawElements, LayerId, InWidgetStyle, ShouldBeEnabled( bParentEnabled ) );
+		return TheChild.Widget->Paint( Args.WithNewParent(this), TheChild.Geometry, ChildClippingRect, OutDrawElements, LayerId, InWidgetStyle, ShouldBeEnabled( bParentEnabled ) );
 	}
 	return LayerId;
 }
diff --git a/Engine/Source/Runtime/Slate/Private/Widgets/Layout/SFxWidget.cpp b/Engine/Source/Runtime/Slate/Private/Widgets/Layout/SFxWidget.cpp
--- a/Engine/Source/Runtime/Slate/Private/Widgets/Layout/SFxWidget.cpp
+++ b/Engine/Source/Runtime/Slate/Private/Widgets/Layout/SFxWidget.cpp
@@ -69,7 +69,7 @@ static FGeometry MakeScaledGeometry( const FGeometry& Geometry, const float Rend
 		Geometry.Scale * RenderScale );
 }
 
-int32 SFxWidget::OnPaint( const FGeometry& AllottedGeometry, const FSlateRect& MyClippingRect, FSlateWindowElementList& OutDrawElements, int32 LayerId, const FWidgetStyle& InWidgetStyle, bool bParentEnabled ) const
+int32 SFxWidget::OnPaint( const FPaintArgs& Args, const FGeometry& AllottedGeometry, const FSlateRect& MyClippingRect, FSlateWindowElementList& OutDrawElements, int32 LayerId, const FWidgetStyle& InWidgetStyle, bool bParentEnabled ) const
 {
 	// Transformt the geometry for rendering only. Layout is unaffected.
 	const FGeometry ModifiedGeometry = MakeScaledGeometry( AllottedGeometry, RenderScale.Get(), RenderScaleOrigin.Get(), VisualOffset.Get() );
@@ -93,7 +93,7 @@ int32 SFxWidget::OnPaint( const FGeometry& AllottedGeometry, const FSlateRect& M
 			.BlendColorAndOpacityTint(ColorAndOpacity.Get())
 			.SetForegroundColor( ForegroundColor );
 
-		return TheChild.Widget->Paint( TheChild.Geometry, ChildClippingRect, OutDrawElements, LayerId + 1, CompoundedWidgetStyle, ShouldBeEnabled( bParentEnabled ) );
+		return TheChild.Widget->Paint( Args.WithNewParent(this), TheChild.Geometry, ChildClippingRect, OutDrawElements, LayerId + 1, CompoundedWidgetStyle, ShouldBeEnabled( bParentEnabled ) );
 	}
 	return LayerId;
 
diff --git a/Engine/Source/Runtime/Slate/Private/Widgets/Layout/SGridPanel.cpp b/Engine/Source/Runtime/Slate/Private/Widgets/Layout/SGridPanel.cpp
--- a/Engine/Source/Runtime/Slate/Private/Widgets/Layout/SGridPanel.cpp
+++ b/Engine/Source/Runtime/Slate/Private/Widgets/Layout/SGridPanel.cpp
@@ -29,13 +29,10 @@ void SGridPanel::Construct( const FArguments& InArgs )
 
 	ColFillCoefficients = InArgs.ColFillCoefficients;
 	RowFillCoefficients = InArgs.RowFillCoefficients;
-
-//	check( ColFillCoefficients.Num() <= Columns.Num() );
-//	check( RowFillCoefficients.Num() <= Rows.Num() );
 }
 
 
-int32 SGridPanel::OnPaint( const FGeometry& AllottedGeometry, const FSlateRect& MyClippingRect, FSlateWindowElementList& OutDrawElements, int32 LayerId, const FWidgetStyle& InWidgetStyle, bool bParentEnabled ) const
+int32 SGridPanel::OnPaint( const FPaintArgs& Args, const FGeometry& AllottedGeometry, const FSlateRect& MyClippingRect, FSlateWindowElementList& OutDrawElements, int32 LayerId, const FWidgetStyle& InWidgetStyle, bool bParentEnabled ) const
 {
 	FArrangedChildren ArrangedChildren(EVisibility::All);
 	this->ArrangeChildren(AllottedGeometry, ArrangedChildren);
@@ -69,6 +66,7 @@ int32 SGridPanel::OnPaint( const FGeometry& AllottedGeometry, const FSlateRect&
 			}
 
 			const int32 CurWidgetsMaxLayerId = CurWidget.Widget->Paint(
+				Args.WithNewParent(this),
 				CurWidget.Geometry,
 				ChildClipRect,
 				OutDrawElements,
diff --git a/Engine/Source/Runtime/Slate/Private/Widgets/Layout/SScrollBox.cpp b/Engine/Source/Runtime/Slate/Private/Widgets/Layout/SScrollBox.cpp
--- a/Engine/Source/Runtime/Slate/Private/Widgets/Layout/SScrollBox.cpp
+++ b/Engine/Source/Runtime/Slate/Private/Widgets/Layout/SScrollBox.cpp
@@ -396,9 +396,9 @@ FCursorReply SScrollBox::OnCursorQuery( const FGeometry& MyGeometry, const FPoin
 	}
 }
 
-int32 SScrollBox::OnPaint( const FGeometry& AllottedGeometry, const FSlateRect& MyClippingRect, FSlateWindowElementList& OutDrawElements, int32 LayerId, const FWidgetStyle& InWidgetStyle, bool bParentEnabled ) const
+int32 SScrollBox::OnPaint( const FPaintArgs& Args, const FGeometry& AllottedGeometry, const FSlateRect& MyClippingRect, FSlateWindowElementList& OutDrawElements, int32 LayerId, const FWidgetStyle& InWidgetStyle, bool bParentEnabled ) const
 {
-	int32 NewLayerId = SCompoundWidget::OnPaint( AllottedGeometry, MyClippingRect, OutDrawElements, LayerId, InWidgetStyle, bParentEnabled );
+	int32 NewLayerId = SCompoundWidget::OnPaint( Args, AllottedGeometry, MyClippingRect, OutDrawElements, LayerId, InWidgetStyle, bParentEnabled );
 
 	if( !bShowSoftwareCursor )
 	{
diff --git a/Engine/Source/Runtime/Slate/Private/Widgets/Layout/SSpacer.cpp b/Engine/Source/Runtime/Slate/Private/Widgets/Layout/SSpacer.cpp
--- a/Engine/Source/Runtime/Slate/Private/Widgets/Layout/SSpacer.cpp
+++ b/Engine/Source/Runtime/Slate/Private/Widgets/Layout/SSpacer.cpp
@@ -22,20 +22,7 @@ void SSpacer::Construct( const FArguments& InArgs )
 }
 
 
-/**
- * The widget should respond by populating the OutDrawElements array with FDrawElements 
- * that represent it and any of its children.
- *
- * @param AllottedGeometry  The FGeometry that describes an area in which the widget should appear.
- * @param MyClippingRect    The clipping rectangle allocated for this widget and its children.
- * @param OutDrawElements   A list of FDrawElements to populate with the output.
- * @param LayerId           The Layer onto which this widget should be rendered.
- * @param InColorAndOpacity Color and Opacity to be applied to all the descendants of the widget being painted
- * @param bParentEnabled	True if the parent of this widget is enabled.
- *
- * @return The maximum layer ID attained by this widget or any of its children.
- */
-int32 SSpacer::OnPaint( const FGeometry& AllottedGeometry, const FSlateRect& MyClippingRect, FSlateWindowElementList& OutDrawElements, int32 LayerId, const FWidgetStyle& InWidgetStyle, bool bParentEnabled ) const
+int32 SSpacer::OnPaint( const FPaintArgs& Args, const FGeometry& AllottedGeometry, const FSlateRect& MyClippingRect, FSlateWindowElementList& OutDrawElements, int32 LayerId, const FWidgetStyle& InWidgetStyle, bool bParentEnabled ) const
 {
 	// We did not paint anything. The parent's current LayerId is probably the max we were looking for.
 	return LayerId;
diff --git a/Engine/Source/Runtime/Slate/Private/Widgets/Layout/SSplitter.cpp b/Engine/Source/Runtime/Slate/Private/Widgets/Layout/SSplitter.cpp
--- a/Engine/Source/Runtime/Slate/Private/Widgets/Layout/SSplitter.cpp
+++ b/Engine/Source/Runtime/Slate/Private/Widgets/Layout/SSplitter.cpp
@@ -150,12 +150,12 @@ void SSplitter::OnArrangeChildren( const FGeometry& AllottedGeometry, FArrangedC
 }
 
 
-int32 SSplitter::OnPaint( const FGeometry& AllottedGeometry, const FSlateRect& MyClippingRect, FSlateWindowElementList& OutDrawElements, int32 LayerId, const FWidgetStyle& InWidgetStyle, bool bParentEnabled ) const
+int32 SSplitter::OnPaint( const FPaintArgs& Args, const FGeometry& AllottedGeometry, const FSlateRect& MyClippingRect, FSlateWindowElementList& OutDrawElements, int32 LayerId, const FWidgetStyle& InWidgetStyle, bool bParentEnabled ) const
 {
 	FArrangedChildren ArrangedChildren( EVisibility::Visible );
 	ArrangeChildren( AllottedGeometry, ArrangedChildren );
 
-	int32 MaxLayerId = PaintArrangedChildren( ArrangedChildren, MyClippingRect, OutDrawElements, LayerId, InWidgetStyle, bParentEnabled );
+	int32 MaxLayerId = PaintArrangedChildren( Args, ArrangedChildren, MyClippingRect, OutDrawElements, LayerId, InWidgetStyle, bParentEnabled );
 
 	const FSlateBrush* NormalHandleBrush = &Style->HandleNormalBrush;
 
diff --git a/Engine/Source/Runtime/Slate/Private/Widgets/Notifications/SProgressBar.cpp b/Engine/Source/Runtime/Slate/Private/Widgets/Notifications/SProgressBar.cpp
--- a/Engine/Source/Runtime/Slate/Private/Widgets/Notifications/SProgressBar.cpp
+++ b/Engine/Source/Runtime/Slate/Private/Widgets/Notifications/SProgressBar.cpp
@@ -79,7 +79,7 @@ const FSlateBrush* SProgressBar::GetMarqueeImage() const
 	return MarqueeImage ? MarqueeImage : &Style->MarqueeImage;
 }
 
-int32 SProgressBar::OnPaint( const FGeometry& AllottedGeometry, const FSlateRect& MyClippingRect, FSlateWindowElementList& OutDrawElements, int32 LayerId, const FWidgetStyle& InWidgetStyle, bool bParentEnabled ) const
+int32 SProgressBar::OnPaint( const FPaintArgs& Args, const FGeometry& AllottedGeometry, const FSlateRect& MyClippingRect, FSlateWindowElementList& OutDrawElements, int32 LayerId, const FWidgetStyle& InWidgetStyle, bool bParentEnabled ) const
 {
 	// Used to track the layer ID we will return.
 	int32 RetLayerId = LayerId;
diff --git a/Engine/Source/Runtime/Slate/Private/Widgets/SCanvas.cpp b/Engine/Source/Runtime/Slate/Private/Widgets/SCanvas.cpp
--- a/Engine/Source/Runtime/Slate/Private/Widgets/SCanvas.cpp
+++ b/Engine/Source/Runtime/Slate/Private/Widgets/SCanvas.cpp
@@ -92,7 +92,7 @@ void SCanvas::OnArrangeChildren( const FGeometry& AllottedGeometry, FArrangedChi
 }
 
 
-int32 SCanvas::OnPaint( const FGeometry& AllottedGeometry, const FSlateRect& MyClippingRect, FSlateWindowElementList& OutDrawElements, int32 LayerId, const FWidgetStyle& InWidgetStyle, bool bParentEnabled ) const
+int32 SCanvas::OnPaint( const FPaintArgs& Args, const FGeometry& AllottedGeometry, const FSlateRect& MyClippingRect, FSlateWindowElementList& OutDrawElements, int32 LayerId, const FWidgetStyle& InWidgetStyle, bool bParentEnabled ) const
 {
 	FArrangedChildren ArrangedChildren(EVisibility::Visible);
 	this->ArrangeChildren(AllottedGeometry, ArrangedChildren);
@@ -105,7 +105,7 @@ int32 SCanvas::OnPaint( const FGeometry& AllottedGeometry, const FSlateRect& MyC
 	{
 		FArrangedWidget& CurWidget = ArrangedChildren(ChildIndex);
 		FSlateRect ChildClipRect = MyClippingRect.IntersectionWith(CurWidget.Geometry.GetClippingRect());
-		const int32 CurWidgetsMaxLayerId = CurWidget.Widget->Paint(CurWidget.Geometry, ChildClipRect, OutDrawElements, MaxLayerId + 1, InWidgetStyle, ShouldBeEnabled(bParentEnabled));
+		const int32 CurWidgetsMaxLayerId = CurWidget.Widget->Paint(Args.WithNewParent(this), CurWidget.Geometry, ChildClipRect, OutDrawElements, MaxLayerId + 1, InWidgetStyle, ShouldBeEnabled(bParentEnabled));
 
 		MaxLayerId = FMath::Max(MaxLayerId, CurWidgetsMaxLayerId);
 	}
diff --git a/Engine/Source/Runtime/Slate/Private/Widgets/SViewport.cpp b/Engine/Source/Runtime/Slate/Private/Widgets/SViewport.cpp
--- a/Engine/Source/Runtime/Slate/Private/Widgets/SViewport.cpp
+++ b/Engine/Source/Runtime/Slate/Private/Widgets/SViewport.cpp
@@ -36,7 +36,7 @@ void SViewport::Construct( const FArguments& InArgs )
 }
 
 
-int32 SViewport::OnPaint( const FGeometry& AllottedGeometry, const FSlateRect& MyClippingRect, FSlateWindowElementList& OutDrawElements, int32 LayerId, const FWidgetStyle& InWidgetStyle, bool bParentEnabled ) const
+int32 SViewport::OnPaint( const FPaintArgs& Args, const FGeometry& AllottedGeometry, const FSlateRect& MyClippingRect, FSlateWindowElementList& OutDrawElements, int32 LayerId, const FWidgetStyle& InWidgetStyle, bool bParentEnabled ) const
 {
 #if SLATE_HD_STATS
 	SCOPE_CYCLE_COUNTER( STAT_SlateOnPaint_SViewport );
@@ -75,7 +75,7 @@ int32 SViewport::OnPaint( const FGeometry& AllottedGeometry, const FSlateRect& M
 		}
 	}
 
-	int32 Layer = SCompoundWidget::OnPaint(AllottedGeometry, MyClippingRect, OutDrawElements, LayerId, InWidgetStyle, bEnabled );
+	int32 Layer = SCompoundWidget::OnPaint(Args, AllottedGeometry, MyClippingRect, OutDrawElements, LayerId, InWidgetStyle, bEnabled );
 
 	if( ViewportInterfacePin.IsValid() && ViewportInterfacePin->IsSoftwareCursorVisible() )
 	{
diff --git a/Engine/Source/Runtime/Slate/Private/Widgets/SWeakWidget.cpp b/Engine/Source/Runtime/Slate/Private/Widgets/SWeakWidget.cpp
--- a/Engine/Source/Runtime/Slate/Private/Widgets/SWeakWidget.cpp
+++ b/Engine/Source/Runtime/Slate/Private/Widgets/SWeakWidget.cpp
@@ -44,7 +44,7 @@ FChildren* SWeakWidget::GetChildren()
 	return &WeakChild;
 }
 
-int32 SWeakWidget::OnPaint( const FGeometry& AllottedGeometry, const FSlateRect& MyClippingRect, FSlateWindowElementList& OutDrawElements, int32 LayerId, const FWidgetStyle& InWidgetStyle, bool bParentEnabled ) const
+int32 SWeakWidget::OnPaint( const FPaintArgs& Args, const FGeometry& AllottedGeometry, const FSlateRect& MyClippingRect, FSlateWindowElementList& OutDrawElements, int32 LayerId, const FWidgetStyle& InWidgetStyle, bool bParentEnabled ) const
 {
 	// Just draw the children.
 	FArrangedChildren ArrangedChildren(EVisibility::Visible);
@@ -56,7 +56,9 @@ int32 SWeakWidget::OnPaint( const FGeometry& AllottedGeometry, const FSlateRect&
 		check( ArrangedChildren.Num() == 1 );
 		FArrangedWidget& TheChild = ArrangedChildren(0);
 
-		return TheChild.Widget->Paint( TheChild.Geometry, 
+		return TheChild.Widget->Paint( 
+			Args.WithNewParent(this), 
+			TheChild.Geometry, 
 			MyClippingRect, 
 			OutDrawElements, 
 			LayerId + 1,
diff --git a/Engine/Source/Runtime/Slate/Private/Widgets/Testing/STestSuite.cpp b/Engine/Source/Runtime/Slate/Private/Widgets/Testing/STestSuite.cpp
--- a/Engine/Source/Runtime/Slate/Private/Widgets/Testing/STestSuite.cpp
+++ b/Engine/Source/Runtime/Slate/Private/Widgets/Testing/STestSuite.cpp
@@ -61,20 +61,7 @@ public:
 		return FVector2D(128, 128);
 	}
 
-	/**
-	 * The widget should respond by populating the OutDrawElements array with FDrawElements 
-	 * that represent it and any of its children.
-	 *
-	 * @param AllottedGeometry  The FGeometry that describes an area in which the widget should appear.
-	 * @param MyClippingRect    The clipping rectangle allocated for this widget and its children.
-	 * @param OutDrawElements   A list of FDrawElements to populate with the output.
-	 * @param LayerId           The Layer onto which this widget should be rendered.
-	 * @param InColorAndOpacity Color and Opacity to be applied to all the descendants of the widget being painted
- 	 * @param bParentEnabled	True if the parent of this widget is enabled.
-	 *
-	 * @return The maximum layer ID attained by this widget or any of its children.
-	 */
-	virtual int32 OnPaint( const FGeometry& AllottedGeometry, const FSlateRect& MyClippingRect, FSlateWindowElementList& OutDrawElements, int32 LayerId, const FWidgetStyle& InWidgetStyle, bool bParentEnabled ) const
+	virtual int32 OnPaint( const FPaintArgs& Args, const FGeometry& AllottedGeometry, const FSlateRect& MyClippingRect, FSlateWindowElementList& OutDrawElements, int32 LayerId, const FWidgetStyle& InWidgetStyle, bool bParentEnabled ) const
 	{
 		if( OnPaintHandler.IsBound() )
 		{
@@ -91,7 +78,7 @@ public:
 			);
 		}
 
-		return SCompoundWidget::OnPaint(AllottedGeometry, MyClippingRect, OutDrawElements, LayerId, InWidgetStyle, bParentEnabled && IsEnabled() );
+		return SCompoundWidget::OnPaint( Args, AllottedGeometry, MyClippingRect, OutDrawElements, LayerId, InWidgetStyle, bParentEnabled && IsEnabled() );
 	}
 
 private:
@@ -823,6 +810,8 @@ protected:
 
 };
 
+#if WITH_FANCY_TEXT
+
 class SRichTextTest : public SCompoundWidget
 {
 public:
@@ -1543,6 +1532,8 @@ private:
 	TSharedPtr< SRichTextBlock > InteractiveRichText;
 };
 
+#endif //WITH_FANCY_TEXT
+
 
 class STextEditTest : public SCompoundWidget
 {
@@ -1645,6 +1636,8 @@ public:
 				[
 					SAssignNew(ErrorText, SErrorText)
 				]
+
+#if WITH_FANCY_TEXT
 				+SVerticalBox::Slot().FillHeight(1)
 				[
 					SNew(SHorizontalBox)
@@ -1732,6 +1725,7 @@ public:
 						]
 					]
 				]
+				#endif //WITH_FANCY_TEXT
 			]
 		];
 	}
@@ -1837,8 +1831,12 @@ protected:
 
 	TSharedPtr<SEditableTextBox> NumericInput;
 
+#if WITH_FANCY_TEXT
+
 	TSharedPtr<SRichTextBlock> RichTextBlock;
 
+#endif //WITH_FANCY_TEXT
+
 	TSharedPtr<STextEntryPopup> PopupInput;
 
 	TSharedPtr<SWindow> PopupWindow;
@@ -2531,8 +2529,8 @@ class SFxTest : public SCompoundWidget
 							]
 							+SHorizontalBox::Slot() .AutoWidth()
 							[
-								SNew(STextBlock)
-								.Text( LOCTEXT("FxTextContentLabel", "Content" ) )
+								SNew(SButton)
+								.Text( LOCTEXT("FxTextContentLabel", "Click Me!" ) )
 							]
 						]
 					]
@@ -3259,7 +3257,11 @@ TSharedRef<SDockTab> SpawnTab(const FSpawnTabArgs& Args, FName TabIdentifier)
 			. Label( LOCTEXT("RichTextTestTab", "Rich Text") )
 			. ToolTipText( LOCTEXT( "RichTextTestTabToolTip", "Switches to the Rich Text tab, where you can test the various rich text features." ) )
 			[
+				#if WITH_FANCY_TEXT
 				SNew( SRichTextTest )
+				#else
+				SNew( SSpacer )
+				#endif //WITH_FANCY_TEXT
 			];
 	}
 	else if (TabIdentifier == FName(TEXT("EditableTextTab")))
diff --git a/Engine/Source/Runtime/Slate/Private/Widgets/Text/SMultiLineEditableText.cpp b/Engine/Source/Runtime/Slate/Private/Widgets/Text/SMultiLineEditableText.cpp
--- a/Engine/Source/Runtime/Slate/Private/Widgets/Text/SMultiLineEditableText.cpp
+++ b/Engine/Source/Runtime/Slate/Private/Widgets/Text/SMultiLineEditableText.cpp
@@ -1,6 +1,9 @@
 // Copyright 1998-2014 Epic Games, Inc. All Rights Reserved.
 
 #include "SlatePrivatePCH.h"
+
+#if WITH_FANCY_TEXT
+
 #include "TextEditHelper.h"
 #include "SlateWordWrapper.h"
 
@@ -56,7 +59,7 @@ SMultiLineEditableText::FCursorLineHighlighter::FCursorLineHighlighter(const FCu
 	check(CursorInfo);
 }
 
-int32 SMultiLineEditableText::FCursorLineHighlighter::OnPaint( const FTextLayout::FLineView& Line, const float OffsetX, const float Width, const FTextBlockStyle& DefaultStyle, const FGeometry& AllottedGeometry, const FSlateRect& MyClippingRect, FSlateWindowElementList& OutDrawElements, int32 LayerId, const FWidgetStyle& InWidgetStyle, bool bParentEnabled ) const 
+int32 SMultiLineEditableText::FCursorLineHighlighter::OnPaint(const FPaintArgs& Args, const FTextLayout::FLineView& Line, const float OffsetX, const float Width, const FTextBlockStyle& DefaultStyle, const FGeometry& AllottedGeometry, const FSlateRect& MyClippingRect, FSlateWindowElementList& OutDrawElements, int32 LayerId, const FWidgetStyle& InWidgetStyle, bool bParentEnabled) const
 {
 	const FVector2D Location(Line.Offset.X + OffsetX, Line.Offset.Y);
 	const FVector2D Size(Width, Line.TextSize.Y);
@@ -102,7 +105,7 @@ SMultiLineEditableText::FTextCompositionHighlighter::FTextCompositionHighlighter
 {
 }
 
-int32 SMultiLineEditableText::FTextCompositionHighlighter::OnPaint( const FTextLayout::FLineView& Line, const float OffsetX, const float Width, const FTextBlockStyle& DefaultStyle, const FGeometry& AllottedGeometry, const FSlateRect& MyClippingRect, FSlateWindowElementList& OutDrawElements, int32 LayerId, const FWidgetStyle& InWidgetStyle, bool bParentEnabled ) const
+int32 SMultiLineEditableText::FTextCompositionHighlighter::OnPaint( const FPaintArgs& Args, const FTextLayout::FLineView& Line, const float OffsetX, const float Width, const FTextBlockStyle& DefaultStyle, const FGeometry& AllottedGeometry, const FSlateRect& MyClippingRect, FSlateWindowElementList& OutDrawElements, int32 LayerId, const FWidgetStyle& InWidgetStyle, bool bParentEnabled ) const
 {
 	const FVector2D Location(Line.Offset.X + OffsetX, Line.Offset.Y);
 	const FVector2D Size(Width, Line.TextSize.Y);
@@ -138,7 +141,7 @@ SMultiLineEditableText::FTextSelectionRunRenderer::FTextSelectionRunRenderer()
 {
 }
 
-int32 SMultiLineEditableText::FTextSelectionRunRenderer::OnPaint( const FTextLayout::FLineView& Line, const TSharedRef< ISlateRun >& Run, const TSharedRef< ILayoutBlock >& Block, const FTextBlockStyle& DefaultStyle, const FGeometry& AllottedGeometry, const FSlateRect& MyClippingRect, FSlateWindowElementList& OutDrawElements, int32 LayerId, const FWidgetStyle& InWidgetStyle, bool bParentEnabled ) const 
+int32 SMultiLineEditableText::FTextSelectionRunRenderer::OnPaint( const FPaintArgs& Args, const FTextLayout::FLineView& Line, const TSharedRef< ISlateRun >& Run, const TSharedRef< ILayoutBlock >& Block, const FTextBlockStyle& DefaultStyle, const FGeometry& AllottedGeometry, const FSlateRect& MyClippingRect, FSlateWindowElementList& OutDrawElements, int32 LayerId, const FWidgetStyle& InWidgetStyle, bool bParentEnabled ) const 
 {
 	FVector2D Location( Block->GetLocationOffset() );
 	Location.Y = Line.Offset.Y;
@@ -168,7 +171,7 @@ int32 SMultiLineEditableText::FTextSelectionRunRenderer::OnPaint( const FTextLay
 	//FWidgetStyle WidgetStyle( InWidgetStyle );
 	//WidgetStyle.SetForegroundColor( InvertedForeground );
 
-	return Run->OnPaint( Line, Block, DefaultStyle, AllottedGeometry, MyClippingRect, OutDrawElements, LayerId, InWidgetStyle, bParentEnabled );
+	return Run->OnPaint( Args, Line, Block, DefaultStyle, AllottedGeometry, MyClippingRect, OutDrawElements, LayerId, InWidgetStyle, bParentEnabled );
 }
 
 TSharedRef< SMultiLineEditableText::FTextSelectionRunRenderer > SMultiLineEditableText::FTextSelectionRunRenderer::Create()
@@ -1708,7 +1711,7 @@ void SMultiLineEditableText::Tick( const FGeometry& AllottedGeometry, const doub
 	}
 }
 
-int32 SMultiLineEditableText::OnPaint( const FGeometry& AllottedGeometry, const FSlateRect& MyClippingRect, FSlateWindowElementList& OutDrawElements, int32 LayerId, const FWidgetStyle& InWidgetStyle, bool bParentEnabled ) const
+int32 SMultiLineEditableText::OnPaint( const FPaintArgs& Args, const FGeometry& AllottedGeometry, const FSlateRect& MyClippingRect, FSlateWindowElementList& OutDrawElements, int32 LayerId, const FWidgetStyle& InWidgetStyle, bool bParentEnabled ) const
 {
 	// Update the auto-wrap size now that we have computed paint geometry; won't take affect until text frame
 	// Note: This is done here rather than in Tick(), because Tick() doesn't get called while resizing windows, but OnPaint() does
@@ -1723,7 +1726,7 @@ int32 SMultiLineEditableText::OnPaint( const FGeometry& AllottedGeometry, const
 	{
 		const FVector2D TextLayoutSize = TextLayout->GetSize();
 		const FVector2D Offset( ( AllottedGeometry.Size.X - TextLayoutSize.X ), 0 );
-		LayerId = TextLayout->OnPaint( TextStyle, AllottedGeometry.MakeChild( Offset, TextLayoutSize ), MyClippingRect, OutDrawElements, LayerId, InWidgetStyle, ShouldBeEnabled( bParentEnabled ) );
+		LayerId = TextLayout->OnPaint( Args.WithNewParent(this), TextStyle, AllottedGeometry.MakeChild( Offset, TextLayoutSize ), MyClippingRect, OutDrawElements, LayerId, InWidgetStyle, ShouldBeEnabled( bParentEnabled ) );
 	}
 	else if ( TextJustification == ETextJustify::Center )
 	{
@@ -1736,11 +1739,11 @@ int32 SMultiLineEditableText::OnPaint( const FGeometry& AllottedGeometry, const
 		}
 
 		const FVector2D Offset((AllottedGeometry.Size.X - TextLayoutSize.X) / 2, 0);
-		LayerId = TextLayout->OnPaint( TextStyle, AllottedGeometry.MakeChild( Offset, TextLayoutSize ), MyClippingRect, OutDrawElements, LayerId, InWidgetStyle, ShouldBeEnabled( bParentEnabled ) );
+		LayerId = TextLayout->OnPaint( Args.WithNewParent(this), TextStyle, AllottedGeometry.MakeChild( Offset, TextLayoutSize ), MyClippingRect, OutDrawElements, LayerId, InWidgetStyle, ShouldBeEnabled( bParentEnabled ) );
 	}
 	else
 	{
-		LayerId = TextLayout->OnPaint( TextStyle, AllottedGeometry, MyClippingRect, OutDrawElements, LayerId, InWidgetStyle, ShouldBeEnabled( bParentEnabled ) );
+		LayerId = TextLayout->OnPaint( Args.WithNewParent(this), TextStyle, AllottedGeometry, MyClippingRect, OutDrawElements, LayerId, InWidgetStyle, ShouldBeEnabled( bParentEnabled ) );
 	}
 
 	return LayerId;
@@ -2250,3 +2253,6 @@ void SMultiLineEditableText::FTextInputMethodContext::EndComposition()
 		IsComposing = false;
 	}
 }
+
+
+#endif //WITH_FANCY_TEXT
\ No newline at end of file
diff --git a/Engine/Source/Runtime/Slate/Private/Widgets/Text/SRichTextBlock.cpp b/Engine/Source/Runtime/Slate/Private/Widgets/Text/SRichTextBlock.cpp
--- a/Engine/Source/Runtime/Slate/Private/Widgets/Text/SRichTextBlock.cpp
+++ b/Engine/Source/Runtime/Slate/Private/Widgets/Text/SRichTextBlock.cpp
@@ -1,6 +1,9 @@
 // Copyright 1998-2014 Epic Games, Inc. All Rights Reserved.
 
 #include "SlatePrivatePCH.h"
+
+#if WITH_FANCY_TEXT
+
 #include "TextLayoutEngine.h"
 #include "SlateTextLayout.h"
 #include "SlateTextHighlightRunRenderer.h"
@@ -47,7 +50,7 @@ void SRichTextBlock::Tick( const FGeometry& AllottedGeometry, const double InCur
 	TextLayout->SetScale( AllottedGeometry.Scale );
 }
 
-int32 SRichTextBlock::OnPaint( const FGeometry& AllottedGeometry, const FSlateRect& MyClippingRect, FSlateWindowElementList& OutDrawElements, int32 LayerId, const FWidgetStyle& InWidgetStyle, bool bParentEnabled ) const
+int32 SRichTextBlock::OnPaint( const FPaintArgs& Args, const FGeometry& AllottedGeometry, const FSlateRect& MyClippingRect, FSlateWindowElementList& OutDrawElements, int32 LayerId, const FWidgetStyle& InWidgetStyle, bool bParentEnabled ) const
 {
 	SCOPE_CYCLE_COUNTER( STAT_SlateOnPaint_SRichTextBlock );
 
@@ -60,17 +63,17 @@ int32 SRichTextBlock::OnPaint( const FGeometry& AllottedGeometry, const FSlateRe
 	{
 		const FVector2D TextLayoutSize = TextLayout->GetSize();
 		const FVector2D Offset( ( AllottedGeometry.Size.X - TextLayoutSize.X ), 0 );
-		LayerId = TextLayout->OnPaint( TextStyle, AllottedGeometry.MakeChild( Offset, TextLayoutSize ), MyClippingRect, OutDrawElements, LayerId, InWidgetStyle, ShouldBeEnabled( bParentEnabled ) );
+		LayerId = TextLayout->OnPaint( Args.WithNewParent(this), TextStyle, AllottedGeometry.MakeChild( Offset, TextLayoutSize ), MyClippingRect, OutDrawElements, LayerId, InWidgetStyle, ShouldBeEnabled( bParentEnabled ) );
 	}
 	else if ( TextJustification == ETextJustify::Center )
 	{
 		const FVector2D TextLayoutSize = TextLayout->GetSize();
 		const FVector2D Offset( ( AllottedGeometry.Size.X - TextLayoutSize.X ) / 2, 0 );
-		LayerId = TextLayout->OnPaint( TextStyle, AllottedGeometry.MakeChild( Offset, TextLayoutSize ), MyClippingRect, OutDrawElements, LayerId, InWidgetStyle, ShouldBeEnabled( bParentEnabled ) );
+		LayerId = TextLayout->OnPaint( Args.WithNewParent(this), TextStyle, AllottedGeometry.MakeChild( Offset, TextLayoutSize ), MyClippingRect, OutDrawElements, LayerId, InWidgetStyle, ShouldBeEnabled( bParentEnabled ) );
 	}
 	else
 	{
-		LayerId = TextLayout->OnPaint( TextStyle, AllottedGeometry, MyClippingRect, OutDrawElements, LayerId, InWidgetStyle, ShouldBeEnabled( bParentEnabled ) );
+		LayerId = TextLayout->OnPaint( Args.WithNewParent(this), TextStyle, AllottedGeometry, MyClippingRect, OutDrawElements, LayerId, InWidgetStyle, ShouldBeEnabled( bParentEnabled ) );
 	}
 
 	return LayerId;
@@ -249,3 +252,6 @@ void SRichTextBlock::SetHighlightText( const FText& InHighlightText )
 
 	TextLayout->SetRunRenderers( TextHighlights );
 }
+
+
+#endif //WITH_FANCY_TEXT
\ No newline at end of file
diff --git a/Engine/Source/Runtime/Slate/Private/Widgets/Text/STextBlock.cpp b/Engine/Source/Runtime/Slate/Private/Widgets/Text/STextBlock.cpp
--- a/Engine/Source/Runtime/Slate/Private/Widgets/Text/STextBlock.cpp
+++ b/Engine/Source/Runtime/Slate/Private/Widgets/Text/STextBlock.cpp
@@ -79,7 +79,7 @@ void STextBlock::SetText( const TAttribute< FText >& InText )
 	bRequestCache = true;
 }
 
-int32 STextBlock::OnPaint( const FGeometry& AllottedGeometry, const FSlateRect& MyClippingRect, FSlateWindowElementList& OutDrawElements, int32 LayerId, const FWidgetStyle& InWidgetStyle, bool bParentEnabled ) const
+int32 STextBlock::OnPaint( const FPaintArgs& Args, const FGeometry& AllottedGeometry, const FSlateRect& MyClippingRect, FSlateWindowElementList& OutDrawElements, int32 LayerId, const FWidgetStyle& InWidgetStyle, bool bParentEnabled ) const
 {
 #if SLATE_HD_STATS
 	SCOPE_CYCLE_COUNTER( STAT_SlateOnPaint_STextBlock );
diff --git a/Engine/Source/Runtime/Slate/Private/Widgets/Tutorials/STutorialWrapper.cpp b/Engine/Source/Runtime/Slate/Private/Widgets/Tutorials/STutorialWrapper.cpp
--- a/Engine/Source/Runtime/Slate/Private/Widgets/Tutorials/STutorialWrapper.cpp
+++ b/Engine/Source/Runtime/Slate/Private/Widgets/Tutorials/STutorialWrapper.cpp
@@ -51,7 +51,7 @@ void STutorialWrapper::Tick( const FGeometry& AllottedGeometry, const double InC
 	CachedGeometry = AllottedGeometry;
 }
 
-int32 STutorialWrapper::OnPaint(const FGeometry& AllottedGeometry, const FSlateRect& MyClippingRect, FSlateWindowElementList& OutDrawElements, int32 LayerId, const FWidgetStyle& InWidgetStyle, bool bParentEnabled ) const
+int32 STutorialWrapper::OnPaint(const FPaintArgs& Args, const FGeometry& AllottedGeometry, const FSlateRect& MyClippingRect, FSlateWindowElementList& OutDrawElements, int32 LayerId, const FWidgetStyle& InWidgetStyle, bool bParentEnabled ) const
 {
 	bool IsVisible = !( CachedGeometry.Size.X == 0 || CachedGeometry.Size.Y == 0 );
 
@@ -92,7 +92,7 @@ int32 STutorialWrapper::OnPaint(const FGeometry& AllottedGeometry, const FSlateR
 		FSlateDrawElement::MakeBox(OutDrawElements, OverlayLayerId + 1, BorderGeometry, BorderBrush, WindowClippingRect, ESlateDrawEffect::None, BorderTint);
 	}
 
-	return SBorder::OnPaint(AllottedGeometry, MyClippingRect, OutDrawElements, LayerId, InWidgetStyle, bParentEnabled);
+	return SBorder::OnPaint(Args, AllottedGeometry, MyClippingRect, OutDrawElements, LayerId, InWidgetStyle, bParentEnabled);
 }
 
 void STutorialWrapper::GetAnimationValues(float& OutAlphaFactor, float& OutPulseFactor, FLinearColor& OutShadowTint, FLinearColor& OutBorderTint) const
diff --git a/Engine/Source/Runtime/Slate/Private/Widgets/Views/STableViewBase.cpp b/Engine/Source/Runtime/Slate/Private/Widgets/Views/STableViewBase.cpp
--- a/Engine/Source/Runtime/Slate/Private/Widgets/Views/STableViewBase.cpp
+++ b/Engine/Source/Runtime/Slate/Private/Widgets/Views/STableViewBase.cpp
@@ -494,9 +494,9 @@ bool STableViewBase::IsPendingRefresh() const
 	return bItemsNeedRefresh || ItemsPanel->IsRefreshPending();
 }
 
-int32 STableViewBase::OnPaint( const FGeometry& AllottedGeometry, const FSlateRect& MyClippingRect, FSlateWindowElementList& OutDrawElements, int32 LayerId, const FWidgetStyle& InWidgetStyle, bool bParentEnabled ) const
+int32 STableViewBase::OnPaint( const FPaintArgs& Args, const FGeometry& AllottedGeometry, const FSlateRect& MyClippingRect, FSlateWindowElementList& OutDrawElements, int32 LayerId, const FWidgetStyle& InWidgetStyle, bool bParentEnabled ) const
 {
-	int32 NewLayerId = SCompoundWidget::OnPaint( AllottedGeometry, MyClippingRect, OutDrawElements, LayerId, InWidgetStyle, bParentEnabled );
+	int32 NewLayerId = SCompoundWidget::OnPaint( Args, AllottedGeometry, MyClippingRect, OutDrawElements, LayerId, InWidgetStyle, bParentEnabled );
 
 	if( !bShowSoftwareCursor )
 	{
diff --git a/Engine/Source/Runtime/Slate/Public/Framework/Application/SlateApplication.h b/Engine/Source/Runtime/Slate/Public/Framework/Application/SlateApplication.h
--- a/Engine/Source/Runtime/Slate/Public/Framework/Application/SlateApplication.h
+++ b/Engine/Source/Runtime/Slate/Public/Framework/Application/SlateApplication.h
@@ -8,6 +8,7 @@
 class SToolTip;
 class SViewport;
 class SWindow;
+class FHittestGrid;
 
 
 /** A Delegate for querying whether source code access is possible */
@@ -1070,6 +1071,11 @@ private:
 	FSlateApplication();
 
 private:
+	/**
+	 * Will be invoked when the size of the geometry of the virtual
+	 * desktop changes (e.g. resolution change or monitors re-arranged)
+	 */
+	void OnVirtualDesktopSizeChanged(const FDisplayMetrics& NewDisplayMetric);
 
 	/** Application singleton */
 	static TSharedPtr< FSlateApplication > CurrentApplication;
@@ -1357,4 +1363,17 @@ private:
 
 	/** The icon to use on application windows */
 	const FSlateBrush *AppIcon;
+
+	//
+	// Hittest 2.0
+	//
+
+	// The rectangle that bounds all the physical monitors given their arrangement.
+	// Info comes from the native platform.
+	// e.g. On windows the origin (coordinates X=0, Y=0) is the upper left of the primary monitor,
+	// but there could be another monitor on any of the sides.
+	FSlateRect VirtualDesktopRect;
+
+	// Hittest acceleration structure. Filled out during Paint.
+	TSharedRef<FHittestGrid> HittestGrid;
 };
diff --git a/Engine/Source/Runtime/Slate/Public/Framework/Text/ISlateLineHighlighter.h b/Engine/Source/Runtime/Slate/Public/Framework/Text/ISlateLineHighlighter.h
--- a/Engine/Source/Runtime/Slate/Public/Framework/Text/ISlateLineHighlighter.h
+++ b/Engine/Source/Runtime/Slate/Public/Framework/Text/ISlateLineHighlighter.h
@@ -6,5 +6,5 @@ class SLATE_API ISlateLineHighlighter : public ILineHighlighter
 public:
 	virtual ~ISlateLineHighlighter() {}
 
-	virtual int32 OnPaint( const FTextLayout::FLineView& Line, const float OffsetX, const float Width, const FTextBlockStyle& DefaultStyle, const FGeometry& AllottedGeometry, const FSlateRect& MyClippingRect, FSlateWindowElementList& OutDrawElements, int32 LayerId, const FWidgetStyle& InWidgetStyle, bool bParentEnabled ) const = 0;
+	virtual int32 OnPaint(const FPaintArgs& Args, const FTextLayout::FLineView& Line, const float OffsetX, const float Width, const FTextBlockStyle& DefaultStyle, const FGeometry& AllottedGeometry, const FSlateRect& MyClippingRect, FSlateWindowElementList& OutDrawElements, int32 LayerId, const FWidgetStyle& InWidgetStyle, bool bParentEnabled) const = 0;
 };
diff --git a/Engine/Source/Runtime/Slate/Public/Framework/Text/ISlateRun.h b/Engine/Source/Runtime/Slate/Public/Framework/Text/ISlateRun.h
--- a/Engine/Source/Runtime/Slate/Public/Framework/Text/ISlateRun.h
+++ b/Engine/Source/Runtime/Slate/Public/Framework/Text/ISlateRun.h
@@ -1,15 +1,19 @@
 // Copyright 1998-2014 Epic Games, Inc. All Rights Reserved.
 #pragma once
 
+#if WITH_FANCY_TEXT
+
 #include "TextLayout.h"
 
 class SLATE_API ISlateRun : public IRun
 {
 public:
 
-	virtual int32 OnPaint( const FTextLayout::FLineView& Line, const TSharedRef< ILayoutBlock >& Block, const FTextBlockStyle& DefaultStyle, const FGeometry& AllottedGeometry, const FSlateRect& MyClippingRect, FSlateWindowElementList& OutDrawElements, int32 LayerId, const FWidgetStyle& InWidgetStyle, bool bParentEnabled ) const = 0;
+	virtual int32 OnPaint( const FPaintArgs& Args, const FTextLayout::FLineView& Line, const TSharedRef< ILayoutBlock >& Block, const FTextBlockStyle& DefaultStyle, const FGeometry& AllottedGeometry, const FSlateRect& MyClippingRect, FSlateWindowElementList& OutDrawElements, int32 LayerId, const FWidgetStyle& InWidgetStyle, bool bParentEnabled ) const = 0;
 
 	virtual FChildren* GetChildren() = 0;
 
 	virtual void ArrangeChildren( const TSharedRef< ILayoutBlock >& Block, const FGeometry& AllottedGeometry, FArrangedChildren& ArrangedChildren ) const = 0;
-};
\ No newline at end of file
+};
+
+#endif //WITH_FANCY_TEXT
\ No newline at end of file
diff --git a/Engine/Source/Runtime/Slate/Public/Framework/Text/ISlateRunRenderer.h b/Engine/Source/Runtime/Slate/Public/Framework/Text/ISlateRunRenderer.h
--- a/Engine/Source/Runtime/Slate/Public/Framework/Text/ISlateRunRenderer.h
+++ b/Engine/Source/Runtime/Slate/Public/Framework/Text/ISlateRunRenderer.h
@@ -1,6 +1,8 @@
 // Copyright 1998-2014 Epic Games, Inc. All Rights Reserved.
 #pragma once
 
+#if WITH_FANCY_TEXT
+
 class ISlateRun;
 class ILayoutBlock;
 
@@ -9,5 +11,7 @@ class SLATE_API ISlateRunRenderer : public IRunRenderer
 public:
 	virtual ~ISlateRunRenderer() {}
 
-	virtual int32 OnPaint( const FTextLayout::FLineView& Line, const TSharedRef< ISlateRun >& Run, const TSharedRef< ILayoutBlock >& Block, const FTextBlockStyle& DefaultStyle, const FGeometry& AllottedGeometry, const FSlateRect& MyClippingRect, FSlateWindowElementList& OutDrawElements, int32 LayerId, const FWidgetStyle& InWidgetStyle, bool bParentEnabled ) const = 0;
+	virtual int32 OnPaint( const FPaintArgs& Args, const FTextLayout::FLineView& Line, const TSharedRef< ISlateRun >& Run, const TSharedRef< ILayoutBlock >& Block, const FTextBlockStyle& DefaultStyle, const FGeometry& AllottedGeometry, const FSlateRect& MyClippingRect, FSlateWindowElementList& OutDrawElements, int32 LayerId, const FWidgetStyle& InWidgetStyle, bool bParentEnabled ) const = 0;
 };
+
+#endif //WITH_FANCY_TEXT
\ No newline at end of file
diff --git a/Engine/Source/Runtime/Slate/Public/Framework/Text/ITextDecorator.h b/Engine/Source/Runtime/Slate/Public/Framework/Text/ITextDecorator.h
--- a/Engine/Source/Runtime/Slate/Public/Framework/Text/ITextDecorator.h
+++ b/Engine/Source/Runtime/Slate/Public/Framework/Text/ITextDecorator.h
@@ -1,6 +1,8 @@
 // Copyright 1998-2014 Epic Games, Inc. All Rights Reserved.
 #pragma once
 
+#if WITH_FANCY_TEXT
+
 struct SLATE_API FTextRunParseResults
 {
 	FTextRunParseResults( FString InName, const FTextRange& InOriginalRange)
@@ -70,4 +72,6 @@ public:
 	virtual bool Supports( const FTextRunParseResults& RunInfo, const FString& Text ) const = 0;
 
 	virtual TSharedRef< ISlateRun > Create( const FTextRunParseResults& RunInfo, const FString& OriginalText, const TSharedRef< FString >& ModelText, const ISlateStyle* Style ) = 0;
-};
\ No newline at end of file
+};
+
+#endif //WITH_FANCY_TEXT
\ No newline at end of file
diff --git a/Engine/Source/Runtime/Slate/Public/Framework/Text/SlateHyperlinkRun.h b/Engine/Source/Runtime/Slate/Public/Framework/Text/SlateHyperlinkRun.h
--- a/Engine/Source/Runtime/Slate/Public/Framework/Text/SlateHyperlinkRun.h
+++ b/Engine/Source/Runtime/Slate/Public/Framework/Text/SlateHyperlinkRun.h
@@ -1,6 +1,8 @@
 // Copyright 1998-2014 Epic Games, Inc. All Rights Reserved.
 #pragma once
 
+#if WITH_FANCY_TEXT
+
 class SLATE_API FSlateHyperlinkRun : public ISlateRun, public TSharedFromThis< FSlateHyperlinkRun >
 {
 public:
@@ -50,7 +52,7 @@ public:
 
 	virtual TSharedRef< ILayoutBlock > CreateBlock( int32 StartIndex, int32 EndIndex, FVector2D Size, const TSharedPtr< IRunRenderer >& Renderer ) override;
 
-	virtual int32 OnPaint( const FTextLayout::FLineView& Line, const TSharedRef< ILayoutBlock >& Block, const FTextBlockStyle& DefaultStyle, const FGeometry& AllottedGeometry, const FSlateRect& MyClippingRect, FSlateWindowElementList& OutDrawElements, int32 LayerId, const FWidgetStyle& InWidgetStyle, bool bParentEnabled ) const override;
+	virtual int32 OnPaint( const FPaintArgs& Args, const FTextLayout::FLineView& Line, const TSharedRef< ILayoutBlock >& Block, const FTextBlockStyle& DefaultStyle, const FGeometry& AllottedGeometry, const FSlateRect& MyClippingRect, FSlateWindowElementList& OutDrawElements, int32 LayerId, const FWidgetStyle& InWidgetStyle, bool bParentEnabled ) const override;
 
 	virtual FChildren* GetChildren() override;
 
@@ -91,4 +93,6 @@ private:
 
 	TSharedRef< FWidgetViewModel > ViewModel;
 	TSlotlessChildren< SWidget > Children;
-};
\ No newline at end of file
+};
+
+#endif //WITH_FANCY_TEXT
\ No newline at end of file
diff --git a/Engine/Source/Runtime/Slate/Public/Framework/Text/SlateImageRun.h b/Engine/Source/Runtime/Slate/Public/Framework/Text/SlateImageRun.h
--- a/Engine/Source/Runtime/Slate/Public/Framework/Text/SlateImageRun.h
+++ b/Engine/Source/Runtime/Slate/Public/Framework/Text/SlateImageRun.h
@@ -1,6 +1,8 @@
 // Copyright 1998-2014 Epic Games, Inc. All Rights Reserved.
 #pragma once
 
+#if WITH_FANCY_TEXT
+
 class SLATE_API FSlateImageRun : public ISlateRun, public TSharedFromThis< FSlateImageRun >
 {
 public:
@@ -25,7 +27,7 @@ public:
 
 	virtual TSharedRef< ILayoutBlock > CreateBlock( int32 StartIndex, int32 EndIndex, FVector2D Size, const TSharedPtr< IRunRenderer >& Renderer ) override;
 
-	virtual int32 OnPaint( const FTextLayout::FLineView& Line, const TSharedRef< ILayoutBlock >& Block, const FTextBlockStyle& DefaultStyle, const FGeometry& AllottedGeometry, const FSlateRect& MyClippingRect, FSlateWindowElementList& OutDrawElements, int32 LayerId, const FWidgetStyle& InWidgetStyle, bool bParentEnabled ) const override;
+	virtual int32 OnPaint( const FPaintArgs& Args, const FTextLayout::FLineView& Line, const TSharedRef< ILayoutBlock >& Block, const FTextBlockStyle& DefaultStyle, const FGeometry& AllottedGeometry, const FSlateRect& MyClippingRect, FSlateWindowElementList& OutDrawElements, int32 LayerId, const FWidgetStyle& InWidgetStyle, bool bParentEnabled ) const override;
 
 	virtual FChildren* GetChildren() override;
 
@@ -58,4 +60,6 @@ private:
 	int16 Baseline;
 
 	static FNoChildren NoChildrenInstance;
-};
\ No newline at end of file
+};
+
+#endif //WITH_FANCY_TEXT
\ No newline at end of file
diff --git a/Engine/Source/Runtime/Slate/Public/Framework/Text/SlateTextHighlightRunRenderer.h b/Engine/Source/Runtime/Slate/Public/Framework/Text/SlateTextHighlightRunRenderer.h
--- a/Engine/Source/Runtime/Slate/Public/Framework/Text/SlateTextHighlightRunRenderer.h
+++ b/Engine/Source/Runtime/Slate/Public/Framework/Text/SlateTextHighlightRunRenderer.h
@@ -1,6 +1,8 @@
 // Copyright 1998-2014 Epic Games, Inc. All Rights Reserved.
 #pragma once
 
+#if WITH_FANCY_TEXT
+
 class SLATE_API FSlateTextHighlightRunRenderer : public ISlateRunRenderer
 {
 public:
@@ -9,10 +11,11 @@ public:
 
 	virtual ~FSlateTextHighlightRunRenderer() {}
 
-	virtual int32 OnPaint( const FTextLayout::FLineView& Line, const TSharedRef< ISlateRun >& Run, const TSharedRef< ILayoutBlock >& Block, const FTextBlockStyle& DefaultStyle, const FGeometry& AllottedGeometry, const FSlateRect& MyClippingRect, FSlateWindowElementList& OutDrawElements, int32 LayerId, const FWidgetStyle& InWidgetStyle, bool bParentEnabled ) const override;
+	virtual int32 OnPaint( const FPaintArgs& Args, const FTextLayout::FLineView& Line, const TSharedRef< ISlateRun >& Run, const TSharedRef< ILayoutBlock >& Block, const FTextBlockStyle& DefaultStyle, const FGeometry& AllottedGeometry, const FSlateRect& MyClippingRect, FSlateWindowElementList& OutDrawElements, int32 LayerId, const FWidgetStyle& InWidgetStyle, bool bParentEnabled ) const override;
 
 private:
 
 	FSlateTextHighlightRunRenderer();
 
-};
\ No newline at end of file
+};
+#endif //WITH_FANCY_TEXT
\ No newline at end of file
diff --git a/Engine/Source/Runtime/Slate/Public/Framework/Text/SlateTextLayout.h b/Engine/Source/Runtime/Slate/Public/Framework/Text/SlateTextLayout.h
--- a/Engine/Source/Runtime/Slate/Public/Framework/Text/SlateTextLayout.h
+++ b/Engine/Source/Runtime/Slate/Public/Framework/Text/SlateTextLayout.h
@@ -1,6 +1,8 @@
 // Copyright 1998-2014 Epic Games, Inc. All Rights Reserved.
 #pragma once
 
+#if WITH_FANCY_TEXT
+
 class SLATE_API FSlateTextLayout : public FTextLayout
 {
 public:
@@ -11,7 +13,7 @@ public:
 
 	void ArrangeChildren( const FGeometry& AllottedGeometry, FArrangedChildren& ArrangedChildren ) const;
 
-	int32 OnPaint( const FTextBlockStyle& DefaultTextStyle, const FGeometry& AllottedGeometry, const FSlateRect& MyClippingRect, FSlateWindowElementList& OutDrawElements, int32 LayerId, const FWidgetStyle& InWidgetStyle, bool bParentEnabled ) const;
+	int32 OnPaint( const FPaintArgs& Args, const FTextBlockStyle& DefaultTextStyle, const FGeometry& AllottedGeometry, const FSlateRect& MyClippingRect, FSlateWindowElementList& OutDrawElements, int32 LayerId, const FWidgetStyle& InWidgetStyle, bool bParentEnabled ) const;
 
 	//void AddLine( const TSharedRef< FString >& Text, const FTextBlockStyle& TextStyle )
 	//{
@@ -27,7 +29,7 @@ protected:
 
 	FSlateTextLayout();
 
-	int32 OnPaintHighlights( const FTextLayout::FLineView& LineView, const TArray<FLineViewHighlight>& Highlights, const FTextBlockStyle& DefaultTextStyle, const FGeometry& AllottedGeometry, const FSlateRect& ClippingRect, FSlateWindowElementList& OutDrawElements, const int32 LayerId, const FWidgetStyle& InWidgetStyle, bool bParentEnabled ) const;
+	int32 OnPaintHighlights(const FPaintArgs& Args, const FTextLayout::FLineView& LineView, const TArray<FLineViewHighlight>& Highlights, const FTextBlockStyle& DefaultTextStyle, const FGeometry& AllottedGeometry, const FSlateRect& ClippingRect, FSlateWindowElementList& OutDrawElements, const int32 LayerId, const FWidgetStyle& InWidgetStyle, bool bParentEnabled) const;
 
 	void AggregateChildren();
 
@@ -36,4 +38,6 @@ private:
 	TSlotlessChildren< SWidget > Children;
 
 	friend class FSlateTextLayoutFactory;
-};
\ No newline at end of file
+};
+
+#endif //WITH_FANCY_TEXT
\ No newline at end of file
diff --git a/Engine/Source/Runtime/Slate/Public/Framework/Text/SlateWidgetRun.h b/Engine/Source/Runtime/Slate/Public/Framework/Text/SlateWidgetRun.h
--- a/Engine/Source/Runtime/Slate/Public/Framework/Text/SlateWidgetRun.h
+++ b/Engine/Source/Runtime/Slate/Public/Framework/Text/SlateWidgetRun.h
@@ -1,6 +1,8 @@
 // Copyright 1998-2014 Epic Games, Inc. All Rights Reserved.
 #pragma once
 
+#if WITH_FANCY_TEXT
+
 class SLATE_API FSlateWidgetRun : public ISlateRun, public TSharedFromThis< FSlateWidgetRun >
 {
 public:
@@ -45,7 +47,7 @@ public:
 
 	virtual TSharedRef< ILayoutBlock > CreateBlock( int32 StartIndex, int32 EndIndex, FVector2D Size, const TSharedPtr< IRunRenderer >& Renderer ) override;
 
-	virtual int32 OnPaint( const FTextLayout::FLineView& Line, const TSharedRef< ILayoutBlock >& Block, const FTextBlockStyle& DefaultStyle, const FGeometry& AllottedGeometry, const FSlateRect& MyClippingRect, FSlateWindowElementList& OutDrawElements, int32 LayerId, const FWidgetStyle& InWidgetStyle, bool bParentEnabled ) const override;
+	virtual int32 OnPaint( const FPaintArgs& Args, const FTextLayout::FLineView& Line, const TSharedRef< ILayoutBlock >& Block, const FTextBlockStyle& DefaultStyle, const FGeometry& AllottedGeometry, const FSlateRect& MyClippingRect, FSlateWindowElementList& OutDrawElements, int32 LayerId, const FWidgetStyle& InWidgetStyle, bool bParentEnabled ) const override;
 
 	virtual FChildren* GetChildren() override;
 
@@ -78,4 +80,5 @@ private:
 	FTextRange Range;
 	FWidgetRunInfo Info;
 	TSlotlessChildren< SWidget > Children;
-};
\ No newline at end of file
+};
+#endif //WITH_FANCY_TEXT
\ No newline at end of file
diff --git a/Engine/Source/Runtime/Slate/Public/Widgets/Images/SSpinningImage.h b/Engine/Source/Runtime/Slate/Public/Widgets/Images/SSpinningImage.h
--- a/Engine/Source/Runtime/Slate/Public/Widgets/Images/SSpinningImage.h
+++ b/Engine/Source/Runtime/Slate/Public/Widgets/Images/SSpinningImage.h
@@ -31,8 +31,7 @@ public:
 	/* Construct this widget */
 	void Construct( const FArguments& InArgs );
 
-	/* Override SImage's OnPaint to draw rotated */
-	virtual int32 OnPaint( const FGeometry& AllottedGeometry, const FSlateRect& MyClippingRect, FSlateWindowElementList& OutDrawElements, int32 LayerId, const FWidgetStyle& InWidgetStyle, bool bParentEnabled ) const override;
+	virtual int32 OnPaint( const FPaintArgs& Args, const FGeometry& AllottedGeometry, const FSlateRect& MyClippingRect, FSlateWindowElementList& OutDrawElements, int32 LayerId, const FWidgetStyle& InWidgetStyle, bool bParentEnabled ) const override;
 
 private:
 	/** The sequence to drive the spinning animation */
diff --git a/Engine/Source/Runtime/Slate/Public/Widgets/Input/SSpinBox.h b/Engine/Source/Runtime/Slate/Public/Widgets/Input/SSpinBox.h
--- a/Engine/Source/Runtime/Slate/Public/Widgets/Input/SSpinBox.h
+++ b/Engine/Source/Runtime/Slate/Public/Widgets/Input/SSpinBox.h
@@ -158,20 +158,7 @@ public:
 		];
 	}
 	
-	/**
-	* The widget should respond by populating the OutDrawElements array with FDrawElements 
-	* that represent it and any of its children.
-	*
-	* @param AllottedGeometry  The FGeometry that describes an area in which the widget should appear.
-	* @param MyClippingRect    The clipping rectangle allocated for this widget and its children.
-	* @param OutDrawElements   A list of FDrawElements to populate with the output.
-	* @param LayerId           The Layer onto which this widget should be rendered.
-	* @param InColorAndOpacity Color and Opacity to be applied to all the descendants of the widget being painted
-	* @param bParentEnabled	True if the parent of this widget is enabled.
-	*
-	* @return The maximum layer ID attained by this widget or any of its children.
-	*/
-	virtual int32 OnPaint( const FGeometry& AllottedGeometry, const FSlateRect& MyClippingRect, FSlateWindowElementList& OutDrawElements, int32 LayerId, const FWidgetStyle& InWidgetStyle, bool bParentEnabled ) const override
+	virtual int32 OnPaint( const FPaintArgs& Args, const FGeometry& AllottedGeometry, const FSlateRect& MyClippingRect, FSlateWindowElementList& OutDrawElements, int32 LayerId, const FWidgetStyle& InWidgetStyle, bool bParentEnabled ) const override
 	{
 		const bool bActiveFeedback = IsHovered() || bDragging;
 
@@ -231,7 +218,7 @@ public:
 			}
 		}
 
-		return FMath::Max( FilledLayer, SCompoundWidget::OnPaint(AllottedGeometry, MyClippingRect, OutDrawElements, FilledLayer, InWidgetStyle, bEnabled ) );
+		return FMath::Max( FilledLayer, SCompoundWidget::OnPaint(Args, AllottedGeometry, MyClippingRect, OutDrawElements, FilledLayer, InWidgetStyle, bEnabled ) );
 	}
 
 	/**
diff --git a/Engine/Source/Runtime/Slate/Public/Widgets/Layout/SEnableBox.h b/Engine/Source/Runtime/Slate/Public/Widgets/Layout/SEnableBox.h
--- a/Engine/Source/Runtime/Slate/Public/Widgets/Layout/SEnableBox.h
+++ b/Engine/Source/Runtime/Slate/Public/Widgets/Layout/SEnableBox.h
@@ -19,9 +19,9 @@ public:
 	}
 
 	// SWidget interface
-	virtual int32 OnPaint(const FGeometry& AllottedGeometry, const FSlateRect& MyClippingRect, FSlateWindowElementList& OutDrawElements, int32 LayerId, const FWidgetStyle& InWidgetStyle, bool bParentEnabled) const override
+	virtual int32 OnPaint(const FPaintArgs& Args, const FGeometry& AllottedGeometry, const FSlateRect& MyClippingRect, FSlateWindowElementList& OutDrawElements, int32 LayerId, const FWidgetStyle& InWidgetStyle, bool bParentEnabled) const override
 	{
-		return SBox::OnPaint(AllottedGeometry, MyClippingRect, OutDrawElements, LayerId, InWidgetStyle, /*bParentEnabled=*/ true);
+		return SBox::OnPaint(Args, AllottedGeometry, MyClippingRect, OutDrawElements, LayerId, InWidgetStyle, /*bParentEnabled=*/ true);
 	}
 	// End of SWidget interface
 };
