Refactored MultiBox and TabManager to support presenting the currently active window menu as the global main menu on Mac OS X.

- On Mac OS X we hide the per-window menu bar.
- MultiBlocks have a type so that the Mac-specific menu creation code can generate the proper native widget.
- The menu builder can now be asked for the constructed MultiBox so that on menu creation it can be set on the TabManager.
- That allows the TabManager to be told to update the current global menu bar when the tab or window focus changes.
- A window activation delegate has been added to SWindow to permit this.
- Widgets that used to live in the LevelEditor menu bar have been hoisted into the window titlebar on OS X as on that platform there's no menu bar anymore.
- There's limited support for parsing custom widgets in the menu MultiBlock in order to support the platform list in the File->Package... sub-menu, this is sufficient for now but is quite fragile.
- We can support custom Slate widgets inside native Mac menus should we ever need to - but it is tricky.

diff --git a/Engine/Source/Editor/LevelEditor/Private/LevelEditor.cpp b/Engine/Source/Editor/LevelEditor/Private/LevelEditor.cpp
--- a/Engine/Source/Editor/LevelEditor/Private/LevelEditor.cpp
+++ b/Engine/Source/Editor/LevelEditor/Private/LevelEditor.cpp
@@ -101,7 +101,17 @@ TSharedRef<SDockTab> FLevelEditorModule::SpawnLevelEditor( const FSpawnTabArgs&
 						.ColorAndOpacity( FLinearColor( 1.0f, 1.0f, 1.0f, 0.3f ) )
 					]
 				]
-
+// Put the level editor stats/notification widgets on the main window title bar since we don't have a menu bar on OS X
+#if PLATFORM_MAC
+				+SHorizontalBox::Slot()
+				.AutoWidth()
+				.Padding(16.0f, 0.0f, 0.0f, 0.0f)
+				.VAlign(VAlign_Center)
+				[
+					LevelEditorTab->GetRightContent()
+				]
+#endif
+		
 				+SHorizontalBox::Slot()
 				.AutoWidth()
 				.Padding(16.0f, 0.0f, 0.0f, 0.0f)
diff --git a/Engine/Source/Editor/LevelEditor/Private/SLevelEditor.cpp b/Engine/Source/Editor/LevelEditor/Private/SLevelEditor.cpp
--- a/Engine/Source/Editor/LevelEditor/Private/SLevelEditor.cpp
+++ b/Engine/Source/Editor/LevelEditor/Private/SLevelEditor.cpp
@@ -192,7 +192,9 @@ void SLevelEditor::Initialize( const TSharedRef<SDockTab>& OwnerTab, const TShar
 					Widget1
 				]
 			]
-
+		 
+// For platforms without a global menu bar we can put the perf. tools in the editor window's menu bar
+#if !PLATFORM_MAC
 			+SOverlay::Slot()
 			.HAlign( HAlign_Right )
 			[
@@ -201,6 +203,7 @@ void SLevelEditor::Initialize( const TSharedRef<SDockTab>& OwnerTab, const TShar
 					SAssignNew( NotificationBarBox, SHorizontalBox )
 				]
 			]
+#endif
 		]
 
 		+SVerticalBox::Slot()
@@ -209,6 +212,16 @@ void SLevelEditor::Initialize( const TSharedRef<SDockTab>& OwnerTab, const TShar
 			Widget2
 		]
 	];
+	
+// For OS X we need to put it into the window's title bar since there's no per-window menu bar
+#if PLATFORM_MAC
+	TSharedRef<SWidget> TutorialWidget = SNew( STutorialWrapper, TEXT("PerformanceTools") )
+	[
+		SAssignNew( NotificationBarBox, SHorizontalBox )
+	];
+	
+	OwnerTab->SetRightContent(TutorialWidget);
+#endif
 
 	ConstructNotificationBar();
 
@@ -274,7 +287,7 @@ FText SLevelEditor::GetTabTitle() const
 
 	const bool bIncludeGameName = false;
 
-	const bool bDirtyState = World->GetCurrentLevel()->GetOutermost()->IsDirty();
+	const bool bDirtyState = World && World->GetCurrentLevel()->GetOutermost()->IsDirty();
 
 	FFormatNamedArguments Args;
 	Args.Add( TEXT("LevelName"), FText::FromString( MainFrameModule.GetLoadedLevelName() ) );
diff --git a/Engine/Source/Editor/MainFrame/Private/Menus/MainMenu.cpp b/Engine/Source/Editor/MainFrame/Private/Menus/MainMenu.cpp
--- a/Engine/Source/Editor/MainFrame/Private/Menus/MainMenu.cpp
+++ b/Engine/Source/Editor/MainFrame/Private/Menus/MainMenu.cpp
@@ -285,6 +285,9 @@ TSharedRef<SWidget> FMainMenu::MakeMainMenu( const TSharedPtr<FTabManager>& TabM
 	// Create the menu bar!
 	TSharedRef<SWidget> MenuBarWidget = MenuBuilder.MakeWidget();
 
+	// Tell tab-manager about the multi-box for platforms with a global menu bar
+	TabManager->SetMenuMultiBox(MenuBuilder.GetMultiBox());
+	
 	return MenuBarWidget;
 }
 
diff --git a/Engine/Source/Runtime/Slate/Private/Framework/Docking/SDockingArea.cpp b/Engine/Source/Runtime/Slate/Private/Framework/Docking/SDockingArea.cpp
--- a/Engine/Source/Runtime/Slate/Private/Framework/Docking/SDockingArea.cpp
+++ b/Engine/Source/Runtime/Slate/Private/Framework/Docking/SDockingArea.cpp
@@ -251,6 +251,7 @@ void SDockingArea::SetParentWindow( const TSharedRef<SWindow>& NewParentWindow )
 	}
 
 	ParentWindowPtr = NewParentWindow;
+	NewParentWindow->SetOnWindowActivated( FOnWindowActivated::CreateSP( this, &SDockingArea::OnOwningWindowActivated ) );
 }
 
 TSharedPtr<FTabManager::FLayoutNode> SDockingArea::GatherPersistentLayout() const
@@ -423,6 +424,25 @@ void SDockingArea::OnOwningWindowBeingDestroyed(const TSharedRef<SWindow>& Windo
 
 }
 
+void SDockingArea::OnOwningWindowActivated()
+{
+	// Update the global menu bar when the window activation changes.
+	TSharedPtr< SWidget > MenuBar;
+	TSharedPtr<FTabManager> TabManager = MyTabManager.Pin();
+	if(FGlobalTabmanager::Get() == TabManager)
+	{
+		TSharedPtr<FGlobalTabmanager> GlobalTabManager = FGlobalTabmanager::Get();
+		TSharedPtr<SDockTab> ActiveTab = GlobalTabManager->GetActiveTab();
+		if(ActiveTab.IsValid())
+		{
+			ActiveTab->GetTabManager()->UpdateMainMenu(true);
+		}
+	}
+	else
+	{
+		TabManager->UpdateMainMenu(true);
+	}
+}
 
 void SDockingArea::OnLiveTabAdded()
 {
diff --git a/Engine/Source/Runtime/Slate/Private/Framework/Docking/SDockingArea.h b/Engine/Source/Runtime/Slate/Private/Framework/Docking/SDockingArea.h
--- a/Engine/Source/Runtime/Slate/Private/Framework/Docking/SDockingArea.h
+++ b/Engine/Source/Runtime/Slate/Private/Framework/Docking/SDockingArea.h
@@ -87,6 +87,9 @@ private:
 	/** We were placed in a window, and it is being destroyed */
 	void OnOwningWindowBeingDestroyed(const TSharedRef<SWindow>& WindowBeingDestroyed);
 
+	/** We were placed in a window and it is being activated */
+	void OnOwningWindowActivated();
+	
 	virtual void OnLiveTabAdded() override;
 
 	/**
diff --git a/Engine/Source/Runtime/Slate/Private/Framework/Docking/SDockingTabStack.cpp b/Engine/Source/Runtime/Slate/Private/Framework/Docking/SDockingTabStack.cpp
--- a/Engine/Source/Runtime/Slate/Private/Framework/Docking/SDockingTabStack.cpp
+++ b/Engine/Source/Runtime/Slate/Private/Framework/Docking/SDockingTabStack.cpp
@@ -254,9 +254,9 @@ void SDockingTabStack::Tick( const FGeometry& AllottedGeometry, const double InC
 void SDockingTabStack::OpenTab( const TSharedRef<SDockTab>& InTab, int32 InsertLocationAmongActiveTabs )
 {
 	const int32 TabIndex = OpenPersistentTab( InTab->GetLayoutIdentifier(), InsertLocationAmongActiveTabs );
-	AddTabWidget( InTab, TabIndex );
 	// The tab may be a nomad tab, in which case it should inherit whichever tab manager it is being put into!
 	InTab->SetTabManager( GetDockArea()->GetTabManager() );
+	AddTabWidget( InTab, TabIndex );
 	OnLiveTabAdded();
 	TabWell->RefreshParentContent();
 }
@@ -388,6 +388,7 @@ void SDockingTabStack::OnKeyboardFocusChanging( const FWeakWidgetPath& PreviousF
 		{
 			// If a widget inside this tab stack got focused, activate this tab.
 			FGlobalTabmanager::Get()->SetActiveTab( ForegroundTab );
+			ForegroundTab->ActivateInParent(ETabActivationCause::SetDirectly);
 		}
 	}
 }
diff --git a/Engine/Source/Runtime/Slate/Private/Framework/Docking/SDockingTabWell.cpp b/Engine/Source/Runtime/Slate/Private/Framework/Docking/SDockingTabWell.cpp
--- a/Engine/Source/Runtime/Slate/Private/Framework/Docking/SDockingTabWell.cpp
+++ b/Engine/Source/Runtime/Slate/Private/Framework/Docking/SDockingTabWell.cpp
@@ -408,6 +408,20 @@ void SDockingTabWell::BringTabToFront( int32 TabIndexToActivate )
 
 	// Always force a refresh, even if we don't think the active index changed.
 	RefreshParentContent();
+
+	// Update the native, global menu bar if a tab is in the foreground.
+	if( ForegroundTabIndex != INDEX_NONE )
+	{
+		TSharedPtr<FTabManager> TabManager = Tabs[ForegroundTabIndex]->GetTabManager();
+		if(TabManager == FGlobalTabmanager::Get())
+		{
+			FGlobalTabmanager::Get()->UpdateMainMenu(Tabs[ForegroundTabIndex], false);
+		}
+		else
+		{
+			TabManager->UpdateMainMenu(false);
+		}
+	}
 }
 
 /** Activate the tab specified by TabToActivate SDockTab. */
diff --git a/Engine/Source/Runtime/Slate/Private/Framework/Docking/TabManager.cpp b/Engine/Source/Runtime/Slate/Private/Framework/Docking/TabManager.cpp
--- a/Engine/Source/Runtime/Slate/Private/Framework/Docking/TabManager.cpp
+++ b/Engine/Source/Runtime/Slate/Private/Framework/Docking/TabManager.cpp
@@ -2,6 +2,9 @@
 
 #include "SlatePrivatePCH.h"
 #include "DockingPrivate.h"
+#if PLATFORM_MAC
+#include "../MultiBox/Mac/MacMenu.h"
+#endif
 
 
 const FVector2D FTabManager::FallbackWindowSize( 1000, 600 );
@@ -437,6 +440,10 @@ void FTabManager::FPrivateApi::HideWindows()
 	SetWindowVisibility(TabManager.DockAreas, false);
 }
 
+FTabManager::FPrivateApi& FTabManager::GetPrivateApi()
+{
+	return *PrivateApi;
+}
 
 
 
@@ -444,9 +451,45 @@ void FTabManager::FPrivateApi::HideWindows()
 
 
 
-FTabManager::FPrivateApi& FTabManager::GetPrivateApi()
+
+
+
+void FTabManager::SetMenuMultiBox(const TSharedPtr< FMultiBox >& NewMenuMutliBox)
 {
-	return *PrivateApi;
+	// We only use the platform native global menu bar on Mac
+#if PLATFORM_MAC
+	MenuMultiBox = NewMenuMutliBox;
+	if(MenuMultiBox.IsValid())
+	{
+		UpdateMainMenu(false);
+	}
+#endif
+}
+
+void FTabManager::UpdateMainMenu(bool const bForce)
+{
+	// We only use the platform native global menu bar on Mac
+#if PLATFORM_MAC
+	if(MenuMultiBox.IsValid())
+	{
+		bool bUpdate = bForce;
+		// On OS X opening the tab will set the multi-box and take key focus, but not seemingly send a keyboard focus event into Slate.
+		// I'm still looking into this, so for now we just update the menu bar here if the new tab is foreground in the focused window.
+		TSharedPtr<SDockTab> Tab = OwnerTabPtr.Pin();
+		if(Tab.IsValid() && Tab->IsForeground())
+		{
+			TSharedPtr<SWindow> ParentWindow = Tab->GetParentWindow();
+			if(ParentWindow.IsValid())
+			{
+				bUpdate |= (ParentWindow->HasKeyboardFocus() || ParentWindow->HasFocusedDescendants());
+			}
+		}
+		if(bUpdate)
+		{
+			FSlateMacMenu::UpdateWithMultiBox(MenuMultiBox.ToSharedRef());
+		}
+	}
+#endif
 }
 
 TSharedPtr<FTabManager::FStack> FTabManager::FLayoutNode::AsStack()
@@ -1579,6 +1622,20 @@ TSharedRef<FTabManager> FGlobalTabmanager::NewTabManager( const TSharedRef<SDock
 	return NewTabManager;
 }
 
+void FGlobalTabmanager::UpdateMainMenu(const TSharedRef<SDockTab>& ForTab, bool const bForce)
+{
+	TSharedPtr<FTabManager> Tabmanager = ForTab->GetTabManager();
+	if(Tabmanager == AsShared())
+	{
+		const int32 TabIndex = SubTabManagers.FindMatch( FindByTab(ForTab) );
+		if (TabIndex != INDEX_NONE)
+		{
+			Tabmanager = SubTabManagers[TabIndex].TabManager.Pin();
+		}
+	}
+	Tabmanager->UpdateMainMenu(bForce);
+}
+
 void FGlobalTabmanager::SaveAllVisualState()
 {
 	this->SavePersistentLayout();
diff --git a/Engine/Source/Runtime/Slate/Private/Framework/MultiBox/Mac/MacMenu.cpp b/Engine/Source/Runtime/Slate/Private/Framework/MultiBox/Mac/MacMenu.cpp
new file mode 100644
--- /dev/null
+++ b/Engine/Source/Runtime/Slate/Private/Framework/MultiBox/Mac/MacMenu.cpp
@@ -0,0 +1,348 @@
+// Copyright 1998-2014 Epic Games, Inc. All Rights Reserved.
+
+#include "SlatePrivatePCH.h"
+#include "MacMenu.h"
+
+@interface FMacMenuItem : NSMenuItem
+{
+	TSharedPtr< const FMenuEntryBlock > MenuEntryBlock;
+}
+
+- (void)performAction;
+
+@end
+
+@implementation FMacMenuItem
+
+- (id)initWithMenuEntryBlock:(TSharedRef< const FMenuEntryBlock >&)Block
+{
+	CFStringRef Title = FSlateMacMenu::GetMenuItemTitle(Block);
+	uint32 KeyModifiers;
+	CFStringRef KeyEquivalent = FSlateMacMenu::GetMenuItemKeyEquivalent(Block, &KeyModifiers);
+	self = [super initWithTitle:(NSString*)Title action:nil keyEquivalent:(NSString*)KeyEquivalent];
+	[self setKeyEquivalentModifierMask:KeyModifiers];
+	NSImage* MenuImage = FSlateMacMenu::GetMenuItemIcon(Block);
+	if(MenuImage)
+	{
+		[self setImage:MenuImage];
+		[MenuImage release];
+	}
+	CFRelease(Title);
+	CFRelease(KeyEquivalent);
+
+	MenuEntryBlock = Block;
+
+	if (FSlateMacMenu::IsMenuItemEnabled(Block))
+	{
+		[self setTarget:self];
+		[self setAction:@selector(performAction)];
+	}
+
+	return self;
+}
+
+- (void)performAction
+{
+	FSlateMacMenu::ExecuteMenuItemAction(MenuEntryBlock.ToSharedRef());
+}
+
+@end
+
+@implementation FMacMenu
+
+- (id)initWithMenuEntryBlock:(TSharedRef< const FMenuEntryBlock >&)Block
+{
+	CFStringRef Title = FSlateMacMenu::GetMenuItemTitle(Block);
+	self = [super initWithTitle:(NSString*)Title];
+	CFRelease(Title);
+
+	[self setDelegate:self];
+
+	MenuEntryBlock = Block;
+	
+	return self;
+}
+
+- (void)dealloc
+{
+	[super dealloc];
+}
+
+- (void)menuNeedsUpdate:(NSMenu*)Menu
+{
+	FSlateMacMenu::UpdateMenu(self);
+}
+
+- ( TSharedPtr< const FMenuEntryBlock >& )GetMenuEntryBlock
+{
+	return MenuEntryBlock;
+}
+
+@end
+
+void FSlateMacMenu::UpdateWithMultiBox(const TSharedRef< FMultiBox >& MultiBox)
+{
+	int32 NumItems = [[NSApp mainMenu] numberOfItems];
+	FText HelpTitle = NSLOCTEXT("MainMenu", "HelpMenu", "Help");
+	NSMenuItem* HelpMenu = [[[NSApp mainMenu] itemWithTitle:HelpTitle.ToString().GetNSString()] retain];
+	for (int32 Index = NumItems - 1; Index > 0; Index--)
+	{
+		[[NSApp mainMenu] removeItemAtIndex:Index];
+	}
+	
+	const TArray<TSharedRef<const FMultiBlock> >& MenuBlocks = MultiBox->GetBlocks();
+	
+	for (int32 Index = 0; Index < MenuBlocks.Num(); Index++)
+	{
+		TSharedRef<const FMenuEntryBlock> Block = StaticCastSharedRef<const FMenuEntryBlock>(MenuBlocks[Index]);
+		FMacMenu* Menu = [[FMacMenu alloc] initWithMenuEntryBlock:Block];
+		
+		NSMenuItem* MenuItem = [[NSMenuItem new] autorelease];
+		[MenuItem setTitle:[Menu title]];
+		[[NSApp mainMenu] addItem:MenuItem];
+		[MenuItem setSubmenu:Menu];
+	}
+	
+	if([[NSApp mainMenu] itemWithTitle:HelpTitle.ToString().GetNSString()] == nil && HelpMenu)
+	{
+		[[NSApp mainMenu] addItem:HelpMenu];
+	}
+	if(HelpMenu)
+	{
+		[HelpMenu release];
+	}
+}
+
+static const TSharedRef<SWidget> FindTextBlockWidget(TSharedRef<SWidget> Content)
+{
+	if (Content->GetType() == FName(TEXT("STextBlock")))
+	{
+		return Content;
+	}
+
+	FChildren* Children = Content->GetChildren();
+	int32 NumChildren = Children->Num();
+
+	for (int32 Index=0; Index < NumChildren; ++Index)
+	{
+		const TSharedRef<SWidget> Found = FindTextBlockWidget(Children->GetChildAt(Index));
+		if (Found != SNullWidget::NullWidget)
+		{
+			return Found;
+		}
+	}
+	return SNullWidget::NullWidget;
+}
+
+CFStringRef FSlateMacMenu::GetMenuItemTitle(const TSharedRef<const FMenuEntryBlock>& Block)
+{
+	TAttribute<FText> Label;
+	if (!Block->LabelOverride.IsBound() && Block->LabelOverride.Get().IsEmpty() && Block->GetAction().IsValid())
+	{
+		Label = Block->GetAction()->GetLabel();
+	}
+	else if (!Block->LabelOverride.Get().IsEmpty())
+	{
+		Label = Block->LabelOverride;
+	}
+	else if (Block->EntryWidget.IsValid())
+	{
+		const TSharedRef<SWidget>& TextBlockWidget = FindTextBlockWidget(Block->EntryWidget.ToSharedRef());
+		if (TextBlockWidget != SNullWidget::NullWidget)
+		{
+			Label = FText::FromString(StaticCastSharedRef<STextBlock>(TextBlockWidget)->GetText());
+		}
+	}
+
+	CFStringRef Title = FPlatformString::TCHARToCFString(*Label.Get().ToString());
+	return Title;
+}
+
+NSImage* FSlateMacMenu::GetMenuItemIcon(const TSharedRef<const FMenuEntryBlock>& Block)
+{
+	NSImage* MenuImage = nil;
+	FSlateIcon Icon;
+	if(Block->IconOverride.IsSet())
+	{
+		Icon = Block->IconOverride;
+	}
+	else if(Block->GetAction().IsValid() && Block->GetAction()->GetIcon().IsSet())
+	{
+		Icon = Block->GetAction()->GetIcon();
+	}
+	if(Icon.IsSet())
+	{
+		if(Icon.GetIcon())
+		{
+			FSlateBrush const* IconBrush = Icon.GetIcon();
+			FName ResourceName = IconBrush->GetResourceName();
+			MenuImage = [[NSImage alloc] initWithContentsOfFile: ResourceName.ToString().GetNSString()];
+			if(MenuImage)
+			{
+				[MenuImage setSize:NSMakeSize(16.0f, 16.0f)];
+			}
+		}
+	}
+	return MenuImage;
+}
+
+void FSlateMacMenu::ExecuteMenuItemAction(const TSharedRef< const class FMenuEntryBlock >& Block)
+{
+	TSharedPtr< const FUICommandList > ActionList = Block->GetActionList();
+	if (ActionList.IsValid() && Block->GetAction().IsValid())
+	{
+		ActionList->ExecuteAction(Block->GetAction().ToSharedRef());
+	}
+	else
+	{
+		// There is no action list or action associated with this block via a UI command.  Execute any direct action we have
+		Block->GetDirectActions().Execute();
+	}
+}
+
+CFStringRef FSlateMacMenu::GetMenuItemKeyEquivalent(const TSharedRef<const class FMenuEntryBlock>& Block, uint32* OutModifiers)
+{
+	if (Block->GetAction().IsValid())
+	{
+		const TSharedRef<const FInputGesture>& Gesture = Block->GetAction()->GetActiveGesture();
+
+		*OutModifiers = 0;
+		if (Gesture->bCtrl)
+		{
+			*OutModifiers |= NSCommandKeyMask;
+		}
+		if (Gesture->bShift)
+		{
+			*OutModifiers |= NSShiftKeyMask;
+		}
+		if (Gesture->bAlt)
+		{
+			*OutModifiers |= NSAlternateKeyMask;
+		}
+
+		FString KeyString = Gesture->GetKeyText().ToString().ToLower();
+		return FPlatformString::TCHARToCFString(*KeyString);
+	}
+	return FPlatformString::TCHARToCFString(TEXT(""));
+}
+
+void FSlateMacMenu::UpdateMenu(FMacMenu* Menu)
+{
+	FText WindowLabel = NSLOCTEXT("MainMenu", "WindowMenu", "Window");
+	
+	bool bIsWindowMenu = (WindowLabel.ToString().Compare(FString([Menu title])) == 0);
+	
+	[Menu removeAllItems];
+	
+	if(bIsWindowMenu)
+	{
+		NSMenuItem* MinimizeItem = [[[NSMenuItem alloc] initWithTitle:@"Minimize" action:@selector(miniaturize:) keyEquivalent:@"m"] autorelease];
+		NSMenuItem* ZoomItem = [[[NSMenuItem alloc] initWithTitle:@"Zoom" action:@selector(performZoom:) keyEquivalent:@""] autorelease];
+		NSMenuItem* BringAllToFrontItem = [[[NSMenuItem alloc] initWithTitle:@"Bring All to Front" action:@selector(arrangeInFront:) keyEquivalent:@""] autorelease];
+		
+		[Menu addItem:MinimizeItem];
+		[Menu addItem:ZoomItem];
+		[Menu addItem:[NSMenuItem separatorItem]];
+		[Menu addItem:BringAllToFrontItem];
+		[Menu addItem:[NSMenuItem separatorItem]];
+	}
+
+	const bool bShouldCloseWindowAfterMenuSelection = true;
+	FMenuBuilder MenuBuilder(bShouldCloseWindowAfterMenuSelection, [Menu GetMenuEntryBlock]->GetActionList(), [Menu GetMenuEntryBlock]->Extender);
+	{
+		// Have the menu fill its contents
+		[Menu GetMenuEntryBlock]->EntryBuilder.ExecuteIfBound(MenuBuilder);
+	}
+	TSharedRef< SWidget > Widget = MenuBuilder.MakeWidget();
+
+	TSharedRef<const FMultiBox> MultiBox = StaticCastSharedRef<SMultiBoxWidget>(Widget)->GetMultiBox();
+	const TArray<TSharedRef<const FMultiBlock>>& MenuBlocks = MultiBox->GetBlocks();
+
+	for (int32 Index = 0; Index < MenuBlocks.Num(); Index++)
+	{
+		EMultiBlockType::Type Type = MenuBlocks[Index]->GetType();
+
+		if (Type == EMultiBlockType::MenuEntry)
+		{
+			TSharedRef<const FMenuEntryBlock> Block = StaticCastSharedRef<const FMenuEntryBlock>(MenuBlocks[Index]);
+			FMacMenuItem* MenuItem = [[[FMacMenuItem alloc] initWithMenuEntryBlock:Block] autorelease];
+
+			if (Block->bIsSubMenu)
+			{
+				FMacMenu* SubMenu = [[FMacMenu alloc] initWithMenuEntryBlock:Block];
+				[MenuItem setSubmenu:SubMenu];
+			}
+			else
+			{
+				[MenuItem setState:FSlateMacMenu::GetMenuItemState(Block)];
+			}
+			
+			if(!bIsWindowMenu)
+			{
+				bool bIsWindowSubMenu = false;
+				NSMenu* SuperMenu = [Menu supermenu];
+				while(!bIsWindowSubMenu && SuperMenu)
+				{
+					bIsWindowSubMenu = (WindowLabel.ToString().Compare(FString([SuperMenu title])) == 0);
+					SuperMenu = [SuperMenu supermenu];
+				}
+				
+				if(!bIsWindowSubMenu)
+				{
+					[MenuItem setImage:nil];
+				}
+			}
+
+			[Menu addItem:MenuItem];
+		}
+		else if (Type == EMultiBlockType::MenuSeparator)
+		{
+			[Menu addItem:[NSMenuItem separatorItem]];
+		}
+	}
+	
+	if(bIsWindowMenu)
+	{
+		[NSApp setWindowsMenu:Menu];
+	}
+}
+
+bool FSlateMacMenu::IsMenuItemEnabled(const TSharedRef<const class FMenuEntryBlock>& Block)
+{
+	TSharedPtr<const FUICommandList> ActionList = Block->GetActionList();
+	TSharedPtr<const FUICommandInfo> Action = Block->GetAction();
+	const FUIAction& DirectActions = Block->GetDirectActions();
+
+	bool bEnabled = true;
+	if (ActionList.IsValid() && Action.IsValid())
+	{
+		bEnabled = ActionList->CanExecuteAction(Action.ToSharedRef());
+	}
+	else
+	{
+		// There is no action list or action associated with this block via a UI command.  Execute any direct action we have
+		bEnabled = DirectActions.CanExecute();
+	}
+
+	return bEnabled;
+}
+
+int32 FSlateMacMenu::GetMenuItemState(const TSharedRef<const class FMenuEntryBlock>& Block)
+{
+	TSharedPtr<const FUICommandList> ActionList = Block->GetActionList();
+	TSharedPtr<const FUICommandInfo> Action = Block->GetAction();
+	const FUIAction& DirectActions = Block->GetDirectActions();
+
+	bool bIsChecked = true;
+	if (ActionList.IsValid() && Action.IsValid())
+	{
+		bIsChecked = ActionList->IsChecked(Action.ToSharedRef());
+	}
+	else
+	{
+		// There is no action list or action associated with this block via a UI command.  Execute any direct action we have
+		bIsChecked = DirectActions.IsChecked();
+	}
+
+	return bIsChecked ? NSOnState : NSOffState;
+}
diff --git a/Engine/Source/Runtime/Slate/Private/Framework/MultiBox/Mac/MacMenu.h b/Engine/Source/Runtime/Slate/Private/Framework/MultiBox/Mac/MacMenu.h
new file mode 100644
--- /dev/null
+++ b/Engine/Source/Runtime/Slate/Private/Framework/MultiBox/Mac/MacMenu.h
@@ -0,0 +1,31 @@
+// Copyright 1998-2014 Epic Games, Inc. All Rights Reserved.
+
+#pragma once
+
+#include "MultiBox.h"
+#include "SMenuEntryBlock.h"
+
+#ifdef __OBJC__
+@interface FMacMenu : NSMenu <NSMenuDelegate>
+{
+	TSharedPtr<const FMenuEntryBlock> MenuEntryBlock;
+}
+- (TSharedPtr<const FMenuEntryBlock>&)GetMenuEntryBlock;
+@end
+#else
+class FMacMenu;
+#endif
+
+class SLATE_API FSlateMacMenu
+{
+public:
+
+	static void UpdateWithMultiBox(const TSharedRef< FMultiBox >& MultiBox);
+	static CFStringRef GetMenuItemTitle(const TSharedRef< const class FMenuEntryBlock >& Block);
+	static NSImage* GetMenuItemIcon(const TSharedRef<const FMenuEntryBlock>& Block);
+	static void ExecuteMenuItemAction(const TSharedRef< const class FMenuEntryBlock >& Block);
+	static CFStringRef GetMenuItemKeyEquivalent(const TSharedRef< const class FMenuEntryBlock >& Block, uint32* OutModifiers);
+	static void UpdateMenu(FMacMenu* InMenu);
+	static bool IsMenuItemEnabled(const TSharedRef<const class FMenuEntryBlock>& Block);
+	static int32 GetMenuItemState(const TSharedRef<const class FMenuEntryBlock>& Block);
+};
diff --git a/Engine/Source/Runtime/Slate/Private/Framework/MultiBox/MultiBox.cpp b/Engine/Source/Runtime/Slate/Private/Framework/MultiBox/MultiBox.cpp
--- a/Engine/Source/Runtime/Slate/Private/Framework/MultiBox/MultiBox.cpp
+++ b/Engine/Source/Runtime/Slate/Private/Framework/MultiBox/MultiBox.cpp
@@ -219,7 +219,14 @@ TSharedRef< SMultiBoxWidget > FMultiBox::MakeWidget()
 
 	// Build up the widget
 	NewMultiBoxWidget->BuildMultiBoxWidget();
-
+	
+#if PLATFORM_MAC
+	if(Type == EMultiBoxType::MenuBar)
+	{
+		NewMultiBoxWidget->SetVisibility(EVisibility::Collapsed);
+	}
+#endif
+	
 	return NewMultiBoxWidget;
 }
 
diff --git a/Engine/Source/Runtime/Slate/Private/Framework/MultiBox/MultiBoxBuilder.cpp b/Engine/Source/Runtime/Slate/Private/Framework/MultiBox/MultiBoxBuilder.cpp
--- a/Engine/Source/Runtime/Slate/Private/Framework/MultiBox/MultiBoxBuilder.cpp
+++ b/Engine/Source/Runtime/Slate/Private/Framework/MultiBox/MultiBoxBuilder.cpp
@@ -85,6 +85,10 @@ TSharedRef< class SWidget > FMultiBoxBuilder::MakeWidget()
 	return MultiBox->MakeWidget();
 }
 
+TSharedRef< class FMultiBox > FMultiBoxBuilder::GetMultiBox()
+{
+	return MultiBox;
+}
 
 FBaseMenuBuilder::FBaseMenuBuilder( const EMultiBoxType::Type InType, const bool bInShouldCloseWindowAfterMenuSelection, TSharedPtr< const FUICommandList > InCommandList, bool bInCloseSelfOnly, TSharedPtr<FExtender> InExtender, const ISlateStyle* InStyleSet )
 	: FMultiBoxBuilder( InType, FMultiBoxCustomization::None, bInShouldCloseWindowAfterMenuSelection, InCommandList, InExtender )
diff --git a/Engine/Source/Runtime/Slate/Private/Framework/MultiBox/SButtonRowBlock.cpp b/Engine/Source/Runtime/Slate/Private/Framework/MultiBox/SButtonRowBlock.cpp
--- a/Engine/Source/Runtime/Slate/Private/Framework/MultiBox/SButtonRowBlock.cpp
+++ b/Engine/Source/Runtime/Slate/Private/Framework/MultiBox/SButtonRowBlock.cpp
@@ -7,7 +7,7 @@
 
 
 FButtonRowBlock::FButtonRowBlock( const TSharedPtr< const FUICommandInfo > InCommand, TSharedPtr< const FUICommandList > InCommandList, const TAttribute<FText>& InLabelOverride, const TAttribute<FText>& InToolTipOverride, const FSlateIcon& InIconOverride )
-	: FMultiBlock( InCommand, InCommandList )
+	: FMultiBlock( InCommand, InCommandList, NAME_None, EMultiBlockType::ButtonRow )
 	, LabelOverride( InLabelOverride )
 	, ToolTipOverride( InToolTipOverride )
 	, IconOverride( InIconOverride )
@@ -15,7 +15,7 @@ FButtonRowBlock::FButtonRowBlock( const TSharedPtr< const FUICommandInfo > InCom
 }
 
 FButtonRowBlock::FButtonRowBlock( const TAttribute<FText>& InLabel, const TAttribute<FText>& InToolTip, const FSlateIcon& InIcon, const FUIAction& UIAction, const EUserInterfaceActionType::Type InUserInterfaceActionType )
-	: FMultiBlock( UIAction )
+	: FMultiBlock( UIAction, NAME_None, EMultiBlockType::ButtonRow )
 	, LabelOverride( InLabel )
 	, ToolTipOverride( InToolTip )
 	, IconOverride( InIcon )
diff --git a/Engine/Source/Runtime/Slate/Private/Framework/MultiBox/SEditableTextBlock.cpp b/Engine/Source/Runtime/Slate/Private/Framework/MultiBox/SEditableTextBlock.cpp
--- a/Engine/Source/Runtime/Slate/Private/Framework/MultiBox/SEditableTextBlock.cpp
+++ b/Engine/Source/Runtime/Slate/Private/Framework/MultiBox/SEditableTextBlock.cpp
@@ -9,7 +9,7 @@
 
 
 FEditableTextBlock::FEditableTextBlock( const FText& InLabel, const FText& InToolTip, const FSlateIcon& InIcon, const TAttribute< FText >& InTextAttribute, bool bInReadOnly, const FOnTextCommitted& InOnTextCommitted, const FOnTextChanged& InOnTextChanged )
-	: FMultiBlock( FUIAction() ),
+	: FMultiBlock( FUIAction(), NAME_None, EMultiBlockType::EditableText ),
 	  LabelOverride( InLabel ),
 	  ToolTipOverride( InToolTip ),
 	  IconOverride( InIcon ),
diff --git a/Engine/Source/Runtime/Slate/Private/Framework/MultiBox/SHeadingBlock.cpp b/Engine/Source/Runtime/Slate/Private/Framework/MultiBox/SHeadingBlock.cpp
--- a/Engine/Source/Runtime/Slate/Private/Framework/MultiBox/SHeadingBlock.cpp
+++ b/Engine/Source/Runtime/Slate/Private/Framework/MultiBox/SHeadingBlock.cpp
@@ -11,7 +11,7 @@
  * @param	InHeadingText	Heading text
  */
 FHeadingBlock::FHeadingBlock( const FName& InExtensionHook, const TAttribute< FText >& InHeadingText )
-	: FMultiBlock( NULL, NULL, InExtensionHook )
+	: FMultiBlock( NULL, NULL, InExtensionHook, EMultiBlockType::Heading )
 	, HeadingText( InHeadingText )
 {
 }
diff --git a/Engine/Source/Runtime/Slate/Private/Framework/MultiBox/SMenuEntryBlock.cpp b/Engine/Source/Runtime/Slate/Private/Framework/MultiBox/SMenuEntryBlock.cpp
--- a/Engine/Source/Runtime/Slate/Private/Framework/MultiBox/SMenuEntryBlock.cpp
+++ b/Engine/Source/Runtime/Slate/Private/Framework/MultiBox/SMenuEntryBlock.cpp
@@ -6,7 +6,7 @@
 
 
 FMenuEntryBlock::FMenuEntryBlock( const FName& InExtensionHook, const TSharedPtr< const FUICommandInfo > InCommand, TSharedPtr< const FUICommandList > InCommandList, const TAttribute<FText>& InLabelOverride, const TAttribute<FText>& InToolTipOverride, const FSlateIcon& InIconOverride, bool bInCloseSelfOnly )
-	: FMultiBlock( InCommand, InCommandList, InExtensionHook )
+	: FMultiBlock( InCommand, InCommandList, InExtensionHook, EMultiBlockType::MenuEntry )
 	, LabelOverride( InLabelOverride )
 	, ToolTipOverride( InToolTipOverride )
 	, IconOverride( InIconOverride )
@@ -19,7 +19,7 @@ FMenuEntryBlock::FMenuEntryBlock( const FName& InExtensionHook, const TSharedPtr
 
 
 FMenuEntryBlock::FMenuEntryBlock( const FName& InExtensionHook, const TAttribute<FText>& InLabel, const TAttribute<FText>& InToolTip, const FNewMenuDelegate& InEntryBuilder, TSharedPtr<FExtender> InExtender, bool bInSubMenu, bool bInSubMenuOnClick, const FSlateIcon& InIcon, const FUIAction& InUIAction, const EUserInterfaceActionType::Type InUserInterfaceActionType, bool bInCloseSelfOnly )
-	: FMultiBlock( InUIAction, InExtensionHook )
+	: FMultiBlock( InUIAction, InExtensionHook, EMultiBlockType::MenuEntry )
 	, LabelOverride( InLabel )
 	, ToolTipOverride( InToolTip )
 	, IconOverride( InIcon )
@@ -34,7 +34,7 @@ FMenuEntryBlock::FMenuEntryBlock( const FName& InExtensionHook, const TAttribute
 
 
 FMenuEntryBlock::FMenuEntryBlock( const FName& InExtensionHook, const TAttribute<FText>& InLabel, const TAttribute<FText>& InToolTip, const FSlateIcon& InIcon, const FUIAction& UIAction, const EUserInterfaceActionType::Type InUserInterfaceActionType, bool bInCloseSelfOnly )
-	: FMultiBlock( UIAction, InExtensionHook )
+	: FMultiBlock( UIAction, InExtensionHook, EMultiBlockType::MenuEntry )
 	, LabelOverride( InLabel )
 	, ToolTipOverride( InToolTip )
 	, IconOverride( InIcon )
@@ -47,7 +47,7 @@ FMenuEntryBlock::FMenuEntryBlock( const FName& InExtensionHook, const TAttribute
 
 
 FMenuEntryBlock::FMenuEntryBlock( const FName& InExtensionHook, const TAttribute<FText>& InLabel, const TAttribute<FText>& InToolTip, const FNewMenuDelegate& InEntryBuilder, TSharedPtr<FExtender> InExtender, bool bInSubMenu, bool bInSubMenuOnClick, TSharedPtr< const FUICommandList > InCommandList, bool bInCloseSelfOnly, const FSlateIcon& InIcon )
-	: FMultiBlock( NULL, InCommandList, InExtensionHook )
+	: FMultiBlock( NULL, InCommandList, InExtensionHook, EMultiBlockType::MenuEntry )
 	, LabelOverride( InLabel )
 	, ToolTipOverride( InToolTip )
 	, IconOverride( InIcon )
@@ -62,7 +62,7 @@ FMenuEntryBlock::FMenuEntryBlock( const FName& InExtensionHook, const TAttribute
 
 
 FMenuEntryBlock::FMenuEntryBlock( const FName& InExtensionHook, const TAttribute<FText>& InLabel, const TAttribute<FText>& InToolTip, const FOnGetContent& InMenuBuilder, TSharedPtr<FExtender> InExtender, bool bInSubMenu, bool bInSubMenuOnClick, TSharedPtr< const FUICommandList > InCommandList, bool bInCloseSelfOnly, const FSlateIcon& InIcon )
-	: FMultiBlock( NULL, InCommandList, InExtensionHook )
+	: FMultiBlock( NULL, InCommandList, InExtensionHook, EMultiBlockType::MenuEntry )
 	, LabelOverride( InLabel )
 	, ToolTipOverride( InToolTip )
 	, IconOverride( InIcon )
@@ -77,7 +77,7 @@ FMenuEntryBlock::FMenuEntryBlock( const FName& InExtensionHook, const TAttribute
 
 
 FMenuEntryBlock::FMenuEntryBlock( const FName& InExtensionHook, const TAttribute<FText>& InLabel, const TAttribute<FText>& InToolTip, const TSharedPtr<SWidget>& InEntryWidget, TSharedPtr<FExtender> InExtender, bool bInSubMenu, bool bInSubMenuOnClick, TSharedPtr< const FUICommandList > InCommandList, bool bInCloseSelfOnly, const FSlateIcon& InIcon )
-	: FMultiBlock( NULL, InCommandList, InExtensionHook )
+	: FMultiBlock( NULL, InCommandList, InExtensionHook, EMultiBlockType::MenuEntry )
 	, LabelOverride( InLabel )
 	, ToolTipOverride( InToolTip )
 	, IconOverride( InIcon )
@@ -91,7 +91,7 @@ FMenuEntryBlock::FMenuEntryBlock( const FName& InExtensionHook, const TAttribute
 }
 
 FMenuEntryBlock::FMenuEntryBlock( const FName& InExtensionHook, const FUIAction& UIAction, const TSharedRef< SWidget > Contents, const TAttribute<FText>& InToolTip, const EUserInterfaceActionType::Type InUserInterfaceActionType, bool bInCloseSelfOnly )
-	: FMultiBlock( UIAction, InExtensionHook )
+	: FMultiBlock( UIAction, InExtensionHook, EMultiBlockType::MenuEntry )
 	, ToolTipOverride( InToolTip )
 	, EntryWidget( Contents )
 	, bIsSubMenu( false )
@@ -103,7 +103,7 @@ FMenuEntryBlock::FMenuEntryBlock( const FName& InExtensionHook, const FUIAction&
 
 
 FMenuEntryBlock::FMenuEntryBlock( const FName& InExtensionHook, const TSharedRef< SWidget > Contents, const FNewMenuDelegate& InEntryBuilder, TSharedPtr<FExtender> InExtender, bool bInSubMenu, bool bInSubMenuOnClick, TSharedPtr< const FUICommandList > InCommandList, bool bInCloseSelfOnly )
-	: FMultiBlock( NULL, InCommandList, InExtensionHook )
+	: FMultiBlock( NULL, InCommandList, InExtensionHook, EMultiBlockType::MenuEntry )
 	, EntryBuilder( InEntryBuilder )
 	, EntryWidget( Contents )
 	, bIsSubMenu( bInSubMenu )
@@ -116,7 +116,7 @@ FMenuEntryBlock::FMenuEntryBlock( const FName& InExtensionHook, const TSharedRef
 
 
 FMenuEntryBlock::FMenuEntryBlock( const FName& InExtensionHook, const FUIAction& UIAction, const TSharedRef< SWidget > Contents, const FNewMenuDelegate& InEntryBuilder, TSharedPtr<FExtender> InExtender, bool bInSubMenu, TSharedPtr< const FUICommandList > InCommandList, bool bInCloseSelfOnly )
-	: FMultiBlock( UIAction, InExtensionHook )
+	: FMultiBlock( UIAction, InExtensionHook, EMultiBlockType::MenuEntry )
 	, EntryBuilder( InEntryBuilder )
 	, EntryWidget( Contents )
 	, bIsSubMenu( bInSubMenu )
diff --git a/Engine/Source/Runtime/Slate/Private/Framework/MultiBox/SMenuEntryBlock.h b/Engine/Source/Runtime/Slate/Private/Framework/MultiBox/SMenuEntryBlock.h
--- a/Engine/Source/Runtime/Slate/Private/Framework/MultiBox/SMenuEntryBlock.h
+++ b/Engine/Source/Runtime/Slate/Private/Framework/MultiBox/SMenuEntryBlock.h
@@ -105,6 +105,7 @@ private:
 
 	// Friend our corresponding widget class
 	friend class SMenuEntryBlock;
+	friend class FSlateMacMenu;
 
 	/** Optional overridden text label for this menu entry.  If not set, then the action's label will be used instead. */
 	TAttribute<FText> LabelOverride;
diff --git a/Engine/Source/Runtime/Slate/Private/Framework/MultiBox/SMenuSeparatorBlock.cpp b/Engine/Source/Runtime/Slate/Private/Framework/MultiBox/SMenuSeparatorBlock.cpp
--- a/Engine/Source/Runtime/Slate/Private/Framework/MultiBox/SMenuSeparatorBlock.cpp
+++ b/Engine/Source/Runtime/Slate/Private/Framework/MultiBox/SMenuSeparatorBlock.cpp
@@ -9,7 +9,7 @@
  * Constructor
  */
 FMenuSeparatorBlock::FMenuSeparatorBlock(const FName& InExtensionHook)
-	: FMultiBlock( NULL, NULL, InExtensionHook )
+	: FMultiBlock( NULL, NULL, InExtensionHook, EMultiBlockType::MenuSeparator )
 {
 }
 
diff --git a/Engine/Source/Runtime/Slate/Private/Framework/MultiBox/SToolBarButtonBlock.cpp b/Engine/Source/Runtime/Slate/Private/Framework/MultiBox/SToolBarButtonBlock.cpp
--- a/Engine/Source/Runtime/Slate/Private/Framework/MultiBox/SToolBarButtonBlock.cpp
+++ b/Engine/Source/Runtime/Slate/Private/Framework/MultiBox/SToolBarButtonBlock.cpp
@@ -6,7 +6,7 @@
 
 
 FToolBarButtonBlock::FToolBarButtonBlock( const TSharedPtr< const FUICommandInfo > InCommand, TSharedPtr< const FUICommandList > InCommandList, const TAttribute<FText>& InLabelOverride, const TAttribute<FText>& InToolTipOverride, const TAttribute<FSlateIcon>& InIconOverride )
-	: FMultiBlock( InCommand, InCommandList )
+	: FMultiBlock( InCommand, InCommandList, NAME_None, EMultiBlockType::ToolBarButton )
 	, LabelOverride( InLabelOverride )
 	, ToolTipOverride( InToolTipOverride )
 	, IconOverride( InIconOverride )
diff --git a/Engine/Source/Runtime/Slate/Private/Framework/MultiBox/SToolBarComboButtonBlock.cpp b/Engine/Source/Runtime/Slate/Private/Framework/MultiBox/SToolBarComboButtonBlock.cpp
--- a/Engine/Source/Runtime/Slate/Private/Framework/MultiBox/SToolBarComboButtonBlock.cpp
+++ b/Engine/Source/Runtime/Slate/Private/Framework/MultiBox/SToolBarComboButtonBlock.cpp
@@ -6,7 +6,7 @@
 
 
 FToolBarComboButtonBlock::FToolBarComboButtonBlock( const FUIAction& InAction, const FOnGetContent& InMenuContentGenerator, const TAttribute<FText>& InLabel, const TAttribute<FText>& InToolTip, const TAttribute<FSlateIcon>& InIcon, bool bInSimpleComboBox )
-	: FMultiBlock( InAction )
+	: FMultiBlock( InAction, NAME_None, EMultiBlockType::ToolBarComboButton )
 	, MenuContentGenerator( InMenuContentGenerator )
 	, Label( InLabel )
 	, ToolTip( InToolTip )
diff --git a/Engine/Source/Runtime/Slate/Private/Framework/MultiBox/SToolBarSeparatorBlock.cpp b/Engine/Source/Runtime/Slate/Private/Framework/MultiBox/SToolBarSeparatorBlock.cpp
--- a/Engine/Source/Runtime/Slate/Private/Framework/MultiBox/SToolBarSeparatorBlock.cpp
+++ b/Engine/Source/Runtime/Slate/Private/Framework/MultiBox/SToolBarSeparatorBlock.cpp
@@ -9,7 +9,7 @@
  * Constructor
  */
 FToolBarSeparatorBlock::FToolBarSeparatorBlock(const FName& InExtensionHook)
-	: FMultiBlock( NULL, NULL, InExtensionHook )
+	: FMultiBlock( NULL, NULL, InExtensionHook, EMultiBlockType::ToolBarSeparator )
 {
 }
 
diff --git a/Engine/Source/Runtime/Slate/Private/Framework/MultiBox/SWidgetBlock.cpp b/Engine/Source/Runtime/Slate/Private/Framework/MultiBox/SWidgetBlock.cpp
--- a/Engine/Source/Runtime/Slate/Private/Framework/MultiBox/SWidgetBlock.cpp
+++ b/Engine/Source/Runtime/Slate/Private/Framework/MultiBox/SWidgetBlock.cpp
@@ -11,7 +11,7 @@
  * @param	InHeadingText	Heading text
  */
 FWidgetBlock::FWidgetBlock( TSharedRef<SWidget> InContent, const FText& InLabel, bool bInNoIndent )
-	: FMultiBlock( NULL, NULL )
+	: FMultiBlock( NULL, NULL, NAME_None, EMultiBlockType::Widget )
 	, ContentWidget( InContent )
 	, Label( InLabel )
 	, bNoIndent( bInNoIndent )
diff --git a/Engine/Source/Runtime/Slate/Public/Framework/Docking/TabManager.h b/Engine/Source/Runtime/Slate/Public/Framework/Docking/TabManager.h
--- a/Engine/Source/Runtime/Slate/Public/Framework/Docking/TabManager.h
+++ b/Engine/Source/Runtime/Slate/Public/Framework/Docking/TabManager.h
@@ -585,6 +585,17 @@ class SLATE_API FTabManager : public TSharedFromThis<FTabManager>
 		{
 		}
 
+		/**
+		 * Set the multi-box to use for generating a native, global menu bar.
+		 * @param NewMenuMutliBox The multi-box to generate the global menu bar from.
+		 */
+		void SetMenuMultiBox(const TSharedPtr< FMultiBox >& NewMenuMutliBox);
+
+		/**
+		 * Update the native, global menu bar if it is being used.
+		 * @param bForce Used to force an update even if the parent window doesn't contain the widget with keyboard focus.
+		 */
+		void UpdateMainMenu(bool const bForce);
 
 	protected:
 		void InvokeTabForMenu( FName TabId );
@@ -663,6 +674,9 @@ class SLATE_API FTabManager : public TSharedFromThis<FTabManager>
 		/** A Major tab that contains this TabManager's widgets. */
 		TWeakPtr<SDockTab> OwnerTabPtr;
 
+		/** The current menu multi-box for the tab, used to construct platform native main menus */
+		TSharedPtr< FMultiBox > MenuMultiBox;
+
 		/** Protected private API that must only be accessed by the docking framework internals */
 		TSharedRef<FPrivateApi> PrivateApi;
 
@@ -732,6 +746,13 @@ public:
 	void DrawAttentionToTabManager( const TSharedRef<FTabManager>& ChildManager );
 
 	TSharedRef<FTabManager> NewTabManager( const TSharedRef<SDockTab>& InOwnerTab );
+	
+	/**
+	 * Update the native, global menu bar if it is being used for a specific tab managed by the global tab manager.
+	 * @param ForTab The tab to update the main menu.
+	 * @param bForce Used to force an update even if the parent window doesn't contain the widget with keyboard focus.
+	 */
+	void UpdateMainMenu(const TSharedRef<SDockTab>& ForTab, bool const bForce);
 
 	/** Persist and serialize the layout of every TabManager and the custom visual state of every Tab. */
 	void SaveAllVisualState();
diff --git a/Engine/Source/Runtime/Slate/Public/Framework/MultiBox/MultiBox.h b/Engine/Source/Runtime/Slate/Public/Framework/MultiBox/MultiBox.h
--- a/Engine/Source/Runtime/Slate/Public/Framework/MultiBox/MultiBox.h
+++ b/Engine/Source/Runtime/Slate/Public/Framework/MultiBox/MultiBox.h
@@ -44,10 +44,11 @@ public:
 	 * @param InCommand		The command info that describes what action to take when this block is activated
 	 * @param InCommandList	The list of mappings from command info to delegates so we can find the delegates to process for the provided action
 	 */
-	FMultiBlock( const TSharedPtr< const FUICommandInfo > InCommand, TSharedPtr< const FUICommandList > InCommandList, FName InExtensionHook = NAME_None )
+	FMultiBlock( const TSharedPtr< const FUICommandInfo > InCommand, TSharedPtr< const FUICommandList > InCommandList, FName InExtensionHook = NAME_None, EMultiBlockType::Type InType = EMultiBlockType::None )
 		: Action( InCommand )
 		, ActionList( InCommandList )
 		, ExtensionHook( InExtensionHook )
+		, Type( InType )
 		, TutorialHighlightName( NAME_None )
 	{
 	}
@@ -57,9 +58,10 @@ public:
 	 *
 	 * @InAction UI action delegates that should be used in place of UI commands (dynamic menu items)
 	 */
-	FMultiBlock( const FUIAction& InAction,  FName InExtensionHook = NAME_None )
+	FMultiBlock( const FUIAction& InAction,  FName InExtensionHook = NAME_None, EMultiBlockType::Type InType = EMultiBlockType::None )
 		: DirectActions( InAction )
 		, ExtensionHook( InExtensionHook )
+		, Type( InType )
 		, TutorialHighlightName( NAME_None )
 	{
 	}
@@ -118,6 +120,16 @@ public:
 	 */
 	TSharedRef< class IMultiBlockBaseWidget > MakeWidget( TSharedRef< class SMultiBoxWidget > InOwnerMultiBoxWidget, EMultiBlockLocation::Type InLocation ) const;
 
+	/**
+	 * Gets the type of this MultiBox
+	 *
+	 * @return	The MultiBlock's type
+	 */
+	const EMultiBlockType::Type GetType() const
+	{
+		return Type;
+	}
+
 private:
 	/**
 	 * Allocates a widget for this type of MultiBlock.  Override this in derived classes.
@@ -146,6 +158,9 @@ private:
 	/** Optional extension hook which is used for debug display purposes, so users can see what hooks are where */
 	FName ExtensionHook;
 
+	/** Type of MultiBlock */
+	EMultiBlockType::Type Type;
+
 	/** Name to identify a widget for tutorials */
 	FName TutorialHighlightName;
 };
diff --git a/Engine/Source/Runtime/Slate/Public/Framework/MultiBox/MultiBoxBuilder.h b/Engine/Source/Runtime/Slate/Public/Framework/MultiBox/MultiBoxBuilder.h
--- a/Engine/Source/Runtime/Slate/Public/Framework/MultiBox/MultiBoxBuilder.h
+++ b/Engine/Source/Runtime/Slate/Public/Framework/MultiBox/MultiBoxBuilder.h
@@ -51,6 +51,14 @@ public:
 	 * @return  New widget object
 	 */
 	TSharedRef< class SWidget > MakeWidget();
+	
+	
+	/** 
+	 * Get the multi-box being built.
+	 *
+	 * @return The multi-box being built.
+	 */
+	TSharedRef< class FMultiBox > GetMultiBox();
 
 
 	/**
diff --git a/Engine/Source/Runtime/Slate/Public/Framework/MultiBox/MultiBoxDefs.h b/Engine/Source/Runtime/Slate/Public/Framework/MultiBox/MultiBoxDefs.h
--- a/Engine/Source/Runtime/Slate/Public/Framework/MultiBox/MultiBoxDefs.h
+++ b/Engine/Source/Runtime/Slate/Public/Framework/MultiBox/MultiBoxDefs.h
@@ -27,6 +27,26 @@ namespace EMultiBoxType
 	};
 }
 
+/**
+ * Types of MultiBlocks
+ */
+namespace EMultiBlockType
+{
+	enum Type
+	{
+		None = 0,
+		ButtonRow,
+		EditableText,
+		Heading,
+		MenuEntry,
+		MenuSeparator,
+		ToolBarButton,
+		ToolBarComboButton,
+		ToolBarSeparator,
+		Widget,
+	};
+}
+
 class SLATE_API FMultiBoxSettings
 {
 public:
diff --git a/Engine/Source/Runtime/SlateCore/Private/Widgets/SWindow.cpp b/Engine/Source/Runtime/SlateCore/Private/Widgets/SWindow.cpp
--- a/Engine/Source/Runtime/SlateCore/Private/Widgets/SWindow.cpp
+++ b/Engine/Source/Runtime/SlateCore/Private/Widgets/SWindow.cpp
@@ -929,6 +929,11 @@ bool SWindow::AppearsInTaskbar() const
 	return !bIsPopupWindow && !bIsToolTipWindow && !bIsCursorDecoratorWindow;
 }
 
+void SWindow::SetOnWindowActivated( const FOnWindowActivated& InDelegate )
+{
+	OnWindowActivated = InDelegate;
+}
+
 void SWindow::SetOnWindowDeactivated( const FOnWindowDeactivated& InDelegate )
 {
 	OnWindowDeactivated = InDelegate;
@@ -1209,6 +1214,10 @@ bool SWindow::OnIsActiveChanged( const FWindowActivateEvent& ActivateEvent )
 			WidgetToFocusOnActivate.Reset();
 		}
 	}
+	else
+	{
+		OnWindowActivated.ExecuteIfBound();
+	}
 
 	return true;
 }
@@ -1250,22 +1259,33 @@ bool SWindow::SupportsKeyboardFocus() const
 
 FReply SWindow::OnKeyboardFocusReceived( const FGeometry& MyGeometry, const FKeyboardFocusEvent& InKeyboardFocusEvent )
 {
-	// If we're becoming active and we were set to restore keyboard focus to a specific widget
-	// after reactivating, then do so now
-	TSharedPtr< SWidget > PinnedWidgetToFocusOnActivate( WidgetToFocusOnActivate.Pin() );
-
-	if( PinnedWidgetToFocusOnActivate.IsValid() && 
-		( InKeyboardFocusEvent.GetCause() == EKeyboardFocusCause::WindowActivate  || FSlateApplicationBase::Get().IsExternalUIOpened() ) )
+	if( InKeyboardFocusEvent.GetCause() == EKeyboardFocusCause::WindowActivate  || FSlateApplicationBase::Get().IsExternalUIOpened() )
 	{
 		TArray< TSharedRef<SWindow> > JustThisWindow;
 		{
 			JustThisWindow.Add( SharedThis(this) );
 		}
-
-		FWidgetPath WidgetToFocusPath;
-		if( FSlateWindowHelper::FindPathToWidget( JustThisWindow, PinnedWidgetToFocusOnActivate.ToSharedRef(), WidgetToFocusPath ) )
+		
+		// If we're becoming active and we were set to restore keyboard focus to a specific widget
+		// after reactivating, then do so now
+		TSharedPtr< SWidget > PinnedWidgetToFocusOnActivate( WidgetToFocusOnActivate.Pin() );
+		
+		if(PinnedWidgetToFocusOnActivate.IsValid())
+		{
+			FWidgetPath WidgetToFocusPath;
+			if( FSlateWindowHelper::FindPathToWidget( JustThisWindow, PinnedWidgetToFocusOnActivate.ToSharedRef(), WidgetToFocusPath ) )
+			{
+				FSlateApplicationBase::Get().SetKeyboardFocus( WidgetToFocusPath, EKeyboardFocusCause::SetDirectly );
+			}
+		}
+		else
 		{
-			FSlateApplicationBase::Get().SetKeyboardFocus( WidgetToFocusPath, EKeyboardFocusCause::SetDirectly );
+			FWidgetPath WindowWidgetPath;
+			if( FSlateWindowHelper::FindPathToWidget( JustThisWindow, AsShared(), WindowWidgetPath ) )
+			{
+				FWeakWidgetPath WeakWindowPath(WindowWidgetPath);
+				FSlateApplicationBase::Get().SetKeyboardFocus( WeakWindowPath.ToNextFocusedPath(EFocusMoveDirection::Next), EKeyboardFocusCause::SetDirectly );
+			}
 		}
 	}
 
diff --git a/Engine/Source/Runtime/SlateCore/Public/Widgets/SWindow.h b/Engine/Source/Runtime/SlateCore/Public/Widgets/SWindow.h
--- a/Engine/Source/Runtime/SlateCore/Public/Widgets/SWindow.h
+++ b/Engine/Source/Runtime/SlateCore/Public/Widgets/SWindow.h
@@ -7,6 +7,9 @@
 #pragma once
 
 
+/** Notification that a window has been activated */
+DECLARE_DELEGATE( FOnWindowActivated );
+
 /** Notification that a window has been deactivated */
 DECLARE_DELEGATE( FOnWindowDeactivated );
 
@@ -380,6 +383,10 @@ public:
 	/** @return should this window show up in the taskbar */
 	bool AppearsInTaskbar() const;
 
+	/** Sets the delegate to execute when the window is activated */
+	void SetOnWindowActivated( const FOnWindowActivated& InDelegate );
+
+	/** Sets the delegate to execute when the window is deactivated */
 	void SetOnWindowDeactivated( const FOnWindowDeactivated& InDelegate );
 
 	/** Sets the delegate to execute right before the window is closed */
@@ -765,6 +772,9 @@ private:
 	/** The native window that is backing this Slate Window */
 	TSharedPtr<FGenericWindow> NativeWindow;
 	
+	/** Invoked when the window has been activated. */
+	FOnWindowActivated OnWindowActivated;
+	
 	/** Invoked when the window has been deactivated. */
 	FOnWindowDeactivated OnWindowDeactivated;
 	
