Added Window->Close menu item to new Mac menu bar code

diff --git a/Engine/Source/Runtime/Slate/Private/Framework/MultiBox/Mac/MacMenu.cpp b/Engine/Source/Runtime/Slate/Private/Framework/MultiBox/Mac/MacMenu.cpp
--- a/Engine/Source/Runtime/Slate/Private/Framework/MultiBox/Mac/MacMenu.cpp
+++ b/Engine/Source/Runtime/Slate/Private/Framework/MultiBox/Mac/MacMenu.cpp
@@ -238,10 +238,12 @@ void FSlateMacMenu::UpdateMenu(FMacMenu* Menu)
 	{
 		NSMenuItem* MinimizeItem = [[[NSMenuItem alloc] initWithTitle:@"Minimize" action:@selector(miniaturize:) keyEquivalent:@"m"] autorelease];
 		NSMenuItem* ZoomItem = [[[NSMenuItem alloc] initWithTitle:@"Zoom" action:@selector(performZoom:) keyEquivalent:@""] autorelease];
+		NSMenuItem* CloseItem = [[[NSMenuItem alloc] initWithTitle:@"Close" action:@selector(performClose:) keyEquivalent:@"w"] autorelease];
 		NSMenuItem* BringAllToFrontItem = [[[NSMenuItem alloc] initWithTitle:@"Bring All to Front" action:@selector(arrangeInFront:) keyEquivalent:@""] autorelease];
 		
 		[Menu addItem:MinimizeItem];
 		[Menu addItem:ZoomItem];
+		[Menu addItem:CloseItem];
 		[Menu addItem:[NSMenuItem separatorItem]];
 		[Menu addItem:BringAllToFrontItem];
 		[Menu addItem:[NSMenuItem separatorItem]];
