Added support for displaying Control symbol for shortcuts in Mac menu bar menus

diff --git a/Engine/Source/Runtime/Slate/Private/Framework/MultiBox/Mac/MacMenu.cpp b/Engine/Source/Runtime/Slate/Private/Framework/MultiBox/Mac/MacMenu.cpp
--- a/Engine/Source/Runtime/Slate/Private/Framework/MultiBox/Mac/MacMenu.cpp
+++ b/Engine/Source/Runtime/Slate/Private/Framework/MultiBox/Mac/MacMenu.cpp
@@ -219,6 +219,10 @@ CFStringRef FSlateMacMenu::GetMenuItemKeyEquivalent(const TSharedRef<const class
 		{
 			*OutModifiers |= NSAlternateKeyMask;
 		}
+		if (Gesture->bCmd)
+		{
+			*OutModifiers |= NSControlKeyMask;
+		}
 
 		FString KeyString = Gesture->GetKeyText().ToString().ToLower();
 		return FPlatformString::TCHARToCFString(*KeyString);
