Make Slate.ToolTipWrapWidth 1000 instead of 400 by default
SDocumentationTooltip now uses Slate.ToolTipWrapWidth console variable instead of hardcoding 400

diff --git a/Engine/Source/Editor/Documentation/Private/SDocumentationToolTip.cpp b/Engine/Source/Editor/Documentation/Private/SDocumentationToolTip.cpp
--- a/Engine/Source/Editor/Documentation/Private/SDocumentationToolTip.cpp
+++ b/Engine/Source/Editor/Documentation/Private/SDocumentationToolTip.cpp
@@ -58,7 +58,7 @@ void SDocumentationToolTip::ConstructSimpleTipContent()
 				.Text( TextContent )
 				.TextStyle( &StyleInfo )
 				.ColorAndOpacity( ColorAndOpacity )
-				.WrapTextAt( 400 )
+				.WrapTextAt_Static( &SToolTip::GetToolTipWrapWidth )
 			]
 		];
 	}
@@ -74,7 +74,7 @@ void SDocumentationToolTip::ConstructSimpleTipContent()
 				.Text( TextContent )
 				.TextStyle( &StyleInfo )
 				.ColorAndOpacity( ColorAndOpacity )
-				.WrapTextAt( 400 )
+				.WrapTextAt_Static( &SToolTip::GetToolTipWrapWidth )
 			]
 		];
 	}
diff --git a/Engine/Source/Runtime/Slate/Private/Widgets/SToolTip.cpp b/Engine/Source/Runtime/Slate/Private/Widgets/SToolTip.cpp
--- a/Engine/Source/Runtime/Slate/Private/Widgets/SToolTip.cpp
+++ b/Engine/Source/Runtime/Slate/Private/Widgets/SToolTip.cpp
@@ -9,9 +9,13 @@
 
 static TAutoConsoleVariable<float> StaticToolTipWrapWidth(
 	TEXT( "Slate.ToolTipWrapWidth" ),
-	400.0f,
+	1000.0f,
 	TEXT( "Width of Slate tool-tips before we wrap the tool-tip text" ) );
 
+float SToolTip::GetToolTipWrapWidth()
+{
+	return StaticToolTipWrapWidth.GetValueOnGameThread();
+}
 
 /**
  * Construct this widget
@@ -29,21 +33,13 @@ void SToolTip::Construct( const FArguments& InArgs )
 		WidgetContent = InArgs._Content.Widget;
 	}
 
-	struct Local
-	{
-		static float GetToolTipWrapWidth()
-		{
-			return StaticToolTipWrapWidth.GetValueOnGameThread();
-		}
-	};
-
 	TSharedRef<SWidget> ToolTipContent = ( !WidgetContent.IsValid() )
 		? StaticCastSharedRef<SWidget>(
 			SNew( STextBlock )
 			.Text( TextContent )
 			.Font( InArgs._Font )
 			.ColorAndOpacity( InArgs._ColorAndOpacity )
-			.WrapTextAt_Static( &Local::GetToolTipWrapWidth )
+			.WrapTextAt_Static( &SToolTip::GetToolTipWrapWidth )
 		)
 		: InArgs._Content.Widget;
 
diff --git a/Engine/Source/Runtime/Slate/Public/Widgets/SToolTip.h b/Engine/Source/Runtime/Slate/Public/Widgets/SToolTip.h
--- a/Engine/Source/Runtime/Slate/Public/Widgets/SToolTip.h
+++ b/Engine/Source/Runtime/Slate/Public/Widgets/SToolTip.h
@@ -76,6 +76,8 @@ public:
 
 	// End IToolTip interface
 
+	static float GetToolTipWrapWidth();
+
 private:
 
 	// Text block widget.
