Fixes Interactive Tool Tips not being interactive any more.

This TTP is not complete; only did safe step because of 4.4 deadline.
Further cleanup is required. The function SWindow::AcceptsInput() should include the call to IsWindowHousingInteractiveTooltip() and all other calls to IsWindowHousingInteractiveTooltip() should be deleted. However, the implications of that change must be studied. For the moment, IsWindowHousingInteractiveTooltip() is being called strategically alongside AcceptsInput() in the two locations that minimize potential fallout.

diff --git a/Engine/Source/Runtime/Slate/Private/Framework/Application/SlateApplication.cpp b/Engine/Source/Runtime/Slate/Private/Framework/Application/SlateApplication.cpp
--- a/Engine/Source/Runtime/Slate/Private/Framework/Application/SlateApplication.cpp
+++ b/Engine/Source/Runtime/Slate/Private/Framework/Application/SlateApplication.cpp
@@ -474,7 +474,6 @@ void FSlateApplication::SetCursorPos( const FVector2D& MouseCoordinate )
 	}
 }
 
-
 FWidgetPath FSlateApplication::LocateWindowUnderMouse( FVector2D ScreenspaceMouseCoordinate, const TArray< TSharedRef< SWindow > >& Windows, bool bIgnoreEnabledStatus )
 {
 	bool bPrevWindowWasModal = false;
@@ -494,9 +493,8 @@ FWidgetPath FSlateApplication::LocateWindowUnderMouse( FVector2D ScreenspaceMous
 		// If none of the children were hit, hittest the parent.
 
 		// Only accept input if the current window accepts input and the current window is not under a modal window or an interactive tooltip
-		const TSharedPtr<IToolTip> ActiveToolTipPtr = ActiveToolTip.Pin();
-		const TSharedPtr<SWindow> ToolTipWindowPtr = ToolTipWindow.Pin();
-		const bool AcceptsInput = Window->AcceptsInput() || ( Window == ToolTipWindowPtr && ActiveToolTipPtr.IsValid() && ActiveToolTipPtr->IsInteractive() );
+		
+		const bool AcceptsInput = Window->AcceptsInput() || IsWindowHousingInteractiveTooltip(Window);
 
 		if ( Window->IsVisible() && AcceptsInput && Window->IsScreenspaceMouseWithin(ScreenspaceMouseCoordinate) && !bPrevWindowWasModal )
 		{
@@ -519,6 +517,18 @@ FWidgetPath FSlateApplication::LocateWindowUnderMouse( FVector2D ScreenspaceMous
 	return FWidgetPath();
 }
 
+bool FSlateApplication::IsWindowHousingInteractiveTooltip(const TSharedRef<const SWindow>& WindowToTest) const
+{
+	const TSharedPtr<IToolTip> ActiveToolTipPtr = ActiveToolTip.Pin();
+	const TSharedPtr<SWindow> ToolTipWindowPtr = ToolTipWindow.Pin();
+	const bool bIsHousingInteractiveTooltip =
+		WindowToTest == ToolTipWindowPtr &&
+		ActiveToolTipPtr.IsValid() &&
+		ActiveToolTipPtr->IsInteractive();
+
+	return bIsHousingInteractiveTooltip;
+}
+
 /** 
  * Ticks a single slate window
  *
diff --git a/Engine/Source/Runtime/Slate/Public/Framework/Application/SlateApplication.h b/Engine/Source/Runtime/Slate/Public/Framework/Application/SlateApplication.h
--- a/Engine/Source/Runtime/Slate/Public/Framework/Application/SlateApplication.h
+++ b/Engine/Source/Runtime/Slate/Public/Framework/Application/SlateApplication.h
@@ -954,6 +954,8 @@ public:
 
 	virtual FWidgetPath LocateWindowUnderMouse( FVector2D ScreenspaceMouseCoordinate, const TArray<TSharedRef<SWindow>>& Windows, bool bIgnoreEnabledStatus = false ) override;
 
+	virtual bool IsWindowHousingInteractiveTooltip(const TSharedRef<const SWindow>& WindowToTest) const override;
+
 	virtual TSharedRef<SWidget> MakeImage( const TAttribute<const FSlateBrush*>& Image, const TAttribute<FSlateColor>& Color, const TAttribute<EVisibility>& Visibility ) const override;
 
 	virtual TSharedRef<SWidget> MakeWindowTitleBar( const TSharedRef<SWindow>& Window, const TSharedPtr<SWidget>& CenterContent, EHorizontalAlignment CenterContentAlignment, TSharedPtr<IWindowTitleBar>& OutTitleBar ) const override;
diff --git a/Engine/Source/Runtime/SlateCore/Private/Widgets/SWindow.cpp b/Engine/Source/Runtime/SlateCore/Private/Widgets/SWindow.cpp
--- a/Engine/Source/Runtime/SlateCore/Private/Widgets/SWindow.cpp
+++ b/Engine/Source/Runtime/SlateCore/Private/Widgets/SWindow.cpp
@@ -780,7 +780,7 @@ FSlateColor SWindow::GetWindowOutlineColor() const
 
 EVisibility SWindow::GetWindowVisibility() const
 {
-	return ( AcceptsInput() )
+	return ( AcceptsInput() || FSlateApplicationBase::Get().IsWindowHousingInteractiveTooltip(SharedThis(this)) )
 		? EVisibility::Visible
 		: EVisibility::HitTestInvisible;
 }
diff --git a/Engine/Source/Runtime/SlateCore/Public/Application/SlateApplicationBase.h b/Engine/Source/Runtime/SlateCore/Public/Application/SlateApplicationBase.h
--- a/Engine/Source/Runtime/SlateCore/Public/Application/SlateApplicationBase.h
+++ b/Engine/Source/Runtime/SlateCore/Public/Application/SlateApplicationBase.h
@@ -181,6 +181,9 @@ public:
 	 */
 	virtual FWidgetPath LocateWindowUnderMouse( FVector2D ScreenspaceMouseCoordinate, const TArray< TSharedRef<SWindow > >& Windows, bool bIgnoreEnabledStatus = false ) = 0;
 
+	/** @return true if 'WindowToTest' is being used to display the current tooltip and the tooltip is interactive. */
+	virtual bool IsWindowHousingInteractiveTooltip(const TSharedRef<const SWindow>& WindowToTest) const = 0;
+
 	/**
 	 * Creates an image widget.
 	 *
