Fixed crash in treeview with negative index

#ttp 342172 - CRASH: LIVE: EDITOR: Editor crashes when right arrow is held while in GPU Profile screen

diff --git a/Engine/Source/Runtime/Slate/Public/Widgets/Views/STreeView.h b/Engine/Source/Runtime/Slate/Public/Widgets/Views/STreeView.h
--- a/Engine/Source/Runtime/Slate/Public/Widgets/Views/STreeView.h
+++ b/Engine/Source/Runtime/Slate/Public/Widgets/Views/STreeView.h
@@ -313,7 +313,7 @@ private:
 	virtual bool Private_DoesItemHaveChildren( int32 ItemIndexInList ) const override
 	{
 		bool bHasChildren = false;
-		if (ItemIndexInList < DenseItemInfos.Num())
+		if (DenseItemInfos.IsValidIndex(ItemIndexInList))
 		{
 			bHasChildren = DenseItemInfos[ItemIndexInList].bHasChildren;
 		}
@@ -323,7 +323,7 @@ private:
 	virtual int32 Private_GetNestingDepth( int32 ItemIndexInList ) const override
 	{
 		int32 NextingLevel = 0;
-		if (ItemIndexInList < DenseItemInfos.Num())
+		if (DenseItemInfos.IsValidIndex(ItemIndexInList))
 		{
 			NextingLevel = DenseItemInfos[ItemIndexInList].NestingLevel;
 		}
