Fix hitches caused by Mac menu bar changes (TTP #342455) by returning NO when asked if we have a menu item with a given key-equivalent. This prevents us from needlessly re-populating the menu each time a key is pressed and stops Cocoa from trying to execute key-equivalents on our menus as we handle all keyboard input ourselves elsewhere.

diff --git a/Engine/Source/Runtime/Slate/Private/Framework/MultiBox/Mac/MacMenu.cpp b/Engine/Source/Runtime/Slate/Private/Framework/MultiBox/Mac/MacMenu.cpp
--- a/Engine/Source/Runtime/Slate/Private/Framework/MultiBox/Mac/MacMenu.cpp
+++ b/Engine/Source/Runtime/Slate/Private/Framework/MultiBox/Mac/MacMenu.cpp
@@ -73,6 +73,11 @@
 	FSlateMacMenu::UpdateMenu(self);
 }
 
+- (BOOL)menuHasKeyEquivalent:(NSMenu*)menu forEvent:(NSEvent*)event target:(id*)target action:(SEL*)action
+{
+	return NO;
+}
+
 - ( TSharedPtr< const FMenuEntryBlock >& )GetMenuEntryBlock
 {
 	return MenuEntryBlock;
