Added FText::IdenticalTo and made SRichTextBlock take its text as an attribute

FText::IdenticalTo compares the string pointers inside the FTexts to see if they're identical. This can be used as an early out when checking if a cached copy of an FText is the same as the current version of an FText (since if the string pointers are identical, then the strings themselves shouldn't have changed).

SMultiLineEditableText::Tick has been changed to use FText::IdenticalTo to speed up its equality testing when using a bound FText attribute.

SRichTextBlock has also been updated to take its text as an attribute. This wasn't done originally due to the performance considerations when caching and comparing the FText, but FText::IdenticalTo alleviates those concerns.

diff --git a/Engine/Source/Runtime/Core/Private/Internationalization/Text.cpp b/Engine/Source/Runtime/Core/Private/Internationalization/Text.cpp
--- a/Engine/Source/Runtime/Core/Private/Internationalization/Text.cpp
+++ b/Engine/Source/Runtime/Core/Private/Internationalization/Text.cpp
@@ -780,4 +780,11 @@ TSharedPtr< FString, ESPMode::ThreadSafe > FText::GetSourceString() const
 	return DisplayString;
 }
 
+bool FText::IdenticalTo( const FText& Other ) const
+{
+	// If both instances point to the same string, then both instances are considered identical
+	// This is fast as it skips a lexical compare, however it can also return true for two instances that have identical strings, but in different pointers
+	return DisplayString == Other.DisplayString;
+}
+
 #undef LOCTEXT_NAMESPACE
diff --git a/Engine/Source/Runtime/Core/Public/Internationalization/Text.h b/Engine/Source/Runtime/Core/Public/Internationalization/Text.h
--- a/Engine/Source/Runtime/Core/Public/Internationalization/Text.h
+++ b/Engine/Source/Runtime/Core/Public/Internationalization/Text.h
@@ -211,6 +211,14 @@ public:
 	bool EqualTo( const FText& Other, const ETextComparisonLevel::Type ComparisonLevel = ETextComparisonLevel::Default ) const;
 	bool EqualToCaseIgnored( const FText& Other ) const;
 
+	/**
+	 * Check to see if this FText is identical to the other FText
+	 * 
+	 * Note:	This doesn't compare the text, but only checks that the internal string pointers have the same target (which makes it very fast!)
+	 *			If you actually want to perform a lexical comparison, then you need to use EqualTo instead
+	 */
+	bool IdenticalTo( const FText& Other ) const;
+
 	class CORE_API FSortPredicate
 	{
 	public:
diff --git a/Engine/Source/Runtime/Slate/Private/Widgets/Text/SMultiLineEditableText.cpp b/Engine/Source/Runtime/Slate/Private/Widgets/Text/SMultiLineEditableText.cpp
--- a/Engine/Source/Runtime/Slate/Private/Widgets/Text/SMultiLineEditableText.cpp
+++ b/Engine/Source/Runtime/Slate/Private/Widgets/Text/SMultiLineEditableText.cpp
@@ -1702,10 +1702,16 @@ void SMultiLineEditableText::Tick( const FGeometry& AllottedGeometry, const doub
 	if (!bShouldAppearFocused && BoundText.IsBound())
 	{
 		const FText& TextToSet = BoundText.Get(FText::GetEmpty());
-		if (!BoundTextLastTick.ToString().Equals(TextToSet.ToString(), ESearchCase::CaseSensitive))
+		if (!BoundTextLastTick.IdenticalTo(TextToSet))
 		{
-			// The source text has changed, so update the internal editable text
-			SetEditableText(TextToSet);
+			// The pointer used by the bound text has changed, however the text may still be the same - check that now
+			if (!BoundTextLastTick.ToString().Equals(TextToSet.ToString(), ESearchCase::CaseSensitive))
+			{
+				// The source text has changed, so update the internal editable text
+				SetEditableText(TextToSet);
+			}
+
+			// Update this even if the text is lexically identical, as it will update the pointer compared by IdenticalTo for the next Tick
 			BoundTextLastTick = TextToSet;
 		}
 	}
diff --git a/Engine/Source/Runtime/Slate/Private/Widgets/Text/SRichTextBlock.cpp b/Engine/Source/Runtime/Slate/Private/Widgets/Text/SRichTextBlock.cpp
--- a/Engine/Source/Runtime/Slate/Private/Widgets/Text/SRichTextBlock.cpp
+++ b/Engine/Source/Runtime/Slate/Private/Widgets/Text/SRichTextBlock.cpp
@@ -48,6 +48,23 @@ void SRichTextBlock::Construct( const FArguments& InArgs )
 void SRichTextBlock::Tick( const FGeometry& AllottedGeometry, const double InCurrentTime, const float InDeltaTime )
 {
 	TextLayout->SetScale( AllottedGeometry.Scale );
+
+	if (BoundText.IsBound())
+	{
+		const FText& TextToSet = GetText();
+		if (!BoundTextLastTick.IdenticalTo(TextToSet))
+		{
+			// The pointer used by the bound text has changed, however the text may still be the same - check that now
+			if (!BoundTextLastTick.ToString().Equals(TextToSet.ToString(), ESearchCase::CaseSensitive))
+			{
+				// The source text has changed, so update the internal editable text
+				SetText(TextToSet);
+			}
+
+			// Update this even if the text is lexically identical, as it will update the pointer compared by IdenticalTo for the next Tick
+			BoundTextLastTick = TextToSet;
+		}
+	}
 }
 
 int32 SRichTextBlock::OnPaint( const FPaintArgs& Args, const FGeometry& AllottedGeometry, const FSlateRect& MyClippingRect, FSlateWindowElementList& OutDrawElements, int32 LayerId, const FWidgetStyle& InWidgetStyle, bool bParentEnabled ) const
@@ -154,9 +171,17 @@ TSharedPtr< ITextDecorator > SRichTextBlock::TryGetDecorator( const TArray< TSha
 	return NULL;
 }
 
-void SRichTextBlock::SetText( const FText& InText )
+void SRichTextBlock::SetText( const TAttribute<FText>& InTextAttr )
 {
-	Text = InText;
+	BoundText = InTextAttr;
+
+	const FText& InText = GetText();
+
+	// Update the cached BoundText value to prevent it triggering another SetText update again next Tick
+	if (BoundText.IsBound())
+	{
+		BoundTextLastTick = InText;
+	}
 
 	const FString& InString = InText.ToString();
 	TArray<FTextLineParseResults> LineParseResultsArray;
diff --git a/Engine/Source/Runtime/Slate/Public/Widgets/Text/SRichTextBlock.h b/Engine/Source/Runtime/Slate/Public/Widgets/Text/SRichTextBlock.h
--- a/Engine/Source/Runtime/Slate/Public/Widgets/Text/SRichTextBlock.h
+++ b/Engine/Source/Runtime/Slate/Public/Widgets/Text/SRichTextBlock.h
@@ -24,7 +24,7 @@ public:
 		, _Parser()
 	{}
 		/** The text displayed in this text block */
-		SLATE_ARGUMENT( FText, Text )
+		SLATE_ATTRIBUTE( FText, Text )
 
 		/** Whether text wraps onto a new line when it's length exceeds this width; if this value is zero or negative, no wrapping occurs. */
 		SLATE_ATTRIBUTE( float, WrapTextAt )
@@ -116,25 +116,28 @@ public:
 	 */
 	const FText& GetText() const
 	{
-		return Text;
+		return BoundText.Get(FText::GetEmpty());
 	}
 
 	/**
 	 * Sets the text for this text block
 	 */
-	void SetText( const FText& InText );
+	void SetText( const TAttribute<FText>& InTextAttr );
 
 	void SetHighlightText( const FText& InHighlightText );
 
 
 private:
-	
+
 	TSharedPtr< ITextDecorator > TryGetDecorator( const TArray< TSharedRef< ITextDecorator > >& Decorators, const FString& Line, const FTextRunParseResults& TextRun ) const;
 
 private:
 
 	/** The text displayed in this text block */
-	FText Text;
+	TAttribute< FText > BoundText;
+
+	/** The state of BoundText last Tick() (only used when BoundText is bound to a delegate providing the source text) */
+	FText BoundTextLastTick;
 
 	TSharedPtr< FSlateTextLayout > TextLayout;
 
