SScrollBox - Add support for horizontal scrolling to the core widget and UMG.

There is more copy/paste here than I wanted.  I had some unified layout/construct code, but with all of the templates and additional function calls, it was much harder to follow, and performance would have probably taken a hit, so I opted for the clearer duplicated code in a few spots.

diff --git a/Engine/Source/Runtime/Slate/Private/Widgets/Layout/SScrollBox.cpp b/Engine/Source/Runtime/Slate/Private/Widgets/Layout/SScrollBox.cpp
--- a/Engine/Source/Runtime/Slate/Private/Widgets/Layout/SScrollBox.cpp
+++ b/Engine/Source/Runtime/Slate/Private/Widgets/Layout/SScrollBox.cpp
@@ -18,6 +18,9 @@ public:
 		{
 			_Visibility = EVisibility::SelfHitTestInvisible;
 		}
+
+		SLATE_ARGUMENT(EOrientation, Orientation)
+
 	SLATE_END_ARGS()
 
 	SScrollPanel()
@@ -33,10 +36,21 @@ public:
 		{
 			Children.Add( InSlots[ SlotIndex ] );
 		}
+		Orientation = InArgs._Orientation;
 	}
 
 public:
 
+	EOrientation GetOrientation()
+	{
+		return Orientation;
+	}
+
+	void SetOrientation(EOrientation InOrientation)
+	{
+		Orientation = InOrientation;
+	}
+
 	virtual void OnArrangeChildren( const FGeometry& AllottedGeometry, FArrangedChildren& ArrangedChildren ) const override
 	{
 		float CurChildOffset = -PhysicalOffset;
@@ -48,20 +62,19 @@ public:
 
 			if ( ChildVisibility != EVisibility::Collapsed )
 			{
-				const FMargin& ThisPadding = ThisSlot.SlotPadding.Get();
-				const FVector2D& WidgetDesiredSize = ThisSlot.GetWidget()->GetDesiredSize();
-				const float ThisSlotDesiredHeight = WidgetDesiredSize.Y + ThisSlot.SlotPadding.Get().GetTotalSpaceAlong<Orient_Vertical>();
-
-				// Figure out the size and local position of the child within the slot
-				// There is no vertical alignment, because it does not make sense in a panel where items are stacked end-to-end
-				AlignmentArrangeResult XAlignmentResult = AlignChild<Orient_Horizontal>( AllottedGeometry.Size.X, ThisSlot, ThisPadding );
-
-				ArrangedChildren.AddWidget( AllottedGeometry.MakeChild( ThisSlot.GetWidget(), FVector2D(XAlignmentResult.Offset, CurChildOffset + ThisPadding.Top), FVector2D(XAlignmentResult.Size, WidgetDesiredSize.Y) ) );
-				CurChildOffset += ThisSlotDesiredHeight;
+				if (Orientation == Orient_Vertical)
+				{
+					CurChildOffset = ArrangeChildVerticalAndReturnOffset(AllottedGeometry, ArrangedChildren, ThisSlot, CurChildOffset);
+				}
+				else
+				{
+					CurChildOffset = ArrangeChildHorizontalAndReturnOffset(AllottedGeometry, ArrangedChildren, ThisSlot, CurChildOffset);
+				}
 			}
 		}
 	}
 
+
 	virtual FVector2D ComputeDesiredSize() const override
 	{
 		FVector2D ThisDesiredSize = FVector2D::ZeroVector;
@@ -69,8 +82,16 @@ public:
 		{
 			const SScrollBox::FSlot& ThisSlot = Children[SlotIndex];
 			const FVector2D ChildDesiredSize = ThisSlot.GetWidget()->GetDesiredSize();
-			ThisDesiredSize.X  = FMath::Max(ChildDesiredSize.X, ThisDesiredSize.X);
-			ThisDesiredSize.Y += ChildDesiredSize.Y + ThisSlot.SlotPadding.Get().GetTotalSpaceAlong<Orient_Vertical>();
+			if (Orientation == Orient_Vertical)
+			{
+				ThisDesiredSize.X = FMath::Max(ChildDesiredSize.X, ThisDesiredSize.X);
+				ThisDesiredSize.Y += ChildDesiredSize.Y + ThisSlot.SlotPadding.Get().GetTotalSpaceAlong<Orient_Vertical>();
+			}
+			else
+			{
+				ThisDesiredSize.X += ChildDesiredSize.X + ThisSlot.SlotPadding.Get().GetTotalSpaceAlong<Orient_Horizontal>();
+				ThisDesiredSize.Y = FMath::Max(ChildDesiredSize.Y, ThisDesiredSize.Y);
+			}
 		}
 
 		return ThisDesiredSize;
@@ -83,6 +104,40 @@ public:
 
 	float PhysicalOffset;
 	TPanelChildren<SScrollBox::FSlot> Children;	
+
+private:
+	
+	float ArrangeChildVerticalAndReturnOffset(const FGeometry& AllottedGeometry, FArrangedChildren& ArrangedChildren, const SScrollBox::FSlot& ThisSlot, float CurChildOffset) const
+	{
+		const FMargin& ThisPadding = ThisSlot.SlotPadding.Get();
+		const FVector2D& WidgetDesiredSize = ThisSlot.GetWidget()->GetDesiredSize();
+		const float ThisSlotDesiredHeight = WidgetDesiredSize.Y + ThisPadding.GetTotalSpaceAlong<Orient_Vertical>();
+
+		// Figure out the size and local position of the child within the slot.  There is no vertical alignment, because 
+		// it does not make sense in a panel where items are stacked vertically end-to-end.
+		AlignmentArrangeResult XAlignmentResult = AlignChild<Orient_Horizontal>(AllottedGeometry.Size.X, ThisSlot, ThisPadding);
+
+		ArrangedChildren.AddWidget(AllottedGeometry.MakeChild(ThisSlot.GetWidget(), FVector2D(XAlignmentResult.Offset, CurChildOffset + ThisPadding.Top), FVector2D(XAlignmentResult.Size, WidgetDesiredSize.Y)));
+		return CurChildOffset + ThisSlotDesiredHeight;
+	}
+
+	float ArrangeChildHorizontalAndReturnOffset(const FGeometry& AllottedGeometry, FArrangedChildren& ArrangedChildren, const SScrollBox::FSlot& ThisSlot, float CurChildOffset) const
+	{
+		const FMargin& ThisPadding = ThisSlot.SlotPadding.Get();
+		const FVector2D& WidgetDesiredSize = ThisSlot.GetWidget()->GetDesiredSize();
+		const float ThisSlotDesiredWidth = WidgetDesiredSize.X + ThisPadding.GetTotalSpaceAlong<Orient_Horizontal>();
+
+		// Figure out the size and local position of the child within the slot.  There is no horizontal alignment, because
+		// it doesn't make sense in a panel where items are stacked horizontally end-to-end.
+		AlignmentArrangeResult YAlignmentResult = AlignChild<Orient_Vertical>(AllottedGeometry.Size.Y, ThisSlot, ThisPadding);
+
+		ArrangedChildren.AddWidget(AllottedGeometry.MakeChild(ThisSlot.GetWidget(), FVector2D(CurChildOffset + ThisPadding.Left, YAlignmentResult.Offset), FVector2D(WidgetDesiredSize.X, YAlignmentResult.Size)));
+		return CurChildOffset + ThisSlotDesiredWidth;
+	}
+
+private:
+
+	EOrientation Orientation;
 };
 
 
@@ -91,6 +146,8 @@ void SScrollBox::Construct( const FArguments& InArgs )
 {
 	check(InArgs._Style);
 
+	Style = InArgs._Style;
+	ScrollBarStyle = InArgs._ScrollBarStyle;
 	DesiredScrollOffset = 0;
 	bIsScrolling = false;
 	bAnimateScroll = false;
@@ -98,75 +155,152 @@ void SScrollBox::Construct( const FArguments& InArgs )
 	bShowSoftwareCursor = false;
 	SoftwareCursorPosition = FVector2D::ZeroVector;
 	OnUserScrolled = InArgs._OnUserScrolled;
-
-	TSharedPtr<SHorizontalBox> PanelAndScrollbar;
-	this->ChildSlot
-	[
-		SAssignNew(PanelAndScrollbar, SHorizontalBox)
-
-		+ SHorizontalBox::Slot()
-		.FillWidth(1)
-		[
-			SNew(SOverlay)
-
-			+ SOverlay::Slot()
-			.Padding( FMargin(0.0f, 0.0f, 0.0f, 1.0f) )
-			[
-				// Scroll panel that presents the scrolled content
-				SAssignNew( ScrollPanel, SScrollPanel, InArgs.Slots )
-			]
-
-			+ SOverlay::Slot()
-			.HAlign(HAlign_Fill)
-			.VAlign(VAlign_Top)
-			[
-				// Shadow: Hint to scroll up
-				SNew(SImage)
-				.Visibility(EVisibility::HitTestInvisible)
-				.ColorAndOpacity(this, &SScrollBox::GetTopShadowOpacity)
-				.Image( &InArgs._Style->TopShadowBrush )	
-			]
-
-			+ SOverlay::Slot()
-			.HAlign(HAlign_Fill)
-			.VAlign(VAlign_Bottom)
-			[
-				// Shadow: a hint to scroll down
-				SNew(SImage)
-				.Visibility(EVisibility::HitTestInvisible)
-				.ColorAndOpacity(this, &SScrollBox::GetBottomShadowOpacity)
-				.Image( &InArgs._Style->BottomShadowBrush )	
-			]
-		]
-	
-	];
-
+	Orientation = InArgs._Orientation;
 
 	if (InArgs._ExternalScrollbar.IsValid())
 	{
 		// An external scroll bar was specified by the user
 		ScrollBar = InArgs._ExternalScrollbar;
 		ScrollBar->SetOnUserScrolled(FOnUserScrolled::CreateSP(this, &SScrollBox::ScrollBar_OnUserScrolled));
+		bScrollBarIsExternal = true;
 	}
 	else
 	{
 		// Make a scroll bar 
-		ScrollBar = 
-			SNew( SScrollBar)
-			.Thickness(FVector2D(5.0f, 5.0f))
-			.Style( InArgs._ScrollBarStyle )
-			.OnUserScrolled(this, &SScrollBox::ScrollBar_OnUserScrolled);
+		ScrollBar = ConstructScrollBar();
+		bScrollBarIsExternal = false;
+	}
+
+	SAssignNew(ScrollPanel, SScrollPanel, InArgs.Slots)
+		.Orientation(Orientation);
+
+	if (Orientation == Orient_Vertical)
+	{
+		ConstructVerticalLayout();
+	}
+	else
+	{
+		ConstructHorizontalLayout();
+	}
+
+	ScrollBar->SetState( 0.0f, 1.0f );
+}
+
+TSharedPtr<SScrollBar> SScrollBox::ConstructScrollBar()
+{
+	return TSharedPtr<SScrollBar>(SNew(SScrollBar)
+		.Thickness(FVector2D(5.0f, 5.0f))
+		.Style(ScrollBarStyle)
+		.Orientation(Orientation)
+		.OnUserScrolled(this, &SScrollBox::ScrollBar_OnUserScrolled));
+}
+
+void SScrollBox::ConstructVerticalLayout()
+{
+	TSharedPtr<SHorizontalBox> PanelAndScrollbar;
+	this->ChildSlot
+		[
+			SAssignNew(PanelAndScrollbar, SHorizontalBox)
+
+			+ SHorizontalBox::Slot()
+			.FillWidth(1)
+			[
+				SNew(SOverlay)
+
+				+ SOverlay::Slot()
+				.Padding(FMargin(0.0f, 0.0f, 0.0f, 1.0f))
+				[
+					// Scroll panel that presents the scrolled content
+					ScrollPanel.ToSharedRef()
+				]
+
+				+ SOverlay::Slot()
+					.HAlign(HAlign_Fill)
+					.VAlign(VAlign_Top)
+					[
+						// Shadow: Hint to scroll up
+						SNew(SImage)
+						.Visibility(EVisibility::HitTestInvisible)
+						.ColorAndOpacity(this, &SScrollBox::GetStartShadowOpacity)
+						.Image(&Style->TopShadowBrush)
+					]
+
+				+ SOverlay::Slot()
+					.HAlign(HAlign_Fill)
+					.VAlign(VAlign_Bottom)
+					[
+						// Shadow: a hint to scroll down
+						SNew(SImage)
+						.Visibility(EVisibility::HitTestInvisible)
+						.ColorAndOpacity(this, &SScrollBox::GetEndShadowOpacity)
+						.Image(&Style->BottomShadowBrush)
+					]
+			]
 
+		];
+
+	if (!bScrollBarIsExternal)
+	{
 		PanelAndScrollbar->AddSlot()
 			.AutoWidth()
 			[
 				ScrollBar.ToSharedRef()
 			];
 	}
+}
+
+void SScrollBox::ConstructHorizontalLayout()
+{
+	TSharedPtr<SVerticalBox> PanelAndScrollbar;
+	this->ChildSlot
+		[
+			SAssignNew(PanelAndScrollbar, SVerticalBox)
 
+			+ SVerticalBox::Slot()
+			.FillHeight(1)
+			[
+				SNew(SOverlay)
+
+				+ SOverlay::Slot()
+				.Padding(FMargin(0.0f, 0.0f, 1.0f, 0.0f))
+				[
+					// Scroll panel that presents the scrolled content
+					ScrollPanel.ToSharedRef()
+				]
+
+				+ SOverlay::Slot()
+					.HAlign(HAlign_Left)
+					.VAlign(VAlign_Fill)
+					[
+						// Shadow: Hint to left
+						SNew(SImage)
+						.Visibility(EVisibility::HitTestInvisible)
+						.ColorAndOpacity(this, &SScrollBox::GetStartShadowOpacity)
+						.Image(&Style->LeftShadowBrush)
+					]
+
+				+ SOverlay::Slot()
+					.HAlign(HAlign_Right)
+					.VAlign(VAlign_Fill)
+					[
+						// Shadow: a hint to scroll right
+						SNew(SImage)
+						.Visibility(EVisibility::HitTestInvisible)
+						.ColorAndOpacity(this, &SScrollBox::GetEndShadowOpacity)
+						.Image(&Style->RightShadowBrush)
+					]
+			]
 
+		];
 
-	ScrollBar->SetState( 0.0f, 1.0f );
+	if (!bScrollBarIsExternal)
+	{
+		PanelAndScrollbar->AddSlot()
+			.AutoHeight()
+			[
+				ScrollBar.ToSharedRef()
+			];
+	}
 }
 
 /** Adds a slot to SScrollBox */
@@ -212,6 +346,32 @@ void SScrollBox::SetScrollOffset( float NewScrollOffset )
 	DesiredScrollOffset = NewScrollOffset;
 }
 
+EOrientation SScrollBox::GetOrientation()
+{
+	return Orientation;
+}
+
+void SScrollBox::SetOrientation(EOrientation InOrientation)
+{
+	if (Orientation != InOrientation)
+	{
+		Orientation = InOrientation;
+		if (!bScrollBarIsExternal)
+		{
+			ScrollBar = ConstructScrollBar();
+		}
+		ScrollPanel->SetOrientation(Orientation);
+		if (Orientation == Orient_Vertical)
+		{
+			ConstructVerticalLayout();
+		}
+		else
+		{
+			ConstructHorizontalLayout();
+		}
+	}
+}
+
 void SScrollBox::Tick( const FGeometry& AllottedGeometry, const double InCurrentTime, const float InDeltaTime )
 {
 	if ( !IsRightClickScrolling() )
@@ -227,14 +387,14 @@ void SScrollBox::Tick( const FGeometry& AllottedGeometry, const double InCurrent
 
 	const FGeometry ScrollPanelGeometry = FindChildGeometry( AllottedGeometry, ScrollPanel.ToSharedRef() );
 	
-	const float ContentHeight = ScrollPanel->GetDesiredSize().Y;
+	const float ContentSize = GetScrollComponentFromVector(ScrollPanel->GetDesiredSize());
 
 	// If this scroll box has no size, do not compute a view fraction because it will be wrong and causes pop in when the size is available
-	const float ViewFraction = AllottedGeometry.Size.Y > 0 ? ScrollPanelGeometry.Size.Y / ContentHeight : 1;
-	const float ViewOffset = FMath::Clamp<float>( DesiredScrollOffset/ContentHeight, 0.0, 1.0 - ViewFraction );
+	const float ViewFraction = GetScrollComponentFromVector(AllottedGeometry.Size) > 0 ? GetScrollComponentFromVector(ScrollPanelGeometry.Size) / ContentSize : 1;
+	const float ViewOffset = FMath::Clamp<float>( DesiredScrollOffset/ContentSize, 0.0, 1.0 - ViewFraction );
 	
 	// Update the scrollbar with the clamped version of the offset
-	const float TargetPhysicalOffset = ViewOffset*ScrollPanel->GetDesiredSize().Y;
+	const float TargetPhysicalOffset = GetScrollComponentFromVector(ViewOffset*ScrollPanel->GetDesiredSize());
 	bIsScrolling = !FMath::IsNearlyEqual(TargetPhysicalOffset, ScrollPanel->PhysicalOffset, 0.001f);	
 	ScrollPanel->PhysicalOffset = (bAnimateScroll)
 		? FMath::FInterpTo(ScrollPanel->PhysicalOffset, TargetPhysicalOffset, InDeltaTime, 15.f)
@@ -302,7 +462,7 @@ FReply SScrollBox::OnMouseMove( const FGeometry& MyGeometry, const FPointerEvent
 {	
 	if( MouseEvent.IsMouseButtonDown( EKeys::RightMouseButton ) )
 	{
-		const float ScrollByAmount = MouseEvent.GetCursorDelta().Y;
+		const float ScrollByAmount = GetScrollComponentFromVector(MouseEvent.GetCursorDelta());
 		// If scrolling with the right mouse button, we need to remember how much we scrolled.
 		// If we did not scroll at all, we will bring up the context menu when the mouse is released.
 		AmountScrolledWhileRightMouseDown += FMath::Abs( ScrollByAmount );
@@ -327,7 +487,7 @@ FReply SScrollBox::OnMouseMove( const FGeometry& MyGeometry, const FPointerEvent
 			// Check if the mouse has moved.
 			if( bDidScroll )
 			{
-				SoftwareCursorPosition.Y += MouseEvent.GetCursorDelta().Y;
+				SetScrollComponentOnVector(SoftwareCursorPosition, GetScrollComponentFromVector(SoftwareCursorPosition) + GetScrollComponentFromVector(MouseEvent.GetCursorDelta()));
 			}
 
 			return Reply;
@@ -367,7 +527,7 @@ bool SScrollBox::ScrollBy( const FGeometry& AllottedGeometry, float ScrollAmount
 {
 	bAnimateScroll = InAnimateScroll;
 
-	const float ContentSize = ScrollPanel->GetDesiredSize().Y;
+	const float ContentSize = GetScrollComponentFromVector(ScrollPanel->GetDesiredSize());
 	const FGeometry ScrollPanelGeometry = FindChildGeometry( AllottedGeometry, ScrollPanel.ToSharedRef() );
 
 	float PreviousScrollOffset = DesiredScrollOffset;
@@ -375,7 +535,7 @@ bool SScrollBox::ScrollBy( const FGeometry& AllottedGeometry, float ScrollAmount
 	DesiredScrollOffset += ScrollAmount;
 
 	const float ScrollMin = 0.0f;
-	const float ScrollMax = ContentSize - ScrollPanelGeometry.Size.Y;
+	const float ScrollMax = ContentSize - GetScrollComponentFromVector(ScrollPanelGeometry.Size);
 	DesiredScrollOffset = FMath::Clamp(DesiredScrollOffset, ScrollMin, ScrollMax);
 
 	OnUserScrolled.ExecuteIfBound(DesiredScrollOffset);
@@ -425,13 +585,13 @@ int32 SScrollBox::OnPaint( const FPaintArgs& Args, const FGeometry& AllottedGeom
 
 void SScrollBox::ScrollBar_OnUserScrolled( float InScrollOffsetFraction )
 {
-	DesiredScrollOffset = InScrollOffsetFraction * ScrollPanel->GetDesiredSize().Y;
+	DesiredScrollOffset = InScrollOffsetFraction * GetScrollComponentFromVector(ScrollPanel->GetDesiredSize());
 	OnUserScrolled.ExecuteIfBound(DesiredScrollOffset);
 }
 
 const float ShadowFadeDistance = 32.0f;
 
-FSlateColor SScrollBox::GetTopShadowOpacity() const
+FSlateColor SScrollBox::GetStartShadowOpacity() const
 {
 	// The shadow should only be visible when the user needs a hint that they can scroll up.
 	const float ShadowOpacity = FMath::Clamp( ScrollPanel->PhysicalOffset/ShadowFadeDistance, 0.0f, 1.0f);
@@ -439,10 +599,10 @@ FSlateColor SScrollBox::GetTopShadowOpacity() const
 	return FLinearColor(1.0f, 1.0f, 1.0f, ShadowOpacity);
 }
 
-FSlateColor SScrollBox::GetBottomShadowOpacity() const
+FSlateColor SScrollBox::GetEndShadowOpacity() const
 {
 	// The shadow should only be visible when the user needs a hint that they can scroll down.
-	const float ShadowOpacity = (ScrollBar->DistanceFromBottom() * ScrollPanel->GetDesiredSize().Y) / ShadowFadeDistance;
+	const float ShadowOpacity = (ScrollBar->DistanceFromBottom() * GetScrollComponentFromVector(ScrollPanel->GetDesiredSize()) / ShadowFadeDistance);
 	
 	return FLinearColor(1.0f, 1.0f, 1.0f, ShadowOpacity);
 }
diff --git a/Engine/Source/Runtime/Slate/Public/Widgets/Layout/SScrollBox.h b/Engine/Source/Runtime/Slate/Public/Widgets/Layout/SScrollBox.h
--- a/Engine/Source/Runtime/Slate/Public/Widgets/Layout/SScrollBox.h
+++ b/Engine/Source/Runtime/Slate/Public/Widgets/Layout/SScrollBox.h
@@ -12,9 +12,12 @@ public:
 	public:
 		FSlot()
 			: TSlotBase<FSlot>()
+			// Set both vertical and horizontal alignment to fill by default.  During layout, the
+			// alignment direction parallel to the scroll direction is assumed to be top, or left
+			// since that is how the items are stacked.
 			, HAlignment(HAlign_Fill)
+			, VAlignment(VAlign_Fill)
 		{
-
 		}
 
 		FSlot& HAlign( EHorizontalAlignment InHAlignment )
@@ -22,8 +25,15 @@ public:
 			HAlignment = InHAlignment;
 			return *this;
 		}
+
+		FSlot& VAlign(EVerticalAlignment InVAlignment)
+		{
+			VAlignment = InVAlignment;
+			return *this;
+		}
 		
 		EHorizontalAlignment HAlignment;
+		EVerticalAlignment VAlignment;
 	};
 
 	SLATE_BEGIN_ARGS(SScrollBox)
@@ -31,6 +41,7 @@ public:
 		, _ScrollBarStyle( &FCoreStyle::Get().GetWidgetStyle<FScrollBarStyle>("ScrollBar") )
 		, _ExternalScrollbar()
 		, _OnUserScrolled()
+		, _Orientation(Orient_Vertical)
 		{}
 		
 		SLATE_SUPPORTS_SLOT( FSlot )
@@ -44,6 +55,9 @@ public:
 		/** Custom scroll bar */
 		SLATE_ARGUMENT( TSharedPtr<SScrollBar>, ExternalScrollbar )
 
+		/** The direction that children will be stacked, and also the direction the box will scroll. */
+		SLATE_ARGUMENT(EOrientation, Orientation)
+
 		/** Called when the button is clicked */
 		SLATE_EVENT(FOnUserScrolled, OnUserScrolled)
 
@@ -72,6 +86,12 @@ public:
 
 	void SetScrollOffset( float NewScrollOffset );
 
+	/** Get the current orientation of the scrollbox. */
+	EOrientation GetOrientation();
+
+	/** Sets the current orientation of the scrollbox and updates the layout */
+	void SetOrientation(EOrientation InOrientation);
+
 public:
 
 	// SWidget interface
@@ -88,6 +108,34 @@ public:
 
 private:
 
+	/** Builds a default Scrollbar */
+	TSharedPtr<SScrollBar> ConstructScrollBar();
+
+	/** Constructs internal layout widgets for scrolling vertically using the existing ScrollPanel and ScrollBar. */
+	void ConstructVerticalLayout();
+
+	/** Constructs internal layout widgets for scrolling horizontally using the existing ScrollPanel and ScrollBar. */
+	void ConstructHorizontalLayout();
+
+	/** Gets the component of a vector in the direction of scrolling based on the Orientation property. */
+	inline float GetScrollComponentFromVector(FVector2D Vector) const
+	{
+		return Orientation == Orient_Vertical ? Vector.Y : Vector.X;
+	}
+
+	/** Sets the component of a vector in the direction of scrolling based on the Orientation property. */
+	inline void SetScrollComponentOnVector(FVector2D& InVector, float Value) const
+	{
+		if (Orientation == Orient_Vertical)
+		{
+			InVector.Y = Value;
+		}
+		else
+		{
+			InVector.X = Value;
+		}
+	}
+
 	/** Scroll offset that the user asked for. We will clamp it before actually scrolling there. */
 	float DesiredScrollOffset;
 
@@ -103,14 +151,20 @@ private:
 	/** Invoked when the user scroll via the scrollbar */
 	void ScrollBar_OnUserScrolled( float InScrollOffsetFraction );
 
-	/** Does the user need a hint that they can scroll up? */
-	FSlateColor GetTopShadowOpacity() const;
+	/** Does the user need a hint that they can scroll to the start of the list? */
+	FSlateColor GetStartShadowOpacity() const;
 	
-	/** Does the user need a hint that they can scroll down? */
-	FSlateColor GetBottomShadowOpacity() const;
+	/** Does the user need a hint that they can scroll to the end of the list? */
+	FSlateColor GetEndShadowOpacity() const;
 
+	/** The panel which stacks the child slots */
 	TSharedPtr<class SScrollPanel> ScrollPanel;
+
+	/** The scrollbar which controls scrolling for the scrollbox. */
 	TSharedPtr<SScrollBar> ScrollBar;
+
+	/** Whether or not the user supplied an external scrollbar to control scrolling. */
+	bool bScrollBarIsExternal;
 	
 	/** Are we actively scrolling right now */
 	bool bIsScrolling;
@@ -132,6 +186,15 @@ private:
 
 	/** Fired when the user scrolls the scrollbox */
 	FOnUserScrolled OnUserScrolled;
+
+	/** The scrolling and stacking orientation. */
+	EOrientation Orientation;
+
+	/** Style resource for the scrollbox */
+	const FScrollBoxStyle* Style;
+
+	/** Style resource for the scrollbar */
+	const FScrollBarStyle* ScrollBarStyle;
 };
 
 
diff --git a/Engine/Source/Runtime/SlateCore/Private/Styling/CoreStyle.cpp b/Engine/Source/Runtime/SlateCore/Private/Styling/CoreStyle.cpp
--- a/Engine/Source/Runtime/SlateCore/Private/Styling/CoreStyle.cpp
+++ b/Engine/Source/Runtime/SlateCore/Private/Styling/CoreStyle.cpp
@@ -982,6 +982,8 @@ TSharedRef<ISlateStyle> FCoreStyle::Create( const FName& InStyleSetName )
 		Style->Set( "ScrollBox", FScrollBoxStyle()
 			.SetTopShadowBrush( IMAGE_BRUSH( "Common/ScrollBoxShadowTop", FVector2D(64,8) ) )
 			.SetBottomShadowBrush( IMAGE_BRUSH( "Common/ScrollBoxShadowBottom", FVector2D(64,8) ) )
+			.SetLeftShadowBrush( IMAGE_BRUSH( "Common/ScrollBoxShadowLeft", FVector2D(8, 64) ) )
+			.SetRightShadowBrush( IMAGE_BRUSH( "Common/ScrollBoxShadowRight", FVector2D(8, 64) ) )
 			);
 	}
 
diff --git a/Engine/Source/Runtime/SlateCore/Private/Styling/SlateTypes.cpp b/Engine/Source/Runtime/SlateCore/Private/Styling/SlateTypes.cpp
--- a/Engine/Source/Runtime/SlateCore/Private/Styling/SlateTypes.cpp
+++ b/Engine/Source/Runtime/SlateCore/Private/Styling/SlateTypes.cpp
@@ -568,6 +568,8 @@ void FScrollBoxStyle::GetResources( TArray< const FSlateBrush* >& OutBrushes ) c
 {
 	OutBrushes.Add( &TopShadowBrush );
 	OutBrushes.Add( &BottomShadowBrush );
+	OutBrushes.Add( &LeftShadowBrush );
+	OutBrushes.Add( &RightShadowBrush );
 }
 
 const FName FScrollBoxStyle::TypeName( TEXT("FScrollBoxStyle") );
diff --git a/Engine/Source/Runtime/SlateCore/Public/Styling/SlateTypes.h b/Engine/Source/Runtime/SlateCore/Public/Styling/SlateTypes.h
--- a/Engine/Source/Runtime/SlateCore/Public/Styling/SlateTypes.h
+++ b/Engine/Source/Runtime/SlateCore/Public/Styling/SlateTypes.h
@@ -1244,6 +1244,24 @@ struct SLATECORE_API FScrollBoxStyle : public FSlateWidgetStyle
 	UPROPERTY(EditAnywhere, Category=Appearance)
 	FSlateBrush BottomShadowBrush;
 	FScrollBoxStyle& SetBottomShadowBrush( const FSlateBrush& InBottomShadowBrush ){ BottomShadowBrush = InBottomShadowBrush; return *this; }
+
+	/** Brush used to draw the left shadow of a scrollbox */
+	UPROPERTY(EditAnywhere, Category = Appearance)
+	FSlateBrush LeftShadowBrush;
+	FScrollBoxStyle& SetLeftShadowBrush(const FSlateBrush& InLeftShadowBrush)
+	{
+		LeftShadowBrush = InLeftShadowBrush;
+		return *this;
+	}
+
+	/** Brush used to draw the right shadow of a scrollbox */
+	UPROPERTY(EditAnywhere, Category = Appearance)
+	FSlateBrush RightShadowBrush;
+	FScrollBoxStyle& SetRightShadowBrush(const FSlateBrush& InRightShadowBrush)
+	{
+		RightShadowBrush = InRightShadowBrush;
+		return *this;
+	}
 };
 
 
diff --git a/Engine/Source/Runtime/UMG/Private/Components/ScrollBox.cpp b/Engine/Source/Runtime/UMG/Private/Components/ScrollBox.cpp
--- a/Engine/Source/Runtime/UMG/Private/Components/ScrollBox.cpp
+++ b/Engine/Source/Runtime/UMG/Private/Components/ScrollBox.cpp
@@ -14,6 +14,7 @@ UScrollBox::UScrollBox(const FPostConstructInitializeProperties& PCIP)
 
 	SScrollBox::FArguments Defaults;
 	Visiblity = UWidget::ConvertRuntimeToSerializedVisiblity(Defaults._Visibility.Get());
+	Orientation = Defaults._Orientation;
 }
 
 void UScrollBox::ReleaseNativeWidget()
@@ -68,7 +69,8 @@ TSharedRef<SWidget> UScrollBox::RebuildWidget()
 
 	MyScrollBox = SNew(SScrollBox)
 		.Style(StylePtr)
-		.ScrollBarStyle(BarStylePtr);
+		.ScrollBarStyle(BarStylePtr)
+		.Orientation(Orientation);
 
 	for ( UPanelSlot* Slot : Slots )
 	{
@@ -85,6 +87,7 @@ TSharedRef<SWidget> UScrollBox::RebuildWidget()
 void UScrollBox::SyncronizeProperties()
 {
 	MyScrollBox->SetScrollOffset(DesiredScrollOffset);
+	MyScrollBox->SetOrientation(Orientation);
 }
 
 void UScrollBox::ClearChildren()
diff --git a/Engine/Source/Runtime/UMG/Public/Components/ScrollBox.h b/Engine/Source/Runtime/UMG/Public/Components/ScrollBox.h
--- a/Engine/Source/Runtime/UMG/Public/Components/ScrollBox.h
+++ b/Engine/Source/Runtime/UMG/Public/Components/ScrollBox.h
@@ -20,6 +20,10 @@ class UMG_API UScrollBox : public UPanelWidget
 
 	virtual void ReleaseNativeWidget() override;
 
+	/** The orientation of the scrolling and stacking in the box. */
+	UPROPERTY(EditDefaultsOnly, Category = "Behavior")
+	TEnumAsByte<EOrientation> Orientation;
+
 	/** Removes all child widget from the scrollbox */
 	UFUNCTION(BlueprintCallable, Category="Widget")
 	void ClearChildren();
