Slate: Added a test case for disabled state to the Slate test suite.

diff --git a/Engine/Source/Runtime/Slate/Private/Widgets/Testing/STestSuite.cpp b/Engine/Source/Runtime/Slate/Private/Widgets/Testing/STestSuite.cpp
--- a/Engine/Source/Runtime/Slate/Private/Widgets/Testing/STestSuite.cpp
+++ b/Engine/Source/Runtime/Slate/Private/Widgets/Testing/STestSuite.cpp
@@ -526,11 +526,37 @@ public:
 		
 		this->ChildSlot
 		[
-			SNew(SListView< TSharedRef<FDocumentInfo> >)
-			.ItemHeight(24)
-			.SelectionMode(ESelectionMode::None)
-			.ListItemsSource( &Documents )
-			.OnGenerateRow( this, &SDocumentsTest::GenerateListRow )
+			SNew(SVerticalBox)
+			+ SVerticalBox::Slot()
+			[
+				SNew(SListView< TSharedRef<FDocumentInfo> >)
+				.ItemHeight(24)
+				.SelectionMode(ESelectionMode::None)
+				.ListItemsSource( &Documents )
+				.OnGenerateRow( this, &SDocumentsTest::GenerateListRow )
+			]
+			+ SVerticalBox::Slot() .HAlign(HAlign_Center) .VAlign(VAlign_Center)
+			[
+				SNew(SButton)
+				.ContentPadding(20).HAlign(HAlign_Center).VAlign(VAlign_Center)
+				[
+					SNew(SButton)
+					.IsEnabled( false )
+					.ContentPadding(20).HAlign(HAlign_Center).VAlign(VAlign_Center)
+					[
+						SNew(SButton)
+						.ContentPadding(20).HAlign(HAlign_Center).VAlign(VAlign_Center)
+						[
+							SNew(SButton)
+							.ContentPadding(20).HAlign(HAlign_Center).VAlign(VAlign_Center)
+							[
+								SNew(STextBlock)
+								.Text(NSLOCTEXT("x","ClickMe","ClickMe!"))
+							]
+						]
+					]
+				]
+			]
 		];
 	}
 
