Fixed some issues with Ctrl+{Key} shortcuts on Mac

diff --git a/Engine/Source/Editor/BlueprintGraph/Private/K2Node_InputKey.cpp b/Engine/Source/Editor/BlueprintGraph/Private/K2Node_InputKey.cpp
--- a/Engine/Source/Editor/BlueprintGraph/Private/K2Node_InputKey.cpp
+++ b/Engine/Source/Editor/BlueprintGraph/Private/K2Node_InputKey.cpp
@@ -314,6 +314,7 @@ void UK2Node_InputKey::CreateInputKeyEvent(FKismetCompilerContext& CompilerConte
 		InputKeyEvent->InputChord.bCtrl = bControl;
 		InputKeyEvent->InputChord.bAlt = bAlt;
 		InputKeyEvent->InputChord.bShift = bShift;
+		InputKeyEvent->InputChord.bCmd = bCommand;
 		InputKeyEvent->bConsumeInput = bConsumeInput;
 		InputKeyEvent->bExecuteWhenPaused = bExecuteWhenPaused;
 		InputKeyEvent->bOverrideParentBinding = bOverrideParentBinding;
diff --git a/Engine/Source/Editor/DetailCustomizations/Private/InputStructCustomization.cpp b/Engine/Source/Editor/DetailCustomizations/Private/InputStructCustomization.cpp
--- a/Engine/Source/Editor/DetailCustomizations/Private/InputStructCustomization.cpp
+++ b/Engine/Source/Editor/DetailCustomizations/Private/InputStructCustomization.cpp
@@ -59,6 +59,7 @@ void FInputActionMappingCustomization::CustomizeChildren( TSharedRef<class IProp
 	TSharedPtr<IPropertyHandle> ShiftHandle = InStructPropertyHandle->GetChildHandle(GET_MEMBER_NAME_CHECKED(FInputActionKeyMapping, bShift));
 	TSharedPtr<IPropertyHandle> CtrlHandle = InStructPropertyHandle->GetChildHandle(GET_MEMBER_NAME_CHECKED(FInputActionKeyMapping, bCtrl));
 	TSharedPtr<IPropertyHandle> AltHandle = InStructPropertyHandle->GetChildHandle(GET_MEMBER_NAME_CHECKED(FInputActionKeyMapping, bAlt));
+	TSharedPtr<IPropertyHandle> CmdHandle = InStructPropertyHandle->GetChildHandle(GET_MEMBER_NAME_CHECKED(FInputActionKeyMapping, bCmd));
 
 	TSharedRef<SWidget> RemoveButton = PropertyCustomizationHelpers::MakeDeleteButton(FSimpleDelegate::CreateSP(this, &FInputActionMappingCustomization::RemoveActionMappingButton_OnClick),
 		LOCTEXT("RemoveActionMappingToolTip", "Removes Action Mapping"));
@@ -129,6 +130,22 @@ void FInputActionMappingCustomization::CustomizeChildren( TSharedRef<class IProp
 		.HAlign(HAlign_Center)
 		.VAlign(VAlign_Center)
 		.AutoWidth()
+		[
+			CmdHandle->CreatePropertyNameWidget()
+		]
+		+ SHorizontalBox::Slot()
+		.Padding(InputConstants::PropertyPadding)
+		.HAlign(HAlign_Center)
+		.VAlign(VAlign_Center)
+		.AutoWidth()
+		[
+			CmdHandle->CreatePropertyValueWidget()
+		]
+		+ SHorizontalBox::Slot()
+		.Padding(InputConstants::PropertyPadding)
+		.HAlign(HAlign_Center)
+		.VAlign(VAlign_Center)
+		.AutoWidth()
 		[
 			RemoveButton
 		]
diff --git a/Engine/Source/Editor/GraphEditor/Private/SNodePanel.cpp b/Engine/Source/Editor/GraphEditor/Private/SNodePanel.cpp
--- a/Engine/Source/Editor/GraphEditor/Private/SNodePanel.cpp
+++ b/Engine/Source/Editor/GraphEditor/Private/SNodePanel.cpp
@@ -844,6 +844,7 @@ FReply SNodePanel::OnMouseButtonUp( const FGeometry& MyGeometry, const FPointerE
 				LastKeyGestureDetected.bAlt = false;
 				LastKeyGestureDetected.bCtrl = false;
 				LastKeyGestureDetected.bShift = false;
+				LastKeyGestureDetected.bCmd = false;
 			}
 		}
 		else if ( Marquee.IsValid() )
@@ -912,6 +913,7 @@ FReply SNodePanel::OnKeyDown( const FGeometry& MyGeometry, const FKeyboardEvent&
 		LastKeyGestureDetected.bAlt = InKeyboardEvent.IsAltDown();
 		LastKeyGestureDetected.bCtrl = InKeyboardEvent.IsControlDown();
 		LastKeyGestureDetected.bShift = InKeyboardEvent.IsShiftDown();
+		LastKeyGestureDetected.bCmd = InKeyboardEvent.IsCommandDown();
 	}
 
 	return FReply::Unhandled();
@@ -925,6 +927,7 @@ FReply SNodePanel::OnKeyUp( const FGeometry& MyGeometry, const FKeyboardEvent& I
 		LastKeyGestureDetected.bAlt = false;
 		LastKeyGestureDetected.bCtrl = false;
 		LastKeyGestureDetected.bShift = false;
+		LastKeyGestureDetected.bCmd = false;
 	}
 
 	return FReply::Unhandled();
@@ -936,6 +939,7 @@ void SNodePanel::OnKeyboardFocusLost( const FKeyboardFocusEvent& InKeyboardFocus
 	LastKeyGestureDetected.bAlt = false;
 	LastKeyGestureDetected.bCtrl = false;
 	LastKeyGestureDetected.bShift = false;
+	LastKeyGestureDetected.bCmd = false;
 }
 
 FReply SNodePanel::OnTouchGesture( const FGeometry& MyGeometry, const FPointerEvent& GestureEvent )
diff --git a/Engine/Source/Editor/Kismet/Private/BlueprintEditorCommands.cpp b/Engine/Source/Editor/Kismet/Private/BlueprintEditorCommands.cpp
--- a/Engine/Source/Editor/Kismet/Private/BlueprintEditorCommands.cpp
+++ b/Engine/Source/Editor/Kismet/Private/BlueprintEditorCommands.cpp
@@ -344,6 +344,7 @@ void FBlueprintSpawnNodeCommands::RegisterCommands()
 			bool bShift = false;
 			bool bCtrl = false;
 			bool bAlt = false;
+			bool bCmd = false;
 			
 			// Parse the keybinding information
 			FString KeyString;
@@ -357,12 +358,14 @@ void FBlueprintSpawnNodeCommands::RegisterCommands()
 				FParse::Bool(*NodeSpawns[x], TEXT("Shift="), bShift);
 				FParse::Bool(*NodeSpawns[x], TEXT("Alt="), bAlt);
 				FParse::Bool(*NodeSpawns[x], TEXT("Ctrl="), bCtrl);
+				FParse::Bool(*NodeSpawns[x], TEXT("Cmd="), bCmd);
 			}
 
 			FInputGesture Gesture(Key);
 			Gesture.bAlt = bAlt;
 			Gesture.bCtrl = bCtrl;
 			Gesture.bShift = bShift;
+			Gesture.bCmd = bCmd;
 
 			FText CommandLabelText = FText::FromString( CommandLabel );
 			FText Description = FText::Format( NSLOCTEXT("BlueprintEditor", "NodeSpawnDescription", "Hold down the bound keys and left click in the graph panel to spawn a {0} node."), CommandLabelText );
diff --git a/Engine/Source/Editor/UnrealEd/Private/EditorViewportClient.cpp b/Engine/Source/Editor/UnrealEd/Private/EditorViewportClient.cpp
--- a/Engine/Source/Editor/UnrealEd/Private/EditorViewportClient.cpp
+++ b/Engine/Source/Editor/UnrealEd/Private/EditorViewportClient.cpp
@@ -207,6 +207,7 @@ FEditorViewportClient::FEditorViewportClient(FEditorModeTools& InModeTools, FPre
 	, bDrawAxes(true)
 	, bSetListenerPosition(false)
 	, LandscapeLODOverride(-1)
+	, ModeTools(&InModeTools)
 	, Widget(new FWidget)
 	, MouseDeltaTracker(new FMouseDeltaTracker)
 	, RecordingInterpEd(NULL)
@@ -229,7 +230,6 @@ FEditorViewportClient::FEditorViewportClient(FEditorModeTools& InModeTools, FPre
 	, bForceAudioRealtime(false)
 	, bIsRealtime(false)
 	, bStoredRealtime(false)
-	, ModeTools(&InModeTools)
 	, bStoredShowStats(false)
 	, bShowStats(false)
 	, bHasAudioFocus(false)
@@ -1992,6 +1992,11 @@ bool FEditorViewportClient::IsShiftPressed() const
 	return Viewport->KeyState(EKeys::LeftShift) || Viewport->KeyState(EKeys::RightShift);
 }
 
+bool FEditorViewportClient::IsCmdPressed() const
+{
+	return Viewport->KeyState(EKeys::LeftCommand) || Viewport->KeyState(EKeys::RightCommand);
+}
+
 
 void FEditorViewportClient::ProcessDoubleClickInViewport( const struct FInputEventState& InputState, FSceneView& View )
 {
diff --git a/Engine/Source/Editor/UnrealEd/Private/LevelEditorViewport.cpp b/Engine/Source/Editor/UnrealEd/Private/LevelEditorViewport.cpp
--- a/Engine/Source/Editor/UnrealEd/Private/LevelEditorViewport.cpp
+++ b/Engine/Source/Editor/UnrealEd/Private/LevelEditorViewport.cpp
@@ -1282,8 +1282,8 @@ bool FLevelEditorViewportClient::CanApplyMaterialToHitProxy( const HHitProxy* Hi
 }
 
 FTrackingTransaction::FTrackingTransaction()
-	: ScopedTransaction(NULL)
-	, TransCount(0)
+	: TransCount(0)
+	, ScopedTransaction(NULL)
 	, TrackingTransactionState(ETransactionState::Inactive)
 {}
 
@@ -1366,8 +1366,6 @@ FLevelEditorViewportClient::FLevelEditorViewportClient()
 	: FEditorViewportClient(GLevelEditorModeTools(), NULL)
 	, ViewHiddenLayers()
 	, VolumeActorVisibility()
-	, ActorLockedToCamera(NULL)
-	, ActorLockedByMatinee(NULL)
 	, LastEditorViewLocation( FVector::ZeroVector )
 	, LastEditorViewRotation( FRotator::ZeroRotator )
 	, ColorScale( FVector(1,1,1) )
@@ -1380,13 +1378,15 @@ FLevelEditorViewportClient::FLevelEditorViewportClient()
 	, bDuplicateActorsOnNextDrag( false )
 	, bDuplicateActorsInProgress( false )
 	, bIsTrackingBrushModification( false )
+	, bLockedCameraView(true)
 	, SpriteCategoryVisibility()
 	, World(NULL)
 	, TrackingTransaction()
 	, DropPreviewMouseX(0)
 	, DropPreviewMouseY(0)
-	, bLockedCameraView(true)
 	, bWasControlledByOtherViewport(false)
+	, ActorLockedByMatinee(NULL)
+	, ActorLockedToCamera(NULL)
 {
 	// By default a level editor viewport is pointed to the editor world
 	SetReferenceToWorldContext(GEditor->GetEditorWorldContext());
@@ -2182,6 +2182,7 @@ static bool CommandAcceptsInput( FLevelEditorViewportClient& ViewportClient, FKe
 	return (!Gesture.bCtrl	|| ViewportClient.IsCtrlPressed() ) 
 		&& (!Gesture.bAlt	|| ViewportClient.IsAltPressed() ) 
 		&& (!Gesture.bShift || ViewportClient.IsShiftPressed() ) 
+		&& (!Gesture.bCmd	|| ViewportClient.IsCmdPressed() )
 		&& Gesture.Key == Key;
 }
 
diff --git a/Engine/Source/Editor/UnrealEd/Public/EditorViewportClient.h b/Engine/Source/Editor/UnrealEd/Public/EditorViewportClient.h
--- a/Engine/Source/Editor/UnrealEd/Public/EditorViewportClient.h
+++ b/Engine/Source/Editor/UnrealEd/Public/EditorViewportClient.h
@@ -855,6 +855,7 @@ public:
 	bool IsAltPressed() const;
 	bool IsCtrlPressed() const;
 	bool IsShiftPressed() const;
+	bool IsCmdPressed() const;
 
 	/** @return True if the window is in an immersive viewport */
 	bool IsInImmersiveViewport() const;
diff --git a/Engine/Source/Runtime/Engine/Private/InputComponent.cpp b/Engine/Source/Runtime/Engine/Private/InputComponent.cpp
--- a/Engine/Source/Runtime/Engine/Private/InputComponent.cpp
+++ b/Engine/Source/Runtime/Engine/Private/InputComponent.cpp
@@ -10,19 +10,22 @@ FInputChord::RelationshipType FInputChord::GetRelationship(const FInputChord& Ot
 	{
 		if (    bAlt == OtherChord.bAlt
 			 && bCtrl == OtherChord.bCtrl
-			 && bShift == OtherChord.bShift)
+			 && bShift == OtherChord.bShift
+			 && bCmd == OtherChord.bCmd)
 		{
 			Relationship = Same;
 		}
 		else if (    (bAlt || !OtherChord.bAlt)
 				  && (bCtrl || !OtherChord.bCtrl)
-				  && (bShift || !OtherChord.bShift))
+				  && (bShift || !OtherChord.bShift)
+				  && (bCmd || !OtherChord.bCmd))
 		{
 			Relationship = Masks;
 		}
 		else if (    (!bAlt || OtherChord.bAlt)
 				  && (!bCtrl || OtherChord.bCtrl)
-				  && (!bShift || OtherChord.bShift))
+				  && (!bShift || OtherChord.bShift)
+				  && (!bCmd || OtherChord.bCmd))
 		{
 			Relationship = Masked;
 		}
diff --git a/Engine/Source/Runtime/Slate/Private/Framework/Commands/InputBindingManager.cpp b/Engine/Source/Runtime/Slate/Private/Framework/Commands/InputBindingManager.cpp
--- a/Engine/Source/Runtime/Slate/Private/Framework/Commands/InputBindingManager.cpp
+++ b/Engine/Source/Runtime/Slate/Private/Framework/Commands/InputBindingManager.cpp
@@ -82,6 +82,7 @@ void FUserDefinedGestures::LoadGestures()
 					const TSharedPtr<FJsonValue> CtrlObj = GestureInfoObj->Values.FindRef( TEXT("Control") );
 					const TSharedPtr<FJsonValue> AltObj = GestureInfoObj->Values.FindRef( TEXT("Alt") );
 					const TSharedPtr<FJsonValue> ShiftObj = GestureInfoObj->Values.FindRef( TEXT("Shift") );
+					const TSharedPtr<FJsonValue> CmdObj = GestureInfoObj->Values.FindRef( TEXT("Command") );
 					const TSharedPtr<FJsonValue> KeyObj = GestureInfoObj->Values.FindRef( TEXT("Key") );
 
 					const FName BindingContext = *BindingContextObj->AsString();
@@ -93,6 +94,7 @@ void FUserDefinedGestures::LoadGestures()
 					UserDefinedGesture.bCtrl = CtrlObj->AsBool();
 					UserDefinedGesture.bAlt = AltObj->AsBool();
 					UserDefinedGesture.bShift = ShiftObj->AsBool();
+					UserDefinedGesture.bCmd = CmdObj->AsBool();
 					UserDefinedGesture.Key = *KeyObj->AsString();
 				}
 			}
@@ -136,6 +138,7 @@ void FUserDefinedGestures::LoadGestures()
 						const TSharedPtr<FJsonValue> CtrlObj = CommandObj->Values.FindRef( TEXT("Control") );
 						const TSharedPtr<FJsonValue> AltObj = CommandObj->Values.FindRef( TEXT("Alt") );
 						const TSharedPtr<FJsonValue> ShiftObj = CommandObj->Values.FindRef( TEXT("Shift") );
+						const TSharedPtr<FJsonValue> CmdObj = CommandObj->Values.FindRef( TEXT("Command") );
 						const TSharedPtr<FJsonValue> KeyObj = CommandObj->Values.FindRef( TEXT("Key") );
 
 						const FUserDefinedGestureKey GestureKey(BindingContext, CommandName);
@@ -144,6 +147,7 @@ void FUserDefinedGestures::LoadGestures()
 						UserDefinedGesture.bCtrl = CtrlObj->AsBool();
 						UserDefinedGesture.bAlt = AltObj->AsBool();
 						UserDefinedGesture.bShift = ShiftObj->AsBool();
+						UserDefinedGesture.bCmd = CmdObj->AsBool();
 						UserDefinedGesture.Key = *KeyObj->AsString();
 					}
 				}
@@ -169,6 +173,7 @@ void FUserDefinedGestures::SaveGestures() const
 			GestureInfoObj->Values.Add( TEXT("Control"), MakeShareable( new FJsonValueBoolean( GestureInfo.Value.bCtrl ) ) );
 			GestureInfoObj->Values.Add( TEXT("Alt"), MakeShareable( new FJsonValueBoolean( GestureInfo.Value.bAlt ) ) );
 			GestureInfoObj->Values.Add( TEXT("Shift"), MakeShareable( new FJsonValueBoolean( GestureInfo.Value.bShift ) ) );
+			GestureInfoObj->Values.Add( TEXT("Command"), MakeShareable( new FJsonValueBoolean( GestureInfo.Value.bCmd ) ) );
 			GestureInfoObj->Values.Add( TEXT("Key"), MakeShareable( new FJsonValueString( GestureInfo.Value.Key.ToString() ) ) );
 
 			auto JsonWriter = TJsonWriterFactory< TCHAR, TCondensedJsonPrintPolicy<TCHAR> >::Create( &GestureRawJsonContent );
diff --git a/Engine/Source/Runtime/Slate/Private/Framework/Commands/UICommandList.cpp b/Engine/Source/Runtime/Slate/Private/Framework/Commands/UICommandList.cpp
--- a/Engine/Source/Runtime/Slate/Private/Framework/Commands/UICommandList.cpp
+++ b/Engine/Source/Runtime/Slate/Private/Framework/Commands/UICommandList.cpp
@@ -186,7 +186,7 @@ bool FUICommandList::IsChecked( const TSharedRef< const FUICommandInfo > InUICom
  */
 bool FUICommandList::ProcessCommandBindings( const FKeyboardEvent& InKeyboardEvent ) const 
 {
-	return ConditionalProcessCommandBindings( InKeyboardEvent.GetKey(), InKeyboardEvent.IsControlDown(), InKeyboardEvent.IsAltDown(), InKeyboardEvent.IsShiftDown(), InKeyboardEvent.IsRepeat() );
+	return ConditionalProcessCommandBindings( InKeyboardEvent.GetKey(), InKeyboardEvent.IsControlDown(), InKeyboardEvent.IsAltDown(), InKeyboardEvent.IsShiftDown(), InKeyboardEvent.IsCommandDown(), InKeyboardEvent.IsRepeat() );
 }
 
 
@@ -200,7 +200,7 @@ bool FUICommandList::ProcessCommandBindings( const FKeyboardEvent& InKeyboardEve
  */
 bool FUICommandList::ProcessCommandBindings( const FPointerEvent& InMouseEvent ) const
 {
-	return ConditionalProcessCommandBindings( InMouseEvent.GetEffectingButton(), InMouseEvent.IsControlDown(), InMouseEvent.IsAltDown(), InMouseEvent.IsShiftDown(), InMouseEvent.IsRepeat() );
+	return ConditionalProcessCommandBindings( InMouseEvent.GetEffectingButton(), InMouseEvent.IsControlDown(), InMouseEvent.IsAltDown(), InMouseEvent.IsShiftDown(), InMouseEvent.IsCommandDown(), InMouseEvent.IsRepeat() );
 }
 
 
@@ -220,6 +220,7 @@ bool FUICommandList::ProcessCommandBindings( const FKey Key, const FModifierKeys
 		ModifierKeysState.IsControlDown(),
 		ModifierKeysState.IsAltDown(),
 		ModifierKeysState.IsShiftDown(),
+		ModifierKeysState.IsCommandDown(),
 		bRepeat );
 }
 
@@ -238,7 +239,7 @@ bool FUICommandList::ProcessCommandBindings( const FKey Key, const FModifierKeys
  * @param bRepeat	True if command is repeating (held)
  * @return True if a command was executed, False otherwise
  */
-bool FUICommandList::ConditionalProcessCommandBindings( const FKey Key, bool bCtrl, bool bAlt, bool bShift, bool bRepeat ) const
+bool FUICommandList::ConditionalProcessCommandBindings( const FKey Key, bool bCtrl, bool bAlt, bool bShift, bool bCmd, bool bRepeat ) const
 {
 	if ( !bRepeat && !FSlateApplication::Get().IsDragDropping() )
 	{
@@ -246,6 +247,7 @@ bool FUICommandList::ConditionalProcessCommandBindings( const FKey Key, bool bCt
 		CheckGesture.bCtrl = bCtrl;
 		CheckGesture.bAlt = bAlt;
 		CheckGesture.bShift = bShift;
+		CheckGesture.bCmd = bCmd;
 
 		if( CheckGesture.IsValidGesture() )
 		{
diff --git a/Engine/Source/Runtime/Slate/Public/Framework/Commands/UICommandList.h b/Engine/Source/Runtime/Slate/Public/Framework/Commands/UICommandList.h
--- a/Engine/Source/Runtime/Slate/Public/Framework/Commands/UICommandList.h
+++ b/Engine/Source/Runtime/Slate/Public/Framework/Commands/UICommandList.h
@@ -154,7 +154,7 @@ protected:
 	 * @param bRepeat	True if command is repeating (held)
 	 * @return True if a command was executed, False otherwise
 	 */
-	bool ConditionalProcessCommandBindings( const FKey Key, bool bCtrl, bool bAlt, bool bShift, bool bRepeat ) const;
+	bool ConditionalProcessCommandBindings( const FKey Key, bool bCtrl, bool bAlt, bool bShift, bool bCmd, bool bRepeat ) const;
 
 	/** 
 	  * Attempts to find an action for the specified command in the current UICommandList. If it is not found, the action for the
