Slate: multi line editable text now has a dedicated tab in the test suite.

diff --git a/Engine/Source/Runtime/Slate/Private/Widgets/Testing/STestSuite.cpp b/Engine/Source/Runtime/Slate/Private/Widgets/Testing/STestSuite.cpp
--- a/Engine/Source/Runtime/Slate/Private/Widgets/Testing/STestSuite.cpp
+++ b/Engine/Source/Runtime/Slate/Private/Widgets/Testing/STestSuite.cpp
@@ -838,6 +838,126 @@ protected:
 
 #if WITH_FANCY_TEXT
 
+class SMultiLineEditingTest : public SCompoundWidget
+{
+	SLATE_BEGIN_ARGS( SMultiLineEditingTest )
+	{}
+	SLATE_END_ARGS()
+
+	void Construct( const FArguments& InArgs )
+	{
+		MultilineEditableText = LOCTEXT( "MultiLineTextTest", "He has refused his Assent to Laws, the most wholesome and necessary for the public good.\nHe has forbidden his Governors to pass Laws of immediate and pressing importance, unless suspended in their operation till his Assent should be obtained; and when so suspended, he has utterly neglected to attend to them.\nHe has refused to pass other Laws for the accommodation of large districts of people, unless those people would relinquish the right of Representation in the Legislature, a right inestimable to them and formidable to tyrants only.\n\nHe has called together legislative bodies at places unusual, uncomfortable, and distant from the depository of their public Records, for the sole purpose of fatiguing them into compliance with his measures.\nHe has dissolved Representative Houses repeatedly, for opposing with manly firmness his invasions on the rights of the people.\nHe has refused for a long time, after such dissolutions, to cause others to be elected; whereby the Legislative powers, incapable of Annihilation, have returned to the People at large for their exercise; the State remaining in the mean time exposed to all the dangers of invasion from without, and convulsions within.\nHe has endeavoured to prevent the population of these States; for that purpose obstructing the Laws for Naturalization of Foreigners; refusing to pass others to encourage their migrations hither, and raising the conditions of new Appropriations of Lands.\n" );
+
+		TSharedRef<SScrollBar> FirstScrollBar = SNew( SScrollBar );
+		TSharedRef<SScrollBar> SecondScrollBar = SNew( SScrollBar );
+
+		this->ChildSlot
+		[
+			SNew(SVerticalBox)
+			+SVerticalBox::Slot().FillHeight( 1 )
+			[
+				SNew( SHorizontalBox )
+				+ SHorizontalBox::Slot().AutoWidth().VAlign( VAlign_Top ).Padding( 2 )
+				[
+					SNew( SVerticalBox )
+					+ SVerticalBox::Slot()
+					.AutoHeight()
+					[
+						SNew( STextBlock )
+						.Text( LOCTEXT( "MultiLineTextWrapAt600px", "Multi-line editable text wrapping at 600px" ) )
+					]
+					+ SVerticalBox::Slot()
+					[
+						SNew( SHorizontalBox )
+						+ SHorizontalBox::Slot()
+						[
+							SNew( SMultiLineEditableTextBox )
+							.Text( this, &SMultiLineEditingTest::GetMultilineEditableText )
+							.Font( FSlateFontInfo( FPaths::EngineContentDir() / TEXT( "Slate/Fonts/Roboto-Regular.ttf" ), 12 ) )
+							.Justification( ETextJustify::Center )
+							.LineHeightPercentage( 2.0f )
+							.OnTextCommitted( this, &SMultiLineEditingTest::HandleMultilineEditableTextCommitted )
+							.WrapTextAt( 600.0f )
+						]
+						+ SHorizontalBox::Slot().AutoWidth()
+						[
+							FirstScrollBar
+						]
+					]
+				]
+				+ SHorizontalBox::Slot().FillWidth( 1 ).Padding( 2 )
+				[
+					SNew( SVerticalBox )
+					+ SVerticalBox::Slot()
+					.AutoHeight()
+					[
+						SNew( STextBlock )
+						.Text( LOCTEXT( "MultiLineTextAutoWrap", "Multi-line editable text auto-wrapping" ) )
+					]
+					+ SVerticalBox::Slot()
+						[
+							SNew( SHorizontalBox )
+							+ SHorizontalBox::Slot()
+							[
+								SNew( SMultiLineEditableTextBox )
+								.Margin( 10 )
+								.Text( MultilineEditableText )
+								//.Justification(ETextJustify::Right)
+								.Font( FSlateFontInfo( FPaths::EngineContentDir() / TEXT( "Slate/Fonts/Roboto-Regular.ttf" ), 12 ) )
+								.AutoWrapText( true )
+								.HintText( LOCTEXT( "TypehereTextHint", "Type Here" ) )
+							]
+							+ SHorizontalBox::Slot().AutoWidth()
+								[
+									SecondScrollBar
+								]
+						]
+				]
+			]
+			+ SVerticalBox::Slot().AutoHeight()
+			[
+				SNew( SHorizontalBox )
+				+ SHorizontalBox::Slot().AutoWidth().VAlign( VAlign_Top ).Padding( 2 )
+				[
+					SNew( SVerticalBox )
+					+ SVerticalBox::Slot()
+					.AutoHeight()
+					[
+						SNew( STextBlock )
+						.Text( LOCTEXT( "MultiLineTextNoWrap", "Multi-line editable text without wrapping" ) )
+					]
+					+ SVerticalBox::Slot()
+						.AutoHeight()
+						[
+							SNew( SHorizontalBox )
+							+ SHorizontalBox::Slot()
+							[
+								SNew( SMultiLineEditableTextBox )
+								.Font( FSlateFontInfo( FPaths::EngineContentDir() / TEXT( "Slate/Fonts/Roboto-Regular.ttf" ), 12 ) )
+								//.WrapTextAt(300.0f)
+							]
+						]
+				]
+			]
+		];
+	}
+
+private:
+	FText GetMultilineEditableText() const
+	{
+		return MultilineEditableText;
+	}
+
+	void HandleMultilineEditableTextCommitted( const FText& Text, ETextCommit::Type Type )
+	{
+		MultilineEditableText = Text;
+	}
+
+
+	FText MultilineEditableText;
+};
+
+
 class SRichTextTest : public SCompoundWidget
 {
 public:
@@ -1577,14 +1697,9 @@ public:
 	{
 		InlineEditableText = LOCTEXT( "TestingInlineEditableTextBlock", "Testing inline editable text block!" );
 
-		MultilineEditableText = LOCTEXT("MultiLineTextTest", "He has refused his Assent to Laws, the most wholesome and necessary for the public good.\nHe has forbidden his Governors to pass Laws of immediate and pressing importance, unless suspended in their operation till his Assent should be obtained; and when so suspended, he has utterly neglected to attend to them.\nHe has refused to pass other Laws for the accommodation of large districts of people, unless those people would relinquish the right of Representation in the Legislature, a right inestimable to them and formidable to tyrants only.\n\nHe has called together legislative bodies at places unusual, uncomfortable, and distant from the depository of their public Records, for the sole purpose of fatiguing them into compliance with his measures.\nHe has dissolved Representative Houses repeatedly, for opposing with manly firmness his invasions on the rights of the people.\nHe has refused for a long time, after such dissolutions, to cause others to be elected; whereby the Legislative powers, incapable of Annihilation, have returned to the People at large for their exercise; the State remaining in the mean time exposed to all the dangers of invasion from without, and convulsions within.\nHe has endeavoured to prevent the population of these States; for that purpose obstructing the Laws for Naturalization of Foreigners; refusing to pass others to encourage their migrations hither, and raising the conditions of new Appropriations of Lands.\n");
-
 		Animation = FCurveSequence(0, 5);
 		Animation.Play();
 
-		TSharedRef<SScrollBar> FirstScrollBar = SNew(SScrollBar);
-		TSharedRef<SScrollBar> SecondScrollBar = SNew(SScrollBar);
-
 		this->ChildSlot
 		[
 			SNew( SScrollBox )
@@ -1662,111 +1777,11 @@ public:
 				[
 					SAssignNew(ErrorText, SErrorText)
 				]
-
-#if WITH_FANCY_TEXT
-				+SVerticalBox::Slot().FillHeight(1)
-				[
-					SNew(SHorizontalBox)
-					+SHorizontalBox::Slot().AutoWidth() .VAlign(VAlign_Top) .Padding(2)
-					[
-						SNew(SVerticalBox)
-						+ SVerticalBox::Slot()
-						.AutoHeight()
-						[
-							SNew(STextBlock)
-							.Text(LOCTEXT("MultiLineTextWrapAt600px", "Multi-line editable text wrapping at 600px"))
-						]
-						+ SVerticalBox::Slot()
-						.AutoHeight()
-						[
-							SNew(SHorizontalBox)
-							+ SHorizontalBox::Slot()
-							[
-								SNew(SMultiLineEditableTextBox)
-								.Text(this, &STextEditTest::GetMultilineEditableText)
-								.Font(FSlateFontInfo(FPaths::EngineContentDir() / TEXT("Slate/Fonts/Roboto-Regular.ttf"), 12))
-								.Justification(ETextJustify::Center)
-								.LineHeightPercentage(2.0f)
-								.OnTextCommitted(this, &STextEditTest::HandleMultilineEditableTextCommitted)
-								.WrapTextAt(600.0f)
-							]
-							+ SHorizontalBox::Slot().AutoWidth()
-							[
-								FirstScrollBar
-							]
-						]
-					]
-					+SHorizontalBox::Slot().FillWidth(1).Padding(2)
-					[
-						SNew(SVerticalBox)
-						+ SVerticalBox::Slot()
-						.AutoHeight()
-						[
-							SNew(STextBlock)
-							.Text(LOCTEXT("MultiLineTextAutoWrap", "Multi-line editable text auto-wrapping"))
-						]
-						+ SVerticalBox::Slot()
-						.AutoHeight()
-						[
-							SNew(SHorizontalBox)
-							+SHorizontalBox::Slot()
-							[
-								SNew(SMultiLineEditableTextBox)
-								.Margin(10)
-								.Text(MultilineEditableText)
-								//.Justification(ETextJustify::Right)
-								.Font(FSlateFontInfo(FPaths::EngineContentDir() / TEXT("Slate/Fonts/Roboto-Regular.ttf"), 12))
-								.AutoWrapText(true)
-								.HintText(LOCTEXT("TypehereTextHint", "Type Here"))
-							]
-							+SHorizontalBox::Slot() .AutoWidth()
-							[
-								SecondScrollBar
-							]
-						]
-					]
-				]
-				+SVerticalBox::Slot().AutoHeight()
-				[
-					SNew(SHorizontalBox)
-					+SHorizontalBox::Slot().AutoWidth() .VAlign(VAlign_Top) .Padding(2)
-					[
-						SNew(SVerticalBox)
-						+ SVerticalBox::Slot()
-						.AutoHeight()
-						[
-							SNew(STextBlock)
-							.Text(LOCTEXT("MultiLineTextNoWrap", "Multi-line editable text without wrapping"))
-						]
-						+ SVerticalBox::Slot()
-						.AutoHeight()
-						[
-							SNew(SHorizontalBox)
-							+ SHorizontalBox::Slot()
-							[
-								SNew(SMultiLineEditableTextBox)
-								.Font(FSlateFontInfo(FPaths::EngineContentDir() / TEXT("Slate/Fonts/Roboto-Regular.ttf"), 12))
-								//.WrapTextAt(300.0f)
-							]
-						]
-					]
-				]
-				#endif //WITH_FANCY_TEXT
 			]
 		];
 	}
 	END_SLATE_FUNCTION_BUILD_OPTIMIZATION
 
-	FText GetMultilineEditableText() const
-	{
-		return MultilineEditableText;
-	}
-
-	void HandleMultilineEditableTextCommitted(const FText& Text, ETextCommit::Type Type)
-	{
-		MultilineEditableText = Text;
-	}
-
 	void FocusDefaultWidget()
 	{
 		// Set focus to the editable text, so the user doesn't have to click initially to type
@@ -1870,7 +1885,6 @@ protected:
 	TSharedPtr<SInlineEditableTextBlock> InlineEditableTextBlock;
 	FText InlineEditableText;
 
-	FText MultilineEditableText;
  };
 
 /** Demonstrates the brokenness of our current approach to trading smoothness for sharpness. */
@@ -3281,7 +3295,6 @@ TSharedRef<SDockTab> SpawnTab(const FSpawnTabArgs& Args, FName TabIdentifier)
 	{
 		return SNew(SDockTab)
 			. Label( LOCTEXT("RichTextTestTab", "Rich Text") )
-			. ToolTipText( LOCTEXT( "RichTextTestTabToolTip", "Switches to the Rich Text tab, where you can test the various rich text features." ) )
 			[
 				#if WITH_FANCY_TEXT
 				SNew( SRichTextTest )
@@ -3290,6 +3303,18 @@ TSharedRef<SDockTab> SpawnTab(const FSpawnTabArgs& Args, FName TabIdentifier)
 				#endif //WITH_FANCY_TEXT
 			];
 	}
+	else if ( TabIdentifier == FName( TEXT( "MultiLineEditTab" ) ) )
+	{
+		return SNew(SDockTab)
+			. Label( LOCTEXT("MultiLineEditTab", "MultiLine Edit") )
+			[
+				#if WITH_FANCY_TEXT
+				SNew( SMultiLineEditingTest )
+				#else
+				SNew( SSpacer )
+				#endif //WITH_FANCY_TEXT
+			];
+	}
 	else if (TabIdentifier == FName(TEXT("EditableTextTab")))
 	{
 		return SNew(SDockTab)
@@ -3437,7 +3462,8 @@ TSharedRef<SDockTab> SpawnTestSuite1( const FSpawnTabArgs& Args )
 				->AddTab("LayoutRoundingTab", ETabState::OpenedTab)
 				->AddTab("EditableTextTab", ETabState::OpenedTab)
 				->AddTab("RichTextTab", ETabState::OpenedTab)
-				)
+				->AddTab("MultiLineEditTab", ETabState::OpenedTab)
+			)
 		)
 	)
 #if PLATFORM_SUPPORTS_MULTIPLE_NATIVE_WINDOWS
@@ -3480,6 +3506,10 @@ TSharedRef<SDockTab> SpawnTestSuite1( const FSpawnTabArgs& Args )
 		.SetDisplayName( NSLOCTEXT("TestSuite1", "RichTextTab", "Rich Text Test") )
 		.SetGroup(TestSuiteMenu::SuiteTabs);
 
+	TestSuite1TabManager->RegisterTabSpawner( "MultiLineEditTab", FOnSpawnTab::CreateStatic( &SpawnTab, FName( "MultiLineEditTab" ) ) )
+		.SetDisplayName( NSLOCTEXT( "TestSuite1", "MultiLineEditTab", "Multiline Edit" ) )
+		.SetGroup( TestSuiteMenu::SuiteTabs );
+
 	TestSuite1TabManager->RegisterTabSpawner( "LayoutRoundingTab", FOnSpawnTab::CreateStatic( &SpawnTab, FName("LayoutRoundingTab") ) )
 		.SetDisplayName( NSLOCTEXT("TestSuite1", "LayoutRoundingTab", "Layout Rounding") )
 		.SetGroup(TestSuiteMenu::SuiteTabs);
