Fixed loading of key bindings for users with config files created before command key support was added

diff --git a/Engine/Source/Runtime/Slate/Private/Framework/Commands/InputBindingManager.cpp b/Engine/Source/Runtime/Slate/Private/Framework/Commands/InputBindingManager.cpp
--- a/Engine/Source/Runtime/Slate/Private/Framework/Commands/InputBindingManager.cpp
+++ b/Engine/Source/Runtime/Slate/Private/Framework/Commands/InputBindingManager.cpp
@@ -94,7 +94,7 @@ void FUserDefinedGestures::LoadGestures()
 					UserDefinedGesture.bCtrl = CtrlObj->AsBool();
 					UserDefinedGesture.bAlt = AltObj->AsBool();
 					UserDefinedGesture.bShift = ShiftObj->AsBool();
-					UserDefinedGesture.bCmd = CmdObj->AsBool();
+					UserDefinedGesture.bCmd = CmdObj.IsValid() ? CmdObj->AsBool() : false; // Old config files may not have this
 					UserDefinedGesture.Key = *KeyObj->AsString();
 				}
 			}
@@ -147,7 +147,7 @@ void FUserDefinedGestures::LoadGestures()
 						UserDefinedGesture.bCtrl = CtrlObj->AsBool();
 						UserDefinedGesture.bAlt = AltObj->AsBool();
 						UserDefinedGesture.bShift = ShiftObj->AsBool();
-						UserDefinedGesture.bCmd = CmdObj->AsBool();
+						UserDefinedGesture.bCmd = CmdObj.IsValid() ? CmdObj->AsBool() : false; // Old config files may not have this
 						UserDefinedGesture.Key = *KeyObj->AsString();
 					}
 				}
