Slate: Removed file header comments

Also:
- NULL to nullptr
- removed redundant documentation
- consistent line spacing
- added/fixed comments

diff --git a/Engine/Source/Runtime/Slate/Private/Framework/Application/MenuStack.cpp b/Engine/Source/Runtime/Slate/Private/Framework/Application/MenuStack.cpp
--- a/Engine/Source/Runtime/Slate/Private/Framework/Application/MenuStack.cpp
+++ b/Engine/Source/Runtime/Slate/Private/Framework/Application/MenuStack.cpp
@@ -1,9 +1,5 @@
 // Copyright 1998-2014 Epic Games, Inc. All Rights Reserved.
 
-/*=============================================================================
-	MenuStack.cpp: Implements the FMenuStack class.
-=============================================================================*/
-
 #include "SlatePrivatePCH.h"
 
 
diff --git a/Engine/Source/Runtime/Slate/Private/Framework/Application/SWindowTitleBar.h b/Engine/Source/Runtime/Slate/Private/Framework/Application/SWindowTitleBar.h
--- a/Engine/Source/Runtime/Slate/Private/Framework/Application/SWindowTitleBar.h
+++ b/Engine/Source/Runtime/Slate/Private/Framework/Application/SWindowTitleBar.h
@@ -1,9 +1,5 @@
 // Copyright 1998-2014 Epic Games, Inc. All Rights Reserved.
 
-/*=============================================================================
-	SWindowTitleBar.h: Declares the SWindowTitleBar class.
-=============================================================================*/
-
 #pragma once
 
 
diff --git a/Engine/Source/Runtime/Slate/Private/Framework/Application/SlateIcon.cpp b/Engine/Source/Runtime/Slate/Private/Framework/Application/SlateIcon.cpp
--- a/Engine/Source/Runtime/Slate/Private/Framework/Application/SlateIcon.cpp
+++ b/Engine/Source/Runtime/Slate/Private/Framework/Application/SlateIcon.cpp
@@ -1,9 +1,5 @@
 // Copyright 1998-2014 Epic Games, Inc. All Rights Reserved.
 
-/*=============================================================================
-	SlateIcon.cpp: Implements the FSlateIcon structure.
-=============================================================================*/
-
 #include "SlatePrivatePCH.h"
 
 
diff --git a/Engine/Source/Runtime/Slate/Private/Framework/Commands/InputBindingManager.cpp b/Engine/Source/Runtime/Slate/Private/Framework/Commands/InputBindingManager.cpp
--- a/Engine/Source/Runtime/Slate/Private/Framework/Commands/InputBindingManager.cpp
+++ b/Engine/Source/Runtime/Slate/Private/Framework/Commands/InputBindingManager.cpp
@@ -1,9 +1,5 @@
 // Copyright 1998-2014 Epic Games, Inc. All Rights Reserved.
 
-/*=============================================================================
-	InputBindingManager.cpp: Implements the FInputBindingManager class.
-=============================================================================*/
-
 #include "SlatePrivatePCH.h"
 
 
diff --git a/Engine/Source/Runtime/Slate/Private/Framework/Commands/InputGesture.cpp b/Engine/Source/Runtime/Slate/Private/Framework/Commands/InputGesture.cpp
--- a/Engine/Source/Runtime/Slate/Private/Framework/Commands/InputGesture.cpp
+++ b/Engine/Source/Runtime/Slate/Private/Framework/Commands/InputGesture.cpp
@@ -1,9 +1,5 @@
 // Copyright 1998-2014 Epic Games, Inc. All Rights Reserved.
 
-/*=============================================================================
-	InputGesture.cpp: Implements the FInputGesture class.
-=============================================================================*/
-
 #include "SlatePrivatePCH.h"
 
 
diff --git a/Engine/Source/Runtime/Slate/Private/Framework/Commands/UICommandDragDropOp.cpp b/Engine/Source/Runtime/Slate/Private/Framework/Commands/UICommandDragDropOp.cpp
--- a/Engine/Source/Runtime/Slate/Private/Framework/Commands/UICommandDragDropOp.cpp
+++ b/Engine/Source/Runtime/Slate/Private/Framework/Commands/UICommandDragDropOp.cpp
@@ -1,9 +1,5 @@
 // Copyright 1998-2014 Epic Games, Inc. All Rights Reserved.
 
-/*=============================================================================
-	UICommandDragDropOp.cpp: Implements the FUICommandDragDropOp class.
-=============================================================================*/
-
 #include "SlatePrivatePCH.h"
 
 
diff --git a/Engine/Source/Runtime/Slate/Private/Framework/Commands/UICommandInfo.cpp b/Engine/Source/Runtime/Slate/Private/Framework/Commands/UICommandInfo.cpp
--- a/Engine/Source/Runtime/Slate/Private/Framework/Commands/UICommandInfo.cpp
+++ b/Engine/Source/Runtime/Slate/Private/Framework/Commands/UICommandInfo.cpp
@@ -1,9 +1,5 @@
 // Copyright 1998-2014 Epic Games, Inc. All Rights Reserved.
 
-/*=============================================================================
-	.cpp: Implements the  class.
-=============================================================================*/
-
 #include "SlatePrivatePCH.h"
 
 
diff --git a/Engine/Source/Runtime/Slate/Private/Framework/Commands/UICommandList.cpp b/Engine/Source/Runtime/Slate/Private/Framework/Commands/UICommandList.cpp
--- a/Engine/Source/Runtime/Slate/Private/Framework/Commands/UICommandList.cpp
+++ b/Engine/Source/Runtime/Slate/Private/Framework/Commands/UICommandList.cpp
@@ -1,48 +1,23 @@
 // Copyright 1998-2014 Epic Games, Inc. All Rights Reserved.
 
-/*=============================================================================
-	UICommandList.cpp: Implements the FUICommandList class.
-=============================================================================*/
-
 #include "SlatePrivatePCH.h"
 
 
 /* FUICommandList interface
  *****************************************************************************/
 
-/**
- * Maps a command info to a series of delegates that are executed by a multibox or mouse/keyboard input
- *
- * @param InUICommandInfo	The command info to map
- * @param ExecuteAction		The delegate to call when the command should be executed
- */
 void FUICommandList::MapAction( const TSharedPtr< const FUICommandInfo > InUICommandInfo, FExecuteAction ExecuteAction )
 {
 	MapAction( InUICommandInfo, ExecuteAction, FCanExecuteAction(), FIsActionChecked(), FIsActionButtonVisible() );
 }
 
 
-/**
- * Maps a command info to a series of delegates that are executed by a multibox or mouse/keyboard input
- *
- * @param InUICommandInfo	The command info to map
- * @param ExecuteAction		The delegate to call when the command should be executed
- * @param CanExecuteAction	The delegate to call to see if the command can be executed
- */
 void FUICommandList::MapAction( const TSharedPtr< const FUICommandInfo > InUICommandInfo, FExecuteAction ExecuteAction, FCanExecuteAction CanExecuteAction )
 {
 	MapAction( InUICommandInfo, ExecuteAction, CanExecuteAction, FIsActionChecked(), FIsActionButtonVisible() );
 }
 
 
-/**
- * Maps a command info to a series of delegates that are executed by a multibox or mouse/keyboard input
- *
- * @param InUICommandInfo	The command info to map
- * @param ExecuteAction		The delegate to call when the command should be executed
- * @param CanExecuteAction	The delegate to call to see if the command can be executed
- * @param IsCheckedDelegate	The delegate to call to see if the command should appear checked when visualized in a multibox
- */
 void FUICommandList::MapAction( const TSharedPtr< const FUICommandInfo > InUICommandInfo, FExecuteAction ExecuteAction, FCanExecuteAction CanExecuteAction, FIsActionChecked IsCheckedDelegate )
 {
 	MapAction( InUICommandInfo, ExecuteAction, CanExecuteAction, IsCheckedDelegate, FIsActionButtonVisible() );
@@ -88,12 +63,6 @@ void FUICommandList::Append( const TSharedRef<FUICommandList>& InCommandsToAppen
 }
 
 
-/**
- * Executes the action associated with the provided command info
- * Note: It is assumed at this point that CanExecuteAction was already checked
- *
- * @param InUICommandInfo	The command info execute
- */
 bool FUICommandList::ExecuteAction( const TSharedRef< const FUICommandInfo > InUICommandInfo ) const
 {
 	const FUIAction* Action = GetActionForCommand(InUICommandInfo);
@@ -109,11 +78,6 @@ bool FUICommandList::ExecuteAction( const TSharedRef< const FUICommandInfo > InU
 }
 
 
-/**
- * Calls the CanExecuteAction associated with the provided command info to see if ExecuteAction can be called
- *
- * @param InUICommandInfo	The command info execute
- */
 bool FUICommandList::CanExecuteAction( const TSharedRef< const FUICommandInfo > InUICommandInfo ) const
 {
 	const FUIAction* Action = GetActionForCommand(InUICommandInfo);
@@ -139,11 +103,6 @@ bool FUICommandList::TryExecuteAction( const TSharedRef< const FUICommandInfo >
 }
 
 
-/**
- * Calls the IsVisible delegate associated with the provided command info to see if the command should be visible in a toolbar
- *
- * @param InUICommandInfo	The command info execute
- */
 EVisibility FUICommandList::GetVisibility( const TSharedRef< const FUICommandInfo > InUICommandInfo ) const
 {
 	const FUIAction* Action = GetActionForCommand(InUICommandInfo);
@@ -158,11 +117,6 @@ EVisibility FUICommandList::GetVisibility( const TSharedRef< const FUICommandInf
 }
 
 
-/**
- * Calls the IsChecked delegate to see if the visualization of this command in a multibox should appear checked
- *
- * @param InUICommandInfo	The command info execute
- */
 bool FUICommandList::IsChecked( const TSharedRef< const FUICommandInfo > InUICommandInfo ) const
 {
 	const FUIAction* Action = GetActionForCommand(InUICommandInfo);
@@ -176,43 +130,18 @@ bool FUICommandList::IsChecked( const TSharedRef< const FUICommandInfo > InUICom
 }
 
 
-/**
- * Processes any user interface actions which are activated by the specified keyboard event
- *
- * @param InBindingContext	The context in which the actions are valid
- * @param InKeyboardEvent	The keyboard event to check
- *
- * @return true if an action was processed
- */
 bool FUICommandList::ProcessCommandBindings( const FKeyboardEvent& InKeyboardEvent ) const 
 {
 	return ConditionalProcessCommandBindings( InKeyboardEvent.GetKey(), InKeyboardEvent.IsControlDown(), InKeyboardEvent.IsAltDown(), InKeyboardEvent.IsShiftDown(), InKeyboardEvent.IsCommandDown(), InKeyboardEvent.IsRepeat() );
 }
 
 
-/**
- * Processes any user interface actions which are activated by the specified mouse event
- *
- * @param InBindingContext	The context in which the actions are valid
- * @param InKeyboardEvent	The mouse event to check
- *
- * @return true if an action was processed
- */
 bool FUICommandList::ProcessCommandBindings( const FPointerEvent& InMouseEvent ) const
 {
 	return ConditionalProcessCommandBindings( InMouseEvent.GetEffectingButton(), InMouseEvent.IsControlDown(), InMouseEvent.IsAltDown(), InMouseEvent.IsShiftDown(), InMouseEvent.IsCommandDown(), InMouseEvent.IsRepeat() );
 }
 
 
-/**
- * Processes any UI commands which are activated by the specified key, modifier keys state and input event
- *
- * @param Key				The current key that is pressed
- * @param ModifierKeysState	Pressed state of keys that are commonly used as modifiers
- * @param bRepeat			True if the input is repeating (held)
- *
- * @return true if an action was processed
- */
 bool FUICommandList::ProcessCommandBindings( const FKey Key, const FModifierKeysState& ModifierKeysState, const bool bRepeat ) const
 {
 	return ConditionalProcessCommandBindings(
@@ -228,17 +157,6 @@ bool FUICommandList::ProcessCommandBindings( const FKey Key, const FModifierKeys
 /* FUICommandList implementation
  *****************************************************************************/
 
-/**
- * Helper function to execute an interface action delegate or exec command if valid
- *
- * @param InBindingContext The binding context where commands are currently valid
- * @param Key		The current key that is pressed
- * @param bCtrl		True if control is pressed
- * @param bAlt		True if alt is pressed
- * @param bShift	True if shift is pressed
- * @param bRepeat	True if command is repeating (held)
- * @return True if a command was executed, False otherwise
- */
 bool FUICommandList::ConditionalProcessCommandBindings( const FKey Key, bool bCtrl, bool bAlt, bool bShift, bool bCmd, bool bRepeat ) const
 {
 	if ( !bRepeat && !FSlateApplication::Get().IsDragDropping() )
diff --git a/Engine/Source/Runtime/Slate/Private/Framework/Layout/ScrollyZoomy.cpp b/Engine/Source/Runtime/Slate/Private/Framework/Layout/ScrollyZoomy.cpp
--- a/Engine/Source/Runtime/Slate/Private/Framework/Layout/ScrollyZoomy.cpp
+++ b/Engine/Source/Runtime/Slate/Private/Framework/Layout/ScrollyZoomy.cpp
@@ -1,9 +1,5 @@
 // Copyright 1998-2014 Epic Games, Inc. All Rights Reserved.
 
-/*=============================================================================
-	ScrollyZoomy.h: Implements the FScrollyZoomy class.
-=============================================================================*/
-
 #include "SlatePrivatePCH.h"
 
 
diff --git a/Engine/Source/Runtime/Slate/Private/Framework/MultiBox/SEditableTextBlock.cpp b/Engine/Source/Runtime/Slate/Private/Framework/MultiBox/SEditableTextBlock.cpp
--- a/Engine/Source/Runtime/Slate/Private/Framework/MultiBox/SEditableTextBlock.cpp
+++ b/Engine/Source/Runtime/Slate/Private/Framework/MultiBox/SEditableTextBlock.cpp
@@ -1,9 +1,5 @@
 // Copyright 1998-2014 Epic Games, Inc. All Rights Reserved.
 
-/*=============================================================================
-	SEditableTextBlock.h: Implements the SEditableTextBlock class.
-=============================================================================*/
-
 #include "SlatePrivatePCH.h"
 #include "SEditableTextBlock.h"
 
diff --git a/Engine/Source/Runtime/Slate/Private/Framework/MultiBox/SEditableTextBlock.h b/Engine/Source/Runtime/Slate/Private/Framework/MultiBox/SEditableTextBlock.h
--- a/Engine/Source/Runtime/Slate/Private/Framework/MultiBox/SEditableTextBlock.h
+++ b/Engine/Source/Runtime/Slate/Private/Framework/MultiBox/SEditableTextBlock.h
@@ -1,18 +1,14 @@
 // Copyright 1998-2014 Epic Games, Inc. All Rights Reserved.
 
-/*=============================================================================
-	SEditableTextBlock.cpp: Declares the SEditableTextBlock class.
-=============================================================================*/
-
 #pragma once
 
+
 /**
  * Editable text block
  */
 class FEditableTextBlock
 	: public FMultiBlock
 {
-
 public:
 
 	/**
@@ -62,7 +58,6 @@ private:
 
 	/** If true the text box is read only */
 	bool bReadOnly;
-	
 };
 
 
@@ -95,10 +90,7 @@ protected:
 	/**
 	 * Called by Slate to determine if this widget is enabled
 	 * 
-	 * @return True if the widget is enabled, false otherwise
+	 * @return true if the widget is enabled, false otherwise
 	 */
 	bool IsEnabled() const;
-
-private:
-
 };
diff --git a/Engine/Source/Runtime/Slate/Private/Framework/MultiBox/SMenuEntryBlock.cpp b/Engine/Source/Runtime/Slate/Private/Framework/MultiBox/SMenuEntryBlock.cpp
--- a/Engine/Source/Runtime/Slate/Private/Framework/MultiBox/SMenuEntryBlock.cpp
+++ b/Engine/Source/Runtime/Slate/Private/Framework/MultiBox/SMenuEntryBlock.cpp
@@ -47,7 +47,7 @@ FMenuEntryBlock::FMenuEntryBlock( const FName& InExtensionHook, const TAttribute
 
 
 FMenuEntryBlock::FMenuEntryBlock( const FName& InExtensionHook, const TAttribute<FText>& InLabel, const TAttribute<FText>& InToolTip, const FNewMenuDelegate& InEntryBuilder, TSharedPtr<FExtender> InExtender, bool bInSubMenu, bool bInSubMenuOnClick, TSharedPtr< const FUICommandList > InCommandList, bool bInCloseSelfOnly, const FSlateIcon& InIcon )
-	: FMultiBlock( NULL, InCommandList, InExtensionHook, EMultiBlockType::MenuEntry )
+	: FMultiBlock( nullptr, InCommandList, InExtensionHook, EMultiBlockType::MenuEntry )
 	, LabelOverride( InLabel )
 	, ToolTipOverride( InToolTip )
 	, IconOverride( InIcon )
@@ -62,7 +62,7 @@ FMenuEntryBlock::FMenuEntryBlock( const FName& InExtensionHook, const TAttribute
 
 
 FMenuEntryBlock::FMenuEntryBlock( const FName& InExtensionHook, const TAttribute<FText>& InLabel, const TAttribute<FText>& InToolTip, const FOnGetContent& InMenuBuilder, TSharedPtr<FExtender> InExtender, bool bInSubMenu, bool bInSubMenuOnClick, TSharedPtr< const FUICommandList > InCommandList, bool bInCloseSelfOnly, const FSlateIcon& InIcon )
-	: FMultiBlock( NULL, InCommandList, InExtensionHook, EMultiBlockType::MenuEntry )
+	: FMultiBlock( nullptr, InCommandList, InExtensionHook, EMultiBlockType::MenuEntry )
 	, LabelOverride( InLabel )
 	, ToolTipOverride( InToolTip )
 	, IconOverride( InIcon )
@@ -77,7 +77,7 @@ FMenuEntryBlock::FMenuEntryBlock( const FName& InExtensionHook, const TAttribute
 
 
 FMenuEntryBlock::FMenuEntryBlock( const FName& InExtensionHook, const TAttribute<FText>& InLabel, const TAttribute<FText>& InToolTip, const TSharedPtr<SWidget>& InEntryWidget, TSharedPtr<FExtender> InExtender, bool bInSubMenu, bool bInSubMenuOnClick, TSharedPtr< const FUICommandList > InCommandList, bool bInCloseSelfOnly, const FSlateIcon& InIcon )
-	: FMultiBlock( NULL, InCommandList, InExtensionHook, EMultiBlockType::MenuEntry )
+	: FMultiBlock( nullptr, InCommandList, InExtensionHook, EMultiBlockType::MenuEntry )
 	, LabelOverride( InLabel )
 	, ToolTipOverride( InToolTip )
 	, IconOverride( InIcon )
@@ -103,7 +103,7 @@ FMenuEntryBlock::FMenuEntryBlock( const FName& InExtensionHook, const FUIAction&
 
 
 FMenuEntryBlock::FMenuEntryBlock( const FName& InExtensionHook, const TSharedRef< SWidget > Contents, const FNewMenuDelegate& InEntryBuilder, TSharedPtr<FExtender> InExtender, bool bInSubMenu, bool bInSubMenuOnClick, TSharedPtr< const FUICommandList > InCommandList, bool bInCloseSelfOnly )
-	: FMultiBlock( NULL, InCommandList, InExtensionHook, EMultiBlockType::MenuEntry )
+	: FMultiBlock( nullptr, InCommandList, InExtensionHook, EMultiBlockType::MenuEntry )
 	, EntryBuilder( InEntryBuilder )
 	, EntryWidget( Contents )
 	, bIsSubMenu( bInSubMenu )
@@ -159,8 +159,8 @@ void SMenuEntryBlock::Construct( const FArguments& InArgs )
 	SubMenuRequestState = Idle;
 
 	// No images by default
-	CheckedImage = NULL;
-	UncheckedImage = NULL;
+	CheckedImage = nullptr;
+	UncheckedImage = nullptr;
 
 	this->SetForegroundColor( TAttribute<FSlateColor>::Create( TAttribute<FSlateColor>::FGetter::CreateRaw( this, &SMenuEntryBlock::InvertOnHover ) ) );
 }
@@ -431,7 +431,7 @@ TSharedRef< SWidget > SMenuEntryBlock::BuildMenuEntryWidget( const FMenuEntryBui
 		.ClickMethod( ButtonClickMethod )
 
 		// Pass along the block's tool-tip string
-		.ToolTip( FMultiBoxSettings::ToolTipConstructor.Execute( ToolTip, NULL, UICommand ) )
+		.ToolTip( FMultiBoxSettings::ToolTipConstructor.Execute( ToolTip, nullptr, UICommand ) )
 
 		.ContentPadding(FMargin(0, 2))
 
@@ -1142,7 +1142,7 @@ TSharedRef< SWidget > SMenuEntryBlock::MakeNewMenuWidget() const
 	{
 		const bool bShouldCloseWindowAfterMenuSelection = true;
 		const bool bCloseSelfOnly = false;
-		FMenuBuilder MenuBuilder( bShouldCloseWindowAfterMenuSelection, NULL, TSharedPtr<FExtender>(), bCloseSelfOnly, StyleSet );
+		FMenuBuilder MenuBuilder( bShouldCloseWindowAfterMenuSelection, nullptr, TSharedPtr<FExtender>(), bCloseSelfOnly, StyleSet );
 		{
 			MenuBuilder.AddWidget( MenuEntryBlock->EntryWidget.ToSharedRef(), FText::GetEmpty() );
 		}
diff --git a/Engine/Source/Runtime/Slate/Private/Framework/MultiBox/SMenuSeparatorBlock.cpp b/Engine/Source/Runtime/Slate/Private/Framework/MultiBox/SMenuSeparatorBlock.cpp
--- a/Engine/Source/Runtime/Slate/Private/Framework/MultiBox/SMenuSeparatorBlock.cpp
+++ b/Engine/Source/Runtime/Slate/Private/Framework/MultiBox/SMenuSeparatorBlock.cpp
@@ -9,7 +9,7 @@
  * Constructor
  */
 FMenuSeparatorBlock::FMenuSeparatorBlock(const FName& InExtensionHook)
-	: FMultiBlock( NULL, NULL, InExtensionHook, EMultiBlockType::MenuSeparator )
+	: FMultiBlock( nullptr, nullptr, InExtensionHook, EMultiBlockType::MenuSeparator )
 {
 }
 
diff --git a/Engine/Source/Runtime/Slate/Private/Framework/MultiBox/SToolBarButtonBlock.cpp b/Engine/Source/Runtime/Slate/Private/Framework/MultiBox/SToolBarButtonBlock.cpp
--- a/Engine/Source/Runtime/Slate/Private/Framework/MultiBox/SToolBarButtonBlock.cpp
+++ b/Engine/Source/Runtime/Slate/Private/Framework/MultiBox/SToolBarButtonBlock.cpp
@@ -230,7 +230,7 @@ void SToolBarButtonBlock::BuildMultiBlockWidget(const ISlateStyle* StyleSet, con
 				.OnClicked( this, &SToolBarButtonBlock::OnClicked )
 
 				// Pass along the block's tool-tip string
-				.ToolTip( FMultiBoxSettings::ToolTipConstructor.Execute( ActualToolTip, NULL, Action.Pin() ) )
+				.ToolTip( FMultiBoxSettings::ToolTipConstructor.Execute( ActualToolTip, nullptr, Action.Pin() ) )
 		];
 	}
 	else if( ensure( UserInterfaceType == EUserInterfaceActionType::ToggleButton || UserInterfaceType == EUserInterfaceActionType::RadioButton ) )
@@ -250,7 +250,7 @@ void SToolBarButtonBlock::BuildMultiBlockWidget(const ISlateStyle* StyleSet, con
 				.IsFocusable( bIsFocusable )
 
 				// Pass along the block's tool-tip string
-				.ToolTip( FMultiBoxSettings::ToolTipConstructor.Execute( ActualToolTip, NULL, Action.Pin() ) )
+				.ToolTip( FMultiBoxSettings::ToolTipConstructor.Execute( ActualToolTip, nullptr, Action.Pin() ) )
 				[
 					ButtonContent
 				]
diff --git a/Engine/Source/Runtime/Slate/Private/Framework/MultiBox/SToolBarSeparatorBlock.cpp b/Engine/Source/Runtime/Slate/Private/Framework/MultiBox/SToolBarSeparatorBlock.cpp
--- a/Engine/Source/Runtime/Slate/Private/Framework/MultiBox/SToolBarSeparatorBlock.cpp
+++ b/Engine/Source/Runtime/Slate/Private/Framework/MultiBox/SToolBarSeparatorBlock.cpp
@@ -9,7 +9,7 @@
  * Constructor
  */
 FToolBarSeparatorBlock::FToolBarSeparatorBlock(const FName& InExtensionHook)
-	: FMultiBlock( NULL, NULL, InExtensionHook, EMultiBlockType::ToolBarSeparator )
+	: FMultiBlock( nullptr, nullptr, InExtensionHook, EMultiBlockType::ToolBarSeparator )
 {
 }
 
diff --git a/Engine/Source/Runtime/Slate/Private/Framework/MultiBox/SWidgetBlock.cpp b/Engine/Source/Runtime/Slate/Private/Framework/MultiBox/SWidgetBlock.cpp
--- a/Engine/Source/Runtime/Slate/Private/Framework/MultiBox/SWidgetBlock.cpp
+++ b/Engine/Source/Runtime/Slate/Private/Framework/MultiBox/SWidgetBlock.cpp
@@ -11,7 +11,7 @@
  * @param	InHeadingText	Heading text
  */
 FWidgetBlock::FWidgetBlock( TSharedRef<SWidget> InContent, const FText& InLabel, bool bInNoIndent )
-	: FMultiBlock( NULL, NULL, NAME_None, EMultiBlockType::Widget )
+	: FMultiBlock( nullptr, nullptr, NAME_None, EMultiBlockType::Widget )
 	, ContentWidget( InContent )
 	, Label( InLabel )
 	, bNoIndent( bInNoIndent )
@@ -86,4 +86,3 @@ void SWidgetBlock::BuildMultiBlockWidget(const ISlateStyle* StyleSet, const FNam
 		]
 	];
 }
-
diff --git a/Engine/Source/Runtime/Slate/Private/Framework/Text/SlateImageRun.cpp b/Engine/Source/Runtime/Slate/Private/Framework/Text/SlateImageRun.cpp
--- a/Engine/Source/Runtime/Slate/Private/Framework/Text/SlateImageRun.cpp
+++ b/Engine/Source/Runtime/Slate/Private/Framework/Text/SlateImageRun.cpp
@@ -8,7 +8,7 @@
 
 TSharedRef< FSlateImageRun > FSlateImageRun::Create( const TSharedRef< const FString >& InText, const FSlateBrush* InImage, int16 InBaseline )
 {
-	if ( InImage == NULL )
+	if ( InImage == nullptr)
 	{
 		InImage = FStyleDefaults::GetNoBrush();
 	}
@@ -18,7 +18,7 @@ TSharedRef< FSlateImageRun > FSlateImageRun::Create( const TSharedRef< const FSt
 
 TSharedRef< FSlateImageRun > FSlateImageRun::Create( const TSharedRef< const FString >& InText, const FSlateBrush* InImage, int16 InBaseline, const FTextRange& InRange )
 {
-	if ( InImage == NULL )
+	if ( InImage == nullptr)
 	{
 		InImage = FStyleDefaults::GetNoBrush();
 	}
@@ -32,7 +32,7 @@ FSlateImageRun::FSlateImageRun( const TSharedRef< const FString >& InText, const
 	, Image( InImage )
 	, Baseline( InBaseline )
 {
-	check( Image != NULL );
+	check( Image != nullptr);
 }
 
 FSlateImageRun::FSlateImageRun( const TSharedRef< const FString >& InText, const FSlateBrush* InImage, int16 InBaseline ) 
@@ -41,7 +41,7 @@ FSlateImageRun::FSlateImageRun( const TSharedRef< const FString >& InText, const
 	, Image( InImage )
 	, Baseline( InBaseline )
 {
-	check( Image != NULL );
+	check( Image != nullptr);
 }
 
 const TArray< TSharedRef<SWidget> >& FSlateImageRun::GetChildren() 
diff --git a/Engine/Source/Runtime/Slate/Private/SlateModule.cpp b/Engine/Source/Runtime/Slate/Private/SlateModule.cpp
--- a/Engine/Source/Runtime/Slate/Private/SlateModule.cpp
+++ b/Engine/Source/Runtime/Slate/Private/SlateModule.cpp
@@ -1,9 +1,5 @@
 // Copyright 1998-2014 Epic Games, Inc. All Rights Reserved.
 
-/*=============================================================================
-	SlateModule.cpp: Implements the FSlateModule class.
-=============================================================================*/
-
 #include "SlatePrivatePCH.h"
 #include "ModuleManager.h"
 
@@ -36,8 +32,7 @@ DEFINE_STAT(STAT_SlateMiscTime); // Useful for locally pref testing miscellaneou
  */
 class FSlateModule
 	: public IModuleInterface
-{
-};
+{ };
 
 
 IMPLEMENT_MODULE(FSlateModule, Slate);
diff --git a/Engine/Source/Runtime/Slate/Private/SlatePrivatePCH.h b/Engine/Source/Runtime/Slate/Private/SlatePrivatePCH.h
--- a/Engine/Source/Runtime/Slate/Private/SlatePrivatePCH.h
+++ b/Engine/Source/Runtime/Slate/Private/SlatePrivatePCH.h
@@ -1,15 +1,11 @@
 // Copyright 1998-2014 Epic Games, Inc. All Rights Reserved.
 
-/*=============================================================================
-	SlatePrivatePCH.h: Pre-compiled header file for the Slate module.
-=============================================================================*/
-
 #pragma once
 
 #include "Slate.h"
 
 
-/* Dependencies
+/* Private dependencies
  *****************************************************************************/
 
 #include "InputCore.h"
diff --git a/Engine/Source/Runtime/Slate/Private/Widgets/Colors/SColorSpectrum.cpp b/Engine/Source/Runtime/Slate/Private/Widgets/Colors/SColorSpectrum.cpp
--- a/Engine/Source/Runtime/Slate/Private/Widgets/Colors/SColorSpectrum.cpp
+++ b/Engine/Source/Runtime/Slate/Private/Widgets/Colors/SColorSpectrum.cpp
@@ -1,9 +1,5 @@
 // Copyright 1998-2014 Epic Games, Inc. All Rights Reserved.
 
-/*=============================================================================
-	SColorSpectrum.cpp: Implements the SColorSpectrum class.
-=============================================================================*/
-
 #include "SlatePrivatePCH.h"
 
 
@@ -31,6 +27,12 @@ FVector2D SColorSpectrum::ComputeDesiredSize( ) const
 }
 
 
+FReply SColorSpectrum::OnMouseButtonDoubleClick( const FGeometry& InMyGeometry, const FPointerEvent& InMouseEvent )
+{
+	return FReply::Handled();
+}
+
+
 FReply SColorSpectrum::OnMouseButtonDown( const FGeometry& MyGeometry, const FPointerEvent& MouseEvent )
 {
 	if (MouseEvent.GetEffectingButton() == EKeys::LeftMouseButton)
diff --git a/Engine/Source/Runtime/Slate/Private/Widgets/Colors/SColorWheel.cpp b/Engine/Source/Runtime/Slate/Private/Widgets/Colors/SColorWheel.cpp
--- a/Engine/Source/Runtime/Slate/Private/Widgets/Colors/SColorWheel.cpp
+++ b/Engine/Source/Runtime/Slate/Private/Widgets/Colors/SColorWheel.cpp
@@ -1,9 +1,5 @@
 // Copyright 1998-2014 Epic Games, Inc. All Rights Reserved.
 
-/*=============================================================================
-	SColorWheel.cpp: Implements the SColorWheel class.
-=============================================================================*/
-
 #include "SlatePrivatePCH.h"
 
 
@@ -31,6 +27,12 @@ FVector2D SColorWheel::ComputeDesiredSize( ) const
 }
 
 
+FReply SColorWheel::OnMouseButtonDoubleClick( const FGeometry& InMyGeometry, const FPointerEvent& InMouseEvent )
+{
+	return FReply::Handled();
+}
+
+
 FReply SColorWheel::OnMouseButtonDown( const FGeometry& MyGeometry, const FPointerEvent& MouseEvent )
 {
 	if (MouseEvent.GetEffectingButton() == EKeys::LeftMouseButton)
diff --git a/Engine/Source/Runtime/Slate/Private/Widgets/Images/SImage.cpp b/Engine/Source/Runtime/Slate/Private/Widgets/Images/SImage.cpp
--- a/Engine/Source/Runtime/Slate/Private/Widgets/Images/SImage.cpp
+++ b/Engine/Source/Runtime/Slate/Private/Widgets/Images/SImage.cpp
@@ -23,7 +23,7 @@ int32 SImage::OnPaint( const FPaintArgs& Args, const FGeometry& AllottedGeometry
 #endif
 	const FSlateBrush* ImageBrush = Image.Get();
 
-	if ((ImageBrush != NULL) && (ImageBrush->DrawAs != ESlateBrushDrawType::NoDrawType))
+	if ((ImageBrush != nullptr) && (ImageBrush->DrawAs != ESlateBrushDrawType::NoDrawType))
 	{
 		const bool bIsEnabled = ShouldBeEnabled(bParentEnabled);
 		const uint32 DrawEffects = bIsEnabled ? ESlateDrawEffect::None : ESlateDrawEffect::DisabledEffect;
@@ -63,7 +63,7 @@ FReply SImage::OnMouseButtonDown( const FGeometry& MyGeometry, const FPointerEve
 FVector2D SImage::ComputeDesiredSize() const
 {
 	const FSlateBrush* ImageBrush = Image.Get();
-	if (ImageBrush != NULL)
+	if (ImageBrush != nullptr)
 	{
 		return ImageBrush->ImageSize;
 	}
diff --git a/Engine/Source/Runtime/Slate/Private/Widgets/Images/SSpinningImage.cpp b/Engine/Source/Runtime/Slate/Private/Widgets/Images/SSpinningImage.cpp
--- a/Engine/Source/Runtime/Slate/Private/Widgets/Images/SSpinningImage.cpp
+++ b/Engine/Source/Runtime/Slate/Private/Widgets/Images/SSpinningImage.cpp
@@ -19,7 +19,7 @@ int32 SSpinningImage::OnPaint( const FPaintArgs& Args, const FGeometry& Allotted
 {
 	const FSlateBrush* ImageBrush = Image.Get();
 
-	if ((ImageBrush != NULL) && (ImageBrush->DrawAs != ESlateBrushDrawType::NoDrawType))
+	if ((ImageBrush != nullptr) && (ImageBrush->DrawAs != ESlateBrushDrawType::NoDrawType))
 	{
 		const bool bIsEnabled = ShouldBeEnabled(bParentEnabled);
 		const uint32 DrawEffects = bIsEnabled ? ESlateDrawEffect::None : ESlateDrawEffect::DisabledEffect;
diff --git a/Engine/Source/Runtime/Slate/Private/Widgets/Input/SCheckBox.cpp b/Engine/Source/Runtime/Slate/Private/Widgets/Input/SCheckBox.cpp
--- a/Engine/Source/Runtime/Slate/Private/Widgets/Input/SCheckBox.cpp
+++ b/Engine/Source/Runtime/Slate/Private/Widgets/Input/SCheckBox.cpp
@@ -10,7 +10,7 @@
  */
 void SCheckBox::Construct( const SCheckBox::FArguments& InArgs )
 {
-	check(InArgs._Style != NULL);
+	check(InArgs._Style != nullptr);
 	Style = InArgs._Style;
 
 	UncheckedImage = InArgs._UncheckedImage;
diff --git a/Engine/Source/Runtime/Slate/Private/Widgets/Input/SComboButton.cpp b/Engine/Source/Runtime/Slate/Private/Widgets/Input/SComboButton.cpp
--- a/Engine/Source/Runtime/Slate/Private/Widgets/Input/SComboButton.cpp
+++ b/Engine/Source/Runtime/Slate/Private/Widgets/Input/SComboButton.cpp
@@ -72,7 +72,7 @@ void SComboButton::Construct( const FArguments& InArgs )
 
 FReply SComboButton::OnButtonClicked()
 {
-	TSharedPtr<SWidget> Content = NULL;
+	TSharedPtr<SWidget> Content = nullptr;
 	if( OnGetMenuContent.IsBound() )
 	{
 		Content = OnGetMenuContent.Execute();
diff --git a/Engine/Source/Runtime/Slate/Private/Widgets/Input/SEditableText.cpp b/Engine/Source/Runtime/Slate/Private/Widgets/Input/SEditableText.cpp
--- a/Engine/Source/Runtime/Slate/Private/Widgets/Input/SEditableText.cpp
+++ b/Engine/Source/Runtime/Slate/Private/Widgets/Input/SEditableText.cpp
@@ -3,8 +3,8 @@
 #include "SlatePrivatePCH.h"
 #include "TextEditHelper.h"
 
-/** A pointer to the editable text currently under the mouse cursor. NULL when there isn't one. */
-SEditableText* SEditableText::EditableTextUnderCursor = NULL;
+/** A pointer to the editable text currently under the mouse cursor. nullptr when there isn't one. */
+SEditableText* SEditableText::EditableTextUnderCursor = nullptr;
 
 
 /** Constructor */
diff --git a/Engine/Source/Runtime/Slate/Private/Widgets/Input/SEditableTextBox.cpp b/Engine/Source/Runtime/Slate/Private/Widgets/Input/SEditableTextBox.cpp
--- a/Engine/Source/Runtime/Slate/Private/Widgets/Input/SEditableTextBox.cpp
+++ b/Engine/Source/Runtime/Slate/Private/Widgets/Input/SEditableTextBox.cpp
@@ -67,11 +67,12 @@ void SEditableTextBox::Construct( const FArguments& InArgs )
 
 }
 
+
 void SEditableTextBox::SetStyle(const FEditableTextBoxStyle* InStyle)
 {
 	Style = InStyle;
 
-	if ( Style == NULL )
+	if ( Style == nullptr )
 	{
 		FArguments Defaults;
 		Style = Defaults._Style;
@@ -83,16 +84,19 @@ void SEditableTextBox::SetStyle(const FEditableTextBoxStyle* InStyle)
 	BorderImageReadOnly = &Style->BackgroundImageReadOnly;
 }
 
+
 void SEditableTextBox::SetText( const TAttribute< FText >& InNewText )
 {
 	EditableText->SetText( InNewText );
 }
 
+
 void SEditableTextBox::SetError( const FText& InError )
 {
 	SetError( InError.ToString() );
 }
 
+
 void SEditableTextBox::SetError( const FString& InError )
 {
 	const bool bHaveError = !InError.IsEmpty();
@@ -113,6 +117,7 @@ void SEditableTextBox::SetError( const FString& InError )
 	ErrorReporting->SetError( InError );
 }
 
+
 /** @return Border image for the text box based on the hovered and focused state */
 const FSlateBrush* SEditableTextBox::GetBorderImage() const
 {
@@ -137,17 +142,20 @@ const FSlateBrush* SEditableTextBox::GetBorderImage() const
 	}
 }
 
+
 bool SEditableTextBox::SupportsKeyboardFocus() const
 {
 	return StaticCastSharedPtr<SWidget>(EditableText)->SupportsKeyboardFocus();
 }
 
+
 bool SEditableTextBox::HasKeyboardFocus() const
 {
 	// Since keyboard focus is forwarded to our editable text, we will test it instead
 	return SBorder::HasKeyboardFocus() || EditableText->HasKeyboardFocus();
 }
 
+
 FReply SEditableTextBox::OnKeyboardFocusReceived( const FGeometry& MyGeometry, const FKeyboardFocusEvent& InKeyboardFocusEvent )
 {
 	FReply Reply = FReply::Handled();
@@ -161,6 +169,7 @@ FReply SEditableTextBox::OnKeyboardFocusReceived( const FGeometry& MyGeometry, c
 	return Reply;
 }
 
+
 FReply SEditableTextBox::OnKeyDown( const FGeometry& MyGeometry, const FKeyboardEvent& InKeyboardEvent )
 {
 	FKey Key = InKeyboardEvent.GetKey();
@@ -175,16 +184,19 @@ FReply SEditableTextBox::OnKeyDown( const FGeometry& MyGeometry, const FKeyboard
 	return FReply::Unhandled();
 }
 
+
 void SEditableTextBox::SetHintText(const TAttribute< FText >& InHintText)
 {
 	EditableText->SetHintText(InHintText);
 }
 
+
 void SEditableTextBox::SetIsReadOnly(TAttribute< bool > InIsReadOnly)
 {
 	EditableText->SetIsReadOnly(InIsReadOnly);
 }
 
+
 void SEditableTextBox::SetIsPassword(TAttribute< bool > InIsPassword)
 {
 	EditableText->SetIsPassword(InIsPassword);
diff --git a/Engine/Source/Runtime/Slate/Private/Widgets/Input/SMenuAnchor.cpp b/Engine/Source/Runtime/Slate/Private/Widgets/Input/SMenuAnchor.cpp
--- a/Engine/Source/Runtime/Slate/Private/Widgets/Input/SMenuAnchor.cpp
+++ b/Engine/Source/Runtime/Slate/Private/Widgets/Input/SMenuAnchor.cpp
@@ -10,7 +10,7 @@
  */
 void SMenuAnchor::Construct( const FArguments& InArgs )
 {
-	PopupLayerSlot = NULL;
+	PopupLayerSlot = nullptr;
 
 	this->ChildSlot
 	.Padding(0)
@@ -107,7 +107,7 @@ void SMenuAnchor::Tick( const FGeometry& AllottedGeometry, const double InCurren
 		{
 			// Using the PopupLayer to show this popup instead of a window.
 			// Update the Popup's position so it appears anchored to the menu anchor.
-			if (ensure(PopupLayerSlot != NULL))
+			if (ensure(PopupLayerSlot != nullptr))
 			{
 				FVector2D NewPosition = FVector2D::ZeroVector;
 				FVector2D NewSize = PopupLayerSlot->GetWidget()->GetDesiredSize();
@@ -255,10 +255,10 @@ void SMenuAnchor::SetIsOpen( bool InIsOpen, const bool bFocusMenu )
 				// Close the popup overlay
 				ensure(Method==UseCurrentWindow);
 				check(PopupWindowPtr.IsValid());
-				check(PopupLayerSlot != NULL);
+				check(PopupLayerSlot != nullptr);
 				PopupWindowPtr.Pin()->RemovePopupLayerSlot(PopupLayerSlot->GetWidget());
 				// We no longer have a popup open, so reset all the tracking state associated.
-				PopupLayerSlot = NULL;
+				PopupLayerSlot = nullptr;
 				PopupWindowPtr.Reset();
 			}
 
@@ -320,13 +320,13 @@ SMenuAnchor::~SMenuAnchor()
 			PopupWindow->RequestDestroyWindow();
 		}
 		// Close the popup overlay
-		else if (Method == UseCurrentWindow && PopupLayerSlot != NULL)
+		else if (Method == UseCurrentWindow && PopupLayerSlot != nullptr)
 		{
 			PopupWindow->RemovePopupLayerSlot(PopupLayerSlot->GetWidget());
 		}
 		
 		// We no longer have a popup open, so reset all the tracking state associated.
 		PopupWindowPtr.Reset();
-		PopupLayerSlot = NULL;
+		PopupLayerSlot = nullptr;
 	}
 }
diff --git a/Engine/Source/Runtime/Slate/Private/Widgets/Input/SSuggestionTextBox.cpp b/Engine/Source/Runtime/Slate/Private/Widgets/Input/SSuggestionTextBox.cpp
--- a/Engine/Source/Runtime/Slate/Private/Widgets/Input/SSuggestionTextBox.cpp
+++ b/Engine/Source/Runtime/Slate/Private/Widgets/Input/SSuggestionTextBox.cpp
@@ -1,9 +1,5 @@
 // Copyright 1998-2014 Epic Games, Inc. All Rights Reserved.
 
-/*=============================================================================
-	SSuggestionTextBox.cpp: Implements the SSuggestionTextBox class.
-=============================================================================*/
-
 #include "SlatePrivatePCH.h"
 
 
diff --git a/Engine/Source/Runtime/Slate/Private/Widgets/Layout/SExpandableArea.cpp b/Engine/Source/Runtime/Slate/Private/Widgets/Layout/SExpandableArea.cpp
--- a/Engine/Source/Runtime/Slate/Private/Widgets/Layout/SExpandableArea.cpp
+++ b/Engine/Source/Runtime/Slate/Private/Widgets/Layout/SExpandableArea.cpp
@@ -1,9 +1,5 @@
 // Copyright 1998-2014 Epic Games, Inc. All Rights Reserved.
 
-/*=============================================================================
-	SExpandableArea.cpp: Implements the SExpandableArea class.
-=============================================================================*/
-
 #include "SlatePrivatePCH.h"
 
 
diff --git a/Engine/Source/Runtime/Slate/Private/Widgets/Layout/SHeader.cpp b/Engine/Source/Runtime/Slate/Private/Widgets/Layout/SHeader.cpp
--- a/Engine/Source/Runtime/Slate/Private/Widgets/Layout/SHeader.cpp
+++ b/Engine/Source/Runtime/Slate/Private/Widgets/Layout/SHeader.cpp
@@ -5,31 +5,38 @@
 
 void SHeader::Construct( const FArguments& InArgs )
 {
-	SHorizontalBox::FSlot* FirstSlot = NULL;
-	SHorizontalBox::FSlot* LastSlot = NULL;
+	SHorizontalBox::FSlot* FirstSlot = nullptr;
+	SHorizontalBox::FSlot* LastSlot = nullptr;
 
 	SHorizontalBox::Construct( SHorizontalBox::FArguments()
-		+SHorizontalBox::Slot()
+
+	+ SHorizontalBox::Slot()
 		.Expose( FirstSlot )
 		.FillWidth(1)
 		.VAlign(VAlign_Center)
 		[
-			SNew(SSeparator) .SeparatorImage(FCoreStyle::Get().GetBrush("Header.Pre")) .Orientation(Orient_Horizontal)
+			SNew(SSeparator)
+				.SeparatorImage(FCoreStyle::Get().GetBrush("Header.Pre"))
+				.Orientation(Orient_Horizontal)
 		]
-		+SHorizontalBox::Slot()
-		.AutoWidth()
-		.VAlign(VAlign_Center)
-		.Padding(5, 0)
-		[
-			InArgs._Content.Widget
-		]
-		+SHorizontalBox::Slot()
-		.Expose( LastSlot )
-		.FillWidth(1)
-		.VAlign(VAlign_Center)
-		[
-			SNew(SSeparator) .SeparatorImage(FCoreStyle::Get().GetBrush("Header.Post")) .Orientation(Orient_Horizontal)
-		]	
+
+		+ SHorizontalBox::Slot()
+			.AutoWidth()
+			.VAlign(VAlign_Center)
+			.Padding(5, 0)
+			[
+				InArgs._Content.Widget
+			]
+
+		+ SHorizontalBox::Slot()
+			.Expose( LastSlot )
+			.FillWidth(1)
+			.VAlign(VAlign_Center)
+			[
+				SNew(SSeparator)
+					.SeparatorImage(FCoreStyle::Get().GetBrush("Header.Post"))
+					.Orientation(Orient_Horizontal)
+			]	
 	);
 
 	// Either the left or right side of the header needs to be auto-sized based on the alignment of the content.
@@ -39,15 +46,14 @@ void SHeader::Construct( const FArguments& InArgs )
 		default:
 		case HAlign_Center:
 		case HAlign_Fill:
-		break;
+			break;
 
 		case HAlign_Left:
 			FirstSlot->AutoWidth();
-		break;
+			break;
 
 		case HAlign_Right:
 			LastSlot->AutoWidth();
-		break;
-
+			break;
 	};
 }
diff --git a/Engine/Source/Runtime/Slate/Private/Widgets/Layout/SMissingWidget.cpp b/Engine/Source/Runtime/Slate/Private/Widgets/Layout/SMissingWidget.cpp
--- a/Engine/Source/Runtime/Slate/Private/Widgets/Layout/SMissingWidget.cpp
+++ b/Engine/Source/Runtime/Slate/Private/Widgets/Layout/SMissingWidget.cpp
@@ -1,9 +1,5 @@
 // Copyright 1998-2014 Epic Games, Inc. All Rights Reserved.
 
-/*=============================================================================
-	SMissingWidget.h: Declares the SMissingWidget class.
-=============================================================================*/
-
 #include "SlatePrivatePCH.h"
 
 
diff --git a/Engine/Source/Runtime/Slate/Private/Widgets/Layout/SWidgetSwitcher.cpp b/Engine/Source/Runtime/Slate/Private/Widgets/Layout/SWidgetSwitcher.cpp
--- a/Engine/Source/Runtime/Slate/Private/Widgets/Layout/SWidgetSwitcher.cpp
+++ b/Engine/Source/Runtime/Slate/Private/Widgets/Layout/SWidgetSwitcher.cpp
@@ -1,20 +1,22 @@
 // Copyright 1998-2014 Epic Games, Inc. All Rights Reserved.
 
-/*=============================================================================
-	SWidgetSwitcher.cpp: Implements the SWidgetSwitcher class.
-=============================================================================*/
-
 #include "SlatePrivatePCH.h"
 #include "LayoutUtils.h"
 
 
+/* Local constants
+ *****************************************************************************/
+
+const TAttribute<FVector2D> ContentScale = FVector2D::UnitVector;
+
+
 /* SWidgetSwitcher interface
  *****************************************************************************/
 
- SWidgetSwitcher::SWidgetSwitcher()
- : AllChildren()
- {
- }
+SWidgetSwitcher::SWidgetSwitcher()
+	: AllChildren()
+{ }
+
 
 SWidgetSwitcher::FSlot& SWidgetSwitcher::AddSlot( int32 SlotIndex )
 {
@@ -33,21 +35,6 @@ SWidgetSwitcher::FSlot& SWidgetSwitcher::AddSlot( int32 SlotIndex )
 }
 
 
-int32 SWidgetSwitcher::RemoveSlot(TSharedRef<SWidget> WidgetToRemove)
-{
-	for( int32 SlotIndex=0; SlotIndex < AllChildren.Num(); ++SlotIndex )
-	{
-		if ( AllChildren[SlotIndex].GetWidget() == WidgetToRemove )
-		{
-			AllChildren.RemoveAt(SlotIndex);
-			return SlotIndex;
-		}
-	}
-
-	return -1;
-}
-
-
 void SWidgetSwitcher::Construct( const FArguments& InArgs )
 {
 	OneDynamicChild = FOneDynamicChild( &AllChildren, &WidgetIndex );
@@ -69,7 +56,7 @@ TSharedPtr<SWidget> SWidgetSwitcher::GetActiveWidget( ) const
 		return AllChildren[ActiveWidgetIndex].GetWidget();
 	}
 
-	return NULL;
+	return nullptr;
 }
 
 
@@ -80,7 +67,7 @@ TSharedPtr<SWidget> SWidgetSwitcher::GetWidget( int32 SlotIndex ) const
 		return AllChildren[SlotIndex].GetWidget();
 	}
 
-	return NULL;
+	return nullptr;
 }
 
 
@@ -99,13 +86,27 @@ int32 SWidgetSwitcher::GetWidgetIndex( TSharedRef<SWidget> Widget ) const
 	return INDEX_NONE;
 }
 
+
+int32 SWidgetSwitcher::RemoveSlot( TSharedRef<SWidget> WidgetToRemove )
+{
+	for (int32 SlotIndex=0; SlotIndex < AllChildren.Num(); ++SlotIndex)
+	{
+		if (AllChildren[SlotIndex].GetWidget() == WidgetToRemove)
+		{
+			AllChildren.RemoveAt(SlotIndex);
+			return SlotIndex;
+		}
+	}
+
+	return -1;
+}
+
+
 void SWidgetSwitcher::SetActiveWidgetIndex( int32 Index )
 {
 	WidgetIndex = Index;
 }
 
-TAttribute<FVector2D> ContentScale = FVector2D::UnitVector;
-
 
 /* SCompoundWidget interface
  *****************************************************************************/
diff --git a/Engine/Source/Runtime/Slate/Private/Widgets/Notifications/SNotificationList.cpp b/Engine/Source/Runtime/Slate/Private/Widgets/Notifications/SNotificationList.cpp
--- a/Engine/Source/Runtime/Slate/Private/Widgets/Notifications/SNotificationList.cpp
+++ b/Engine/Source/Runtime/Slate/Private/Widgets/Notifications/SNotificationList.cpp
@@ -580,7 +580,7 @@ TSharedRef<SNotificationItem> SNotificationList::AddNotification(const FNotifica
 		SNew(SNotificationItemImpl)
 		.Text(Info.Text)
 		.ButtonDetails(Info.ButtonDetails)
-		.Image((Info.Image != NULL) ? Info.Image : CachedImage)
+		.Image((Info.Image != nullptr) ? Info.Image : CachedImage)
 		.FadeInDuration(Info.FadeInDuration)
 		.ExpireDuration(Info.ExpireDuration)
 		.FadeOutDuration(Info.FadeOutDuration)
diff --git a/Engine/Source/Runtime/Slate/Private/Widgets/SCanvas.cpp b/Engine/Source/Runtime/Slate/Private/Widgets/SCanvas.cpp
--- a/Engine/Source/Runtime/Slate/Private/Widgets/SCanvas.cpp
+++ b/Engine/Source/Runtime/Slate/Private/Widgets/SCanvas.cpp
@@ -1,16 +1,16 @@
 // Copyright 1998-2014 Epic Games, Inc. All Rights Reserved.
 
-/*=============================================================================
-	SCanvas.cpp: Implements the SCanvas class.
-=============================================================================*/
-
 #include "SlatePrivatePCH.h"
 #include "LayoutUtils.h"
 
+
+/* SCanvas structors
+ *****************************************************************************/
+
 SCanvas::SCanvas()
-: Children()
-{
-}
+	: Children()
+{ }
+
 
 /* SCanvas interface
  *****************************************************************************/
diff --git a/Engine/Source/Runtime/Slate/Private/Widgets/SToolTip.cpp b/Engine/Source/Runtime/Slate/Private/Widgets/SToolTip.cpp
--- a/Engine/Source/Runtime/Slate/Private/Widgets/SToolTip.cpp
+++ b/Engine/Source/Runtime/Slate/Private/Widgets/SToolTip.cpp
@@ -1,9 +1,5 @@
 // Copyright 1998-2014 Epic Games, Inc. All Rights Reserved.
 
-/*=============================================================================
-	SToolTip.cpp: Implements the SToolTip class.
-=============================================================================*/
-
 #include "SlatePrivatePCH.h"
 
 
@@ -12,11 +8,13 @@ static TAutoConsoleVariable<float> StaticToolTipWrapWidth(
 	1000.0f,
 	TEXT( "Width of Slate tool-tips before we wrap the tool-tip text" ) );
 
+
 float SToolTip::GetToolTipWrapWidth()
 {
 	return StaticToolTipWrapWidth.GetValueOnGameThread();
 }
 
+
 /**
  * Construct this widget
  *
@@ -55,12 +53,14 @@ void SToolTip::Construct( const FArguments& InArgs )
 	
 }
 
+
 /** @return True if the tool tip has no content to display at this particular moment (remember, attribute callbacks can change this at any time!) */
 bool SToolTip::IsEmpty() const
 {
 	return !WidgetContent.IsValid() && TextContent.Get().IsEmpty();
 }
 
+
 bool SToolTip::IsInteractive() const
 {
 	return bIsInteractive.Get();
diff --git a/Engine/Source/Runtime/Slate/Private/Widgets/SViewport.cpp b/Engine/Source/Runtime/Slate/Private/Widgets/SViewport.cpp
--- a/Engine/Source/Runtime/Slate/Private/Widgets/SViewport.cpp
+++ b/Engine/Source/Runtime/Slate/Private/Widgets/SViewport.cpp
@@ -1,9 +1,5 @@
 // Copyright 1998-2014 Epic Games, Inc. All Rights Reserved.
 
-/*=============================================================================
-	SViewport.cpp: Implements the SViewport class.
-=============================================================================*/
-
 #include "SlatePrivatePCH.h"
 
 
@@ -11,7 +7,7 @@
  *****************************************************************************/
 
 SViewport::SViewport( )
-	: ViewportInterface(NULL)
+	: ViewportInterface(nullptr)
 	, bRenderDirectlyToWindow(false)
 	, bEnableGammaCorrection(true)
 { }
@@ -63,7 +59,7 @@ int32 SViewport::OnPaint( const FPaintArgs& Args, const FGeometry& AllottedGeome
 	// Only draw a quad if not rendering directly to the backbuffer
 	if( !ShouldRenderDirectly() )
 	{
-		if( ViewportInterfacePin.IsValid() && ViewportInterfacePin->GetViewportRenderTargetTexture() != NULL )
+		if( ViewportInterfacePin.IsValid() && ViewportInterfacePin->GetViewportRenderTargetTexture() != nullptr )
 		{
 			FSlateDrawElement::MakeViewport( OutDrawElements, LayerId, AllottedGeometry.ToPaintGeometry(), ViewportInterfacePin, MyClippingRect, bEnableGammaCorrection, bEnableBlending, DrawEffects, InWidgetStyle.GetColorAndOpacityTint() );
 		}
diff --git a/Engine/Source/Runtime/Slate/Private/Widgets/SWeakWidget.cpp b/Engine/Source/Runtime/Slate/Private/Widgets/SWeakWidget.cpp
--- a/Engine/Source/Runtime/Slate/Private/Widgets/SWeakWidget.cpp
+++ b/Engine/Source/Runtime/Slate/Private/Widgets/SWeakWidget.cpp
@@ -1,21 +1,19 @@
 // Copyright 1998-2014 Epic Games, Inc. All Rights Reserved.
 
-/*=============================================================================
-	SWeakWidget.cpp: Implements the SWeakWidget class.
-=============================================================================*/
-
 #include "SlatePrivatePCH.h"
 
+
 SWeakWidget::SWeakWidget()
-: WeakChild()
-{
-}
+	: WeakChild()
+{ }
+
 
 void SWeakWidget::Construct(const FArguments& InArgs)
 {
 	WeakChild.AttachWidget( InArgs._PossiblyNullContent );
 }
 
+
 FVector2D SWeakWidget::ComputeDesiredSize() const
 {	
 	TSharedRef<SWidget> ReferencedWidget = WeakChild.GetWidget();
@@ -28,6 +26,7 @@ FVector2D SWeakWidget::ComputeDesiredSize() const
 	return FVector2D::ZeroVector;
 }
 
+
 void SWeakWidget::OnArrangeChildren( const FGeometry& AllottedGeometry, FArrangedChildren& ArrangedChildren ) const
 {
 	// We just want to show the child that we are presenting. Always stretch it to occupy all of the space.
@@ -42,12 +41,14 @@ void SWeakWidget::OnArrangeChildren( const FGeometry& AllottedGeometry, FArrange
 			) );
 	}
 }
-		
+
+
 FChildren* SWeakWidget::GetChildren()
 {
 	return &WeakChild;
 }
 
+
 int32 SWeakWidget::OnPaint( const FPaintArgs& Args, const FGeometry& AllottedGeometry, const FSlateRect& MyClippingRect, FSlateWindowElementList& OutDrawElements, int32 LayerId, const FWidgetStyle& InWidgetStyle, bool bParentEnabled ) const
 {
 	// Just draw the children.
@@ -72,11 +73,13 @@ int32 SWeakWidget::OnPaint( const FPaintArgs& Args, const FGeometry& AllottedGeo
 	return LayerId;
 }
 
+
 void SWeakWidget::SetContent(const TSharedRef<SWidget>& InWidget)
 {
 	WeakChild.AttachWidget( InWidget );
 }
 
+
 bool SWeakWidget::ChildWidgetIsValid() const
 {
 	return WeakChild.GetWidget() != SNullWidget::NullWidget;
diff --git a/Engine/Source/Runtime/Slate/Private/Widgets/Testing/STestSuite.cpp b/Engine/Source/Runtime/Slate/Private/Widgets/Testing/STestSuite.cpp
--- a/Engine/Source/Runtime/Slate/Private/Widgets/Testing/STestSuite.cpp
+++ b/Engine/Source/Runtime/Slate/Private/Widgets/Testing/STestSuite.cpp
@@ -1480,9 +1480,9 @@ public:
 	{
 		const FString* url = Metadata.Find( TEXT("href") );
 
-		if ( url != NULL )
+		if ( url != nullptr )
 		{
-			FPlatformProcess::LaunchURL( **url, NULL, NULL );
+			FPlatformProcess::LaunchURL( **url, nullptr, nullptr );
 		}
 		else
 		{
diff --git a/Engine/Source/Runtime/Slate/Private/Widgets/Text/SRichTextBlock.cpp b/Engine/Source/Runtime/Slate/Private/Widgets/Text/SRichTextBlock.cpp
--- a/Engine/Source/Runtime/Slate/Private/Widgets/Text/SRichTextBlock.cpp
+++ b/Engine/Source/Runtime/Slate/Private/Widgets/Text/SRichTextBlock.cpp
@@ -168,7 +168,7 @@ TSharedPtr< ITextDecorator > SRichTextBlock::TryGetDecorator( const TArray< TSha
 		}
 	}
 
-	return NULL;
+	return nullptr;
 }
 
 void SRichTextBlock::SetText( const TAttribute<FText>& InTextAttr )
diff --git a/Engine/Source/Runtime/Slate/Public/Framework/Application/MenuStack.h b/Engine/Source/Runtime/Slate/Public/Framework/Application/MenuStack.h
--- a/Engine/Source/Runtime/Slate/Public/Framework/Application/MenuStack.h
+++ b/Engine/Source/Runtime/Slate/Public/Framework/Application/MenuStack.h
@@ -1,9 +1,5 @@
 // Copyright 1998-2014 Epic Games, Inc. All Rights Reserved.
 
-/*=============================================================================
-	MenuStack.h: Declares the FMenuStack class.
-=============================================================================*/
-
 #pragma once
 
 
diff --git a/Engine/Source/Runtime/Slate/Public/Framework/Commands/InputBindingManager.h b/Engine/Source/Runtime/Slate/Public/Framework/Commands/InputBindingManager.h
--- a/Engine/Source/Runtime/Slate/Public/Framework/Commands/InputBindingManager.h
+++ b/Engine/Source/Runtime/Slate/Public/Framework/Commands/InputBindingManager.h
@@ -1,9 +1,5 @@
 // Copyright 1998-2014 Epic Games, Inc. All Rights Reserved.
 
-/*=============================================================================
-	InputBindingManager.h: Declares the FInputBindingManager class.
-=============================================================================*/
-
 #pragma once
 
 
@@ -21,6 +17,7 @@ DECLARE_MULTICAST_DELEGATE_OneParam(FOnUserDefinedGestureChanged, const FUIComma
 class SLATE_API FInputBindingManager
 {
 public:
+
 	/**
 	 * @return The instance of this manager                   
 	 */
diff --git a/Engine/Source/Runtime/Slate/Public/Framework/Commands/InputGesture.h b/Engine/Source/Runtime/Slate/Public/Framework/Commands/InputGesture.h
--- a/Engine/Source/Runtime/Slate/Public/Framework/Commands/InputGesture.h
+++ b/Engine/Source/Runtime/Slate/Public/Framework/Commands/InputGesture.h
@@ -1,9 +1,5 @@
 // Copyright 1998-2014 Epic Games, Inc. All Rights Reserved.
 
-/*=============================================================================
-	InputGesture.h: Declares the FInputGesture class.
-=============================================================================*/
-
 #pragma once
 
 
@@ -53,9 +49,7 @@ struct SLATE_API FInputGesture
 
 public:
 
-	/**
-	 * Default constructor.
-	 */
+	/** Default constructor. */
 	FInputGesture( )
 	{
 		bCtrl = bAlt = bShift = bCmd = 0;
@@ -121,7 +115,6 @@ public:
 	 * Compares this input gesture with another for equality.
 	 *
 	 * @param Other The other gesture to compare with.
-	 *
 	 * @return true if the gestures are equal, false otherwise.
 	 */
 	bool operator!=( const FInputGesture& Other ) const
@@ -133,7 +126,6 @@ public:
 	 * Compares this input gesture with another for inequality.
 	 *
 	 * @param Other The other gesture to compare with.
-	 *
 	 * @return true if the gestures are not equal, false otherwise.
 	 */
 	bool operator==( const FInputGesture& Other ) const
@@ -197,7 +189,6 @@ public:
 	 * Gets a type hash value for the specified gesture.
 	 *
 	 * @param Gesture The input gesture to get the hash value for.
-	 *
 	 * @return The hash value.
 	 */
 	friend uint32 GetTypeHash( const FInputGesture& Gesture )
diff --git a/Engine/Source/Runtime/Slate/Public/Framework/Commands/UIAction.h b/Engine/Source/Runtime/Slate/Public/Framework/Commands/UIAction.h
--- a/Engine/Source/Runtime/Slate/Public/Framework/Commands/UIAction.h
+++ b/Engine/Source/Runtime/Slate/Public/Framework/Commands/UIAction.h
@@ -1,9 +1,5 @@
 // Copyright 1998-2014 Epic Games, Inc. All Rights Reserved.
 
-/*=============================================================================
-	UIAction.h: Declares the FUIAction structure.
-=============================================================================*/
-
 #pragma once
 
 
@@ -39,9 +35,7 @@ struct FUIAction
 
 public:
 
-	/**
-	 * Default constructor
-	 */
+	/** Default constructor. */
 	FUIAction( ) { }
 
 	/**
@@ -103,8 +97,7 @@ public:
 	/**
 	 * Compares this UI action to another for equality.
 	 *
-	 * @param Other - The other action to compare to.
-	 *
+	 * @param Other The other action to compare to.
 	 * @return true if both actions are equal, false otherwise.
 	 */
 	bool operator==( const FUIAction& Other ) const
diff --git a/Engine/Source/Runtime/Slate/Public/Framework/Commands/UICommandDragDropOp.h b/Engine/Source/Runtime/Slate/Public/Framework/Commands/UICommandDragDropOp.h
--- a/Engine/Source/Runtime/Slate/Public/Framework/Commands/UICommandDragDropOp.h
+++ b/Engine/Source/Runtime/Slate/Public/Framework/Commands/UICommandDragDropOp.h
@@ -1,9 +1,5 @@
 // Copyright 1998-2014 Epic Games, Inc. All Rights Reserved.
 
-/*=============================================================================
-	UICommandDragDropOp.h: Declares the FUICommandDragDropOp class.
-=============================================================================*/
-
 #pragma once
 
 
@@ -24,7 +20,7 @@ public:
 		, OriginMultiBox( InOriginMultiBox )
 		, CustomDecorator( InCustomDecorator )
 		, Offset( DecoratorOffset )
-	{}
+	{ }
 
 	/**
 	 * Sets a delegate that will be called when the command is dropped
diff --git a/Engine/Source/Runtime/Slate/Public/Framework/Commands/UICommandInfo.h b/Engine/Source/Runtime/Slate/Public/Framework/Commands/UICommandInfo.h
--- a/Engine/Source/Runtime/Slate/Public/Framework/Commands/UICommandInfo.h
+++ b/Engine/Source/Runtime/Slate/Public/Framework/Commands/UICommandInfo.h
@@ -1,9 +1,5 @@
 // Copyright 1998-2014 Epic Games, Inc. All Rights Reserved.
 
-/*=============================================================================
-	UICommandInfo.h: Declares the UICommandInfo class.
-=============================================================================*/
-
 #pragma once
 
 #include "Layout/Visibility.h"
@@ -11,6 +7,7 @@
 #include "Framework/Commands/InputGesture.h"
 #include "Textures/SlateIcon.h"
 
+
 /** Types of user interfaces that can be associated with a user interface action */
 namespace EUserInterfaceActionType
 {
@@ -130,13 +127,18 @@ public:
 
 	/** A delegate that is called when commands are registered or unregistered with a binding context */
 	static FSimpleMulticastDelegate CommandsChanged; 
+
 private:
+
 	/** The name of the context */
 	FName ContextName;
+
 	/** The name of the parent context */
 	FName ContextParent;
+
 	/** The description of the context */
 	FText ContextDesc;
+
 	/** The style set to find the icons in */
 	FName StyleSetName;
 };
@@ -148,7 +150,12 @@ class SLATE_API FUICommandInfo
 	friend class FUICommandInfoDecl;
 
 public:
-	/** Default constructor */
+
+	/**
+	 * Creates and initializes a new instance.
+	 *
+	 * @param InBindingContext The name of the binding context to use.
+	 */
 	FUICommandInfo( const FName InBindingContext )
 		: ActiveGesture( new FInputGesture )
 		, DefaultGesture( EKeys::Invalid, EModifierKey::None )
diff --git a/Engine/Source/Runtime/Slate/Public/Framework/Commands/UICommandList.h b/Engine/Source/Runtime/Slate/Public/Framework/Commands/UICommandList.h
--- a/Engine/Source/Runtime/Slate/Public/Framework/Commands/UICommandList.h
+++ b/Engine/Source/Runtime/Slate/Public/Framework/Commands/UICommandList.h
@@ -1,9 +1,5 @@
 // Copyright 1998-2014 Epic Games, Inc. All Rights Reserved.
 
-/*=============================================================================
-	UICommandList.h: Declares the FUICommandList class.
-=============================================================================*/
-
 #pragma once
 
 
@@ -107,7 +103,6 @@ public:
 	 * Processes any UI commands which are activated by the specified keyboard event
 	 *
 	 * @param InKeyboardEvent	The keyboard event to check
-	 *
 	 * @return true if an action was processed
 	 */
 	bool ProcessCommandBindings( const FKeyboardEvent& InKeyboardEvent ) const;
@@ -116,7 +111,6 @@ public:
 	 * Processes any UI commands which are activated by the specified mouse event
 	 *
 	 * @param InKeyboardEvent	The mouse event to check
-	 *
 	 * @return true if an action was processed
 	 */
 	bool ProcessCommandBindings( const FPointerEvent& InMouseEvent ) const;
@@ -127,7 +121,6 @@ public:
 	 * @param Key				The current key that is pressed
 	 * @param ModifierKeysState	Pressed state of keys that are commonly used as modifiers
 	 * @param bRepeat			True if input is repeating (held)
-	 *
 	 * @return true if an action was processed
 	 */
 	bool ProcessCommandBindings( const FKey Key, const FModifierKeysState& ModifierKeysState, const bool bRepeat ) const;
diff --git a/Engine/Source/Runtime/Slate/Public/Framework/Layout/ScrollyZoomy.h b/Engine/Source/Runtime/Slate/Public/Framework/Layout/ScrollyZoomy.h
--- a/Engine/Source/Runtime/Slate/Public/Framework/Layout/ScrollyZoomy.h
+++ b/Engine/Source/Runtime/Slate/Public/Framework/Layout/ScrollyZoomy.h
@@ -1,9 +1,5 @@
 // Copyright 1998-2014 Epic Games, Inc. All Rights Reserved.
 
-/*=============================================================================
-	ScrollyZoomy.h: Declares the IScrollableZoomable interface and FScrollyZoomy class.
-=============================================================================*/
-
 #pragma once
 
 
@@ -17,8 +13,7 @@ public:
 	/**
 	 * Override this to scroll your widget's content.
 	 *
-	 * @param Offset - The 2D offset to scroll by.
-	 *
+	 * @param Offset The 2D offset to scroll by.
 	 * @return True if anything was scrolled.
 	 */
 	virtual bool ScrollBy( const FVector2D& Offset ) = 0;
@@ -26,8 +21,7 @@ public:
 	/**
 	 * Override this to zoom your widget's content.
 	 *
-	 * @param Amount - The amount to zoom by.
-	 *
+	 * @param Amount The amount to zoom by.
 	 * @return True if anything was zoomed.
 	 */
 	virtual bool ZoomBy( const float Amount ) = 0;
@@ -57,7 +51,6 @@ public:
 	 * Should be called when a mouse button is pressed
 	 *
 	 * @param	MouseEvent	The mouse event passed to a widget's OnMouseButtonDown() function
-	 *
 	 * @return	If FReply::Handled is returned, that should be passed as the result of the calling function
 	 */
 	FReply OnMouseButtonDown( const FPointerEvent& MouseEvent );
@@ -68,7 +61,6 @@ public:
 	 * @param	MyWidget	Pointer to the widget that owns this object
 	 * @param	MyGeometry	Geometry of the widget we're scrolling
 	 * @param	MouseEvent	The mouse event passed to a widget's OnMouseButtonUp() function
-	 *
 	 * @return	If FReply::Handled is returned, that should be passed as the result of the calling function
 	 */
 	FReply OnMouseButtonUp( const TSharedRef<SWidget> MyWidget, const FGeometry& MyGeometry, const FPointerEvent& MouseEvent );
@@ -80,7 +72,6 @@ public:
 	 * @param	ScrollableZoomable	Interface to the widget to scroll/zoom
 	 * @param	MyGeometry	Geometry of the widget we're scrolling
 	 * @param	MouseEvent	The mouse event passed to a widget's OnMouseMove() function
-	 *
 	 * @return	If FReply::Handled is returned, that should be passed as the result of the calling function
 	 */
 	FReply OnMouseMove( const TSharedRef<SWidget> MyWidget, IScrollableZoomable& ScrollableZoomable, const FGeometry& MyGeometry, const FPointerEvent& MouseEvent );
@@ -98,7 +89,6 @@ public:
 	 *
 	 * @param	MouseEvent	The event passed to your widget's OnMouseWheel function
 	 * @param	ScrollableZoomable	Interface to the widget to scroll/zoom
-	 *
 	 * @return	If FReply::Handled is returned, that should be passed as the result of the calling function
 	 */
 	FReply OnMouseWheel( const FPointerEvent& MouseEvent, IScrollableZoomable& ScrollableZoomable );
@@ -126,7 +116,6 @@ public:
 	 * @param	MyClippingRect			Widget clipping rect passed into OnPaint()
 	 * @param	OutDrawElements			The draw element list
 	 * @param	LayerId					Layer Id
-	 *
 	 * @return	New layer Id
 	 */
 	int32 PaintSoftwareCursorIfNeeded( const FGeometry& AllottedGeometry, const FSlateRect& MyClippingRect, FSlateWindowElementList& OutDrawElements, int32 LayerId ) const;
diff --git a/Engine/Source/Runtime/Slate/Public/Framework/SlateDelegates.h b/Engine/Source/Runtime/Slate/Public/Framework/SlateDelegates.h
--- a/Engine/Source/Runtime/Slate/Public/Framework/SlateDelegates.h
+++ b/Engine/Source/Runtime/Slate/Public/Framework/SlateDelegates.h
@@ -1,9 +1,5 @@
 // Copyright 1998-2014 Epic Games, Inc. All Rights Reserved.
 
-/*=============================================================================
-	SlateDelegates.h: Declares various common delegate types used in Slate.
-=============================================================================*/
-
 #pragma once
 
 
@@ -99,7 +95,6 @@ DECLARE_DELEGATE_OneParam( FOnTextChanged, const FText& );
 DECLARE_DELEGATE_RetVal_OneParam( bool, FOnIsTypedCharValid, const TCHAR /*TypedCharacter*/ );
 
 
-
 /** Delegate for hooking up to an editable text box's 'OnTextCommitted' 
 
 	The first parameter (NewText) is the new text string.
@@ -126,6 +121,7 @@ DECLARE_DELEGATE_TwoParams( FOnInt32ValueCommitted, int32, ETextCommit::Type);
 /** Notification for FLinearColor value change */
 DECLARE_DELEGATE_OneParam( FOnLinearColorValueChanged, FLinearColor )
 
+
 template< typename ArgumentType >
 class TSlateDelegates
 {
@@ -201,5 +197,6 @@ public:
 	)
 };
 
+
 /** Notification for when a keyboard event occurs */
 DECLARE_DELEGATE_RetVal_OneParam( FReply, FOnKeyboardEvent, const FKeyboardEvent& )
diff --git a/Engine/Source/Runtime/Slate/Public/SlateMaterialBrush.h b/Engine/Source/Runtime/Slate/Public/SlateMaterialBrush.h
--- a/Engine/Source/Runtime/Slate/Public/SlateMaterialBrush.h
+++ b/Engine/Source/Runtime/Slate/Public/SlateMaterialBrush.h
@@ -4,6 +4,7 @@
 
 #include "Slate.h"
 
+
 /**
  * Dynamic rush for referencing a UMaterial.  
  *
@@ -11,14 +12,20 @@
  */
 struct FSlateMaterialBrush : public FSlateBrush
 {
+	/**
+	 * Creates and initializes a new instance.
+	 *
+	 * @param InMaterial The material to use.
+	 * @param InImageSize The material's dimensions.
+	 */
 	FSlateMaterialBrush( class UMaterialInterface& InMaterial, const FVector2D& InImageSize )
 		: FSlateBrush( ESlateBrushDrawType::Image, FName(TEXT("None")), FMargin(0), ESlateBrushTileType::NoTile, ESlateBrushImageType::FullColor, InImageSize, FLinearColor::White, &InMaterial )
 	{
 		ResourceName = FName( *InMaterial.GetFullName() );
 	}
 
-
-	virtual ~FSlateMaterialBrush()
+	/** Virtual destructor. */
+	virtual ~FSlateMaterialBrush( )
 	{
 		if(FSlateApplication::IsInitialized() && FSlateApplication::Get().GetRenderer().IsValid())
 		{
diff --git a/Engine/Source/Runtime/Slate/Public/Widgets/Colors/SColorBlock.h b/Engine/Source/Runtime/Slate/Public/Widgets/Colors/SColorBlock.h
--- a/Engine/Source/Runtime/Slate/Public/Widgets/Colors/SColorBlock.h
+++ b/Engine/Source/Runtime/Slate/Public/Widgets/Colors/SColorBlock.h
@@ -47,11 +47,12 @@ public:
 	void Construct( const FArguments& InArgs );
 
 private:
-	// SWidget interface
+
+	// SWidget overrides
+
 	virtual int32 OnPaint( const FPaintArgs& Args, const FGeometry& AllottedGeometry, const FSlateRect& MyClippingRect, FSlateWindowElementList& OutDrawElements, int32 LayerId, const FWidgetStyle& InWidgetStyle, bool bParentEnabled ) const override;
 	virtual FReply OnMouseButtonDown( const FGeometry& MyGeometry, const FPointerEvent& MouseEvent ) override;
 	virtual FVector2D ComputeDesiredSize() const override;
-	// End of SWidget interface
 
 private:
 
diff --git a/Engine/Source/Runtime/Slate/Public/Widgets/Colors/SColorPicker.h b/Engine/Source/Runtime/Slate/Public/Widgets/Colors/SColorPicker.h
--- a/Engine/Source/Runtime/Slate/Public/Widgets/Colors/SColorPicker.h
+++ b/Engine/Source/Runtime/Slate/Public/Widgets/Colors/SColorPicker.h
@@ -45,7 +45,7 @@ struct FColorChannels
 {
 	FColorChannels()
 	{
-		Red = Green = Blue = Alpha = NULL;
+		Red = Green = Blue = Alpha = nullptr;
 	}
 
 	float* Red;
@@ -142,9 +142,7 @@ public:
 
 public:
 
-	/**
-	 *	Default Destructor
-	 */
+	/**	Destructor. */
 	~SColorPicker();
 
 public:
@@ -196,16 +194,13 @@ protected:
 	void HideSmallTrash();
 	void ShowSmallTrash();
 
-	/**
-	 * Cycles the color picker's mode.
-	 */
+	/** Cycles the color picker's mode. */
 	void CycleMode( );
 
 	/**
 	 * Creates a color slider widget for the specified channel.
 	 *
-	 * @param Channel - The color channel to create the widget for.
-	 *
+	 * @param Channel The color channel to create the widget for.
 	 * @return The new slider.
 	 */
 	TSharedRef<SWidget> MakeColorSlider( EColorPickerChannels::Type Channel ) const;
@@ -213,8 +208,7 @@ protected:
 	/**
 	 * Creates a color spin box widget for the specified channel.
 	 *
-	 * @param Channel - The color channel to create the widget for.
-	 *
+	 * @param Channel The color channel to create the widget for.
 	 * @return The new spin box.
 	 */
 	TSharedRef<SWidget> MakeColorSpinBox( EColorPickerChannels::Type Channel ) const;
@@ -495,14 +489,14 @@ struct FColorPickerArgs
 
 	FColorPickerArgs()
 		: bIsModal(false)
-		, ParentWidget(NULL)
+		, ParentWidget(nullptr)
 		, bUseAlpha(false)
 		, bOnlyRefreshOnMouseUp(false)
 		, bOnlyRefreshOnOk(false)
 		, DisplayGamma(2.2f)
-		, ColorArray(NULL)
-		, LinearColorArray(NULL)
-		, ColorChannelsArray(NULL)
+		, ColorArray(nullptr)
+		, LinearColorArray(nullptr)
+		, ColorChannelsArray(nullptr)
 		, OnColorCommitted()
 		, PreColorCommitted()
 		, OnColorPickerWindowClosed()
diff --git a/Engine/Source/Runtime/Slate/Public/Widgets/Colors/SColorSpectrum.h b/Engine/Source/Runtime/Slate/Public/Widgets/Colors/SColorSpectrum.h
--- a/Engine/Source/Runtime/Slate/Public/Widgets/Colors/SColorSpectrum.h
+++ b/Engine/Source/Runtime/Slate/Public/Widgets/Colors/SColorSpectrum.h
@@ -1,9 +1,5 @@
 // Copyright 1998-2014 Epic Games, Inc. All Rights Reserved.
 
-/*=============================================================================
-	SColorSpectrum.h: Declares the SColorSpectrum class.
-=============================================================================*/
-
 #pragma once
 
 
@@ -41,30 +37,20 @@ public:
 	/**
 	 * Construct this widget.
 	 *
-	 * @param InArgs - The declaration data for this widget.
+	 * @param InArgs The declaration data for this widget.
 	 */
 	void Construct( const FArguments& InArgs );
 
 public:
 
-	// Begin overrides
+	// SWidget overrides
 
 	virtual FVector2D ComputeDesiredSize( ) const;
-
-	virtual FReply OnMouseButtonDoubleClick( const FGeometry& InMyGeometry, const FPointerEvent& InMouseEvent )
-	{
-		return FReply::Handled();
-	}
-
+	virtual FReply OnMouseButtonDoubleClick( const FGeometry& InMyGeometry, const FPointerEvent& InMouseEvent );
 	virtual FReply OnMouseButtonDown( const FGeometry& MyGeometry, const FPointerEvent& MouseEvent );
-	
 	virtual FReply OnMouseButtonUp( const FGeometry& MyGeometry, const FPointerEvent& MouseEvent );
-
 	virtual FReply OnMouseMove( const FGeometry& MyGeometry, const FPointerEvent& MouseEvent );
-
 	virtual int32 OnPaint( const FPaintArgs& Args, const FGeometry& AllottedGeometry, const FSlateRect& MyClippingRect, FSlateWindowElementList& OutDrawElements, int32 LayerId, const FWidgetStyle& InWidgetStyle, bool bParentEnabled ) const override;
-
-	// End overrides
 	
 protected:
 
diff --git a/Engine/Source/Runtime/Slate/Public/Widgets/Colors/SColorWheel.h b/Engine/Source/Runtime/Slate/Public/Widgets/Colors/SColorWheel.h
--- a/Engine/Source/Runtime/Slate/Public/Widgets/Colors/SColorWheel.h
+++ b/Engine/Source/Runtime/Slate/Public/Widgets/Colors/SColorWheel.h
@@ -1,9 +1,5 @@
 // Copyright 1998-2014 Epic Games, Inc. All Rights Reserved.
 
-/*=============================================================================
-	SColorWheel.h: Declares the SColorWheel class.
-=============================================================================*/
-
 #pragma once
 
 
@@ -22,16 +18,16 @@ public:
 		, _OnValueChanged()
 	{ }
 	
-		/** The current color selected by the user */
+		/** The current color selected by the user. */
 		SLATE_ATTRIBUTE(FLinearColor, SelectedColor)
 		
-		/** Invoked when the mouse is pressed and a capture begins */
+		/** Invoked when the mouse is pressed and a capture begins. */
 		SLATE_EVENT(FSimpleDelegate, OnMouseCaptureBegin)
 
-		/** Invoked when the mouse is released and a capture ends */
+		/** Invoked when the mouse is released and a capture ends. */
 		SLATE_EVENT(FSimpleDelegate, OnMouseCaptureEnd)
 
-		/** Invoked when a new value is selected on the color wheel */
+		/** Invoked when a new value is selected on the color wheel. */
 		SLATE_EVENT(FOnLinearColorValueChanged, OnValueChanged)
 
 	SLATE_END_ARGS()
@@ -41,30 +37,20 @@ public:
 	/**
 	 * Construct this widget.
 	 *
-	 * @param InArgs - The declaration data for this widget.
+	 * @param InArgs The declaration data for this widget.
 	 */
 	void Construct( const FArguments& InArgs );
 
 public:
 
-	// Begin SWidget overrides
+	// SWidget overrides
 
 	virtual FVector2D ComputeDesiredSize( ) const;
-
-	virtual FReply OnMouseButtonDoubleClick( const FGeometry& InMyGeometry, const FPointerEvent& InMouseEvent )
-	{
-		return FReply::Handled();
-	}
-
+	virtual FReply OnMouseButtonDoubleClick( const FGeometry& InMyGeometry, const FPointerEvent& InMouseEvent );
 	virtual FReply OnMouseButtonDown( const FGeometry& MyGeometry, const FPointerEvent& MouseEvent );
-	
 	virtual FReply OnMouseButtonUp( const FGeometry& MyGeometry, const FPointerEvent& MouseEvent );
-
 	virtual FReply OnMouseMove( const FGeometry& MyGeometry, const FPointerEvent& MouseEvent );
-
 	virtual int32 OnPaint( const FPaintArgs& Args, const FGeometry& AllottedGeometry, const FSlateRect& MyClippingRect, FSlateWindowElementList& OutDrawElements, int32 LayerId, const FWidgetStyle& InWidgetStyle, bool bParentEnabled ) const override;
-
-	// End SWidget overrides
 	
 protected:
 
@@ -77,23 +63,23 @@ protected:
 
 private:
 
-	/** The color wheel image to show */
+	/** The color wheel image to show. */
 	const FSlateBrush* Image;
 	
-	/** The current color selected by the user */
+	/** The current color selected by the user. */
 	TAttribute< FLinearColor > SelectedColor;
 
-	/** The color selector image to show */
+	/** The color selector image to show. */
 	const FSlateBrush* SelectorImage;
 
 private:
 
-	/** Invoked when the mouse is pressed and a capture begins */
+	/** Invoked when the mouse is pressed and a capture begins. */
 	FSimpleDelegate OnMouseCaptureBegin;
 
-	/** Invoked when the mouse is let up and a capture ends */
+	/** Invoked when the mouse is let up and a capture ends. */
 	FSimpleDelegate OnMouseCaptureEnd;
 
-	/** Invoked when a new value is selected on the color wheel */
+	/** Invoked when a new value is selected on the color wheel. */
 	FOnLinearColorValueChanged OnValueChanged;
 };
diff --git a/Engine/Source/Runtime/Slate/Public/Widgets/Images/SImage.h b/Engine/Source/Runtime/Slate/Public/Widgets/Images/SImage.h
--- a/Engine/Source/Runtime/Slate/Public/Widgets/Images/SImage.h
+++ b/Engine/Source/Runtime/Slate/Public/Widgets/Images/SImage.h
@@ -3,8 +3,11 @@
 #pragma once
 
 
-/** SImage displays a single image at a desired Width and Height. */
-class SLATE_API SImage : public SLeafWidget
+/**
+ * Implements a widget that displays an image with a desired width and height.
+ */
+class SLATE_API SImage
+	: public SLeafWidget
 {
 public:
 	SLATE_BEGIN_ARGS( SImage )
@@ -31,12 +34,6 @@ public:
 	 */
 	void Construct( const FArguments& InArgs );
 
-	// SWidget interface
-	virtual int32 OnPaint( const FPaintArgs& Args, const FGeometry& AllottedGeometry, const FSlateRect& MyClippingRect, FSlateWindowElementList& OutDrawElements, int32 LayerId, const FWidgetStyle& InWidgetStyle, bool bParentEnabled ) const override;
-	virtual FReply OnMouseButtonDown( const FGeometry& MyGeometry, const FPointerEvent& MouseEvent ) override;
-	virtual FVector2D ComputeDesiredSize() const override;
-	// End of SWidget interface
-
 public:
 
 	/** See the ColorAndOpacity attribute */
@@ -51,6 +48,14 @@ public:
 	/** See OnMouseButtonDown event */
 	void SetOnMouseButtonDown(FPointerEventHandler EventHandler);
 
+public:
+
+	// SWidget overrides
+
+	virtual int32 OnPaint( const FPaintArgs& Args, const FGeometry& AllottedGeometry, const FSlateRect& MyClippingRect, FSlateWindowElementList& OutDrawElements, int32 LayerId, const FWidgetStyle& InWidgetStyle, bool bParentEnabled ) const override;
+	virtual FReply OnMouseButtonDown( const FGeometry& MyGeometry, const FPointerEvent& MouseEvent ) override;
+	virtual FVector2D ComputeDesiredSize() const override;
+
 protected:
 
 	/** The FName of the image resource to show */
@@ -61,5 +66,4 @@ protected:
 
 	/** Invoked when the mouse is pressed in the image */
 	FPointerEventHandler OnMouseButtonDownHandler;
-
 };
diff --git a/Engine/Source/Runtime/Slate/Public/Widgets/Input/SButton.h b/Engine/Source/Runtime/Slate/Public/Widgets/Input/SButton.h
--- a/Engine/Source/Runtime/Slate/Public/Widgets/Input/SButton.h
+++ b/Engine/Source/Runtime/Slate/Public/Widgets/Input/SButton.h
@@ -6,7 +6,8 @@
 /**
  * Slate's Buttons are clickable Widgets that can contain arbitrary widgets as its Content().
  */
-class SLATE_API SButton : public SBorder
+class SLATE_API SButton
+	: public SBorder
 {
 public:
 
@@ -79,27 +80,6 @@ public:
 
 	/** @return An image that represents this button's border*/
 	const FSlateBrush* GetBorder() const;
-	
-	/**
-	 * See SWidget::SupportsKeyboardFocus().
-	 *
-	 * @return  True if this widget can take keyboard focus
-	 */
-	virtual bool SupportsKeyboardFocus() const override;
-
-	virtual FReply OnKeyDown( const FGeometry& MyGeometry, const FKeyboardEvent& InKeyboardEvent ) override;
-
-	virtual FReply OnMouseButtonDown( const FGeometry& MyGeometry, const FPointerEvent& MouseEvent ) override;
-
-	virtual FReply OnMouseButtonDoubleClick( const FGeometry& InMyGeometry, const FPointerEvent& InMouseEvent ) override;
-	
-	virtual FReply OnMouseButtonUp( const FGeometry& MyGeometry, const FPointerEvent& MouseEvent ) override;
-	
-	virtual FReply OnMouseMove( const FGeometry& MyGeometry, const FPointerEvent& MouseEvent ) override;
-	
-	virtual void OnMouseEnter( const FGeometry& MyGeometry, const FPointerEvent& MouseEvent ) override;
-
-	virtual void OnMouseLeave( const FPointerEvent& MouseEvent ) override;
 
 	/**
 	 * Returns true if this button is currently pressed
@@ -133,7 +113,28 @@ public:
 	/** See ButtonStyle attribute */
 	void SetButtonStyle(const FButtonStyle* ButtonStyle);
 
+public:
+
+	// SWidget overrides
+
+	virtual bool SupportsKeyboardFocus() const override;
+
+	virtual FReply OnKeyDown( const FGeometry& MyGeometry, const FKeyboardEvent& InKeyboardEvent ) override;
+
+	virtual FReply OnMouseButtonDown( const FGeometry& MyGeometry, const FPointerEvent& MouseEvent ) override;
+
+	virtual FReply OnMouseButtonDoubleClick( const FGeometry& InMyGeometry, const FPointerEvent& InMouseEvent ) override;
+	
+	virtual FReply OnMouseButtonUp( const FGeometry& MyGeometry, const FPointerEvent& MouseEvent ) override;
+	
+	virtual FReply OnMouseMove( const FGeometry& MyGeometry, const FPointerEvent& MouseEvent ) override;
+	
+	virtual void OnMouseEnter( const FGeometry& MyGeometry, const FPointerEvent& MouseEvent ) override;
+
+	virtual void OnMouseLeave( const FPointerEvent& MouseEvent ) override;
+
 protected:
+
 	/** @return combines the user-specified margin and the button's internal margin. */
 	FMargin GetCombinedPadding() const;
 
@@ -184,5 +185,6 @@ protected:
 	FSlateSound PressedSound;
 
 private:
+
 	virtual FVector2D ComputeDesiredSize() const override;
 };
diff --git a/Engine/Source/Runtime/Slate/Public/Widgets/Input/SCheckBox.h b/Engine/Source/Runtime/Slate/Public/Widgets/Input/SCheckBox.h
--- a/Engine/Source/Runtime/Slate/Public/Widgets/Input/SCheckBox.h
+++ b/Engine/Source/Runtime/Slate/Public/Widgets/Input/SCheckBox.h
@@ -43,15 +43,15 @@ public:
 		, _ForegroundColor()
 		, _BorderBackgroundColor ()
 		, _IsFocusable( true )
-		, _UncheckedImage( NULL )
-		, _UncheckedHoveredImage( NULL )
-		, _UncheckedPressedImage( NULL )
-		, _CheckedImage( NULL )
-		, _CheckedHoveredImage( NULL )
-		, _CheckedPressedImage( NULL )
-		, _UndeterminedImage( NULL )
-		, _UndeterminedHoveredImage( NULL )
-		, _UndeterminedPressedImage( NULL )
+		, _UncheckedImage( nullptr )
+		, _UncheckedHoveredImage( nullptr )
+		, _UncheckedPressedImage( nullptr )
+		, _CheckedImage( nullptr )
+		, _CheckedHoveredImage( nullptr )
+		, _CheckedPressedImage( nullptr )
+		, _UndeterminedImage( nullptr )
+		, _UndeterminedHoveredImage( nullptr )
+		, _UndeterminedPressedImage( nullptr )
 		{}
 
 		/** Content to be placed next to the check box, or for a toggle button, the content to be placed inside the button */
diff --git a/Engine/Source/Runtime/Slate/Public/Widgets/Input/SComboBox.h b/Engine/Source/Runtime/Slate/Public/Widgets/Input/SComboBox.h
--- a/Engine/Source/Runtime/Slate/Public/Widgets/Input/SComboBox.h
+++ b/Engine/Source/Runtime/Slate/Public/Widgets/Input/SComboBox.h
@@ -79,7 +79,7 @@ public:
 		, _OptionsSource()
 		, _OnSelectionChanged()
 		, _OnGenerateWidget()
-		, _InitiallySelectedItem( NULL )
+		, _InitiallySelectedItem( nullptr)
 		, _Method( SMenuAnchor::CreateNewWindow )
 		, _MaxListHeight(450.0f)
 		, _HasDownArrow( true )
diff --git a/Engine/Source/Runtime/Slate/Public/Widgets/Input/SEditableComboBox.h b/Engine/Source/Runtime/Slate/Public/Widgets/Input/SEditableComboBox.h
--- a/Engine/Source/Runtime/Slate/Public/Widgets/Input/SEditableComboBox.h
+++ b/Engine/Source/Runtime/Slate/Public/Widgets/Input/SEditableComboBox.h
@@ -1,9 +1,5 @@
 // Copyright 1998-2014 Epic Games, Inc. All Rights Reserved.
 
-/*=============================================================================
-	SEditableComboBox.h: Declares the SEditableComboBox class.
-=============================================================================*/
-
 #pragma once
 
 
@@ -32,7 +28,7 @@ public:
 		: _ButtonStyle( &FCoreStyle::Get().GetWidgetStyle< FButtonStyle >( "Button" ) )
 		, _Content()
 		, _ContentPadding(FMargin(4.0, 2.0))
-		, _InitiallySelectedItem(NULL)
+		, _InitiallySelectedItem(nullptr)
 		, _MaxListHeight(450.0f)
 		, _IsRenameVisible(EVisibility::Visible)
 		, _OnGetEditableText()
@@ -79,9 +75,7 @@ public:
 
 public:
 
-	/**
-	 * Clears the combo box selection.
-	 */
+	/** Clears the combo box selection. */
 	void ClearSelection( )
 	{
 		ComboBox->ClearSelection();
@@ -90,7 +84,7 @@ public:
 	/**
 	 * Constructs the widget.
 	 *
-	 * @param InArgs - The Slate argument list.
+	 * @param InArgs The Slate argument list.
 	 */
 	void Construct( const FArguments& InArgs )
 	{
@@ -235,7 +229,7 @@ public:
 	/**
 	 * Gets the item that is currently selected in the combo box.
 	 *
-	 * @return The selected item, or NULL if no item is selected.
+	 * @return The selected item, or nullptr if no item is selected.
 	 */
 	OptionType GetSelectedItem( ) 
 	{
@@ -259,7 +253,7 @@ public:
 	/**
 	 * Sets the item that is selected in the combo box.
 	 *
-	 * @param InSelectedItem - The item to select.
+	 * @param InSelectedItem The item to select.
 	 */
 	void SetSelectedItem( OptionType InSelectedItem )
 	{
@@ -268,7 +262,6 @@ public:
 		ComboBox->SetSelectedItem(InSelectedItem);
 	}
 
-
 private:
 
 	// Callback for getting the enabled state of the 'Add' button.
@@ -328,7 +321,6 @@ private:
 		}
 	}
 
-
 private:
 
 	// Holds the combo box.
@@ -340,7 +332,6 @@ private:
 	// Holds the text box.
 	TSharedPtr<SEditableTextBox> TextBox;
 
-
 private:
 
 	// Holds a delegate to be invoked before the editable text box is populated and shown.
diff --git a/Engine/Source/Runtime/Slate/Public/Widgets/Input/SEditableText.h b/Engine/Source/Runtime/Slate/Public/Widgets/Input/SEditableText.h
--- a/Engine/Source/Runtime/Slate/Public/Widgets/Input/SEditableText.h
+++ b/Engine/Source/Runtime/Slate/Public/Widgets/Input/SEditableText.h
@@ -454,7 +454,7 @@ private:
 
 private:
 
-	/** A pointer to the editable text currently under the mouse cursor. NULL when there isn't one. */
+	/** A pointer to the editable text currently under the mouse cursor. nullptr when there isn't one. */
 	static SEditableText* EditableTextUnderCursor;
 
 	/** The text content for this editable text widget */
diff --git a/Engine/Source/Runtime/Slate/Public/Widgets/Input/SExpandableButton.h b/Engine/Source/Runtime/Slate/Public/Widgets/Input/SExpandableButton.h
--- a/Engine/Source/Runtime/Slate/Public/Widgets/Input/SExpandableButton.h
+++ b/Engine/Source/Runtime/Slate/Public/Widgets/Input/SExpandableButton.h
@@ -2,12 +2,15 @@
 
 #pragma once
 
+
 /**
  * A button that can either be collapsed or expanded, containing different content in each state.
  */
-class SLATE_API SExpandableButton : public SBorder
+class SLATE_API SExpandableButton
+	: public SBorder
 {
 public:
+
 	SLATE_BEGIN_ARGS( SExpandableButton )
 		: _IsExpanded( true )
 		{}
@@ -46,9 +49,13 @@ public:
 	void Construct(const FArguments& InArgs);
 
 protected:
-	/** The attribute of the current expansion state */
-	TAttribute<bool> IsExpanded;
+
 	/** Callbacks to determine visibility of parts that should be shown when the button state is collapsed or expanded */
 	EVisibility GetCollapsedVisibility() const;
 	EVisibility GetExpandedVisibility() const;
+
+protected:
+
+	/** The attribute of the current expansion state */
+	TAttribute<bool> IsExpanded;
 };
diff --git a/Engine/Source/Runtime/Slate/Public/Widgets/Input/SHyperlink.h b/Engine/Source/Runtime/Slate/Public/Widgets/Input/SHyperlink.h
--- a/Engine/Source/Runtime/Slate/Public/Widgets/Input/SHyperlink.h
+++ b/Engine/Source/Runtime/Slate/Public/Widgets/Input/SHyperlink.h
@@ -2,19 +2,21 @@
 
 #pragma once
 
+
 /**
  * A hyperlink widget is what you would expect from a browser hyperlink.
  * When a hyperlink is clicked in invokes an OnNavigate() delegate.
  */
-class SHyperlink : public SButton
+class SHyperlink
+	: public SButton
 {
 public:
 
 	SLATE_BEGIN_ARGS(SHyperlink)
 		: _Text()
 		, _Style(&FCoreStyle::Get().GetWidgetStyle< FHyperlinkStyle >("Hyperlink"))
-		, _TextStyle(NULL)
-		, _UnderlineStyle(NULL)
+		, _TextStyle(nullptr)
+		, _UnderlineStyle(nullptr)
 		, _Padding()
 		, _OnNavigate()
 		{}
@@ -37,8 +39,8 @@ public:
 		this->OnNavigate = InArgs._OnNavigate;
 
 		check (InArgs._Style);
-		const FButtonStyle* UnderlineStyle = InArgs._UnderlineStyle != NULL ? InArgs._UnderlineStyle : &InArgs._Style->UnderlineStyle;
-		const FTextBlockStyle* TextStyle = InArgs._TextStyle != NULL ? InArgs._TextStyle : &InArgs._Style->TextStyle;
+		const FButtonStyle* UnderlineStyle = InArgs._UnderlineStyle != nullptr ? InArgs._UnderlineStyle : &InArgs._Style->UnderlineStyle;
+		const FTextBlockStyle* TextStyle = InArgs._TextStyle != nullptr ? InArgs._TextStyle : &InArgs._Style->TextStyle;
 		TAttribute<FMargin> Padding = InArgs._Padding.IsSet() ? InArgs._Padding : InArgs._Style->Padding;
 
 		SButton::Construct(
@@ -52,12 +54,14 @@ public:
 		);
 	}
 
-	// SWidget implementation
+public:
+
+	// SWidget overrides
+
 	virtual FCursorReply OnCursorQuery( const FGeometry& MyGeometry, const FPointerEvent& CursorEvent ) const override
 	{
 		return FCursorReply::Cursor( EMouseCursor::Hand );
 	}
-	// End SWidget implementation
 
 protected:
 
@@ -71,4 +75,4 @@ protected:
 
 	/** The delegate to invoke when someone clicks the hyperlink */
 	FSimpleDelegate OnNavigate;
-};
\ No newline at end of file
+};
diff --git a/Engine/Source/Runtime/Slate/Public/Widgets/Input/SSearchBox.h b/Engine/Source/Runtime/Slate/Public/Widgets/Input/SSearchBox.h
--- a/Engine/Source/Runtime/Slate/Public/Widgets/Input/SSearchBox.h
+++ b/Engine/Source/Runtime/Slate/Public/Widgets/Input/SSearchBox.h
@@ -4,8 +4,10 @@
 
 #define LOCTEXT_NAMESPACE "Slate"
 
+
 /** A text box that is used for searching. Meant to be as easy to use as possible with as few options as possible. */
-class SLATE_API SSearchBox : public SEditableTextBox
+class SLATE_API SSearchBox
+	: public SEditableTextBox
 {
 
 public:
@@ -27,7 +29,7 @@ public:
 		, _OnSearch()
 		, _SelectAllTextWhenFocused( true )
 		, _DelayChangeNotificationsWhileTyping( false )
-		{}
+	{ }
 
 		/** Style used to draw this search box */
 		SLATE_STYLE_ARGUMENT( FSearchBoxStyle, Style )
@@ -59,14 +61,19 @@ public:
 	SLATE_END_ARGS()
 
 	void Construct( const FArguments& InArgs );
+
+public:
+
+	// SWidget overrides
+
 	virtual void Tick( const FGeometry& AllottedGeometry, const double InCurrentTime, const float InDeltaTime ) override;
 
 private:
 
-	/** Handler for when text in the editable text box changed */
+	/** Callback for changes in the editable text box. */
 	void HandleTextChanged(const FText& NewText);
 
-	/** Handler for when text in the editable text box is committed */
+	/** Callback for committing changes in the editable text box. */
 	void HandleTextCommitted(const FText& NewText, ETextCommit::Type CommitType);
 
 private:
@@ -108,4 +115,5 @@ private:
 	FText LastPendingTextChangedValue;
 };
 
+
 #undef LOCTEXT_NAMESPACE
diff --git a/Engine/Source/Runtime/Slate/Public/Widgets/Input/SSlider.h b/Engine/Source/Runtime/Slate/Public/Widgets/Input/SSlider.h
--- a/Engine/Source/Runtime/Slate/Public/Widgets/Input/SSlider.h
+++ b/Engine/Source/Runtime/Slate/Public/Widgets/Input/SSlider.h
@@ -59,7 +59,7 @@ public:
 	/**
 	 * Construct the widget.
 	 * 
-	 * @param InDeclaration - A declaration from which to construct the widget.
+	 * @param InDeclaration A declaration from which to construct the widget.
 	 */
 	void Construct( const SSlider::FArguments& InDeclaration );
 
@@ -86,14 +86,12 @@ public:
 
 public:
 
-	virtual int32 OnPaint( const FPaintArgs& Args, const FGeometry& AllottedGeometry, const FSlateRect& MyClippingRect, FSlateWindowElementList& OutDrawElements, int32 LayerId, const FWidgetStyle& InWidgetStyle, bool bParentEnabled ) const override;
+	// SWidget overrides
 
+	virtual int32 OnPaint( const FPaintArgs& Args, const FGeometry& AllottedGeometry, const FSlateRect& MyClippingRect, FSlateWindowElementList& OutDrawElements, int32 LayerId, const FWidgetStyle& InWidgetStyle, bool bParentEnabled ) const override;
 	virtual FVector2D ComputeDesiredSize() const;
-
 	virtual FReply OnMouseButtonDown( const FGeometry& MyGeometry, const FPointerEvent& MouseEvent ) override;
-	
 	virtual FReply OnMouseButtonUp( const FGeometry& MyGeometry, const FPointerEvent& MouseEvent ) override;
-	
 	virtual FReply OnMouseMove( const FGeometry& MyGeometry, const FPointerEvent& MouseEvent ) override;
 
 protected:
@@ -101,16 +99,15 @@ protected:
 	/**
 	 * Commits the specified slider value.
 	 *
-	 * @param NewValue - The value to commit.
+	 * @param NewValue The value to commit.
 	 */
 	void CommitValue(float NewValue);
 
 	/**
 	 * Calculates the new value based on the given absolute coordinates.
 	 *
-	 * @param MyGeometry - The slider's geometry.
-	 * @param AbsolutePosition - The absolute position of the slider.
-	 *
+	 * @param MyGeometry The slider's geometry.
+	 * @param AbsolutePosition The absolute position of the slider.
 	 * @return The new value.
 	 */
 	float PositionToValue( const FGeometry& MyGeometry, const FVector2D& AbsolutePosition );
diff --git a/Engine/Source/Runtime/Slate/Public/Widgets/Input/SSpinBox.h b/Engine/Source/Runtime/Slate/Public/Widgets/Input/SSpinBox.h
--- a/Engine/Source/Runtime/Slate/Public/Widgets/Input/SSpinBox.h
+++ b/Engine/Source/Runtime/Slate/Public/Widgets/Input/SSpinBox.h
@@ -2,6 +2,7 @@
 
 #pragma once
 
+
 /**
  * A Slate SpinBox resembles traditional spin boxes in that it is a widget that provides
  * keyboard-based and mouse-based manipulation of a numeric value.
@@ -9,7 +10,8 @@
  * Keyboard-based manipulation: click on the spinbox to enter text mode.
  */
 template<typename NumericType>
-class SSpinBox : public SCompoundWidget
+class SSpinBox
+	: public SCompoundWidget
 {
 public:
 	/** Notification for numeric value change */
@@ -226,7 +228,6 @@ public:
 	 *
 	 * @param MyGeometry The Geometry of the widget receiving the event
 	 * @param MouseEvent Information about the input event
-	 *
 	 * @return Whether the event was handled along with possible requests for the system to take action.
 	 */
 	virtual FReply OnMouseButtonDown( const FGeometry& MyGeometry, const FPointerEvent& MouseEvent ) override
@@ -249,7 +250,6 @@ public:
 	 *
 	 * @param MyGeometry The Geometry of the widget receiving the event
 	 * @param MouseEvent Information about the input event
-	 *
 	 * @return Whether the event was handled along with possible requests for the system to take action.
 	 */
 	virtual FReply OnMouseButtonUp( const FGeometry& MyGeometry, const FPointerEvent& MouseEvent ) override
@@ -282,7 +282,6 @@ public:
 	 *
 	 * @param MyGeometry The Geometry of the widget receiving the event
 	 * @param MouseEvent Information about the input event
-	 *
 	 * @return Whether the event was handled along with possible requests for the system to take action.
 	 */
 	virtual FReply OnMouseMove( const FGeometry& MyGeometry, const FPointerEvent& MouseEvent ) override
diff --git a/Engine/Source/Runtime/Slate/Public/Widgets/Input/SSuggestionTextBox.h b/Engine/Source/Runtime/Slate/Public/Widgets/Input/SSuggestionTextBox.h
--- a/Engine/Source/Runtime/Slate/Public/Widgets/Input/SSuggestionTextBox.h
+++ b/Engine/Source/Runtime/Slate/Public/Widgets/Input/SSuggestionTextBox.h
@@ -1,9 +1,5 @@
 // Copyright 1998-2014 Epic Games, Inc. All Rights Reserved.
 
-/*=============================================================================
-	SSuggestionTextBox.h: Declares the SSuggestionTextBox class.
-=============================================================================*/
-
 #pragma once
 
 
@@ -49,125 +45,81 @@ public:
 		, _SuggestionTextStyle( &FCoreStyle::Get().GetWidgetStyle<FTextBlockStyle>("SuggestionTextBox.Text") )
 		, _SuggestionListMaxHeight(250.0)
 		, _Text()
-		{ }
+	{ }
 
-		/**
-		 * The color of the background/border around the editable text.
-		 */
+		/** The color of the background/border around the editable text. */
 		SLATE_ATTRIBUTE(FSlateColor, BackgroundColor)
 
-		/**
-		 * The image of the background/border around the editable text.
-		 */
+		/** The image of the background/border around the editable text. */
 		SLATE_ARGUMENT(const FSlateBrush*, BackgroundImage)
 
-		/**
-		 * Whether to clear keyboard focus when pressing enter to commit changes.
-		 */
+		/** Whether to clear keyboard focus when pressing enter to commit changes. */
 		SLATE_ATTRIBUTE(bool, ClearKeyboardFocusOnCommit)
 
-		/**
-		 * Provides an alternative mechanism for error reporting.
-		 */
+		/** Provides an alternative mechanism for error reporting. */
 		SLATE_ARGUMENT(TSharedPtr<class IErrorReportingWidget>, ErrorReporting)
 
-		/**
-		 * Font color and opacity.
-		 */
+		/** Font color and opacity. */
 		SLATE_ATTRIBUTE(FSlateFontInfo, Font)
 
-		/**
-		 * Text color and opacity.
-		 */
+		/** Text color and opacity. */
 		SLATE_ATTRIBUTE(FSlateColor, ForegroundColor)
 
-		/**
-		 * Hint text that appears when there is no text in the text box.
-		 */
+		/** Hint text that appears when there is no text in the text box. */
 		SLATE_ATTRIBUTE(FText, HintText)
 
-		/**
-		 * Workaround as we loose focus when the auto completion closes.
-		 */
+		/** Workaround as we loose focus when the auto completion closes. */
 		SLATE_ATTRIBUTE(bool, IsCaretMovedWhenGainFocus)
 
-		/**
-		 * Minimum width that a text block should be.
-		 */
+		/** Minimum width that a text block should be. */
 		SLATE_ATTRIBUTE(float, MinDesiredWidth)
 
-		/**
-		 * Called before the history list is shown.
-		 */
+		/** Called before the history list is shown. */
 		SLATE_EVENT(FOnShowingHistory, OnShowingHistory)
 
-		/**
-		 * Called before the suggestion list is shown.
-		 */
+		/** Called before the suggestion list is shown. */
 		SLATE_EVENT(FOnShowingSuggestions, OnShowingSuggestions)
 
-		/**
-		 * Called whenever the text is changed interactively by the user.
-		 */
+		/** Called whenever the text is changed interactively by the user. */
 		SLATE_EVENT(FOnTextChanged, OnTextChanged)
 
-		/**
-		 * Called when the text has been committed.
-		 */
+		/** Called when the text has been committed. */
 		SLATE_EVENT(FOnTextCommitted, OnTextCommitted)
 
-		/**
-		 * Whether to allow the user to back out of changes when they press the escape key.
-		 */
+		/** Whether to allow the user to back out of changes when they press the escape key. */
 		SLATE_ATTRIBUTE(bool, RevertTextOnEscape)
 
-		/**
-		 * Whether to select all text when pressing enter to commit changes.
-		 */
+		/** Whether to select all text when pressing enter to commit changes. */
 		SLATE_ATTRIBUTE(bool, SelectAllTextOnCommit)
 
-		/**
-		 * Whether to select all text when the user clicks to give focus on the widget.
-		 */
+		/** Whether to select all text when the user clicks to give focus on the widget. */
 		SLATE_ATTRIBUTE(bool, SelectAllTextWhenFocused)
 
-		/**
-		 * The styling of the text box.
-		 */
+		/** The styling of the text box. */
 		SLATE_STYLE_ARGUMENT(FEditableTextBoxStyle, TextStyle)
 
-		/**
-		 * The styling of the suggestion text.
-		 */
+		/** The styling of the suggestion text. */
 		SLATE_STYLE_ARGUMENT(FTextBlockStyle, SuggestionTextStyle)
 
-		/**
-		 * The maximum height of the suggestion list.
-		 */
+		/** The maximum height of the suggestion list. */
 		SLATE_ATTRIBUTE(float, SuggestionListMaxHeight)
 
-		/**
-		 * Sets the text content for this editable text box widget.
-		 */
+		/** Sets the text content for this editable text box widget. */
 		SLATE_ATTRIBUTE(FText, Text)
 
 	SLATE_END_ARGS()
 
-
 public:
 
-	/**
-	 * Default constructor.
-	 */
+	/** Default constructor. */
 	SSuggestionTextBox( );
 
-
 public:
 
 	/**
 	 * Construct this widget.  Called by the SNew() Slate macro.
 	 *
-	 * @param InArgs - Declaration used by the SNew() macro to construct this widget.
+	 * @param InArgs Declaration used by the SNew() macro to construct this widget.
 	 */
 	void Construct( const FArguments& InArgs );
 
@@ -187,7 +139,7 @@ public:
 	 * If InError is a non-empty string the TextBox will the ErrorReporting provided during construction.
 	 * If no error reporting was provided, the TextBox will create a default error reporter.
 	 *
-	 * @param InError - The error string to set. 
+	 * @param InError The error string to set. 
 	 */
 	void SetError( const FString& InError )
 	{
@@ -197,21 +149,16 @@ public:
 	/**
 	 * Sets the text string currently being edited 
 	 *
-	 * @param InNewText - The new text string.
+	 * @param InNewText The new text string.
 	 */
 	void SetText( const TAttribute< FText >& InNewText );
 
-
 protected:
 	
-	/**
-	 * Clears the list of suggestions and hides the suggestions list.
-	 */
+	/** Clears the list of suggestions and hides the suggestions list. */
 	void ClearSuggestions( );
 
-	/**
-	 * Sets the keyboard focus to the text box.
-	 */
+	/** Sets the keyboard focus to the text box. */
 	void FocusTextBox( );
 
 	/**
@@ -221,29 +168,25 @@ protected:
 	 */
 	FString GetSelectedSuggestionString( ) const;
 
-	/**
-	 * Highlights the selected suggestion in the suggestion list.
-	 */
+	/** Highlights the selected suggestion in the suggestion list. */
 	void MarkActiveSuggestion( );
 
 	/**
 	 * Sets the list of suggestions.
 	 *
-	 * @param SuggestionStrings - The suggestion strings.
-	 * @param InHistoryMode - Whether the suggestions represent the input history.
+	 * @param SuggestionStrings The suggestion strings.
+	 * @param InHistoryMode Whether the suggestions represent the input history.
 	 */
 	void SetSuggestions( TArray<FString>& SuggestionStrings, bool InHistoryMode );
 
-
 protected:
 
-	virtual void OnKeyboardFocusLost( const FKeyboardFocusEvent& InKeyboardFocusEvent ) override;
+	// SWidget overrides
 
+	virtual void OnKeyboardFocusLost( const FKeyboardFocusEvent& InKeyboardFocusEvent ) override;
 	virtual FReply OnKeyDown( const FGeometry& MyGeometry, const FKeyboardEvent& KeyboardEvent ) override;
-
 	virtual bool SupportsKeyboardFocus( ) const override;
 
-
 private:
 
 	// Callback for generating a row widget for the suggestion list view.
@@ -261,7 +204,6 @@ private:
 	// Callback committing the text box input.
 	void HandleTextBoxTextCommitted( const FText& InText, ETextCommit::Type CommitInfo );
 
-
 private:
 
 	// Holds a flag to prevent recursive calls in UI callbacks.
@@ -285,7 +227,6 @@ private:
 	// The styling of the suggestion text.
 	const FTextBlockStyle* SuggestionTextStyle;
 
-
 private:
 
 	// Holds a delegate that is executed before the input history is shown.
@@ -299,4 +240,4 @@ private:
 
 	// Holds a delegate that is executed when the text has been committed.
 	FOnTextCommitted OnTextCommitted;
-};
\ No newline at end of file
+};
diff --git a/Engine/Source/Runtime/Slate/Public/Widgets/Layout/SEnableBox.h b/Engine/Source/Runtime/Slate/Public/Widgets/Layout/SEnableBox.h
--- a/Engine/Source/Runtime/Slate/Public/Widgets/Layout/SEnableBox.h
+++ b/Engine/Source/Runtime/Slate/Public/Widgets/Layout/SEnableBox.h
@@ -2,6 +2,7 @@
 
 #pragma once
 
+
 /**
  * A SEnableBox contains a widget that is lied to about whether the parent hierarchy is enabled or not, being told the parent is always enabled
  */
@@ -18,10 +19,12 @@ public:
 		SBox::Construct(SBox::FArguments().Content()[InArgs._Content.Widget]);
 	}
 
+public:
+
 	// SWidget interface
+	
 	virtual int32 OnPaint(const FPaintArgs& Args, const FGeometry& AllottedGeometry, const FSlateRect& MyClippingRect, FSlateWindowElementList& OutDrawElements, int32 LayerId, const FWidgetStyle& InWidgetStyle, bool bParentEnabled) const override
 	{
 		return SBox::OnPaint(Args, AllottedGeometry, MyClippingRect, OutDrawElements, LayerId, InWidgetStyle, /*bParentEnabled=*/ true);
 	}
-	// End of SWidget interface
 };
