Slate: Removed file header comments

diff --git a/Engine/Source/Runtime/Slate/Public/Widgets/Layout/SExpandableArea.h b/Engine/Source/Runtime/Slate/Public/Widgets/Layout/SExpandableArea.h
--- a/Engine/Source/Runtime/Slate/Public/Widgets/Layout/SExpandableArea.h
+++ b/Engine/Source/Runtime/Slate/Public/Widgets/Layout/SExpandableArea.h
@@ -1,9 +1,5 @@
 // Copyright 1998-2014 Epic Games, Inc. All Rights Reserved.
 
-/*=============================================================================
-	SExpandableArea.h: Declares the SExpandableArea class.
-=============================================================================*/
-
 #pragma once
 
 
@@ -51,6 +47,12 @@ public:
 
 
 public:
+
+	/**
+	 * Constructs a new widget.
+	 *
+	 * @param InArgs Construction arguments.
+	 */
 	void Construct( const FArguments& InArgs );
 
 	/**
@@ -64,6 +66,7 @@ public:
 	virtual void SetExpanded( bool bExpanded );
 
 protected:
+
 	/**
 	 * Constructs the header area widget
 	 * 
@@ -79,9 +82,7 @@ protected:
 
 	FReply OnHeaderClicked();
 
-	/**
-	 * Toggles selection visibility when the panel is clicked                   
-	 */
+	/** Toggles selection visibility when the panel is clicked. */
 	virtual void OnToggleContentVisibility();
 
 	/**
@@ -99,8 +100,8 @@ protected:
 	 */
 	virtual FVector2D ComputeDesiredSize() const override;
 
-
 protected:
+
 	/** Border widget for the header area */
 	TSharedPtr<SBorder> HeaderBorder;
 
@@ -118,6 +119,7 @@ protected:
 
 	/** Image to use when the area is collapsed */
 	const FSlateBrush* CollapsedImage;
+
 	/** Image to use when the area is expanded */
 	const FSlateBrush* ExpandedImage;
-};
\ No newline at end of file
+};
diff --git a/Engine/Source/Runtime/Slate/Public/Widgets/Layout/SGridPanel.h b/Engine/Source/Runtime/Slate/Public/Widgets/Layout/SGridPanel.h
--- a/Engine/Source/Runtime/Slate/Public/Widgets/Layout/SGridPanel.h
+++ b/Engine/Source/Runtime/Slate/Public/Widgets/Layout/SGridPanel.h
@@ -2,7 +2,9 @@
 
 #pragma once
 
-class SLATE_API SGridPanel : public SPanel
+
+class SLATE_API SGridPanel
+	: public SPanel
 {
 public:
 	// Used by the mandatory named parameter in FSlot
@@ -88,7 +90,6 @@ public:
 				NudgeParam = Nudge;
 				return *this;
 			}
-
 	};
 
 	/**
@@ -149,15 +150,6 @@ public:
 	
 	void Construct( const FArguments& InArgs );
 
-	// SWidget interface
-	virtual int32 OnPaint( const FPaintArgs& Args, const FGeometry& AllottedGeometry, const FSlateRect& MyClippingRect, FSlateWindowElementList& OutDrawElements, int32 LayerId, const FWidgetStyle& InWidgetStyle, bool bParentEnabled ) const override;
-	virtual void OnArrangeChildren( const FGeometry& AllottedGeometry, FArrangedChildren& ArrangedChildren ) const override;
-	virtual void CacheDesiredSize() override;
-	virtual FVector2D ComputeDesiredSize() const override;
-	virtual FChildren* GetChildren() override;
-	// End of SWidget interface
-
-	
 	/**
 	 * GetDesiredSize of a subregion in the graph.
 	 *
@@ -174,6 +166,16 @@ public:
 	/** Specify a row to stretch instead of sizing to content. */
 	void SetRowFill( int32 RowId, float Coefficient );
 
+public:
+
+	// SWidget interface
+
+	virtual int32 OnPaint( const FPaintArgs& Args, const FGeometry& AllottedGeometry, const FSlateRect& MyClippingRect, FSlateWindowElementList& OutDrawElements, int32 LayerId, const FWidgetStyle& InWidgetStyle, bool bParentEnabled ) const override;
+	virtual void OnArrangeChildren( const FGeometry& AllottedGeometry, FArrangedChildren& ArrangedChildren ) const override;
+	virtual void CacheDesiredSize() override;
+	virtual FVector2D ComputeDesiredSize() const override;
+	virtual FChildren* GetChildren() override;
+
 private:
 
 	/**
@@ -188,11 +190,10 @@ private:
 	static void DistributeSizeContributions( float SizeContribution, TArray<float>& DistributeOverMe, int32 StartIndex, int32 UpperBound );
 
 	/**
-	 * Inser the slot InSlot into the list of Slots based on its LayerParam, such that Slots are sorter by layer.
+	 * Inserts the given slot into the list of Slots based on its LayerParam, such that Slots are sorter by layer.
 	 *
 	 * @param The newly-allocated slot to insert.
-	 *
-	 * @param A reference to the added slot
+	 * @return A reference to the added slot
 	 */
 	FSlot& InsertSlot( FSlot* InSlot );	
 
@@ -209,6 +210,8 @@ private:
 	 */
 	void NotifySlotChanged( FSlot* InSlot );
 
+private:
+
 	/** The slots that are placed into various grid locations */
 	TPanelChildren<FSlot> Slots;
 
diff --git a/Engine/Source/Runtime/Slate/Public/Widgets/Layout/SMissingWidget.h b/Engine/Source/Runtime/Slate/Public/Widgets/Layout/SMissingWidget.h
--- a/Engine/Source/Runtime/Slate/Public/Widgets/Layout/SMissingWidget.h
+++ b/Engine/Source/Runtime/Slate/Public/Widgets/Layout/SMissingWidget.h
@@ -1,9 +1,5 @@
 // Copyright 1998-2014 Epic Games, Inc. All Rights Reserved.
 
-/*=============================================================================
-	SMissingWidget.h: Declares the SMissingWidget class.
-=============================================================================*/
-
 #pragma once
 
 
diff --git a/Engine/Source/Runtime/Slate/Public/Widgets/Layout/SWidgetSwitcher.h b/Engine/Source/Runtime/Slate/Public/Widgets/Layout/SWidgetSwitcher.h
--- a/Engine/Source/Runtime/Slate/Public/Widgets/Layout/SWidgetSwitcher.h
+++ b/Engine/Source/Runtime/Slate/Public/Widgets/Layout/SWidgetSwitcher.h
@@ -1,9 +1,5 @@
 // Copyright 1998-2014 Epic Games, Inc. All Rights Reserved.
 
-/*=============================================================================
-	SWidgetSwitcher.h: Declares the SWidgetSwitcher class.
-=============================================================================*/
-
 #pragma once
 
 
@@ -17,13 +13,17 @@ class SLATE_API SWidgetSwitcher
 {
 public:
 
-	class FSlot : public TSlotBase<FSlot>, public TSupportsContentAlignmentMixin<FSlot>, public TSupportsContentPaddingMixin<FSlot>
+	class FSlot
+		: public TSlotBase<FSlot>
+		, public TSupportsContentAlignmentMixin<FSlot>
+		, public TSupportsContentPaddingMixin<FSlot>
 	{
 	public:
+
 		FSlot()
 			: TSlotBase<FSlot>()
 			,TSupportsContentAlignmentMixin<FSlot>( HAlign_Fill, VAlign_Fill )
-		{}
+		{ }
 	};
 
 	SLATE_BEGIN_ARGS(SWidgetSwitcher)
@@ -46,17 +46,10 @@ public:
 	/**
 	 * Adds a slot to the widget switcher at the specified location.
 	 *
-	 * @param SlotIndex - The index at which to insert the slot, or INDEX_NONE to append.
+	 * @param SlotIndex The index at which to insert the slot, or INDEX_NONE to append.
 	 */
 	FSlot& AddSlot( int32 SlotIndex = INDEX_NONE );
 
-	/**
-	 * Removes a slot with the corresponding widget in it.  Returns the index where the widget was found, otherwise -1.
-	 *
-	 * @param Widget - The widget to find and remove.
-	 */
-	int32 RemoveSlot( TSharedRef<SWidget> WidgetToRemove );
-
 	/**
 	 * Constructs the widget.
 	 */
@@ -81,6 +74,8 @@ public:
 
 	/**
 	 * Gets the number of widgets that this switcher manages.
+	 *
+	 * @return Number of widgets.
 	 */
 	int32 GetNumWidgets( ) const
 	{
@@ -90,25 +85,30 @@ public:
 	/**
 	 * Gets the widget in the specified slot.
 	 *
-	 * @param SlotIndex - The slot index of the widget to get.
-	 *
-	 * @return The widget, or NULL if the slot does not exist.
+	 * @param SlotIndex The slot index of the widget to get.
+	 * @return The widget, or nullptr if the slot does not exist.
 	 */
 	TSharedPtr<SWidget> GetWidget( int32 SlotIndex ) const;
 
 	/**
 	 * Gets the slot index of the specified widget.
 	 *
-	 * @param Widget - The widget to get the index for.
-	 *
+	 * @param Widget The widget to get the index for.
 	 * @return The slot index, or INDEX_NONE if the widget does not exist.
 	 */
 	int32 GetWidgetIndex( TSharedRef<SWidget> Widget ) const;
 
+	/**
+	 * Removes a slot with the corresponding widget in it.  Returns the index where the widget was found, otherwise -1.
+	 *
+	 * @param Widget The widget to find and remove.
+	 */
+	int32 RemoveSlot( TSharedRef<SWidget> WidgetToRemove );
+
 	/**
 	 * Sets the active widget.
 	 *
-	 * @param Widget - The widget to activate.
+	 * @param Widget The widget to activate.
 	 */
 	void SetActiveWidget( TSharedRef<SWidget> Widget )
 	{
@@ -118,7 +118,7 @@ public:
 	/**
 	 * Activates the widget at the specified index.
 	 *
-	 * @param Index - The slot index.
+	 * @param Index The slot index.
 	 */
 	void SetActiveWidgetIndex( int32 Index );
 
@@ -136,16 +136,12 @@ public:
 
 protected:
 
-	// Begin SCompoundWidget interface
+	// SCompoundWidget interface
 
 	virtual void OnArrangeChildren( const FGeometry& AllottedGeometry, FArrangedChildren& ArrangedChildren ) const override;
-	
 	virtual FVector2D ComputeDesiredSize( ) const override;
-
 	virtual FChildren* GetChildren( ) override;
 
-	// End SCompoundWidget interface
-
 private:
 
 	/** Required to implement GetChildren() in a way that can dynamically return the currently active child. */
@@ -154,7 +150,7 @@ private:
 	{
 	public:
 
-		FOneDynamicChild( TPanelChildren<FSlot>* InAllChildren = NULL, const TAttribute<int32>* InWidgetIndex = NULL )
+		FOneDynamicChild( TPanelChildren<FSlot>* InAllChildren = nullptr, const TAttribute<int32>* InWidgetIndex = nullptr )
 			: AllChildren( InAllChildren )
 			, WidgetIndex( InWidgetIndex )
 		{ }
diff --git a/Engine/Source/Runtime/Slate/Public/Widgets/Navigation/SWizard.h b/Engine/Source/Runtime/Slate/Public/Widgets/Navigation/SWizard.h
--- a/Engine/Source/Runtime/Slate/Public/Widgets/Navigation/SWizard.h
+++ b/Engine/Source/Runtime/Slate/Public/Widgets/Navigation/SWizard.h
@@ -1,9 +1,5 @@
 // Copyright 1998-2014 Epic Games, Inc. All Rights Reserved.
 
-/*=============================================================================
-	SWizard.h: Declares the SWizard class.
-=============================================================================*/
-
 #pragma once
 
 
@@ -53,13 +49,12 @@ public:
 
 		SLATE_END_ARGS()
 
-
 	public:
 
 		/**
 		 * Creates and initializes a new instance.
 		 *
-		 * @param InArgs - The arguments for the page.
+		 * @param InArgs The arguments for the page.
 		 */
 		FWizardPage( const FArguments& InArgs )
 			: ButtonContent(InArgs._ButtonContent)
@@ -69,13 +64,12 @@ public:
 			, PageContent(InArgs._PageContent)
 		{ }
 
-
 	public:
 
 		/**
 		 * Checks whether the page can be shown.
 		 *
-		 * @return bool - true if the page can be shown, false otherwise.
+		 * @return true if the page can be shown, false otherwise.
 		 */
 		bool CanShow( ) const
 		{
@@ -95,7 +89,7 @@ public:
 		/**
 		 * Gets the page content.
 		 *
-		 * @return TAlwaysValidWidget - The page contents.
+		 * @return The page contents.
 		 */
 		const TSharedRef<SWidget>& GetPageContent( ) const
 		{
@@ -105,7 +99,7 @@ public:
 		/**
 		 * Gets a delegate to be invoked when this page is being entered.
 		 *
-		 * @return FSimpleDelegate - The delegate.
+		 * @return The delegate.
 		 */
 		FSimpleDelegate& OnEnter( )
 		{
@@ -115,14 +109,13 @@ public:
 		/**
 		 * Gets a delegate to be invoked when this page is being left.
 		 *
-		 * @return FSimpleDelegate - The delegate.
+		 * @return The delegate.
 		 */
 		FSimpleDelegate& OnLeave( )
 		{
 			return OnLeaveDelegate;
 		}
 
-
 	private:
 
 		// Holds the button content.
@@ -141,7 +134,6 @@ public:
 		TAlwaysValidWidget PageContent;
 	};
 
-
 public:
 
 	SLATE_BEGIN_ARGS(SWizard)
@@ -194,22 +186,20 @@ public:
 
 	SLATE_END_ARGS()
 
-
 public:
 
 	/**
 	 * Checks whether the page with the specified index can be shown.
 	 *
-	 * @param PageIndex - The index of the page to check.
-	 *
-	 * @return bool - true if the page can be shown, false otherwise.
+	 * @param PageIndex The index of the page to check.
+	 * @return true if the page can be shown, false otherwise.
 	 */
 	bool CanShowPage( int32 PageIndex ) const;
 
 	/**
 	 * Construct this widget.
 	 *
-	 * @param InArgs - The declaration data for this widget.
+	 * @param InArgs The declaration data for this widget.
 	 */
 	void Construct( const FArguments& InArgs );
 
@@ -226,8 +216,7 @@ public:
 	/**
 	 * Gets the index of the specified wizard page widget.
 	 *
-	 * @param PageWidget - The page widget to get the index for.
-	 *
+	 * @param PageWidget The page widget to get the index for.
 	 * @return The index of the page, or INDEX_NONE if not found.
 	 */
 	int32 GetPageIndex( const TSharedRef<SWidget>& PageWidget ) const
@@ -238,27 +227,22 @@ public:
 	/**
 	 * Attempts to show the page with the specified index.
 	 *
-	 * @param PageIndex - The index of the page to show.
+	 * @param PageIndex The index of the page to show.
 	 */
 	void ShowPage( int32 PageIndex );
 
-
 public:
 
-	// Begin SCompoundWidget interface
+	// SCompoundWidget interface
 
 	virtual FVector2D ComputeDesiredSize( ) const override;
 
-	// End SCompoundWidget interface
-
-
 public:
 
 	/**
 	 * Returns a new slot for a page.
 	 *
-	 * @param PageName - The name of the page.
-	 *
+	 * @param PageName The name of the page.
 	 * @return A new slot.
 	 */
 	static FWizardPage::FArguments Page( )
@@ -267,7 +251,6 @@ public:
 		return Args;
 	}
 
-
 private:
 
 	// Callback for clicking the 'Cancel' button.
@@ -303,7 +286,6 @@ private:
 	// Callback for getting the visibility of the 'Previous' button.
 	EVisibility HandlePrevButtonVisibility( ) const;
 
-
 private:
 
 	// Holds the wizard's desired size.
@@ -315,7 +297,6 @@ private:
 	// Holds the widget switcher.
 	TSharedPtr<SWidgetSwitcher> WidgetSwitcher;
 
-
 private:
 
 	// Holds a delegate to be invoked when the 'Cancel' button has been clicked.
diff --git a/Engine/Source/Runtime/Slate/Public/Widgets/Notifications/SErrorHint.h b/Engine/Source/Runtime/Slate/Public/Widgets/Notifications/SErrorHint.h
--- a/Engine/Source/Runtime/Slate/Public/Widgets/Notifications/SErrorHint.h
+++ b/Engine/Source/Runtime/Slate/Public/Widgets/Notifications/SErrorHint.h
@@ -2,9 +2,12 @@
  
 #pragma once
 
-class SLATE_API SErrorHint : public SCompoundWidget, public IErrorReportingWidget
+class SLATE_API SErrorHint
+	: public SCompoundWidget
+	, public IErrorReportingWidget
 {
 public:
+
 	SLATE_BEGIN_ARGS( SErrorHint )
 		: _ErrorText()
 		{}
@@ -13,18 +16,17 @@ public:
 
 	void Construct(const FArguments& InArgs);
 
+public:
+
 	// IErrorReportingWidget interface
 
 	virtual void SetError( const FText& InErrorText ) override;
 	virtual void SetError( const FString& InErrorText ) override;
-
 	virtual bool HasError() const override;
-
 	virtual TSharedRef<SWidget> AsWidget() override;
 
-	// IErrorReportingWidget interface
-
 private:
+
 	TAttribute<EVisibility> CustomVisibility;
 	EVisibility MyVisibility() const;
 
@@ -35,5 +37,3 @@ private:
 	FText ErrorText;
 	FText GetErrorText() const;
 };
-
-
diff --git a/Engine/Source/Runtime/Slate/Public/Widgets/Notifications/SErrorText.h b/Engine/Source/Runtime/Slate/Public/Widgets/Notifications/SErrorText.h
--- a/Engine/Source/Runtime/Slate/Public/Widgets/Notifications/SErrorText.h
+++ b/Engine/Source/Runtime/Slate/Public/Widgets/Notifications/SErrorText.h
@@ -2,24 +2,36 @@
  
 #pragma once
 
+
+/**
+ * Interface for error reporting widgets.
+ */
 class SLATE_API IErrorReportingWidget
 {
 public:
+
 	virtual void SetError( const FString& InErrorText ) = 0;
 	virtual void SetError( const FText& InErrorText ) = 0;
 	virtual bool HasError() const = 0;
 	virtual TSharedRef<SWidget> AsWidget() = 0;
 };
 
-class SLATE_API SErrorText : public SBorder, public IErrorReportingWidget
+
+/**
+ * Implements a widget that displays an error text message.
+ */
+class SLATE_API SErrorText
+	: public SBorder
+	, public IErrorReportingWidget
 {
 public:
+
 	SLATE_BEGIN_ARGS( SErrorText )
 		: _ErrorText()
 		, _BackgroundColor(FCoreStyle::Get().GetColor("ErrorReporting.BackgroundColor"))
 		, _Font()
 		, _AutoWrapText(false)
-		{}
+	{ }
 
 		SLATE_TEXT_ARGUMENT(ErrorText)
 		SLATE_ATTRIBUTE(FSlateColor, BackgroundColor)
@@ -30,18 +42,17 @@ public:
 
 	void Construct(const FArguments& InArgs);
 
+public:
+
 	// IErrorReportingWidget interface
 
 	virtual void SetError( const FText& InErrorText ) override;
 	virtual void SetError( const FString& InErrorText ) override;
-
 	virtual bool HasError() const override;
-
 	virtual TSharedRef<SWidget> AsWidget() override;
 
-	// IErrorReportingWidget interface
-
 private:
+
 	TAttribute< FSlateFontInfo > Font;
 
 	TAttribute<EVisibility> CustomVisibility;
@@ -52,5 +63,3 @@ private:
 	FVector2D GetDesiredSizeScale() const;
 	FCurveSequence ExpandAnimation;
 };
-
-
diff --git a/Engine/Source/Runtime/Slate/Public/Widgets/Notifications/SNotificationList.h b/Engine/Source/Runtime/Slate/Public/Widgets/Notifications/SNotificationList.h
--- a/Engine/Source/Runtime/Slate/Public/Widgets/Notifications/SNotificationList.h
+++ b/Engine/Source/Runtime/Slate/Public/Widgets/Notifications/SNotificationList.h
@@ -2,11 +2,13 @@
 
 #pragma once
 
+
 /**
  * Interface class for an item in the event message list.
  * Real implementation is found in SEventMessageItemImpl
  */
-class SLATE_API SNotificationItem : public SCompoundWidget
+class SLATE_API SNotificationItem
+	: public SCompoundWidget
 {
 public:
 	enum ECompletionState
@@ -39,9 +41,10 @@ public:
 	virtual void Fadeout() = 0;
 };
 
+
 /**
-	* Setup class to initialize buttons on a notification
-	*/
+ * Setup class to initialize buttons on a notification
+ */
 struct FNotificationButtonInfo
 {
 	FNotificationButtonInfo( const FText& InText, const FText& InToolTip, FSimpleDelegate InCallback, SNotificationItem::ECompletionState VisibleInState = SNotificationItem::CS_Pending )
@@ -52,7 +55,7 @@ struct FNotificationButtonInfo
 		, VisibilityOnPending( VisibleInState == SNotificationItem::CS_Pending ? EVisibility::Visible : EVisibility::Collapsed )
 		, VisibilityOnSuccess( VisibleInState == SNotificationItem::CS_Success ? EVisibility::Visible : EVisibility::Collapsed )
 		, VisibilityOnFail( VisibleInState == SNotificationItem::CS_Fail ? EVisibility::Visible : EVisibility::Collapsed )
-	{}
+	{ }
 
 	/** Message on the button */
 	FText Text;
@@ -71,6 +74,7 @@ struct FNotificationButtonInfo
 	EVisibility VisibilityOnFail;
 };
 
+
 /**
  * Setup class to initialize a notification.
  */
@@ -84,7 +88,7 @@ struct FNotificationInfo
 	FNotificationInfo( const FText& InText )
 		: Text(InText),
 		ButtonDetails(),
-		Image(NULL),
+		Image(nullptr),
 		FadeInDuration(0.5f),
 		FadeOutDuration(2.0f),
 		ExpireDuration(1.0f),
@@ -99,8 +103,7 @@ struct FNotificationInfo
 		Hyperlink(),
 		HyperlinkText( NSLOCTEXT("EditorNotification", "DefaultHyperlinkText", "Show Log") ),
 		bAllowThrottleWhenFrameRateIsLow(true)
-	{
-	};
+	{ };
 
 	/** The text displayed in this text block */
 	FText Text;
@@ -154,33 +157,38 @@ struct FNotificationInfo
 	bool bAllowThrottleWhenFrameRateIsLow;
 };
 
+
 /**
- * A list of non-intrusive messages about the status of currently active work
+ * A list of non-intrusive messages about the status of currently active work.
  */
-class SLATE_API SNotificationList : public SCompoundWidget
+class SLATE_API SNotificationList
+	: public SCompoundWidget
 {
 	friend class SNotificationExtendable;
 	friend class FSlateNotificationManager;
 
 public:
+
 	SLATE_BEGIN_ARGS( SNotificationList ){}
 
 	SLATE_END_ARGS()
 	
 	/**
-	 * Constructs this widget
+	 * Constructs this widget.
 	 *
-	 * @param InArgs    Declaration from which to construct the widget
+	 * @param InArgs    Declaration from which to construct the widget.
 	 */
 	void Construct( const FArguments& InArgs );
 
 	/**
-	 * Adds a floating notification
-	 *@param Info 		Contains various settings used to initialize the notification
+	 * Adds a floating notification.
+	 *
+	 * @param Info 		Contains various settings used to initialize the notification.
 	 */
 	virtual TSharedRef<SNotificationItem> AddNotification(const FNotificationInfo& Info);
 
 protected:
+
 	/**
 	 * Invoked when a notification item has finished fading out. Called by the faded item.
 	 *
@@ -189,12 +197,13 @@ protected:
 	virtual void NotificationItemFadedOut (const TSharedRef<SNotificationItem>& NotificationItem);
 
 protected:
-	/** The vertical box containing the list of message items */
+
+	/** The vertical box containing the list of message items. */
 	TSharedPtr<SVerticalBox> MessageItemBoxPtr;
 
-	/** The parent window of this list */
+	/** The parent window of this list. */
 	TWeakPtr<SWindow> ParentWindowPtr;
 
-	/** Flag to auto-destroy this list */
+	/** Flag to auto-destroy this list. */
 	bool bDone;
 };
diff --git a/Engine/Source/Runtime/Slate/Public/Widgets/SCanvas.h b/Engine/Source/Runtime/Slate/Public/Widgets/SCanvas.h
--- a/Engine/Source/Runtime/Slate/Public/Widgets/SCanvas.h
+++ b/Engine/Source/Runtime/Slate/Public/Widgets/SCanvas.h
@@ -1,9 +1,5 @@
 // Copyright 1998-2014 Epic Games, Inc. All Rights Reserved.
 
-/*=============================================================================
-	SCanvas.h: Declares the SCanvas class.
-=============================================================================*/
-
 #pragma once
 
 
@@ -138,18 +134,13 @@ public:
 
 public:
 
-	// Begin SWidget overrides
+	// SWidget overrides
 
 	virtual void OnArrangeChildren( const FGeometry& AllottedGeometry, FArrangedChildren& ArrangedChildren ) const override;
-	
 	virtual int32 OnPaint( const FPaintArgs& Args, const FGeometry& AllottedGeometry, const FSlateRect& MyClippingRect, FSlateWindowElementList& OutDrawElements, int32 LayerId, const FWidgetStyle& InWidgetStyle, bool bParentEnabled ) const override;
-
 	virtual FVector2D ComputeDesiredSize() const override;
-
 	virtual FChildren* GetChildren() override;
 
-	// End SWidget overrides
-
 protected:
 
 	/** The canvas widget's children. */
diff --git a/Engine/Source/Runtime/Slate/Public/Widgets/SToolTip.h b/Engine/Source/Runtime/Slate/Public/Widgets/SToolTip.h
--- a/Engine/Source/Runtime/Slate/Public/Widgets/SToolTip.h
+++ b/Engine/Source/Runtime/Slate/Public/Widgets/SToolTip.h
@@ -1,9 +1,5 @@
 // Copyright 1998-2014 Epic Games, Inc. All Rights Reserved.
 
-/*=============================================================================
-	SToolTip.h: Declares the SToolTip class.
-=============================================================================*/
-
 #pragma once
 
 
@@ -58,7 +54,7 @@ public:
 
 public:
 
-	// Begin IToolTip interface
+	// IToolTip interface
 
 	virtual TSharedRef<class SWidget> AsWidget( ) override
 	{
@@ -71,10 +67,9 @@ public:
 	}
 
 	virtual bool IsEmpty( ) const override;
-
 	virtual bool IsInteractive( ) const override;
 
-	// End IToolTip interface
+public:
 
 	static float GetToolTipWrapWidth();
 
diff --git a/Engine/Source/Runtime/Slate/Public/Widgets/SViewport.h b/Engine/Source/Runtime/Slate/Public/Widgets/SViewport.h
--- a/Engine/Source/Runtime/Slate/Public/Widgets/SViewport.h
+++ b/Engine/Source/Runtime/Slate/Public/Widgets/SViewport.h
@@ -1,9 +1,5 @@
 // Copyright 1998-2014 Epic Games, Inc. All Rights Reserved.
 
-/*=============================================================================
-	SViewport.h: Declares the SViewport class.
-=============================================================================*/
-
 #pragma once
 
 
@@ -22,8 +18,10 @@ public:
 		{}
 
 		SLATE_DEFAULT_SLOT( FArguments, Content )
+
 		/** Whether or not to show the disabled effect when this viewport is disabled */
 		SLATE_ATTRIBUTE( bool, ShowEffectWhenDisabled )
+		
 		/** 
 		 * Whether or not to render directly to the window's backbuffer or an offscreen render target that is applied to the window later 
 		 * Rendering to an offscreen target is the most common option in the editor where there may be many frames which this viewport's interface may wish to not re-render but use a cached buffer instead
@@ -54,9 +52,7 @@ public:
 
 	SLATE_END_ARGS()
 
-	/**
-	 * Default constructor.
-	 */
+	/** Default constructor. */
 	SViewport( );
 
 	/**
@@ -73,6 +69,8 @@ public:
 
 	/** 
 	 * Computes the ideal size necessary to display this widget.
+	 *
+	 * @return The desired width and height.
 	 */
 	virtual FVector2D ComputeDesiredSize() const override
 	{
@@ -82,7 +80,7 @@ public:
 	/**
 	 * Sets the interface to be used by this viewport for rendering and I/O
 	 *
-	 * @param InViewportInterface	The interface to use
+	 * @param InViewportInterface The interface to use
 	 */
 	void SetViewportInterface( TSharedRef<ISlateViewport> InViewportInterface )
 	{
@@ -160,8 +158,6 @@ public:
 	virtual FReply OnControllerAnalogValueChanged( const FGeometry& MyGeometry, const FControllerEvent& ControllerEvent ) override;
 	virtual FReply OnMotionDetected( const FGeometry& MyGeometry, const FMotionEvent& MotionEvent ) override;
 
-	// End of SWidget interface
-
 private:
 
 	FSimpleSlot& Decl_GetContent()
diff --git a/Engine/Source/Runtime/Slate/Public/Widgets/SWeakWidget.h b/Engine/Source/Runtime/Slate/Public/Widgets/SWeakWidget.h
--- a/Engine/Source/Runtime/Slate/Public/Widgets/SWeakWidget.h
+++ b/Engine/Source/Runtime/Slate/Public/Widgets/SWeakWidget.h
@@ -23,6 +23,7 @@ public:
 	SLATE_END_ARGS()
 
 public:
+
 	SWeakWidget();
 
 	void Construct(const FArguments& InArgs);
@@ -35,18 +36,13 @@ public:
 
 public:
 
-	// Begin SWidget interface
+	// SWidget interface
 
 	virtual FVector2D ComputeDesiredSize() const override;
-
 	virtual void OnArrangeChildren( const FGeometry& AllottedGeometry, FArrangedChildren& ArrangedChildren ) const override;
-		
 	virtual FChildren* GetChildren() override;
-
 	virtual int32 OnPaint( const FPaintArgs& Args, const FGeometry& AllottedGeometry, const FSlateRect& MyClippingRect, FSlateWindowElementList& OutDrawElements, int32 LayerId, const FWidgetStyle& InWidgetStyle, bool bParentEnabled ) const override;
 
-	// End SWidget interface
-
 private:
 
 	TWeakChild<SWidget> WeakChild;
diff --git a/Engine/Source/Runtime/Slate/Public/Widgets/Views/SExpanderArrow.h b/Engine/Source/Runtime/Slate/Public/Widgets/Views/SExpanderArrow.h
--- a/Engine/Source/Runtime/Slate/Public/Widgets/Views/SExpanderArrow.h
+++ b/Engine/Source/Runtime/Slate/Public/Widgets/Views/SExpanderArrow.h
@@ -2,11 +2,17 @@
 
 #pragma once
 
-/** Expander arrow and indentation component that can be placed in a TableRow of a TreeView. Intended for use by TMultiColumnRow in TreeViews. */
+
+/**
+ * Expander arrow and indentation component that can be placed in a TableRow
+ * of a TreeView. Intended for use by TMultiColumnRow in TreeViews.
+ */
 class SLATE_API SExpanderArrow : public SCompoundWidget
 {
 public:
-	SLATE_BEGIN_ARGS( SExpanderArrow ){}
+
+	SLATE_BEGIN_ARGS( SExpanderArrow )
+	{ }
 		SLATE_ATTRIBUTE(float, IndentAmount)
 	SLATE_END_ARGS()
 
diff --git a/Engine/Source/Runtime/Slate/Public/Widgets/Views/SListView.h b/Engine/Source/Runtime/Slate/Public/Widgets/Views/SListView.h
--- a/Engine/Source/Runtime/Slate/Public/Widgets/Views/SListView.h
+++ b/Engine/Source/Runtime/Slate/Public/Widgets/Views/SListView.h
@@ -2,6 +2,7 @@
 
 #pragma once
 
+
 /**
  * A ListView widget observes an array of data items and creates visual representations of these items.
  * ListView relies on the property that holding a reference to a value ensures its existence. In other words,
@@ -39,7 +40,6 @@ public:
 	typedef typename TSlateDelegates< NullableItemType >::FOnSelectionChanged FOnSelectionChanged;
 	typedef typename TSlateDelegates< ItemType >::FOnMouseButtonDoubleClick FOnMouseButtonDoubleClick;
 
-
 public:
 
 	class FColumnHeaderSlot
@@ -54,7 +54,7 @@ public:
 
 	SLATE_BEGIN_ARGS( SListView<ItemType> )
 		: _OnGenerateRow()
-		, _ListItemsSource( static_cast<const TArray<ItemType>*>(NULL) ) //@todo Slate Syntax: Initializing from NULL without a cast
+		, _ListItemsSource( static_cast<const TArray<ItemType>*>(nullptr) ) //@todo Slate Syntax: Initializing from nullptr without a cast
 		, _ItemHeight(16)
 		, _OnContextMenuOpening()
 		, _OnMouseButtonDoubleClick()
@@ -62,7 +62,7 @@ public:
 		, _SelectionMode(ESelectionMode::Multi)
 		, _ClearSelectionOnClick(true)
 		, _ExternalScrollbar()
-	{}
+	{ }
 
 		SLATE_EVENT( FOnGenerateRow, OnGenerateRow )
 
@@ -116,7 +116,7 @@ public:
 				ErrorString += TEXT("Please specify an OnGenerateRow. \n");
 			}
 
-			if ( this->ItemsSource == NULL )
+			if ( this->ItemsSource == nullptr )
 			{
 				ErrorString += TEXT("Please specify a ListItemsSource. \n");
 			}
@@ -147,16 +147,17 @@ public:
 
 	SListView( ETableViewMode::Type InListMode = ETableViewMode::List )
 		: STableViewBase( InListMode )
-		, SelectorItem( NullableItemType(NULL) )
-		, RangeSelectionStart( NullableItemType(NULL) )
-		, ItemsSource( NULL )
-		, ItemToScrollIntoView( NullableItemType(NULL) )
-		, ItemToNotifyWhenInView( NullableItemType(NULL) ) 
-	{
-	}
+		, SelectorItem( NullableItemType(nullptr) )
+		, RangeSelectionStart( NullableItemType(nullptr) )
+		, ItemsSource( nullptr )
+		, ItemToScrollIntoView( NullableItemType(nullptr) )
+		, ItemToNotifyWhenInView( NullableItemType(nullptr) ) 
+	{ }
 
 public:
-	// Inherited from SWidget
+
+	// SWidget overrides
+
 	virtual FReply OnKeyDown( const FGeometry& MyGeometry, const FKeyboardEvent& InKeyboardEvent ) override
 	{
 		const TArray<ItemType>& ItemsSourceRef = (*this->ItemsSource);
@@ -165,7 +166,7 @@ public:
 		if ( ItemsSourceRef.Num() > 0 && !InKeyboardEvent.IsAltDown() )
 		{
 			bool bWasHandled = false;
-			NullableItemType ItemNavigatedTo( NULL );
+			NullableItemType ItemNavigatedTo( nullptr );
 
 			// Check for selection manipulation keys (Up, Down, Home, End, PageUp, PageDown)
 			if ( InKeyboardEvent.GetKey() == EKeys::Home )
@@ -391,19 +392,18 @@ private:
 		 * Find a widget for this item if it has already been constructed.
 		 *
 		 * @param Item  The item for which to find the widget.
-		 *
-		 * @return A pointer to the corresponding widget if it exists; otherwise NULL.
+		 * @return A pointer to the corresponding widget if it exists; otherwise nullptr.
 		 */
 		TSharedPtr<ITableRow> GetWidgetForItem( const ItemType& Item )
 		{
 			TSharedRef<ITableRow>* LookupResult = ItemToWidgetMap.Find( Item );
-			if ( LookupResult != NULL )
+			if ( LookupResult != nullptr )
 			{
 				return *LookupResult;
 			}
 			else
 			{
-				return TSharedPtr<ITableRow>(NULL);
+				return TSharedPtr<ITableRow>(nullptr);
 			}
 		}
 
@@ -416,7 +416,7 @@ private:
 		void OnItemSeen( ItemType InItem, TSharedRef<ITableRow> InGeneratedWidget)
 		{
 			TSharedRef<ITableRow>* LookupResult = ItemToWidgetMap.Find( InItem );
-			const bool bWidgetIsNewlyGenerated = (LookupResult == NULL);
+			const bool bWidgetIsNewlyGenerated = (LookupResult == nullptr);
 			if ( bWidgetIsNewlyGenerated )
 			{
 				// It's a newly generated item!
@@ -452,7 +452,7 @@ private:
 			{
 				ItemType ItemToBeCleanedUp = ItemsToBeCleanedUp[ItemIndex];
 				const TSharedRef<ITableRow>* FindResult = ItemToWidgetMap.Find( ItemToBeCleanedUp );
-				if ( FindResult != NULL )
+				if ( FindResult != nullptr )
 				{
 					const TSharedRef<ITableRow> WidgetToCleanUp = *FindResult;
 					ItemToWidgetMap.Remove( ItemToBeCleanedUp );
@@ -478,7 +478,7 @@ private:
 			{
 				ItemType ItemToBeCleanedUp = ItemsToBeCleanedUp[ItemIndex];
 				const TSharedRef<ITableRow>* FindResult = ItemToWidgetMap.Find( ItemToBeCleanedUp );
-				if ( FindResult != NULL )
+				if ( FindResult != nullptr )
 				{
 					const TSharedRef<ITableRow> WidgetToCleanUp = *FindResult;
 					ItemToWidgetMap.Remove( ItemToBeCleanedUp );
@@ -491,21 +491,22 @@ private:
 
 		/** Map of DataItems to corresponding SWidgets */
 		TMap< ItemType, TSharedRef<ITableRow> > ItemToWidgetMap;
+
 		/** Map of SWidgets to DataItems from which they were generated */
 		TMap< const ITableRow*, ItemType > WidgetMapToItem;
+
 		/** A set of Items that currently have a generated widget */
 		TArray< ItemType > ItemsWithGeneratedWidgets;
+
 		/** Total number of DataItems the last time we performed a generation pass. */
 		int32 TotalItemsLastGeneration;
+
 		/** Items that need their widgets destroyed because they are no longer on screen. */
 		TArray<ItemType> ItemsToBeCleanedUp;
 	};
 
-
 public:
-	//
-	// Private Interface
-	//
+
 	// A low-level interface for use various widgets generated by ItemsWidgets(Lists, Trees, etc).
 	// These handle selection, expansion, and other such properties common to ItemsWidgets.
 	//
@@ -596,7 +597,7 @@ public:
 	virtual const ItemType* Private_ItemFromWidget( const ITableRow* TheWidget ) const override
 	{
 		ItemType const * LookupResult = WidgetGenerator.WidgetMapToItem.Find( TheWidget );
-		return LookupResult == NULL ? NULL : LookupResult;
+		return LookupResult == nullptr ? nullptr : LookupResult;
 	}
 
 	virtual bool Private_UsesSelectorFocus() const override
@@ -611,7 +612,7 @@ public:
 
 	virtual bool Private_IsItemSelected( const ItemType& TheItem ) const override
 	{
-		return NULL != SelectedItems.Find(TheItem);
+		return nullptr != SelectedItems.Find(TheItem);
 	}
 
 	virtual bool Private_IsItemExpanded( const ItemType& TheItem ) const override
@@ -677,8 +678,6 @@ public:
 		return SharedThis(this);
 	}
 
-
-
 public:	
 
 	/**
@@ -690,7 +689,7 @@ public:
 		if ( TableViewMode != ETableViewMode::Tree )
 		{
 			bool bSelectionChanged = false;
-			if ( ItemsSource == NULL )
+			if ( ItemsSource == nullptr )
 			{
 				// We are no longer observing items so there is no more selection.
 				this->Private_ClearSelection();
@@ -704,7 +703,7 @@ public:
 				for ( int32 ItemIndex = 0; ItemIndex < ItemsSource->Num(); ++ItemIndex )
 				{
 					ItemType CurItem = (*ItemsSource)[ItemIndex];
-					const bool bItemIsSelected = ( NULL != SelectedItems.Find( CurItem ) );
+					const bool bItemIsSelected = ( nullptr != SelectedItems.Find( CurItem ) );
 					if ( bItemIsSelected )
 					{
 						NewSelectedItems.Add( CurItem );
@@ -740,7 +739,7 @@ public:
 		FGenerationPassGuard GenerationPassGuard(WidgetGenerator);
 
 		const TArray<ItemType>* SourceItems = ItemsSource;
-		if ( SourceItems != NULL && SourceItems->Num() > 0 )
+		if ( SourceItems != nullptr && SourceItems->Num() > 0 )
 		{
 			// Items in view, including fractional items
 			float ItemsInView = 0.0f;
@@ -883,7 +882,7 @@ public:
 	/** @return how many items there are in the TArray being observed */
 	virtual int32 GetNumItemsBeingObserved() const
 	{
-		return ItemsSource == NULL ? 0 : ItemsSource->Num();
+		return ItemsSource == nullptr ? 0 : ItemsSource->Num();
 	}
 
 	/**
@@ -1041,7 +1040,7 @@ public:
 	 *
 	 * @param InItem  The item for which to find the widget.
 	 *
-	 * @return A pointer to the corresponding widget if it exists; otherwise NULL.
+	 * @return A pointer to the corresponding widget if it exists; otherwise nullptr.
 	*/
 	TSharedPtr<ITableRow> WidgetFromItem( const ItemType& InItem )
 	{
@@ -1068,7 +1067,7 @@ protected:
 	 */
 	virtual void ScrollIntoView( const FGeometry& ListViewGeometry ) override
 	{
-		if ( TListTypeTraits<ItemType>::IsPtrValid(ItemToScrollIntoView) && ItemsSource != NULL )
+		if ( TListTypeTraits<ItemType>::IsPtrValid(ItemToScrollIntoView) && ItemsSource != nullptr )
 		{
 			const int32 IndexOfItem = ItemsSource->Find( TListTypeTraits<ItemType>::NullableItemTypeConvertToItemType( ItemToScrollIntoView ) );
 			if (IndexOfItem != INDEX_NONE)
@@ -1146,7 +1145,7 @@ protected:
 			//           Generate widgets on demand so we can figure out how tall they are.
 
 			const TArray<ItemType>* SourceItems = ItemsSource;
-			if ( SourceItems != NULL && SourceItems->Num() > 0 )
+			if ( SourceItems != nullptr && SourceItems->Num() > 0 )
 			{
 				int ItemIndex = StartingItemIndex;
 				while( AbsScrollByAmount != 0 && ItemIndex < SourceItems->Num() && ItemIndex >= 0 )
@@ -1292,6 +1291,7 @@ protected:
 	}
 
 protected:
+
 	/** A widget generator component */
 	FWidgetGenerator WidgetGenerator;
 
diff --git a/Engine/Source/Runtime/Slate/Public/Widgets/Views/STableRow.h b/Engine/Source/Runtime/Slate/Public/Widgets/Views/STableRow.h
--- a/Engine/Source/Runtime/Slate/Public/Widgets/Views/STableRow.h
+++ b/Engine/Source/Runtime/Slate/Public/Widgets/Views/STableRow.h
@@ -2,6 +2,7 @@
 
 #pragma once
 
+
 /**
  * Interface for table views to talk to their rows.
  */
@@ -35,13 +36,14 @@ class SLATE_API ITableRow
 		virtual void Private_OnExpanderArrowShiftClicked() = 0;
 };
 
+
 template <typename ItemType> class SListView;
 
 DECLARE_DELEGATE_OneParam(FOnTableRowDragEnter, FDragDropEvent const&);
-
 DECLARE_DELEGATE_OneParam(FOnTableRowDragLeave, FDragDropEvent const&);
-
 DECLARE_DELEGATE_RetVal_OneParam(FReply, FOnTableRowDrop, FDragDropEvent const&);
+
+
 /**
  * The ListView is populated by Selectable widgets.
  * A Selectable widget is away of the ListView containing it (OwnerWidget) and holds arbitrary Content (Content).
@@ -54,7 +56,6 @@ class STableRow : public ITableRow, public SBorder
 
 public:
 
-
 	SLATE_BEGIN_ARGS( STableRow< ItemType > )
 		: _Style( &FCoreStyle::Get().GetWidgetStyle<FTableRowStyle>("TableView.Row") )
 		, _Padding( FMargin(0) )
@@ -160,10 +161,9 @@ public:
 	/**
 	 * Called when a mouse button is double clicked.  Override this in derived classes.
 	 *
-	 * @param  InMyGeometry  Widget geometry
-	 * @param  InMouseEvent  Mouse button event
-	 *
-	 * @return  Returns whether the event was handled, along with other possible actions
+	 * @param  InMyGeometry  Widget geometry.
+	 * @param  InMouseEvent  Mouse button event.
+	 * @return  Returns whether the event was handled, along with other possible actions.
 	 */
 	virtual FReply OnMouseButtonDoubleClick( const FGeometry& InMyGeometry, const FPointerEvent& InMouseEvent )
 	{
@@ -191,9 +191,8 @@ public:
 	/**
 	 * See SWidget::OnMouseButtonDown
 	 *
-	 * @param MyGeometry The Geometry of the widget receiving the event
-	 * @param MouseEvent Information about the input event
-	 *
+	 * @param MyGeometry The Geometry of the widget receiving the event.
+	 * @param MouseEvent Information about the input event.
 	 * @return Whether the event was handled along with possible requests for the system to take action.
 	 */	
 	virtual FReply OnMouseButtonDown( const FGeometry& MyGeometry, const FPointerEvent& MouseEvent )
@@ -251,7 +250,7 @@ public:
 					const ItemType* MyItem = OwnerWidget->Private_ItemFromWidget( this );
 					const bool bIsSelected = OwnerWidget->Private_IsItemSelected( *MyItem );
 
-					check(MyItem != NULL);
+					check(MyItem != nullptr);
 
 					if ( MouseEvent.IsControlDown() )
 					{
@@ -285,9 +284,8 @@ public:
 	/**
 	 * See SWidget::OnMouseButtonUp
 	 *
-	 * @param MyGeometry The Geometry of the widget receiving the event
-	 * @param MouseEvent Information about the input event
-	 *
+	 * @param MyGeometry The Geometry of the widget receiving the event.
+	 * @param MouseEvent Information about the input event.
 	 * @return Whether the event was handled along with possible requests for the system to take action.
 	 */
 	virtual FReply OnMouseButtonUp( const FGeometry& MyGeometry, const FPointerEvent& MouseEvent )
@@ -333,7 +331,7 @@ public:
 							if ( !ChangedSelectionOnMouseDown && !MouseEvent.IsControlDown() && !MouseEvent.IsShiftDown() )
 							{
 								const ItemType* MyItem = OwnerWidget->Private_ItemFromWidget( this );
-								check(MyItem != NULL);
+								check(MyItem != nullptr);
 
 								const bool bIsSelected = OwnerWidget->Private_IsItemSelected( *MyItem );
 								if ( bIsSelected && OwnerWidget->Private_GetNumSelectedItems() > 1 )
@@ -657,12 +655,12 @@ public:
 
 	/** Protected constructor; SWidgets should only be instantiated via declarative syntax. */
 	STableRow()
-	: IndexInList(0)
-	, bShowSelection(true)
-	{
-	}
+		: IndexInList(0)
+		, bShowSelection(true)
+	{ }
+
+protected:
 
-protected:	
 	/**
 	 * An internal method to construct and setup this row widget (purposely avoids child construction). 
 	 * Split out from Construct() so that sub-classes can invoke super construction without invoking 
@@ -753,11 +751,9 @@ protected:
 
 	/** Did the current a touch interaction start in this item?*/
 	bool bProcessingSelectionTouch;
-
 };
 
 
-
 template<typename ItemType>
 class SMultiColumnTableRow : public STableRow<ItemType>
 {
@@ -770,7 +766,6 @@ public:
 	 * Make a widget and return it.
 	 *
 	 * @param ColumnName    A unique ID for a column in this TableView; see SHeaderRow::FColumn for more info.
-	 *
 	 * @return a widget to represent the contents of a cell in this row of a TableView. 
 	 */
 	virtual TSharedRef<SWidget> GenerateWidgetForColumn( const FName& ColumnName ) = 0;
@@ -831,7 +826,7 @@ protected:
 			const TSharedRef< SWidget >* const ExistingWidget = ColumnIdToSlotContents.Find( Column.ColumnId );
 
 			TSharedRef< SWidget > CellContents = SNullWidget::NullWidget;
-			if ( ExistingWidget != NULL )
+			if ( ExistingWidget != nullptr )
 			{
 				CellContents = *ExistingWidget;
 			}
@@ -880,9 +875,8 @@ protected:
 		ColumnIdToSlotContents.Empty();
 	}
 
-
 private:
 	
 	TSharedPtr<SHorizontalBox> Box;
 	TMap< FName, TSharedRef< SWidget > > ColumnIdToSlotContents;
-};
\ No newline at end of file
+};
diff --git a/Engine/Source/Runtime/Slate/Public/Widgets/Views/STableViewBase.h b/Engine/Source/Runtime/Slate/Public/Widgets/Views/STableViewBase.h
--- a/Engine/Source/Runtime/Slate/Public/Widgets/Views/STableViewBase.h
+++ b/Engine/Source/Runtime/Slate/Public/Widgets/Views/STableViewBase.h
@@ -4,19 +4,24 @@
 
 #include "IScrollableWidget.h"
 
+
 class SListPanel;
 class SHeaderRow;
 
+
 enum class EAllowOverscroll
 {
 	Yes,
 	No
 };
 
+
 /**
  * Contains ListView functionality that does not depend on the type of data being observed by the ListView.
  */
-class SLATE_API STableViewBase : public SCompoundWidget, public IScrollableWidget
+class SLATE_API STableViewBase
+	: public SCompoundWidget
+	, public IScrollableWidget
 {
 public:
 
@@ -36,7 +41,32 @@ public:
 	 */
 	void ScrollBar_OnUserScrolled( float InScrollOffsetFraction );
 
+	/** @return The number of Widgets we currently have generated. */
+	int32 GetNumGeneratedChildren() const;
+
+	TSharedPtr<SHeaderRow> GetHeaderRow() const;
+
+	/** @return Returns true if the user is currently interactively scrolling the view by holding
+		        the right mouse button and dragging. */
+	bool IsRightClickScrolling() const;
+
+	/** @return Returns true if the user is currently interactively scrolling the view by holding
+		        either mouse button and dragging. */
+	bool IsUserScrolling() const;
+
+	/** Mark the list as dirty, so that it will regenerate its widgets on next tick. */
+	void RequestListRefresh();
+
+	/** Return true if there is currently a refresh pending, false otherwise */
+	bool IsPendingRefresh() const;
+
+	/** Is this list backing a tree or just a standalone list */
+	const ETableViewMode::Type TableViewMode;
+
+public:
+
 	// SWidget interface
+
 	virtual void OnKeyboardFocusLost( const FKeyboardFocusEvent& InKeyboardFocusEvent ) override;
 	virtual void OnMouseCaptureLost() override;
 	virtual bool SupportsKeyboardFocus() const override;
@@ -50,42 +80,18 @@ public:
 	virtual FReply OnMouseWheel( const FGeometry& MyGeometry, const FPointerEvent& MouseEvent ) override;
 	virtual FReply OnKeyDown( const FGeometry& MyGeometry, const FKeyboardEvent& InKeyboardEvent ) override;
 	virtual FCursorReply OnCursorQuery( const FGeometry& MyGeometry, const FPointerEvent& CursorEvent ) const override;
+	virtual int32 OnPaint(const FPaintArgs& Args, const FGeometry& AllottedGeometry, const FSlateRect& MyClippingRect, FSlateWindowElementList& OutDrawElements, int32 LayerId, const FWidgetStyle& InWidgetStyle, bool bParentEnabled) const override;
 	virtual FReply OnTouchStarted( const FGeometry& MyGeometry, const FPointerEvent& InTouchEvent ) override;
 	virtual FReply OnTouchMoved( const FGeometry& MyGeometry, const FPointerEvent& InTouchEvent ) override;
 	virtual FReply OnTouchEnded( const FGeometry& MyGeometry, const FPointerEvent& InTouchEvent ) override;
-	// End of SWidget interface
-
-	/** @return The number of Widgets we currently have generated. */
-	int32 GetNumGeneratedChildren() const;
-
-	TSharedPtr<SHeaderRow> GetHeaderRow() const;
-
-	/** @return Returns true if the user is currently interactively scrolling the view by holding
-		        the right mouse button and dragging. */
-	bool IsRightClickScrolling() const;
-
-	/** @return Returns true if the user is currently interactively scrolling the view by holding
-		        either mouse button and dragging. */
-	bool IsUserScrolling() const;
 
-	/**
-	 * Mark the list as dirty, so that it will regenerate its widgets on next tick.
-	 */
-	void RequestListRefresh();
-
-	/** Return true if there is currently a refresh pending, false otherwise */
-	bool IsPendingRefresh() const;
+public:
 
 	// IScrollableWidget interface
+
 	virtual FVector2D GetScrollDistance() override;
 	virtual FVector2D GetScrollDistanceRemaining() override;
 	virtual TSharedRef<class SWidget> GetScrollWidget() override;
-	// End of IScrollableWidget interface
-
-	virtual int32 OnPaint(const FPaintArgs& Args, const FGeometry& AllottedGeometry, const FSlateRect& MyClippingRect, FSlateWindowElementList& OutDrawElements, int32 LayerId, const FWidgetStyle& InWidgetStyle, bool bParentEnabled) const override;
-
-	/** Is this list backing a tree or just a standalone list */
-	const ETableViewMode::Type TableViewMode;
 
 protected:
 
@@ -225,7 +231,7 @@ protected:
 	/** What the list's geometry was the last time a refresh occurred. */
 	FGeometry PanelGeometryLastTick;
 
-	/** Delegate to invoke when the context menu should be opening. If it is NULL, a context menu will not be summoned. */
+	/** Delegate to invoke when the context menu should be opening. If it is nullptr, a context menu will not be summoned. */
 	FOnContextMenuOpening OnContextMenuOpening;
 
 	/** The selection mode that this tree/list is in. Note that it is up to the generated ITableRows to respect this setting. */
diff --git a/Engine/Source/Runtime/Slate/Public/Widgets/Views/STileView.h b/Engine/Source/Runtime/Slate/Public/Widgets/Views/STileView.h
--- a/Engine/Source/Runtime/Slate/Public/Widgets/Views/STileView.h
+++ b/Engine/Source/Runtime/Slate/Public/Widgets/Views/STileView.h
@@ -2,6 +2,7 @@
  
 #pragma once
 
+
 /**
  * A TileView widget is a list which arranges its items horizontally until there is no more space then creates a new row.
  * Items are spaced evenly horizontally.
@@ -21,7 +22,7 @@ public:
 
 	SLATE_BEGIN_ARGS( STileView<ItemType> )
 		: _OnGenerateTile()
-		, _ListItemsSource( static_cast<TArray<ItemType>*>(NULL) ) //@todo Slate Syntax: Initializing from NULL without a cast
+		, _ListItemsSource( static_cast<TArray<ItemType>*>(nullptr) ) //@todo Slate Syntax: Initializing from nullptr without a cast
 		, _ItemHeight(128)
 		, _ItemWidth(128)
 		, _OnContextMenuOpening()
@@ -82,7 +83,7 @@ public:
 				ErrorString += TEXT("Please specify an OnGenerateTile. \n");
 			}
 
-			if ( this->ItemsSource == NULL )
+			if ( this->ItemsSource == nullptr )
 			{
 				ErrorString += TEXT("Please specify a ListItemsSource. \n");
 			}
@@ -112,7 +113,9 @@ public:
 	}
 
 public:
-	// Inherited from SWidget
+
+	// SWidget overrides
+
 	virtual FReply OnKeyDown( const FGeometry& MyGeometry, const FKeyboardEvent& InKeyboardEvent ) override
 	{
 		const TArray<ItemType>& ItemsSourceRef = (*this->ItemsSource);
@@ -158,7 +161,7 @@ public:
 		this->ClearWidgets();
 		
 		const TArray<ItemType>* SourceItems = this->ItemsSource;
-		if ( SourceItems != NULL )
+		if ( SourceItems != nullptr )
 		{
 			// Item width and height is constant by design.
 			const float AllottedWidth = MyGeometry.Size.X;
@@ -238,7 +241,7 @@ public:
 
 	virtual int32 GetNumItemsBeingObserved() const override
 	{
-		const int32 NumItemsBeingObserved = this->ItemsSource == NULL ? 0 : this->ItemsSource->Num();
+		const int32 NumItemsBeingObserved = this->ItemsSource == nullptr ? 0 : this->ItemsSource->Num();
 		const int32 NumItemsWide = GetNumItemsWide();
 		
 		int32 NumEmptySpacesAtEnd = 0;
diff --git a/Engine/Source/Runtime/Slate/Public/Widgets/Views/STreeView.h b/Engine/Source/Runtime/Slate/Public/Widgets/Views/STreeView.h
--- a/Engine/Source/Runtime/Slate/Public/Widgets/Views/STreeView.h
+++ b/Engine/Source/Runtime/Slate/Public/Widgets/Views/STreeView.h
@@ -2,6 +2,7 @@
  
 #pragma once
 
+
 /** Info needed by a (relatively) small fraction of the tree items; some of them may not be visible. */
 struct FSparseItemInfo
 {
@@ -24,6 +25,7 @@ struct FSparseItemInfo
 	bool bHasExpandedChildren;
 };
 
+
 /** Info needed by every visible item in the tree */
 struct FItemInfo
 {
@@ -45,10 +47,6 @@ struct FItemInfo
 };
 
 
-
-
-
-
 /**
  * This assumes you are familiar with SListView; see SListView.
  *
@@ -62,133 +60,137 @@ struct FItemInfo
 template<typename ItemType>
 class STreeView : public SListView< ItemType >
 {
+public:
+
+	typedef typename TListTypeTraits< ItemType >::NullableType NullableItemType;
+
+	typedef typename TSlateDelegates< ItemType >::FOnGetChildren FOnGetChildren;
+	typedef typename TSlateDelegates< ItemType >::FOnGenerateRow FOnGenerateRow;
+	typedef typename TSlateDelegates< ItemType >::FOnSetExpansionRecursive FOnSetExpansionRecursive;
+	typedef typename TSlateDelegates< ItemType >::FOnItemScrolledIntoView FOnItemScrolledIntoView;
+	typedef typename TSlateDelegates< NullableItemType >::FOnSelectionChanged FOnSelectionChanged;
+	typedef typename TSlateDelegates< ItemType >::FOnMouseButtonDoubleClick FOnMouseButtonDoubleClick;
+	typedef typename TSlateDelegates< ItemType >::FOnExpansionChanged FOnExpansionChanged;
 
-	public:
-		typedef typename TListTypeTraits< ItemType >::NullableType NullableItemType;
-
-		typedef typename TSlateDelegates< ItemType >::FOnGetChildren FOnGetChildren;
-		typedef typename TSlateDelegates< ItemType >::FOnGenerateRow FOnGenerateRow;
-		typedef typename TSlateDelegates< ItemType >::FOnSetExpansionRecursive FOnSetExpansionRecursive;
-		typedef typename TSlateDelegates< ItemType >::FOnItemScrolledIntoView FOnItemScrolledIntoView;
-		typedef typename TSlateDelegates< NullableItemType >::FOnSelectionChanged FOnSelectionChanged;
-		typedef typename TSlateDelegates< ItemType >::FOnMouseButtonDoubleClick FOnMouseButtonDoubleClick;
-		typedef typename TSlateDelegates< ItemType >::FOnExpansionChanged FOnExpansionChanged;
-	public:
+public:
 	
-		SLATE_BEGIN_ARGS( STreeView<ItemType> )
-			: _OnGenerateRow()
-			, _OnGetChildren()
-			, _OnSetExpansionRecursive()
-			, _TreeItemsSource( static_cast< TArray<ItemType>* >(NULL) ) //@todo Slate Syntax: Initializing from NULL without a cast
-			, _ItemHeight(16)
-			, _OnContextMenuOpening()
-			, _OnMouseButtonDoubleClick()
-			, _OnSelectionChanged()
-			, _OnExpansionChanged()
-			, _SelectionMode(ESelectionMode::Multi)
-			, _ClearSelectionOnClick(true)
-			, _ExternalScrollbar()
-			{}
+	SLATE_BEGIN_ARGS( STreeView<ItemType> )
+		: _OnGenerateRow()
+		, _OnGetChildren()
+		, _OnSetExpansionRecursive()
+		, _TreeItemsSource( static_cast< TArray<ItemType>* >(nullptr) ) //@todo Slate Syntax: Initializing from nullptr without a cast
+		, _ItemHeight(16)
+		, _OnContextMenuOpening()
+		, _OnMouseButtonDoubleClick()
+		, _OnSelectionChanged()
+		, _OnExpansionChanged()
+		, _SelectionMode(ESelectionMode::Multi)
+		, _ClearSelectionOnClick(true)
+		, _ExternalScrollbar()
+		{}
 
-			SLATE_EVENT( FOnGenerateRow, OnGenerateRow )
+		SLATE_EVENT( FOnGenerateRow, OnGenerateRow )
 
-			SLATE_EVENT( FOnItemScrolledIntoView, OnItemScrolledIntoView )
+		SLATE_EVENT( FOnItemScrolledIntoView, OnItemScrolledIntoView )
 
-			SLATE_EVENT( FOnGetChildren, OnGetChildren )
+		SLATE_EVENT( FOnGetChildren, OnGetChildren )
 
-			SLATE_EVENT( FOnSetExpansionRecursive, OnSetExpansionRecursive )
+		SLATE_EVENT( FOnSetExpansionRecursive, OnSetExpansionRecursive )
 
-			SLATE_ARGUMENT( TArray<ItemType>* , TreeItemsSource )
+		SLATE_ARGUMENT( TArray<ItemType>* , TreeItemsSource )
 
-			SLATE_ATTRIBUTE( float, ItemHeight )
+		SLATE_ATTRIBUTE( float, ItemHeight )
 
-			SLATE_EVENT( FOnContextMenuOpening, OnContextMenuOpening )
+		SLATE_EVENT( FOnContextMenuOpening, OnContextMenuOpening )
 
-			SLATE_EVENT( FOnMouseButtonDoubleClick, OnMouseButtonDoubleClick )
+		SLATE_EVENT( FOnMouseButtonDoubleClick, OnMouseButtonDoubleClick )
 
-			SLATE_EVENT( FOnSelectionChanged, OnSelectionChanged )
+		SLATE_EVENT( FOnSelectionChanged, OnSelectionChanged )
 
-			SLATE_EVENT( FOnExpansionChanged, OnExpansionChanged )
+		SLATE_EVENT( FOnExpansionChanged, OnExpansionChanged )
 
-			SLATE_ATTRIBUTE( ESelectionMode::Type, SelectionMode )
+		SLATE_ATTRIBUTE( ESelectionMode::Type, SelectionMode )
 
-			SLATE_ARGUMENT( TSharedPtr<SHeaderRow>, HeaderRow )
+		SLATE_ARGUMENT( TSharedPtr<SHeaderRow>, HeaderRow )
 
-			SLATE_ARGUMENT ( bool, ClearSelectionOnClick )
+		SLATE_ARGUMENT ( bool, ClearSelectionOnClick )
 
-			SLATE_ARGUMENT( TSharedPtr<SScrollBar>, ExternalScrollbar )
+		SLATE_ARGUMENT( TSharedPtr<SScrollBar>, ExternalScrollbar )
 
-		SLATE_END_ARGS()
+	SLATE_END_ARGS()
 
 		
-		/**
-		 * Construct this widget
-		 *
-		 * @param	InArgs	The declaration data for this widget
-		 */
-		void Construct( const FArguments& InArgs )
+	/**
+	 * Construct this widget
+	 *
+	 * @param	InArgs	The declaration data for this widget.
+	 */
+	void Construct( const FArguments& InArgs )
+	{
+		this->OnGenerateRow = InArgs._OnGenerateRow;
+		this->OnItemScrolledIntoView = InArgs._OnItemScrolledIntoView;
+		this->OnGetChildren = InArgs._OnGetChildren;
+		this->OnSetExpansionRecursive = InArgs._OnSetExpansionRecursive;
+		this->TreeItemsSource = InArgs._TreeItemsSource;
+
+		this->OnContextMenuOpening = InArgs._OnContextMenuOpening;
+		this->OnDoubleClick = InArgs._OnMouseButtonDoubleClick;
+		this->OnSelectionChanged = InArgs._OnSelectionChanged;
+		this->OnExpansionChanged = InArgs._OnExpansionChanged;
+		this->SelectionMode = InArgs._SelectionMode;
+
+		this->bClearSelectionOnClick = InArgs._ClearSelectionOnClick;
+
+		// Check for any parameters that the coder forgot to specify.
+		FString ErrorString;
 		{
-			this->OnGenerateRow = InArgs._OnGenerateRow;
-			this->OnItemScrolledIntoView = InArgs._OnItemScrolledIntoView;
-			this->OnGetChildren = InArgs._OnGetChildren;
-			this->OnSetExpansionRecursive = InArgs._OnSetExpansionRecursive;
-			this->TreeItemsSource = InArgs._TreeItemsSource;
-
-			this->OnContextMenuOpening = InArgs._OnContextMenuOpening;
-			this->OnDoubleClick = InArgs._OnMouseButtonDoubleClick;
-			this->OnSelectionChanged = InArgs._OnSelectionChanged;
-			this->OnExpansionChanged = InArgs._OnExpansionChanged;
-			this->SelectionMode = InArgs._SelectionMode;
-
-			this->bClearSelectionOnClick = InArgs._ClearSelectionOnClick;
-
-			// Check for any parameters that the coder forgot to specify.
-			FString ErrorString;
+			if ( !this->OnGenerateRow.IsBound() )
 			{
-				if ( !this->OnGenerateRow.IsBound() )
-				{
-					ErrorString += TEXT("Please specify an OnGenerateRow. \n");
-				}
-
-				if ( this->TreeItemsSource == NULL )
-				{
-					ErrorString += TEXT("Please specify a TreeItemsSource. \n");
-				}
-
-				if ( !this->OnGetChildren.IsBound() )
-				{
-					ErrorString += TEXT("Please specify an OnGetChildren. \n");
-				}
+				ErrorString += TEXT("Please specify an OnGenerateRow. \n");
 			}
 
-			if (ErrorString.Len() > 0)
+			if ( this->TreeItemsSource == nullptr )
 			{
-				// Let the coder know what they forgot
-				this->ChildSlot
-				.HAlign(HAlign_Center)
-				.VAlign(VAlign_Center)
-				[
-					SNew(STextBlock)
-					.Text(ErrorString)
-				];
+				ErrorString += TEXT("Please specify a TreeItemsSource. \n");
 			}
-			else
+
+			if ( !this->OnGetChildren.IsBound() )
 			{
-				// Make the TableView
-				this->ConstructChildren( 0, InArgs._ItemHeight, InArgs._HeaderRow, InArgs._ExternalScrollbar );
+				ErrorString += TEXT("Please specify an OnGetChildren. \n");
 			}
 		}
 
-		STreeView()
-			: SListView< ItemType >( ETableViewMode::Tree )
-			, TreeItemsSource( NULL )
-			, bTreeItemsAreDirty( true )
+		if (ErrorString.Len() > 0)
 		{
-			this->ItemsSource = &LinearizedItems;
+			// Let the coder know what they forgot
+			this->ChildSlot
+			.HAlign(HAlign_Center)
+			.VAlign(VAlign_Center)
+			[
+				SNew(STextBlock)
+				.Text(ErrorString)
+			];
 		}
+		else
+		{
+			// Make the TableView
+			this->ConstructChildren( 0, InArgs._ItemHeight, InArgs._HeaderRow, InArgs._ExternalScrollbar );
+		}
+	}
+
+	/** Default constructor. */
+	STreeView()
+		: SListView< ItemType >( ETableViewMode::Tree )
+		, TreeItemsSource( nullptr )
+		, bTreeItemsAreDirty( true )
+	{
+		this->ItemsSource = &LinearizedItems;
+	}
 
 public:
-	// Inherited from SWidget
+
+	// SWidget overrides
+
 	virtual FReply OnKeyDown( const FGeometry& MyGeometry, const FKeyboardEvent& InKeyboardEvent ) override
 	{
 		// Check for selection/expansion toggling keys (Left, Right)
@@ -261,16 +263,14 @@ public:
 	}
 	
 private:
-	//
-	// Private Interface
-	//
+
 	// Tree View adds the ability to expand/collapse items.
 	// All the selection functionality is inherited from ListView.
 
 	virtual bool Private_IsItemExpanded( const ItemType& TheItem ) const override
 	{
 		const FSparseItemInfo* ItemInfo = SparseItemInfos.Find(TheItem);
-		return ItemInfo != NULL && ItemInfo->bIsExpanded;
+		return ItemInfo != nullptr && ItemInfo->bIsExpanded;
 	}
 
 	virtual void Private_SetItemExpansion( ItemType TheItem, bool bShouldBeExpanded ) override
@@ -331,295 +331,290 @@ private:
 	}
 
 public:
-		/**
-		 * See SWidget::Tick()
-		 *
-		 * @param  AllottedGeometry The space allotted for this widget
-		 * @param  InCurrentTime  Current absolute real time
-		 * @param  InDeltaTime  Real time passed since last tick
-		 */
-		virtual void Tick( const FGeometry& AllottedGeometry, const double InCurrentTime, const float InDeltaTime )
+
+	/**
+	 * See SWidget::Tick()
+	 *
+	 * @param  AllottedGeometry The space allotted for this widget.
+	 * @param  InCurrentTime  Current absolute real time.
+	 * @param  InDeltaTime  Real time passed since last tick.
+	 */
+	virtual void Tick( const FGeometry& AllottedGeometry, const double InCurrentTime, const float InDeltaTime )
+	{
+		if ( bTreeItemsAreDirty )
 		{
-			if ( bTreeItemsAreDirty )
+			// Check that ItemsPanel was made; we never make it if the user failed to specify all the parameters necessary to make the tree work.
+			if ( this->ItemsPanel.IsValid() )
 			{
-				// Check that ItemsPanel was made; we never make it if the user failed to specify all the parameters necessary to make the tree work.
-				if ( this->ItemsPanel.IsValid() )
+				// We are about to repopulate linearized items; the ListView that TreeView is built on top of will also need to refresh.
+				bTreeItemsAreDirty = false;
+
+				if ( OnGetChildren.IsBound() && TreeItemsSource != nullptr )								
 				{
-					// We are about to repopulate linearized items; the ListView that TreeView is built on top of will also need to refresh.
-					bTreeItemsAreDirty = false;
+					// We make copies of the old expansion and selection sets so that we can remove
+					// any items that are no longer seen by the tree.
+					TSet< ItemType > TempSelectedItemsMap;
+					TMap< ItemType, FSparseItemInfo > TempSparseItemInfo;
+					TArray<FItemInfo> TempDenseItemInfos;
+					
+					// Rebuild the linearized view of the tree data.
+					LinearizedItems.Empty();
+					PopulateLinearizedItems( *TreeItemsSource, LinearizedItems, TempDenseItemInfos, 0, TempSelectedItemsMap, TempSparseItemInfo, true );
 
-					if ( OnGetChildren.IsBound() && TreeItemsSource != NULL )								
+					if(	this->SelectedItems.Num() != TempSelectedItemsMap.Num() || 
+						this->SelectedItems.Difference(TempSelectedItemsMap).Num() > 0 || 
+						TempSelectedItemsMap.Difference(this->SelectedItems).Num() > 0 )
 					{
-						// We make copies of the old expansion and selection sets so that we can remove
-						// any items that are no longer seen by the tree.
-						TSet< ItemType > TempSelectedItemsMap;
-						TMap< ItemType, FSparseItemInfo > TempSparseItemInfo;
-						TArray<FItemInfo> TempDenseItemInfos;
-					
-						// Rebuild the linearized view of the tree data.
-						LinearizedItems.Empty();
-						PopulateLinearizedItems( *TreeItemsSource, LinearizedItems, TempDenseItemInfos, 0, TempSelectedItemsMap, TempSparseItemInfo, true );
+						this->SelectedItems = TempSelectedItemsMap;
 
-						if(	this->SelectedItems.Num() != TempSelectedItemsMap.Num() || 
-							this->SelectedItems.Difference(TempSelectedItemsMap).Num() > 0 || 
-							TempSelectedItemsMap.Difference(this->SelectedItems).Num() > 0 )
+						if ( !TListTypeTraits<ItemType>::IsPtrValid( this->RangeSelectionStart ) ||
+								!this->SelectedItems.Contains( TListTypeTraits<ItemType>::NullableItemTypeConvertToItemType( this->RangeSelectionStart ) ))
 						{
-							this->SelectedItems = TempSelectedItemsMap;
-
-							if ( !TListTypeTraits<ItemType>::IsPtrValid( this->RangeSelectionStart ) ||
-								 !this->SelectedItems.Contains( TListTypeTraits<ItemType>::NullableItemTypeConvertToItemType( this->RangeSelectionStart ) ))
-							{
-									TListTypeTraits< ItemType >::ResetPtr( this->RangeSelectionStart );
-									TListTypeTraits< ItemType >::ResetPtr( this->SelectorItem );
-							}	
-							else if ( !TListTypeTraits<ItemType>::IsPtrValid( this->SelectorItem ) || 
-									  !this->SelectedItems.Contains( TListTypeTraits<ItemType>::NullableItemTypeConvertToItemType( this->SelectorItem ) ) )
-							{
-								this->SelectorItem = this->RangeSelectionStart;
-							}
-
-							this->Private_SignalSelectionChanged(ESelectInfo::Direct);
+								TListTypeTraits< ItemType >::ResetPtr( this->RangeSelectionStart );
+								TListTypeTraits< ItemType >::ResetPtr( this->SelectorItem );
+						}	
+						else if ( !TListTypeTraits<ItemType>::IsPtrValid( this->SelectorItem ) || 
+									!this->SelectedItems.Contains( TListTypeTraits<ItemType>::NullableItemTypeConvertToItemType( this->SelectorItem ) ) )
+						{
+							this->SelectorItem = this->RangeSelectionStart;
 						}
-						
-						// these should come after Private_SignalSelectionChanged(); because through a
-						// series of calls, Private_SignalSelectionChanged() could end up in a child
-						// that indexes into either of these arrays (the index wouldn't be updated yet,
-						// and could be invalid)
-						SparseItemInfos = TempSparseItemInfo;
-						DenseItemInfos  = TempDenseItemInfos;
+
+						this->Private_SignalSelectionChanged(ESelectInfo::Direct);
 					}
+						
+					// these should come after Private_SignalSelectionChanged(); because through a
+					// series of calls, Private_SignalSelectionChanged() could end up in a child
+					// that indexes into either of these arrays (the index wouldn't be updated yet,
+					// and could be invalid)
+					SparseItemInfos = TempSparseItemInfo;
+					DenseItemInfos  = TempDenseItemInfos;
 				}
 			}
-			
-			// Tick the TreeView superclass so that it can refresh.
-			// This may be due to TreeView requesting a refresh or because new items became visible due to resizing or scrolling.
-			SListView< ItemType >::Tick(AllottedGeometry, InCurrentTime, InDeltaTime);
 		}
+			
+		// Tick the TreeView superclass so that it can refresh.
+		// This may be due to TreeView requesting a refresh or because new items became visible due to resizing or scrolling.
+		SListView< ItemType >::Tick(AllottedGeometry, InCurrentTime, InDeltaTime);
+	}
 		
-		/** 
-		 * Given: an array of items (ItemsSource) each of which potentially has a child.
-		 * Task: populate the LinearezedItems array with a flattened version of the visible data items.
-		 *       In the process, remove any items that are not visible while maintaining any collapsed
-		 *       items that may have expanded children.
-		 *
-		 * @param ItemsSource          An array of data items each of which may have 0 or more children.
-		 * @param LinearizedItems      Array to populate with items based on expanded/collapsed state
-		 * @param NewDenseItemInfos    Array representing how nested each item in the Linearized items is, and whether it has children
-		 * @param TreeLevel            The current level of indentation
-		 * @param OutNewSelectedItems  Selected items minus any items that are no longer observed by the list
-		 * @param NewSparseItemInfo    Expanded items and items that have expanded children minus any items that are no longer observed by the list
-		 * @param bAddingItems         Are we adding encountered items to the linearized items list or just testing them for existence
-		 *
-		 * @return true if we encountered expanded children; false otherwise.
-		 */
-		bool PopulateLinearizedItems(
-			const TArray<ItemType>& ItemsSource,
-			TArray< ItemType >& InLinearizedItems,
-			TArray< FItemInfo >& NewDenseItemInfos,
-			int32 TreeLevel,
-			TSet< ItemType >& OutNewSelectedItems,
-			TMap< ItemType, FSparseItemInfo >& NewSparseItemInfo,
-			bool bAddingItems )
+	/** 
+	 * Given: an array of items (ItemsSource) each of which potentially has a child.
+	 * Task: populate the LinearezedItems array with a flattened version of the visible data items.
+	 *       In the process, remove any items that are not visible while maintaining any collapsed
+	 *       items that may have expanded children.
+	 *
+	 * @param ItemsSource          An array of data items each of which may have 0 or more children.
+	 * @param LinearizedItems      Array to populate with items based on expanded/collapsed state.
+	 * @param NewDenseItemInfos    Array representing how nested each item in the Linearized items is, and whether it has children.
+	 * @param TreeLevel            The current level of indentation.
+	 * @param OutNewSelectedItems  Selected items minus any items that are no longer observed by the list.
+	 * @param NewSparseItemInfo    Expanded items and items that have expanded children minus any items that are no longer observed by the list.
+	 * @param bAddingItems         Are we adding encountered items to the linearized items list or just testing them for existence.
+	 *
+	 * @return true if we encountered expanded children; false otherwise.
+	 */
+	bool PopulateLinearizedItems(
+		const TArray<ItemType>& ItemsSource,
+		TArray< ItemType >& InLinearizedItems,
+		TArray< FItemInfo >& NewDenseItemInfos,
+		int32 TreeLevel,
+		TSet< ItemType >& OutNewSelectedItems,
+		TMap< ItemType, FSparseItemInfo >& NewSparseItemInfo,
+		bool bAddingItems )
+	{
+		bool bSawExpandedItems = false;
+		for ( int32 ItemIndex = 0; ItemIndex < ItemsSource.Num(); ++ItemIndex )
 		{
-			bool bSawExpandedItems = false;
-			for ( int32 ItemIndex = 0; ItemIndex < ItemsSource.Num(); ++ItemIndex )
-			{
-				const ItemType& CurItem = ItemsSource[ItemIndex];
+			const ItemType& CurItem = ItemsSource[ItemIndex];
 
-				// Find this items children.
-				TArray<ItemType> ChildItems;
-				OnGetChildren.Execute( ItemsSource[ItemIndex], ChildItems );
+			// Find this items children.
+			TArray<ItemType> ChildItems;
+			OnGetChildren.Execute( ItemsSource[ItemIndex], ChildItems );
 
-				const bool bHasChildren = ChildItems.Num() > 0;
+			const bool bHasChildren = ChildItems.Num() > 0;
 
-				// Is this item expanded, does it have expanded children?
-				const FSparseItemInfo* CurItemInfo = SparseItemInfos.Find( CurItem );
-				const bool bIsExpanded = (CurItemInfo == NULL) ? false : CurItemInfo->bIsExpanded;
-				bool bHasExpandedChildren = (CurItemInfo == NULL) ? false : CurItemInfo->bHasExpandedChildren;
+			// Is this item expanded, does it have expanded children?
+			const FSparseItemInfo* CurItemInfo = SparseItemInfos.Find( CurItem );
+			const bool bIsExpanded = (CurItemInfo == nullptr) ? false : CurItemInfo->bIsExpanded;
+			bool bHasExpandedChildren = (CurItemInfo == nullptr) ? false : CurItemInfo->bHasExpandedChildren;
 								
-				// Add this item to the linearized list and update the selection set.
-				if (bAddingItems)
-				{
-					InLinearizedItems.Add( CurItem );
-
-					NewDenseItemInfos.Add( FItemInfo(TreeLevel, bHasChildren) );
+			// Add this item to the linearized list and update the selection set.
+			if (bAddingItems)
+			{
+				InLinearizedItems.Add( CurItem );
 
-					const bool bIsSelected = this->IsItemSelected( CurItem );
-					if (bIsSelected)
-					{
-						OutNewSelectedItems.Add( CurItem );
-					}
-				}				
-				
-				if ( bIsExpanded || bHasExpandedChildren )
-				{ 
-					// If this item is expanded, we should process all of its children at the next indentation level.
-					// If it is collapsed, process its children but do not add them to the linearized list.
-					const bool bAddChildItems = bAddingItems && bIsExpanded;
-					bHasExpandedChildren = PopulateLinearizedItems( ChildItems, InLinearizedItems, NewDenseItemInfos, TreeLevel+1, OutNewSelectedItems, NewSparseItemInfo, bAddChildItems );
-				}
+				NewDenseItemInfos.Add( FItemInfo(TreeLevel, bHasChildren) );
 
-				if ( bIsExpanded || bHasExpandedChildren )
+				const bool bIsSelected = this->IsItemSelected( CurItem );
+				if (bIsSelected)
 				{
-					// Update the item info for this tree item.
-					NewSparseItemInfo.Add( CurItem, FSparseItemInfo( bIsExpanded, bHasExpandedChildren) );
+					OutNewSelectedItems.Add( CurItem );
 				}
-
-
-				// If we encountered any expanded nodes along the way, the parent has expanded children.
-				bSawExpandedItems = bSawExpandedItems || bIsExpanded || bHasExpandedChildren;				
+			}				
+				
+			if ( bIsExpanded || bHasExpandedChildren )
+			{ 
+				// If this item is expanded, we should process all of its children at the next indentation level.
+				// If it is collapsed, process its children but do not add them to the linearized list.
+				const bool bAddChildItems = bAddingItems && bIsExpanded;
+				bHasExpandedChildren = PopulateLinearizedItems( ChildItems, InLinearizedItems, NewDenseItemInfos, TreeLevel+1, OutNewSelectedItems, NewSparseItemInfo, bAddChildItems );
 			}
 
-			return bSawExpandedItems;
-		}
-		
-		/**
-		 * Given a TreeItem, create a Widget to represent it in the tree view.
-		 *
-		 * @param InItem   Item to visualize
-		 *
-		 * @return A widget that represents this item.
-		 */
-		virtual TSharedRef<ITableRow> GenerateNewWidget( ItemType InItem )
-		{
-			if ( this->OnGenerateRow.IsBound() )
+			if ( bIsExpanded || bHasExpandedChildren )
 			{
-				return this->OnGenerateRow.Execute( InItem, this->SharedThis(this) );
+				// Update the item info for this tree item.
+				NewSparseItemInfo.Add( CurItem, FSparseItemInfo( bIsExpanded, bHasExpandedChildren) );
 			}
-			else
-			{
-				TSharedRef< STreeView<ItemType> > This = StaticCastSharedRef< STreeView<ItemType> >(this->AsShared());
 
-				// The programmer did not provide an OnGenerateRow() handler; let them know.
-				TSharedRef< STableRow<ItemType> > NewTreeItemWidget = 
-					SNew( STableRow<ItemType>, This )
-					.Content()
-					[
-						SNew(STextBlock) .Text( NSLOCTEXT("STreeView", "BrokenSetupMessage", "OnGenerateWidget() not assigned.") )
-					];
 
-				return NewTreeItemWidget;
-			}
+			// If we encountered any expanded nodes along the way, the parent has expanded children.
+			bSawExpandedItems = bSawExpandedItems || bIsExpanded || bHasExpandedChildren;				
 		}
+
+		return bSawExpandedItems;
+	}
 		
-		/**
-		 * Queue up a regeneration of the linearized items on the next tick.
-		 */
-		void RequestTreeRefresh()
+	/**
+	 * Given a TreeItem, create a Widget to represent it in the tree view.
+	 *
+	 * @param InItem   Item to visualize.
+	 * @return A widget that represents this item.
+	 */
+	virtual TSharedRef<ITableRow> GenerateNewWidget( ItemType InItem )
+	{
+		if ( this->OnGenerateRow.IsBound() )
 		{
-			bTreeItemsAreDirty = true;
-			RequestListRefresh();
+			return this->OnGenerateRow.Execute( InItem, this->SharedThis(this) );
 		}
-		
-		/**
-		 * Set whether some data item is expanded or not.
-		 * 
-		 * @param InItem         The item whose expansion state to control
-		 * @param InExpandItem   If true the item should be expanded; otherwise collapsed
-		 */
-		void SetItemExpansion( const ItemType& InItem, bool InShouldExpandItem )
+		else
 		{
-			Private_SetItemExpansion(InItem, InShouldExpandItem);
-		}
+			TSharedRef< STreeView<ItemType> > This = StaticCastSharedRef< STreeView<ItemType> >(this->AsShared());
 
-		/** Collapse all the items in the tree and expand InItem */
-		void SetSingleExpandedItem( const ItemType& InItem )
-		{
-			// Will we have to do any work?
-			const bool bItemAlreadyLoneExpanded = (this->SparseItemInfos.Num() == 1) && this->IsItemExpanded(InItem);
+			// The programmer did not provide an OnGenerateRow() handler; let them know.
+			TSharedRef< STableRow<ItemType> > NewTreeItemWidget = 
+				SNew( STableRow<ItemType>, This )
+				.Content()
+				[
+					SNew(STextBlock) .Text( NSLOCTEXT("STreeView", "BrokenSetupMessage", "OnGenerateWidget() not assigned.") )
+				];
 
-			if (!bItemAlreadyLoneExpanded)
-			{
-				this->SparseItemInfos.Empty();
-				Private_SetItemExpansion(InItem, true);
-			}
+			return NewTreeItemWidget;
 		}
+	}
 		
-		/** 
-		 * @param InItem   The data item whose expansion state to query
-		 *
-		 * @return true if the item is expanded; false otherwise.
-		 */
-		bool IsItemExpanded( const ItemType& InItem ) const
+	/** Queue up a regeneration of the linearized items on the next tick. */
+	void RequestTreeRefresh()
+	{
+		bTreeItemsAreDirty = true;
+		RequestListRefresh();
+	}
+		
+	/**
+	 * Set whether some data item is expanded or not.
+	 * 
+	 * @param InItem         The item whose expansion state to control.
+	 * @param InExpandItem   If true the item should be expanded; otherwise collapsed.
+	 */
+	void SetItemExpansion( const ItemType& InItem, bool InShouldExpandItem )
+	{
+		Private_SetItemExpansion(InItem, InShouldExpandItem);
+	}
+
+	/** Collapse all the items in the tree and expand InItem */
+	void SetSingleExpandedItem( const ItemType& InItem )
+	{
+		// Will we have to do any work?
+		const bool bItemAlreadyLoneExpanded = (this->SparseItemInfos.Num() == 1) && this->IsItemExpanded(InItem);
+
+		if (!bItemAlreadyLoneExpanded)
 		{
-			return Private_IsItemExpanded( InItem );
+			this->SparseItemInfos.Empty();
+			Private_SetItemExpansion(InItem, true);
 		}
+	}
+		
+	/** 
+	 * @param InItem   The data item whose expansion state to query.
+	 *
+	 * @return true if the item is expanded; false otherwise.
+	 */
+	bool IsItemExpanded( const ItemType& InItem ) const
+	{
+		return Private_IsItemExpanded( InItem );
+	}
 
 		
-		/**
-		 * Set the TreeItemsSource. The Tree will generate widgets to represent these items.
-		 *
-		 * @param InItemsSource  A pointer to the array of items that should be observed by this TreeView
-		 */
-		void SetTreeItemsSource( TArray<ItemType>* InItemsSource)
-		{
-			TreeItemsSource = InItemsSource;
-			RequestTreeRefresh();
-		}
+	/**
+	 * Set the TreeItemsSource. The Tree will generate widgets to represent these items.
+	 *
+	 * @param InItemsSource  A pointer to the array of items that should be observed by this TreeView.
+	 */
+	void SetTreeItemsSource( TArray<ItemType>* InItemsSource)
+	{
+		TreeItemsSource = InItemsSource;
+		RequestTreeRefresh();
+	}
 
-		/**
-		 * Generates a set of items that are currently expanded
-		 *
-		 * @param ExpandedItems	The generated set of expanded items.
-		 */
-		void GetExpandedItems( TSet< ItemType >& ExpandedItems ) const
+	/**
+	 * Generates a set of items that are currently expanded.
+	 *
+	 * @param ExpandedItems	The generated set of expanded items.
+	 */
+	void GetExpandedItems( TSet< ItemType >& ExpandedItems ) const
+	{
+		for( typename TMap<ItemType, FSparseItemInfo>::TConstIterator InfoIterator(SparseItemInfos); InfoIterator; ++InfoIterator )
 		{
-			for( typename TMap<ItemType, FSparseItemInfo>::TConstIterator InfoIterator(SparseItemInfos); InfoIterator; ++InfoIterator )
+			if ( InfoIterator.Value().bIsExpanded )
 			{
-				if ( InfoIterator.Value().bIsExpanded )
-				{
-					ExpandedItems.Add( InfoIterator.Key() );
-				}
-			}			
-		}
+				ExpandedItems.Add( InfoIterator.Key() );
+			}
+		}			
+	}
 
-		/**
-		 * Clears the entire set of expanded items
-		 */
-		void ClearExpandedItems()
-		{
-			SparseItemInfos.Empty();
-			RequestTreeRefresh();
-		}
+	/** Clears the entire set of expanded items. */
+	void ClearExpandedItems()
+	{
+		SparseItemInfos.Empty();
+		RequestTreeRefresh();
+	}
 
-	private:
-		/** Hide RequestListRefresh() for the tree widget's interface.  You should always call RequestTreeRefresh()
-		    for trees, not RequestListRefresh() */
-		void RequestListRefresh()
-		{
-			SListView< ItemType >::RequestListRefresh();
-		}
+private:
 
-	protected:
+	/** Hide RequestListRefresh() for the tree widget's interface.  You should always call RequestTreeRefresh()
+		for trees, not RequestListRefresh() */
+	void RequestListRefresh()
+	{
+		SListView< ItemType >::RequestListRefresh();
+	}
 
-		
-		/** The delegate that is invoked whenever we need to gather an item's children */
-		FOnGetChildren OnGetChildren;
+protected:
+	
+	/** The delegate that is invoked whenever we need to gather an item's children. */
+	FOnGetChildren OnGetChildren;
 
-		/** The delegate that is invoked to recursively expand/collapse a tree items children */
-		FOnSetExpansionRecursive OnSetExpansionRecursive;
+	/** The delegate that is invoked to recursively expand/collapse a tree items children. */
+	FOnSetExpansionRecursive OnSetExpansionRecursive;
 
-		/** A pointer to the items being observed by the tree view. */
-		TArray<ItemType>* TreeItemsSource;		
+	/** A pointer to the items being observed by the tree view. */
+	TArray<ItemType>* TreeItemsSource;		
 		
-		/** Info needed by a small fraction of tree items; some of these are not visible to the user */
-		TMap<ItemType, FSparseItemInfo> SparseItemInfos;
+	/** Info needed by a small fraction of tree items; some of these are not visible to the user. */
+	TMap<ItemType, FSparseItemInfo> SparseItemInfos;
 
-		/** Info needed by every item in the linearized version of the tree. */
-		TArray<FItemInfo> DenseItemInfos;
+	/** Info needed by every item in the linearized version of the tree. */
+	TArray<FItemInfo> DenseItemInfos;
 
-		/**
-		 * A linearized version of the items being observed by the tree view.
-		 * Note that we inherit from a ListView, which we point at this linearized version of the tree.
-		 */
-		TArray< ItemType > LinearizedItems;
+	/**
+	 * A linearized version of the items being observed by the tree view.
+	 * Note that we inherit from a ListView, which we point at this linearized version of the tree.
+	 */
+	TArray< ItemType > LinearizedItems;
 
-		/** The delegate that is invoked whnever an item in the tree is expanded or collapsed */
-		FOnExpansionChanged OnExpansionChanged;
+	/** The delegate that is invoked whenever an item in the tree is expanded or collapsed. */
+	FOnExpansionChanged OnExpansionChanged;
 
-	private:		
-		/** true when the LinearizedItems need to be regenerated. */
-		bool bTreeItemsAreDirty;
+private:		
 
+	/** true when the LinearizedItems need to be regenerated. */
+	bool bTreeItemsAreDirty;
 };
-
diff --git a/Engine/Source/Runtime/SlateCore/Private/Animation/CurveHandle.cpp b/Engine/Source/Runtime/SlateCore/Private/Animation/CurveHandle.cpp
--- a/Engine/Source/Runtime/SlateCore/Private/Animation/CurveHandle.cpp
+++ b/Engine/Source/Runtime/SlateCore/Private/Animation/CurveHandle.cpp
@@ -1,9 +1,5 @@
 // Copyright 1998-2014 Epic Games, Inc. All Rights Reserved.
 
-/*=============================================================================
-	CurveHandle.cpp: Implements the FCurveHandle structure.
-=============================================================================*/
-
 #include "SlateCorePrivatePCH.h"
 
 
diff --git a/Engine/Source/Runtime/SlateCore/Private/Animation/CurveSequence.cpp b/Engine/Source/Runtime/SlateCore/Private/Animation/CurveSequence.cpp
--- a/Engine/Source/Runtime/SlateCore/Private/Animation/CurveSequence.cpp
+++ b/Engine/Source/Runtime/SlateCore/Private/Animation/CurveSequence.cpp
@@ -1,9 +1,5 @@
 // Copyright 1998-2014 Epic Games, Inc. All Rights Reserved.
 
-/*=============================================================================
-	CurveSequence.cpp: Implements the FCurveSequence structure.
-=============================================================================*/
-
 #include "SlateCorePrivatePCH.h"
 
 
diff --git a/Engine/Source/Runtime/SlateCore/Private/Application/SlateApplicationBase.cpp b/Engine/Source/Runtime/SlateCore/Private/Application/SlateApplicationBase.cpp
--- a/Engine/Source/Runtime/SlateCore/Private/Application/SlateApplicationBase.cpp
+++ b/Engine/Source/Runtime/SlateCore/Private/Application/SlateApplicationBase.cpp
@@ -1,9 +1,5 @@
 // Copyright 1998-2014 Epic Games, Inc. All Rights Reserved.
 
-/*=============================================================================
-	SlateApplicationBase.cpp: Implements the FSlateApplicationBase class.
-=============================================================================*/
-
 #include "SlateCorePrivatePCH.h"
 
 
diff --git a/Engine/Source/Runtime/SlateCore/Private/Application/SlateWindowHelper.cpp b/Engine/Source/Runtime/SlateCore/Private/Application/SlateWindowHelper.cpp
--- a/Engine/Source/Runtime/SlateCore/Private/Application/SlateWindowHelper.cpp
+++ b/Engine/Source/Runtime/SlateCore/Private/Application/SlateWindowHelper.cpp
@@ -1,9 +1,5 @@
 // Copyright 1998-2014 Epic Games, Inc. All Rights Reserved.
 
-/*=============================================================================
-	SlateWindowHelper.cpp: Implements the FSlateWindowHelper class.
-=============================================================================*/
-
 #include "SlateCorePrivatePCH.h"
 
 
diff --git a/Engine/Source/Runtime/SlateCore/Private/Application/ThrottleManager.cpp b/Engine/Source/Runtime/SlateCore/Private/Application/ThrottleManager.cpp
--- a/Engine/Source/Runtime/SlateCore/Private/Application/ThrottleManager.cpp
+++ b/Engine/Source/Runtime/SlateCore/Private/Application/ThrottleManager.cpp
@@ -1,9 +1,5 @@
 // Copyright 1998-2014 Epic Games, Inc. All Rights Reserved.
 
-/*=============================================================================
-	ThrottleManager.cpp: Implements the FSlateThrottleManager class.
-=============================================================================*/
-
 #include "SlateCorePrivatePCH.h"
 
 
diff --git a/Engine/Source/Runtime/SlateCore/Private/Brushes/SlateDynamicImageBrush.cpp b/Engine/Source/Runtime/SlateCore/Private/Brushes/SlateDynamicImageBrush.cpp
--- a/Engine/Source/Runtime/SlateCore/Private/Brushes/SlateDynamicImageBrush.cpp
+++ b/Engine/Source/Runtime/SlateCore/Private/Brushes/SlateDynamicImageBrush.cpp
@@ -1,9 +1,5 @@
 // Copyright 1998-2014 Epic Games, Inc. All Rights Reserved.
 
-/*=============================================================================
-	SlateDynamicImageBrush.cpp: Implements the FSlateDynamicImageBrush structure.
-=============================================================================*/
-
 #include "SlateCorePrivatePCH.h"
 
 
diff --git a/Engine/Source/Runtime/SlateCore/Private/Fonts/FontCacheUtils.h b/Engine/Source/Runtime/SlateCore/Private/Fonts/FontCacheUtils.h
--- a/Engine/Source/Runtime/SlateCore/Private/Fonts/FontCacheUtils.h
+++ b/Engine/Source/Runtime/SlateCore/Private/Fonts/FontCacheUtils.h
@@ -1,9 +1,5 @@
 // Copyright 1998-2014 Epic Games, Inc. All Rights Reserved.
 
-/*=============================================================================
-	FontCacheUtils.h: Declares the FLRUStringCache class.
-=============================================================================*/
-
 #pragma once
 
 
diff --git a/Engine/Source/Runtime/SlateCore/Private/Fonts/SlateFontInfo.cpp b/Engine/Source/Runtime/SlateCore/Private/Fonts/SlateFontInfo.cpp
--- a/Engine/Source/Runtime/SlateCore/Private/Fonts/SlateFontInfo.cpp
+++ b/Engine/Source/Runtime/SlateCore/Private/Fonts/SlateFontInfo.cpp
@@ -1,9 +1,5 @@
 // Copyright 1998-2014 Epic Games, Inc. All Rights Reserved.
 
-/*=============================================================================
-	SlateFontInfo.cpp: Implements the FSlateFontInfo structure.
-=============================================================================*/
-
 #include "SlateCorePrivatePCH.h"
 
 
diff --git a/Engine/Source/Runtime/SlateCore/Private/Input/Events.cpp b/Engine/Source/Runtime/SlateCore/Private/Input/Events.cpp
--- a/Engine/Source/Runtime/SlateCore/Private/Input/Events.cpp
+++ b/Engine/Source/Runtime/SlateCore/Private/Input/Events.cpp
@@ -1,9 +1,5 @@
 // Copyright 1998-2014 Epic Games, Inc. All Rights Reserved.
 
-/*=============================================================================
-	Events.cpp: Implements various input event related types.
-=============================================================================*/
-
 #include "SlateCorePrivatePCH.h"
 
 
diff --git a/Engine/Source/Runtime/SlateCore/Private/Layout/ArrangedChildren.cpp b/Engine/Source/Runtime/SlateCore/Private/Layout/ArrangedChildren.cpp
--- a/Engine/Source/Runtime/SlateCore/Private/Layout/ArrangedChildren.cpp
+++ b/Engine/Source/Runtime/SlateCore/Private/Layout/ArrangedChildren.cpp
@@ -1,9 +1,5 @@
 // Copyright 1998-2014 Epic Games, Inc. All Rights Reserved.
 
-/*=============================================================================
-	ArrangedChildren.cpp: Implements the FArrangedChildren class.
-=============================================================================*/
-
 #include "SlateCorePrivatePCH.h"
 
 
diff --git a/Engine/Source/Runtime/SlateCore/Private/Layout/ArrangedWidget.cpp b/Engine/Source/Runtime/SlateCore/Private/Layout/ArrangedWidget.cpp
--- a/Engine/Source/Runtime/SlateCore/Private/Layout/ArrangedWidget.cpp
+++ b/Engine/Source/Runtime/SlateCore/Private/Layout/ArrangedWidget.cpp
@@ -1,9 +1,5 @@
 // Copyright 1998-2014 Epic Games, Inc. All Rights Reserved.
 
-/*=============================================================================
-	ArrangedWidget.cpp: Implements the FArrangedWidget class.
-=============================================================================*/
-
 #include "SlateCorePrivatePCH.h"
 
 
diff --git a/Engine/Source/Runtime/SlateCore/Private/Layout/Geometry.cpp b/Engine/Source/Runtime/SlateCore/Private/Layout/Geometry.cpp
--- a/Engine/Source/Runtime/SlateCore/Private/Layout/Geometry.cpp
+++ b/Engine/Source/Runtime/SlateCore/Private/Layout/Geometry.cpp
@@ -1,9 +1,5 @@
 // Copyright 1998-2014 Epic Games, Inc. All Rights Reserved.
 
-/*=============================================================================
-	Geometry.cpp: Implements the FGeometry class.
-=============================================================================*/
-
 #include "SlateCorePrivatePCH.h"
 
 
@@ -15,17 +11,20 @@ FArrangedWidget FGeometry::MakeChild( const TSharedRef<SWidget>& InWidget, const
 	return FArrangedWidget( InWidget, this->MakeChild(ChildOffset, ChildSize, InScale) );
 }
 
+
 FPaintGeometry FGeometry::ToPaintGeometry( FVector2D InOffset, FVector2D InSize, float InScale ) const
 {
 	const float CombinedScale = this->Scale*InScale;
 	return FPaintGeometry( this->AbsolutePosition + InOffset*CombinedScale, InSize * CombinedScale, CombinedScale );
 }
 
+
 FPaintGeometry FGeometry::ToInflatedPaintGeometry( const FVector2D& InflateAmount ) const
 {
 	return FPaintGeometry( this->AbsolutePosition - InflateAmount*this->Scale, (this->Size+InflateAmount*2)*this->Scale, this->Scale );
 }
 
+
 FPaintGeometry FGeometry::CenteredPaintGeometryOnLeft( const FVector2D& SizeBeingAligned, float InScale ) const
 {
 	const float CombinedScale = this->Scale*InScale;
@@ -37,6 +36,7 @@ FPaintGeometry FGeometry::CenteredPaintGeometryOnLeft( const FVector2D& SizeBein
 	);
 }
 
+
 FPaintGeometry FGeometry::CenteredPaintGeometryOnRight( const FVector2D& SizeBeingAligned, float InScale ) const
 {
 	const float CombinedScale = this->Scale*InScale;
@@ -48,6 +48,7 @@ FPaintGeometry FGeometry::CenteredPaintGeometryOnRight( const FVector2D& SizeBei
 	);
 }
 
+
 FPaintGeometry FGeometry::CenteredPaintGeometryBelow( const FVector2D& SizeBeingAligned, float InScale ) const
 {
 	const float CombinedScale = this->Scale*InScale;
diff --git a/Engine/Source/Runtime/SlateCore/Private/Layout/SlateRect.cpp b/Engine/Source/Runtime/SlateCore/Private/Layout/SlateRect.cpp
--- a/Engine/Source/Runtime/SlateCore/Private/Layout/SlateRect.cpp
+++ b/Engine/Source/Runtime/SlateCore/Private/Layout/SlateRect.cpp
@@ -1,9 +1,5 @@
 // Copyright 1998-2014 Epic Games, Inc. All Rights Reserved.
 
-/*=============================================================================
-	SlateRect.cpp: Implements the FSlateRect structure.
-=============================================================================*/
-
 #include "SlateCorePrivatePCH.h"
 
 
diff --git a/Engine/Source/Runtime/SlateCore/Private/Layout/Visibility.cpp b/Engine/Source/Runtime/SlateCore/Private/Layout/Visibility.cpp
--- a/Engine/Source/Runtime/SlateCore/Private/Layout/Visibility.cpp
+++ b/Engine/Source/Runtime/SlateCore/Private/Layout/Visibility.cpp
@@ -1,9 +1,5 @@
 // Copyright 1998-2014 Epic Games, Inc. All Rights Reserved.
 
-/*=============================================================================
-	Visibility.cpp: Implements the FVisibility class.
-=============================================================================*/
-
 #include "SlateCorePrivatePCH.h"
 
 
diff --git a/Engine/Source/Runtime/SlateCore/Private/Logging/EventLogger.cpp b/Engine/Source/Runtime/SlateCore/Private/Logging/EventLogger.cpp
--- a/Engine/Source/Runtime/SlateCore/Private/Logging/EventLogger.cpp
+++ b/Engine/Source/Runtime/SlateCore/Private/Logging/EventLogger.cpp
@@ -1,9 +1,5 @@
 // Copyright 1998-2014 Epic Games, Inc. All Rights Reserved.
 
-/*=============================================================================
-	EventLogger.cpp: Implements various Slate event logger implementations.
-=============================================================================*/
-
 #include "SlateCorePrivatePCH.h"
 
 
diff --git a/Engine/Source/Runtime/SlateCore/Private/Rendering/ElementBatcher.cpp b/Engine/Source/Runtime/SlateCore/Private/Rendering/ElementBatcher.cpp
--- a/Engine/Source/Runtime/SlateCore/Private/Rendering/ElementBatcher.cpp
+++ b/Engine/Source/Runtime/SlateCore/Private/Rendering/ElementBatcher.cpp
@@ -1,9 +1,5 @@
 // Copyright 1998-2014 Epic Games, Inc. All Rights Reserved.
 
-/*=============================================================================
-	ElementBatcher.cpp: Implements the FSlateElementBatcher class.
-=============================================================================*/
-
 #include "SlateCorePrivatePCH.h"
 
 
@@ -68,11 +64,6 @@ FSlateElementBatcher::~FSlateElementBatcher()
 }
 
 
-/** 
- * Adds a new element to be rendered 
- * 
- * @param InElement	The element to render 
- */
 void FSlateElementBatcher::AddElements( const TArray<FSlateDrawElement>& DrawElements )
 {
 	SCOPE_CYCLE_COUNTER(STAT_SlateBuildElements);
@@ -159,13 +150,7 @@ void FSlateElementBatcher::AddQuadElement( const FVector2D& Position, const FVec
 	BatchIndices.Add( IndexStart + 3 );
 }
 
-/** 
- * Creates vertices necessary to draw a box element.  
- *
- * @param Position	The top left screen space position of the element
- * @param Size		The size of the element
- * @param InPayload	The margin and texture data necessary to draw the quad
- */
+
 void FSlateElementBatcher::AddBoxElement( const FVector2D& InPosition, const FVector2D& Size, float Scale, const FSlateDataPayload& InPayload, const FSlateRect& InClippingRect, ESlateDrawEffect::Type InDrawEffects, uint32 Layer )
 {
 	check( InPayload.BrushResource );
@@ -456,12 +441,7 @@ void FSlateElementBatcher::AddBoxElement( const FVector2D& InPosition, const FVe
 	}
 }
 
-/** 
- * Creates vertices necessary to draw a string (one quad per character)
- *
- * @param Position	The top left screen space position of the element
- * @param InPayload	The data payload containing font,text, and size information for the string
- */
+
 void FSlateElementBatcher::AddTextElement( const FVector2D& Position, const FVector2D& Size, float InScale, const FSlateDataPayload& InPayload, const FSlateRect& InClippingRect, ESlateDrawEffect::Type InDrawEffects, uint32 Layer )
 {
 	SCOPE_CYCLE_COUNTER(STAT_SlateBuildText);
@@ -607,13 +587,7 @@ void FSlateElementBatcher::AddTextElement( const FVector2D& Position, const FVec
 	}
 }
 
-/** 
- * Creates vertices necessary to draw a gradient box (horizontal or vertical)
- *
- * @param Position	The top left screen space position of the element
- * @param Size		The size of the element
- * @param InPayload	The data payload containing gradient stops
- */
+
 void FSlateElementBatcher::AddGradientElement( const FVector2D& Position, const FVector2D& Size, float Scale, const class FSlateDataPayload& InPayload, const FSlateRect& InClippingRect, ESlateDrawEffect::Type InDrawEffects, uint32 Layer )
 {
 	SCOPE_CYCLE_COUNTER(STAT_SlateBuildGradients);
@@ -729,6 +703,7 @@ void FSlateElementBatcher::AddGradientElement( const FVector2D& Position, const
 	}
 }
 
+
 void FSlateElementBatcher::AddSplineElement( const FVector2D& Position, float Scale, const FSlateDataPayload& InPayload, const FSlateRect& InClippingRect, ESlateDrawEffect::Type InDrawEffects, uint32 Layer)
 {
 	SCOPE_CYCLE_COUNTER(STAT_SlateBuildSplines);
@@ -815,6 +790,7 @@ void FSlateElementBatcher::AddSplineElement( const FVector2D& Position, float Sc
 	}
 }
 
+
 /**
  * Calculates the intersection of two line segments P1->P2, P3->P4
  * The tolerance setting is used when the lines arent currently intersecting but will intersect in the future  
@@ -1016,14 +992,7 @@ void FSlateElementBatcher::AddLineElement( const FVector2D& Position, const FVec
 	}
 }
 
-/** 
- * Creates vertices necessary to draw viewport
- *
- * @param Position	The top left screen space position of the element
- * @param Size		The size of the element
- * @param InPayload	The data payload containing gradient stops
- * @param InClippingRect	The clipping rect of the element
- */
+
 void FSlateElementBatcher::AddViewportElement( const FVector2D& InPosition, const FVector2D& InSize, float Scale, const FSlateDataPayload& InPayload, const FSlateRect& InClippingRect, ESlateDrawEffect::Type InDrawEffects, uint32 Layer )
 {
 	SCOPE_CYCLE_COUNTER(STAT_SlateBuildViewports);
@@ -1089,14 +1058,7 @@ void FSlateElementBatcher::AddViewportElement( const FVector2D& InPosition, cons
 
 }
 
-/** 
- * Creates vertices necessary to draw a border element
- *
- * @param Position	The top left screen space position of the element
- * @param Size		The size of the element
- * @param InPayload	The data payload containing gradient stops
- * @param InClippingRect	The clipping rect of the element
- */
+
 void FSlateElementBatcher::AddBorderElement( const FVector2D& InPosition, const FVector2D& Size, float Scale, const FSlateDataPayload& InPayload, const FSlateRect& InClippingRect, ESlateDrawEffect::Type InDrawEffects, uint32 Layer )
 {
 	SCOPE_CYCLE_COUNTER(STAT_SlateBuildBorders);
@@ -1301,6 +1263,7 @@ void FSlateElementBatcher::AddBorderElement( const FVector2D& InPosition, const
 	BatchIndices.Add( IndexStart + 31 );
 }
 
+
 void FSlateElementBatcher::AddCustomElement( const FVector2D& Position, const FVector2D& Size, float Scale, const FSlateDataPayload& InPayload, const FSlateRect& InClippingRect, ESlateDrawEffect::Type DrawEffects, uint32 Layer )
 {
 	if( InPayload.CustomDrawer.IsValid() )
@@ -1320,6 +1283,7 @@ void FSlateElementBatcher::AddCustomElement( const FVector2D& Position, const FV
 	}
 }
 
+
 FSlateElementBatch& FSlateElementBatcher::FindBatchForElement( 
 	uint32 Layer, 
 	const FShaderParams& ShaderParams, 
@@ -1390,6 +1354,7 @@ FSlateElementBatch& FSlateElementBatcher::FindBatchForElement(
 	return *ElementBatch;
 }
 
+
 void FSlateElementBatcher::AddVertices( TArray<FSlateVertex>& OutVertices, FSlateElementBatch& ElementBatch, const TArray<FSlateVertex>& VertexBatch )
 {
 	uint32 FirstIndex = OutVertices.Num();
@@ -1407,6 +1372,7 @@ void FSlateElementBatcher::AddVertices( TArray<FSlateVertex>& OutVertices, FSlat
 	ElementBatch.NumVertices = VertexBatch.Num();
 }
 
+
 void FSlateElementBatcher::AddIndices( TArray<SlateIndex>& OutIndices, FSlateElementBatch& ElementBatch, const TArray<SlateIndex>& IndexBatch )
 {
 	uint32 FirstIndex = OutIndices.Num();
@@ -1423,9 +1389,7 @@ void FSlateElementBatcher::AddIndices( TArray<SlateIndex>& OutIndices, FSlateEle
 	ElementBatch.NumIndices = IndexBatch.Num();
 }
 
-/**
- * Updates rendering buffers
- */
+
 void FSlateElementBatcher::FillBatchBuffers( FSlateWindowElementList& WindowElementList, bool& bRequiresStencilTest )
 {
 	SCOPE_CYCLE_COUNTER( STAT_SlateUpdateBufferGTTime );
@@ -1497,9 +1461,7 @@ void FSlateElementBatcher::FillBatchBuffers( FSlateWindowElementList& WindowElem
 
 }
 
-/** 
- * Resets all stored data accumulated during the batching process
- */
+
 void FSlateElementBatcher::ResetBatches()
 {
 	LayerToElementBatches.Reset();
@@ -1523,4 +1485,3 @@ void FSlateElementBatcher::ResetStats()
 	TotalVertexMemory = 0;
 	TotalIndexMemory = 0;
 }
-
diff --git a/Engine/Source/Runtime/SlateCore/Private/Rendering/SlateDrawBuffer.cpp b/Engine/Source/Runtime/SlateCore/Private/Rendering/SlateDrawBuffer.cpp
--- a/Engine/Source/Runtime/SlateCore/Private/Rendering/SlateDrawBuffer.cpp
+++ b/Engine/Source/Runtime/SlateCore/Private/Rendering/SlateDrawBuffer.cpp
@@ -1,9 +1,5 @@
 // Copyright 1998-2014 Epic Games, Inc. All Rights Reserved.
 
-/*=============================================================================
-	SlateDrawBuffer.cpp: Implements the FSlateDrawBuffer class.
-=============================================================================*/
-
 #include "SlateCorePrivatePCH.h"
 
 
diff --git a/Engine/Source/Runtime/SlateCore/Private/Rendering/SlateRenderer.cpp b/Engine/Source/Runtime/SlateCore/Private/Rendering/SlateRenderer.cpp
--- a/Engine/Source/Runtime/SlateCore/Private/Rendering/SlateRenderer.cpp
+++ b/Engine/Source/Runtime/SlateCore/Private/Rendering/SlateRenderer.cpp
@@ -1,9 +1,5 @@
 // Copyright 1998-2014 Epic Games, Inc. All Rights Reserved.
 
-/*=============================================================================
-	Renderer.cpp: Implements the SNullWidget class.
-=============================================================================*/
-
 #include "SlateCorePrivatePCH.h"
 
 
diff --git a/Engine/Source/Runtime/SlateCore/Private/SlateCoreClasses.cpp b/Engine/Source/Runtime/SlateCore/Private/SlateCoreClasses.cpp
--- a/Engine/Source/Runtime/SlateCore/Private/SlateCoreClasses.cpp
+++ b/Engine/Source/Runtime/SlateCore/Private/SlateCoreClasses.cpp
@@ -1,9 +1,5 @@
 // Copyright 1998-2014 Epic Games, Inc. All Rights Reserved.
 
-/*=============================================================================
-	SlateCoreClasses.cpp: Implements the module's script classes.
-=============================================================================*/
-
 #include "SlateCorePrivatePCH.h"
 
 
diff --git a/Engine/Source/Runtime/SlateCore/Private/SlateCoreModule.cpp b/Engine/Source/Runtime/SlateCore/Private/SlateCoreModule.cpp
--- a/Engine/Source/Runtime/SlateCore/Private/SlateCoreModule.cpp
+++ b/Engine/Source/Runtime/SlateCore/Private/SlateCoreModule.cpp
@@ -1,9 +1,5 @@
 // Copyright 1998-2014 Epic Games, Inc. All Rights Reserved.
 
-/*=============================================================================
-	SlateCoreModule.cpp: Implements the FSlateCoreModule class.
-=============================================================================*/
-
 #include "SlateCorePrivatePCH.h"
 #include "ModuleManager.h"
 
