Fix compile error

diff --git a/Engine/Source/Runtime/Slate/Private/Widgets/SViewport.cpp b/Engine/Source/Runtime/Slate/Private/Widgets/SViewport.cpp
--- a/Engine/Source/Runtime/Slate/Private/Widgets/SViewport.cpp
+++ b/Engine/Source/Runtime/Slate/Private/Widgets/SViewport.cpp
@@ -83,7 +83,7 @@ int32 SViewport::OnPaint( const FPaintArgs& Args, const FGeometry& AllottedGeome
 				CursorPosScreenSpace,
 				FVector2D::ZeroVector,
 				TSet<FKey>(),
-				FModifierKeysState(false, false, false, false, false, false, false, false) )
+				FModifierKeysState() )
 		 );
 		EMouseCursor::Type CursorType = Reply.GetCursor();
 
