Fixed scenarios where editable text did not perform proper kerning with DPI scaling enabled.

Also refactored call to MoveCursor() to be a little less obscure.
New API presumes that Word-granularity movement will be supported soon.
Removed excessive use of cached geometry in multiline editable text.

Fixes TTP 305729.

diff --git a/Engine/Source/Runtime/Slate/Private/Framework/Text/TextEditHelper.cpp b/Engine/Source/Runtime/Slate/Private/Framework/Text/TextEditHelper.cpp
--- a/Engine/Source/Runtime/Slate/Private/Framework/Text/TextEditHelper.cpp
+++ b/Engine/Source/Runtime/Slate/Private/Framework/Text/TextEditHelper.cpp
@@ -98,45 +98,45 @@ FReply FTextEditHelper::OnKeyDown( const FKeyboardEvent& InKeyboardEvent, const
 
 	if( Key == EKeys::Left )
 	{
-		return TextEditor->MoveCursor(
+		return TextEditor->MoveCursor( FMoveCursor::Cardinal(
 			// Ctrl moves a whole word instead of one character.	
-			InKeyboardEvent.IsControlDown() ? ECursorMoveMethod::Word : ECursorMoveMethod::CharacterHorizontal,
+			InKeyboardEvent.IsControlDown( ) ? ECursorMoveGranularity::Word : ECursorMoveGranularity::Character,
 			// Move left
-			FVector2D( -1, 0 ),
+			FIntPoint( -1, 0 ),
 			// Shift selects text.	
 			InKeyboardEvent.IsShiftDown() ? ECursorAction::SelectText : ECursorAction::MoveCursor
-		);
+		));
 	}
 	else if( Key == EKeys::Right )
 	{
-		return TextEditor->MoveCursor(
+		return TextEditor->MoveCursor( FMoveCursor::Cardinal(
 			// Ctrl moves a whole word instead of one character.	
-			InKeyboardEvent.IsControlDown() ? ECursorMoveMethod::Word : ECursorMoveMethod::CharacterHorizontal,
+			InKeyboardEvent.IsControlDown( ) ? ECursorMoveGranularity::Word : ECursorMoveGranularity::Character,
 			// Move right
-			FVector2D( +1, 0 ),
+			FIntPoint( +1, 0 ),
 			// Shift selects text.	
 			InKeyboardEvent.IsShiftDown() ? ECursorAction::SelectText : ECursorAction::MoveCursor
-		);
+		));
 	}
 	else if ( Key == EKeys::Up )
 	{
-		return TextEditor->MoveCursor(
-			ECursorMoveMethod::CharacterVertical,
+		return TextEditor->MoveCursor( FMoveCursor::Cardinal(
+			InKeyboardEvent.IsControlDown( ) ? ECursorMoveGranularity::Word : ECursorMoveGranularity::Character,
 			// Move up
-			FVector2D( 0, -1 ),
+			FIntPoint( 0, -1 ),
 			// Shift selects text.	
 			InKeyboardEvent.IsShiftDown() ? ECursorAction::SelectText : ECursorAction::MoveCursor
-		);
+		));
 	}
 	else if ( Key == EKeys::Down )
 	{
-		return TextEditor->MoveCursor(
-			ECursorMoveMethod::CharacterVertical,
+		return TextEditor->MoveCursor( FMoveCursor::Cardinal(
+			InKeyboardEvent.IsControlDown( ) ? ECursorMoveGranularity::Word : ECursorMoveGranularity::Character,
 			// Move down
-			FVector2D( 0, +1 ),
+			FIntPoint( 0, +1 ),
 			// Shift selects text.	
 			InKeyboardEvent.IsShiftDown() ? ECursorAction::SelectText : ECursorAction::MoveCursor
-		);
+		));
 	}
 	else if( Key == EKeys::Home )
 	{
@@ -170,13 +170,13 @@ FReply FTextEditHelper::OnKeyDown( const FKeyboardEvent& InKeyboardEvent, const
 		// Delete to next word boundary (Ctrl+Delete)
 		if (InKeyboardEvent.IsControlDown() && !InKeyboardEvent.IsAltDown() && !InKeyboardEvent.IsShiftDown())
 		{
-			TextEditor->MoveCursor(
-				ECursorMoveMethod::Word, 
+			TextEditor->MoveCursor( FMoveCursor::Cardinal(
+				ECursorMoveGranularity::Word, 
 				// Move right
-				FVector2D(+1, 0),
+				FIntPoint(+1, 0),
 				// selects text.	
 				ECursorAction::SelectText
-			);
+			));
 		}
 
 		FScopedTextTransaction TextTransaction(TextEditor);
@@ -238,12 +238,12 @@ FReply FTextEditHelper::OnKeyDown( const FKeyboardEvent& InKeyboardEvent, const
 	{
 		FScopedTextTransaction TextTransaction(TextEditor);
 
-		TextEditor->MoveCursor(
-			ECursorMoveMethod::Word, 
+		TextEditor->MoveCursor( FMoveCursor::Cardinal(
+			ECursorMoveGranularity::Word, 
 			// Move left
-			FVector2D(-1, 0), 
+			FIntPoint(-1, 0), 
 			ECursorAction::SelectText
-		);
+		));
 		TextEditor->BackspaceChar();
 
 		return FReply::Handled();
@@ -302,13 +302,13 @@ FReply FTextEditHelper::OnMouseButtonDown( const FGeometry& MyGeometry, const FP
 			{
 				if( InMouseEvent.IsShiftDown() )
 				{
-					TextEditor->MoveCursor( ECursorMoveMethod::ScreenPosition, MyGeometry.AbsoluteToLocal( InMouseEvent.GetScreenSpacePosition() ), ECursorAction::SelectText );
+					TextEditor->MoveCursor( FMoveCursor::ViaScreenPointer( MyGeometry.AbsoluteToLocal( InMouseEvent.GetScreenSpacePosition( ) ), MyGeometry.Scale, ECursorAction::SelectText ) );
 				}
 				else
 				{
 					// Deselect any text that was selected
 					TextEditor->ClearSelection();
-					TextEditor->MoveCursor( ECursorMoveMethod::ScreenPosition, MyGeometry.AbsoluteToLocal( InMouseEvent.GetScreenSpacePosition() ), ECursorAction::MoveCursor );
+					TextEditor->MoveCursor( FMoveCursor::ViaScreenPointer( MyGeometry.AbsoluteToLocal( InMouseEvent.GetScreenSpacePosition( ) ), MyGeometry.Scale, ECursorAction::MoveCursor ) );
 				}
 
 				// Start drag selection
@@ -318,7 +318,7 @@ FReply FTextEditHelper::OnMouseButtonDown( const FGeometry& MyGeometry, const FP
 			{
 				// If the user right clicked on a character that wasn't already selected, we'll clear
 				// the selection
-				if (TextEditor->AnyTextSelected() && !TextEditor->IsTextSelectedAt(MyGeometry.AbsoluteToLocal(InMouseEvent.GetScreenSpacePosition())))
+				if ( TextEditor->AnyTextSelected( ) && !TextEditor->IsTextSelectedAt( MyGeometry, InMouseEvent.GetScreenSpacePosition() ) )
 				{
 					// Deselect any text that was selected
 					TextEditor->ClearSelection();
@@ -341,7 +341,7 @@ FReply FTextEditHelper::OnMouseMove( const FGeometry& InMyGeometry, const FPoint
 
 	if( TextEditor->IsDragSelecting() && TextEditor->GetWidget()->HasMouseCapture() )
 	{
-		TextEditor->MoveCursor( ECursorMoveMethod::ScreenPosition, InMyGeometry.AbsoluteToLocal( InMouseEvent.GetScreenSpacePosition() ), ECursorAction::SelectText );
+		TextEditor->MoveCursor( FMoveCursor::ViaScreenPointer( InMyGeometry.AbsoluteToLocal( InMouseEvent.GetScreenSpacePosition( ) ), InMyGeometry.Scale, ECursorAction::SelectText ) );
 		TextEditor->SetHasDragSelectedSinceFocused( true );
 		Reply = FReply::Handled();
 	}
@@ -413,7 +413,7 @@ FReply FTextEditHelper::OnMouseButtonDoubleClick(const FGeometry& InMyGeometry,
 
 	if (InMouseEvent.GetEffectingButton() == EKeys::LeftMouseButton)
 	{
-		TextEditor->SelectWordAt(InMyGeometry.AbsoluteToLocal(InMouseEvent.GetScreenSpacePosition()));
+		TextEditor->SelectWordAt(InMyGeometry, InMouseEvent.GetScreenSpacePosition());
 
 		Reply = FReply::Handled();
 	}
diff --git a/Engine/Source/Runtime/Slate/Private/Widgets/Input/SEditableText.cpp b/Engine/Source/Runtime/Slate/Private/Widgets/Input/SEditableText.cpp
--- a/Engine/Source/Runtime/Slate/Private/Widgets/Input/SEditableText.cpp
+++ b/Engine/Source/Runtime/Slate/Private/Widgets/Input/SEditableText.cpp
@@ -201,7 +201,7 @@ void SEditableText::Tick( const FGeometry& AllottedGeometry, const double InCurr
 	const FString VisibleText = GetStringToRender();
 	const FSlateFontInfo& FontInfo = Font.Get();
 	const TSharedRef< FSlateFontMeasure > FontMeasureService = FSlateApplication::Get().GetRenderer()->GetFontMeasureService();
-	const float FontMaxCharHeight = FontMeasureService->Measure( FString(), FontInfo ).Y;
+	const float FontMaxCharHeight = FontMeasureService->Measure( FString( ), FontInfo, AllottedGeometry.Scale ).Y;
 
 
 	// Update caret position
@@ -212,7 +212,7 @@ void SEditableText::Tick( const FGeometry& AllottedGeometry, const double InCurr
 		if( bShouldAppearFocused )
 		{
 			FString TextBeforeCaret( VisibleText.Left( CaretPosition ) );
-			const FVector2D TextBeforeCaretSize( FontMeasureService->Measure( TextBeforeCaret, FontInfo ) );
+			const FVector2D TextBeforeCaretSize( FontMeasureService->Measure( TextBeforeCaret, FontInfo, AllottedGeometry.Scale ) );
 
 			FinalCaretTargetPosition = TextBeforeCaretSize.X;
 		}
@@ -243,26 +243,27 @@ void SEditableText::Tick( const FGeometry& AllottedGeometry, const double InCurr
 		const float WidgetSpaceCaretLeft = ScrollHelper.FromScrollerSpace( FVector2D( CaretLeft, ScrollHelper.GetPosition().Y ) ).X;
 		const float WidgetSpaceCaretRight = ScrollHelper.FromScrollerSpace( FVector2D( CaretRight, ScrollHelper.GetPosition().Y ) ).X;
 
+		const float AllotedGeometryDrawX = AllottedGeometry.Size.X*AllottedGeometry.Scale;
+
 		// If caret is to the left of the scrolled area, start scrolling left
 		if( WidgetSpaceCaretLeft < 0.0f )
 		{
 			ScrollHelper.SetPosition( FVector2D( CaretLeft, ScrollHelper.GetPosition().Y ) );
 		}
-
 		// If caret is to the right of the scrolled area, start scrolling right
-		else if( WidgetSpaceCaretRight > AllottedGeometry.Size.X )
+		else if ( WidgetSpaceCaretRight > AllotedGeometryDrawX )
 		{
-			ScrollHelper.SetPosition( FVector2D( CaretRight - AllottedGeometry.Size.X, ScrollHelper.GetPosition().Y ) );
+			ScrollHelper.SetPosition( FVector2D( CaretRight - AllotedGeometryDrawX, ScrollHelper.GetPosition( ).Y ) );
 		}
 
 
 		// Make sure text is never scrolled out of view when it doesn't need to be!
 		{
 			// Measure the total widget of text string, plus the caret's width!
-			const float TotalTextWidth = FontMeasureService->Measure( VisibleText, FontInfo ).X + CaretWidth;
-			if( ScrollHelper.GetPosition().X > TotalTextWidth - AllottedGeometry.Size.X )
+			const float TotalTextWidth = FontMeasureService->Measure( VisibleText, FontInfo, AllottedGeometry.Scale ).X + CaretWidth;
+			if ( ScrollHelper.GetPosition( ).X > TotalTextWidth - AllotedGeometryDrawX )
 			{
-				ScrollHelper.SetPosition( FVector2D( TotalTextWidth - AllottedGeometry.Size.X, ScrollHelper.GetPosition().Y ) );
+				ScrollHelper.SetPosition( FVector2D( TotalTextWidth - AllotedGeometryDrawX, ScrollHelper.GetPosition( ).Y ) );
 			}
 			if( ScrollHelper.GetPosition().X < 0.0f )
 			{
@@ -276,8 +277,8 @@ void SEditableText::Tick( const FGeometry& AllottedGeometry, const double InCurr
 	{
 		if( AnyTextSelected() )
 		{
-			const float SelectionLeftX = EditableTextDefs::SelectionRectLeftOffset + FontMeasureService->Measure( VisibleText.Left( Selection.GetMinIndex() ), FontInfo ).X;
-			const float SelectionRightX = EditableTextDefs::SelectionRectRightOffset + FontMeasureService->Measure( VisibleText.Left( Selection.GetMaxIndex() ), FontInfo ).X;
+			const float SelectionLeftX = EditableTextDefs::SelectionRectLeftOffset + FontMeasureService->Measure( VisibleText.Left( Selection.GetMinIndex() ), FontInfo, AllottedGeometry.Scale ).X;
+			const float SelectionRightX = EditableTextDefs::SelectionRectRightOffset + FontMeasureService->Measure( VisibleText.Left( Selection.GetMaxIndex( ) ), FontInfo, AllottedGeometry.Scale ).X;
 
 			SelectionTargetLeftSpring.SetTarget( SelectionLeftX );
 			SelectionTargetRightSpring.SetTarget( SelectionRightX );
@@ -422,8 +423,8 @@ bool SEditableText::FTextInputMethodContext::GetTextBounds(const uint32 BeginInd
 		const FString VisibleText = OwningWidgetPtr->GetStringToRender();
 		const TSharedRef< FSlateFontMeasure > FontMeasureService = FSlateApplication::Get().GetRenderer()->GetFontMeasureService();
 
-		const float TextLeft = FontMeasureService->Measure( VisibleText.Left( BeginIndex ), FontInfo ).X;
-		const float TextRight = FontMeasureService->Measure( VisibleText.Left( BeginIndex + Length ), FontInfo ).X;
+		const float TextLeft = FontMeasureService->Measure( VisibleText.Left( BeginIndex ), FontInfo, CachedGeometry.Scale ).X;
+		const float TextRight = FontMeasureService->Measure( VisibleText.Left( BeginIndex + Length ), FontInfo, CachedGeometry.Scale ).X;
 
 		const float ScrollAreaLeft = OwningWidgetPtr->ScrollHelper.ToScrollerSpace( FVector2D::ZeroVector ).X;
 		const float ScrollAreaRight = OwningWidgetPtr->ScrollHelper.ToScrollerSpace( CachedGeometry.Size ).X;
@@ -633,10 +634,10 @@ static int32 OffsetToTestBasedOnDirection( int8 Direction )
 	return (Direction >= 0) ? 0 : -1;
 }
 
-FReply SEditableText::MoveCursor( ECursorMoveMethod::Type Method, const FVector2D& Direction, ECursorAction::Type Action )
+FReply SEditableText::MoveCursor( FMoveCursor Args )
 {
 	// We can't use the keyboard to move the cursor while composing, as the IME has control over it
-	if(TextInputMethodContext->IsComposing && Method != ECursorMoveMethod::ScreenPosition)
+	if ( TextInputMethodContext->IsComposing && Args.GetMoveMethod() != ECursorMoveMethod::ScreenPosition )
 	{
 		return FReply::Handled();
 	}
@@ -651,9 +652,9 @@ FReply SEditableText::MoveCursor( ECursorMoveMethod::Type Method, const FVector2
 	// This is done regardless of whether the selection was made left-to-right, or right-to-left
 	// This also needs to be done *before* moving to word boundaries, as the start point needs to be 
 	// the start or end of the selection depending on the above rules
-	if (Action == ECursorAction::MoveCursor && Method != ECursorMoveMethod::ScreenPosition && AnyTextSelected())
+	if ( Args.GetAction( ) == ECursorAction::MoveCursor && Args.GetMoveMethod() != ECursorMoveMethod::ScreenPosition && AnyTextSelected( ) )
 	{
-		if (Method == ECursorMoveMethod::CharacterHorizontal)
+		if (Args.IsHorizontalMovement())
 		{
 			// If we're moving the cursor horizontally, we just snap to the start or end of the selection rather than 
 			// move the cursor by the normal movement rules
@@ -661,24 +662,14 @@ FReply SEditableText::MoveCursor( ECursorMoveMethod::Type Method, const FVector2
 		}
 
 		// Work out which edge of the selection we need to start at
-		bool bSnapToSelectionStart = false;
-		switch (Method)
-		{
-		case ECursorMoveMethod::CharacterHorizontal:
-		case ECursorMoveMethod::Word:
-			bSnapToSelectionStart = Direction.X < 0.0f;
-			break;
-
-		default:
-			break;
-		}
+		bool bSnapToSelectionStart = Args.IsHorizontalMovement() && Args.GetMoveDirection().X < 0;
 
 		// Adjust the current cursor position - also set the new cursor position so that the bAllowMoveCursor == false case is handled
 		OldCaretPosition = bSnapToSelectionStart ? Selection.GetMinIndex() : Selection.GetMaxIndex();
 		NewCaretPosition = OldCaretPosition;
 
 		// If we're snapping to a word boundary, but the selection was already at a word boundary, don't let the cursor move any more
-		if (Method == ECursorMoveMethod::Word && IsAtWordStart(NewCaretPosition))
+		if (Args.GetGranularity() == ECursorMoveGranularity::Word && IsAtWordStart(NewCaretPosition))
 		{
 			bAllowMoveCursor = false;
 		}
@@ -690,32 +681,36 @@ FReply SEditableText::MoveCursor( ECursorMoveMethod::Type Method, const FVector2
 
 		// If control is held down, jump to beginning of the current word (or, the previous word, if
 		// the caret is currently on a whitespace character
-		switch( Method ) 
+		if ( Args.GetMoveMethod() == ECursorMoveMethod::Cardinal )
+		{
+			if ( Args.GetGranularity() == ECursorMoveGranularity::Word )
+			{
+				NewCaretPosition = ScanForWordBoundary( CaretPosition, Args.GetMoveDirection().X );
+			}
+			else if ( Args.IsHorizontalMovement() )
+			{
+				checkSlow( Args.GetGranularity() == ECursorMoveGranularity::Character );
+				NewCaretPosition = FMath::Clamp<int32>( CaretPosition + Args.GetMoveDirection( ).X, 0, StringLength );
+			}
+			else
+			{
+				checkSlow( Args.IsVerticalMovement() );
+				return FReply::Unhandled( );
+			}			
+		}
+		else if ( Args.GetMoveMethod( ) == ECursorMoveMethod::ScreenPosition )
+		{			
+			NewCaretPosition = FindClickedCharacterIndex( Args.GetLocalPosition(), Args.GetGeometryScale() );
+		}
+		else
 		{
-		case ECursorMoveMethod::Word:
-			NewCaretPosition = ScanForWordBoundary( CaretPosition, Direction.X );
-			break;
-
-		case ECursorMoveMethod::CharacterHorizontal:
-			NewCaretPosition = FMath::Clamp<int32>( CaretPosition + Direction.X, 0, StringLength );
-			break;
-	
-		case ECursorMoveMethod::CharacterVertical:
-			return FReply::Unhandled();
-	
-		case ECursorMoveMethod::ScreenPosition:
-			NewCaretPosition = FindClickedCharacterIndex( Direction );
-			break;
-
-		default:
 			checkSlow(false, "Unknown ECursorMoveMethod value");
-			break;
 		}
 	}
 
 	SetCaretPosition( NewCaretPosition );
 
-	if( Action == ECursorAction::SelectText )
+	if( Args.GetAction() == ECursorAction::SelectText )
 	{
 		SelectText( OldCaretPosition );
 	}
@@ -740,7 +735,7 @@ FReply SEditableText::MoveCursor( ECursorMoveMethod::Type Method, const FVector2
 }
 
 
-void SEditableText::JumpTo(ETextLocation::Type JumpLocation, ECursorAction::Type Action)
+void SEditableText::JumpTo(ETextLocation JumpLocation, ECursorAction Action)
 {
 	switch( JumpLocation )
 	{
@@ -806,7 +801,7 @@ void SEditableText::SelectAllText()
 
 	// Make sure the animated selection effect starts roughly where the cursor is
 	const TSharedRef< FSlateFontMeasure > FontMeasureService = FSlateApplication::Get().GetRenderer()->GetFontMeasureService();
-	const float SelectionTargetX = EditableTextDefs::SelectionRectLeftOffset + FontMeasureService->Measure( EditedText.ToString().Left( CaretPosition ), Font.Get() ).X;
+	const float SelectionTargetX = EditableTextDefs::SelectionRectLeftOffset + FontMeasureService->Measure( EditedText.ToString( ).Left( CaretPosition ), Font.Get( ), 1.0f/*CachedGeometry.Scale*/ ).X;
 	SelectionTargetLeftSpring.SetPosition( SelectionTargetX );
 	SelectionTargetRightSpring.SetPosition( SelectionTargetX );
 }
@@ -816,10 +811,11 @@ bool SEditableText::SelectAllTextWhenFocused()
 	return bSelectAllTextWhenFocused.Get();
 }
 
-void SEditableText::SelectWordAt(const FVector2D& LocalPosition)
+void SEditableText::SelectWordAt( const FGeometry& MyGeometry, const FVector2D& ScreenSpacePosition )
 {
+	const FVector2D LocalPosition = MyGeometry.AbsoluteToLocal( ScreenSpacePosition );
 	// Figure out where in the string the user clicked
-	const int32 ClickedCharIndex = FindClickedCharacterIndex( LocalPosition );
+	const int32 ClickedCharIndex = FindClickedCharacterIndex( LocalPosition, MyGeometry.Scale );
 
 	// Deselect any text that was selected
 	ClearSelection();
@@ -872,9 +868,9 @@ bool SEditableText::AnyTextSelected() const
 	return Selection.StartIndex != INDEX_NONE && Selection.StartIndex != Selection.FinishIndex;
 }
 
-bool SEditableText::IsTextSelectedAt(const FVector2D& LocalPosition) const
+bool SEditableText::IsTextSelectedAt( const FGeometry& MyGeometry, const FVector2D& ScreenSpacePosition ) const
 {
-	const int32 ClickedCharIndex = FindClickedCharacterIndex(LocalPosition);
+	const int32 ClickedCharIndex = FindClickedCharacterIndex(MyGeometry, ScreenSpacePosition);
 	return ClickedCharIndex >= Selection.GetMinIndex() && ClickedCharIndex < Selection.GetMaxIndex();
 }
 
@@ -1197,11 +1193,11 @@ FVector2D SEditableText::ComputeDesiredSize() const
 	const FString StringToRender = GetStringToRender();
 	if( !StringToRender.IsEmpty() )
 	{
-		TextSize = FontMeasureService->Measure( StringToRender, FontInfo );
+		TextSize = FontMeasureService->Measure( StringToRender, FontInfo, 1.0f/*CachedGeometry.Scale*/ );
 	}
 	else
 	{
-		TextSize = FontMeasureService->Measure( HintText.Get().ToString(), FontInfo );
+		TextSize = FontMeasureService->Measure( HintText.Get( ).ToString( ), FontInfo, 1.0f/*CachedGeometry.Scale*/ );
 	}
 	
 	const float DesiredWidth = FMath::Max(TextSize.X, MinDesiredWidth.Get()) + CaretWidth;
@@ -1219,6 +1215,8 @@ int32 SEditableText::OnPaint( const FPaintArgs& Args, const FGeometry& AllottedG
 	const int32 SelectionLayer = 0;
 	const int32 TextLayer = 1;
 	const TSharedRef< FSlateFontMeasure > FontMeasureService = FSlateApplication::Get().GetRenderer()->GetFontMeasureService();
+	const float ScaleInverse = 1.0f / AllottedGeometry.Scale;
+	const FVector2D DrawGeometrySize = AllottedGeometry.Size * AllottedGeometry.Scale;
 
 	// See if a disabled effect should be used
 	bool bEnabled = ShouldBeEnabled( bParentEnabled );
@@ -1229,7 +1227,7 @@ int32 SEditableText::OnPaint( const FPaintArgs& Args, const FGeometry& AllottedG
 	const FString VisibleText = GetStringToRender();
 	const FLinearColor ThisColorAndOpacity = ColorAndOpacity.Get().GetColor(InWidgetStyle);
 	const FColor ColorAndOpacitySRGB = ThisColorAndOpacity * InWidgetStyle.GetColorAndOpacityTint();
-	const float FontMaxCharHeight = FTextEditHelper::GetFontHeight(FontInfo);
+	const float FontMaxCharHeight = FTextEditHelper::GetFontHeight(FontInfo) * AllottedGeometry.Scale;
 	const double CurrentTime = FSlateApplication::Get().GetCurrentTime();
 
 	// We'll draw with the 'focused' look if we're either focused or we have a context menu summoned
@@ -1241,10 +1239,10 @@ int32 SEditableText::OnPaint( const FPaintArgs& Args, const FGeometry& AllottedG
 		// Figure out a universally visible selection color.
 		const FColor SelectionBackgroundColorAndOpacity = ( (FLinearColor::White - ThisColorAndOpacity)*0.5f + FLinearColor(-0.2f, -0.05f, 0.15f)) * InWidgetStyle.GetColorAndOpacityTint();
 
-		const float SelectionLeftX = EditableTextDefs::SelectionRectLeftOffset + FontMeasureService->Measure( VisibleText.Left( Selection.GetMinIndex() ), FontInfo ).X;
-		const float SelectionRightX = EditableTextDefs::SelectionRectRightOffset + FontMeasureService->Measure( VisibleText.Left( Selection.GetMaxIndex() ), FontInfo ).X;
+		const float SelectionLeftX = EditableTextDefs::SelectionRectLeftOffset + FontMeasureService->Measure( VisibleText.Left( Selection.GetMinIndex() ), FontInfo, AllottedGeometry.Scale ).X;
+		const float SelectionRightX = EditableTextDefs::SelectionRectRightOffset + FontMeasureService->Measure( VisibleText.Left( Selection.GetMaxIndex( ) ), FontInfo, AllottedGeometry.Scale ).X;
 		const float SelectionTopY = 0.0f;
-		const float SelectionBottomY = AllottedGeometry.GetDrawSize().Y;
+		const float SelectionBottomY = DrawGeometrySize.Y;
 
 		const FVector2D DrawPosition = ScrollHelper.FromScrollerSpace( FVector2D( SelectionLeftX, SelectionTopY ) );
 		const FVector2D DrawSize = ScrollHelper.SizeFromScrollerSpace( FVector2D( SelectionRightX - SelectionLeftX, SelectionBottomY - SelectionTopY ) );
@@ -1252,7 +1250,7 @@ int32 SEditableText::OnPaint( const FPaintArgs& Args, const FGeometry& AllottedG
 		FSlateDrawElement::MakeBox(
 			OutDrawElements,
 			LayerId + SelectionLayer,
-			AllottedGeometry.ToPaintGeometry(DrawPosition, DrawSize),	// Position, Size, Scale
+			AllottedGeometry.ToPaintGeometry( DrawPosition, DrawSize, ScaleInverse ),	// Position, Size, Scale
 			BackgroundImageSelected.Get(),								// Image
 			MyClippingRect,												// Clipping rect
 			DrawEffects,												// Effects to use
@@ -1264,10 +1262,10 @@ int32 SEditableText::OnPaint( const FPaintArgs& Args, const FGeometry& AllottedG
 	const FTextRange CompositionRange(TextInputMethodContext->CompositionBeginIndex, TextInputMethodContext->CompositionBeginIndex + TextInputMethodContext->CompositionLength);
 	if( TextInputMethodContext->IsComposing && CompositionRange.InclusiveContains(CaretPosition) )
 	{
-		const float CompositionLeft = FontMeasureService->Measure( VisibleText.Left( CompositionRange.BeginIndex ), FontInfo ).X;
-		const float CompositionRight = FontMeasureService->Measure( VisibleText.Left( CompositionRange.EndIndex ), FontInfo ).X;
+		const float CompositionLeft = FontMeasureService->Measure( VisibleText.Left( CompositionRange.BeginIndex ), FontInfo, AllottedGeometry.Scale ).X;
+		const float CompositionRight = FontMeasureService->Measure( VisibleText.Left( CompositionRange.EndIndex ), FontInfo, AllottedGeometry.Scale ).X;
 		const float CompositionTop = 0.0f;
-		const float CompositionBottom = FontMeasureService->Measure( VisibleText.Mid( TextInputMethodContext->CompositionBeginIndex, TextInputMethodContext->CompositionLength ), FontInfo ).Y;
+		const float CompositionBottom = FontMeasureService->Measure( VisibleText.Mid( TextInputMethodContext->CompositionBeginIndex, TextInputMethodContext->CompositionLength ), FontInfo, AllottedGeometry.Scale ).Y;
 
 		const FVector2D DrawPosition = ScrollHelper.FromScrollerSpace( FVector2D( CompositionLeft, CompositionTop ) );
 		const FVector2D DrawSize = ScrollHelper.SizeFromScrollerSpace( FVector2D( CompositionRight - CompositionLeft, CompositionBottom - CompositionTop ) );
@@ -1275,14 +1273,14 @@ int32 SEditableText::OnPaint( const FPaintArgs& Args, const FGeometry& AllottedG
 		FSlateDrawElement::MakeBox(
 			OutDrawElements,
 			LayerId + SelectionLayer,
-			AllottedGeometry.ToPaintGeometry(DrawPosition, DrawSize),	// Position, Size, Scale
+			AllottedGeometry.ToPaintGeometry( DrawPosition, DrawSize, ScaleInverse ),	// Position, Size, Scale
 			BackgroundImageComposing.Get(),								// Image
 			MyClippingRect,												// Clipping rect
 			DrawEffects,												// Effects to use
 			ColorAndOpacitySRGB );						// Color
 	}
 
-	const float DrawPositionY = ( AllottedGeometry.Size.Y / 2 ) - ( FontMaxCharHeight / 2 );
+	const float DrawPositionY = ( DrawGeometrySize.Y / 2) - (FontMaxCharHeight / 2);
 	if (VisibleText.Len() == 0)
 	{
 		// Draw the hint text.
@@ -1291,7 +1289,7 @@ int32 SEditableText::OnPaint( const FPaintArgs& Args, const FGeometry& AllottedG
 		FSlateDrawElement::MakeText(
 			OutDrawElements,
 			LayerId + TextLayer,
-			AllottedGeometry.ToPaintGeometry( FVector2D( 0, DrawPositionY ), AllottedGeometry.Size ),
+			AllottedGeometry.ToPaintGeometry( FVector2D( 0, DrawPositionY ), AllottedGeometry.Size, ScaleInverse ),
 			ThisHintText,          // Text
 			FontInfo,              // Font information (font name, size)
 			MyClippingRect,        // Clipping rect
@@ -1305,11 +1303,11 @@ int32 SEditableText::OnPaint( const FPaintArgs& Args, const FGeometry& AllottedG
 
 		// Only draw the text that's potentially visible
 		const float ScrollAreaLeft = ScrollHelper.ToScrollerSpace( FVector2D::ZeroVector ).X;
-		const float ScrollAreaRight = ScrollHelper.ToScrollerSpace( AllottedGeometry.Size ).X;
-		const int32 FirstVisibleCharacter = FontMeasureService->FindCharacterIndexAtOffset( VisibleText, FontInfo, ScrollAreaLeft );
-		const int32 LastVisibleCharacter = FontMeasureService->FindCharacterIndexAtOffset( VisibleText, FontInfo, ScrollAreaRight );
+		const float ScrollAreaRight = ScrollHelper.ToScrollerSpace( DrawGeometrySize ).X;
+		const int32 FirstVisibleCharacter = FontMeasureService->FindCharacterIndexAtOffset( VisibleText, FontInfo, ScrollAreaLeft, AllottedGeometry.Scale );
+		const int32 LastVisibleCharacter = FontMeasureService->FindCharacterIndexAtOffset( VisibleText, FontInfo, ScrollAreaRight, AllottedGeometry.Scale );
 		const FString PotentiallyVisibleText( VisibleText.Mid( FirstVisibleCharacter, ( LastVisibleCharacter - FirstVisibleCharacter ) + 1 ) );
-		const float FirstVisibleCharacterPosition = FontMeasureService->Measure( VisibleText.Left( FirstVisibleCharacter ), FontInfo ).X;
+		const float FirstVisibleCharacterPosition = FontMeasureService->Measure( VisibleText.Left( FirstVisibleCharacter ), FontInfo, AllottedGeometry.Scale ).X;
 
 		const float TextVertOffset = EditableTextDefs::TextVertOffsetPercent * FontMaxCharHeight;
 		const FVector2D DrawPosition = ScrollHelper.FromScrollerSpace( FVector2D( FirstVisibleCharacterPosition, DrawPositionY + TextVertOffset ) );
@@ -1375,7 +1373,7 @@ int32 SEditableText::OnPaint( const FPaintArgs& Args, const FGeometry& AllottedG
 			FSlateDrawElement::MakeBox(
 				OutDrawElements,
 				LayerId + TextLayer,
-				AllottedGeometry.ToPaintGeometry( DrawPosition, DrawSize ),	// Position, Size, Scale
+				AllottedGeometry.ToPaintGeometry( DrawPosition, DrawSize, ScaleInverse ),	// Position, Size, Scale
 				SelectionTargetBrush,										// Image
 				MyClippingRect,												// Clipping rect
 				DrawEffects,												// Effects to use
@@ -1415,7 +1413,7 @@ int32 SEditableText::OnPaint( const FPaintArgs& Args, const FGeometry& AllottedG
 		FSlateDrawElement::MakeBox(
 			OutDrawElements,
 			LayerId + TextLayer,
-			AllottedGeometry.ToPaintGeometry( CaretDrawPosition, CaretDrawSize ),		// Position, Size, Scale
+			AllottedGeometry.ToPaintGeometry( CaretDrawPosition, CaretDrawSize, ScaleInverse ),		// Position, Size, Scale
 			CaretImage.Get(),															// Image
 			MyClippingRect,																// Clipping rect
 			DrawEffects,																// Effects to use
@@ -1619,7 +1617,7 @@ FReply SEditableText::OnMouseButtonUp( const FGeometry& InMyGeometry, const FPoi
 			{
 				// Start the selection animation wherever the user started the click/drag
 				const TSharedRef< FSlateFontMeasure > FontMeasureService = FSlateApplication::Get().GetRenderer()->GetFontMeasureService();
-				const float SelectionTargetX = EditableTextDefs::SelectionRectLeftOffset + FontMeasureService->Measure( EditedText.ToString().Left( Selection.StartIndex ), Font.Get() ).X;
+				const float SelectionTargetX = EditableTextDefs::SelectionRectLeftOffset + FontMeasureService->Measure( EditedText.ToString( ).Left( Selection.StartIndex ), Font.Get( ), InMyGeometry.Scale ).X;
 				SelectionTargetLeftSpring.SetPosition( SelectionTargetX );
 				SelectionTargetRightSpring.SetPosition( SelectionTargetX );
 
@@ -1690,7 +1688,7 @@ void SEditableText::DeleteSelectedText()
 
 		// Visually collapse the selection target right before we clear selection
 		const TSharedRef< FSlateFontMeasure > FontMeasureService = FSlateApplication::Get().GetRenderer()->GetFontMeasureService();
-		const float SelectionTargetX = EditableTextDefs::SelectionRectLeftOffset + FontMeasureService->Measure( NewText.Left( CaretPosition ), Font.Get() ).X;
+		const float SelectionTargetX = EditableTextDefs::SelectionRectLeftOffset + FontMeasureService->Measure( NewText.Left( CaretPosition ), Font.Get( ), 1.0f/*CachedGeometry.Scale*/ ).X;
 		SelectionTargetLeftSpring.SetTarget( SelectionTargetX );
 		SelectionTargetRightSpring.SetTarget( SelectionTargetX );
 
@@ -1706,19 +1704,19 @@ int32 SEditableText::FindClickedCharacterIndex( const FGeometry& InMyGeometry, c
 	const FVector2D& ScreenCursorPosition = InScreenCursorPosition;
 	const FVector2D& LocalCursorPosition = InMyGeometry.AbsoluteToLocal( ScreenCursorPosition );
 
-	return FindClickedCharacterIndex(LocalCursorPosition);
+	return FindClickedCharacterIndex( LocalCursorPosition, InMyGeometry.Scale );
 }
 
 
-int32 SEditableText::FindClickedCharacterIndex( const FVector2D& InLocalCursorPosition ) const
+int32 SEditableText::FindClickedCharacterIndex( const FVector2D& InLocalCursorPosition, float GeometryScale ) const
 {
 	// Convert the click position from the local widget space to the scrolled area
-	const FVector2D& PositionInScrollableArea = ScrollHelper.ToScrollerSpace( InLocalCursorPosition );
+	const FVector2D& PositionInScrollableArea = ScrollHelper.ToScrollerSpace( InLocalCursorPosition * GeometryScale );
 	const float ClampedCursorX = FMath::Max( PositionInScrollableArea.X, 0.0f );
 
 	// Figure out where in the string the user clicked
 	const TSharedRef< FSlateFontMeasure > FontMeasureService = FSlateApplication::Get().GetRenderer()->GetFontMeasureService();
-	const int32 ClickedCharIndex = FontMeasureService->FindCharacterIndexAtOffset( EditedText.ToString(), Font.Get(), ClampedCursorX );
+	const int32 ClickedCharIndex = FontMeasureService->FindCharacterIndexAtOffset( EditedText.ToString( ), Font.Get( ), ClampedCursorX, GeometryScale );
 
 	return ClickedCharIndex;
 }
diff --git a/Engine/Source/Runtime/Slate/Private/Widgets/Text/ITextEditorWidget.cpp b/Engine/Source/Runtime/Slate/Private/Widgets/Text/ITextEditorWidget.cpp
new file mode 100644
--- /dev/null
+++ b/Engine/Source/Runtime/Slate/Private/Widgets/Text/ITextEditorWidget.cpp
@@ -0,0 +1,43 @@
+// Copyright 1998-2014 Epic Games, Inc. All Rights Reserved.
+
+#include "SlatePrivatePCH.h"
+#include "ITextEditorWidget.h"
+
+FMoveCursor FMoveCursor::Cardinal( ECursorMoveGranularity Granularity, FIntPoint Direction, ECursorAction Action )
+{
+	return FMoveCursor( Granularity, ECursorMoveMethod::Cardinal, Direction, 1.0f, Action );
+}
+
+FMoveCursor FMoveCursor::ViaScreenPointer( FVector2D LocalPosition, float GeometryScale, ECursorAction Action )
+{
+	return FMoveCursor( ECursorMoveGranularity::Character, ECursorMoveMethod::ScreenPosition, LocalPosition, GeometryScale, Action );
+}
+
+ECursorMoveMethod FMoveCursor::GetMoveMethod( ) const { return Method; }
+
+bool FMoveCursor::IsVerticalMovement( ) const { return DirectionOrPosition.Y != 0.0f; }
+
+bool FMoveCursor::IsHorizontalMovement( ) const { return DirectionOrPosition.X != 0.0f; }
+
+FIntPoint FMoveCursor::GetMoveDirection( ) const { return FIntPoint( DirectionOrPosition.X, DirectionOrPosition.Y ); }
+
+ECursorAction FMoveCursor::GetAction( ) const { return Action; }
+
+FVector2D FMoveCursor::GetLocalPosition( ) const { return DirectionOrPosition; }
+
+ECursorMoveGranularity FMoveCursor::GetGranularity() const { return Granularity; }
+
+float FMoveCursor::GetGeometryScale( ) const { return GeometryScale; }
+
+FMoveCursor::FMoveCursor( ECursorMoveGranularity InGranularity, ECursorMoveMethod InMethod, FVector2D InDirectionOrPosition, float InGeometryScale, ECursorAction InAction )
+: Granularity( InGranularity )
+, Method( InMethod )
+, DirectionOrPosition( InDirectionOrPosition )
+, Action( InAction )
+, GeometryScale(InGeometryScale)
+{
+	// Movement actions are assumed to be exclusively vertical or exclusively horizontal
+	// There is no great reason to assume this, but a lot of the code was written to deal with
+	// events, which must be one or the other.
+	ensure( Method == ECursorMoveMethod::ScreenPosition || (Method == ECursorMoveMethod::Cardinal && IsVerticalMovement( ) != IsHorizontalMovement( )) );
+}
diff --git a/Engine/Source/Runtime/Slate/Private/Widgets/Text/SMultiLineEditableText.cpp b/Engine/Source/Runtime/Slate/Private/Widgets/Text/SMultiLineEditableText.cpp
--- a/Engine/Source/Runtime/Slate/Private/Widgets/Text/SMultiLineEditableText.cpp
+++ b/Engine/Source/Runtime/Slate/Private/Widgets/Text/SMultiLineEditableText.cpp
@@ -788,10 +788,10 @@ void SMultiLineEditableText::DeleteSelectedText()
 	}
 }
 
-FReply SMultiLineEditableText::MoveCursor( ECursorMoveMethod::Type Method, const FVector2D& Direction, ECursorAction::Type Action )
+FReply SMultiLineEditableText::MoveCursor( FMoveCursor Args )
 {
 	// We can't use the keyboard to move the cursor while composing, as the IME has control over it
-	if(TextInputMethodContext->IsComposing && Method != ECursorMoveMethod::ScreenPosition)
+	if(TextInputMethodContext->IsComposing && Args.GetMoveMethod() != ECursorMoveMethod::ScreenPosition)
 	{
 		return FReply::Handled();
 	}
@@ -806,9 +806,9 @@ FReply SMultiLineEditableText::MoveCursor( ECursorMoveMethod::Type Method, const
 	// This is done regardless of whether the selection was made left-to-right, or right-to-left
 	// This also needs to be done *before* moving to word boundaries, or moving vertically, as the 
 	// start point needs to be the start or end of the selection depending on the above rules
-	if (Action == ECursorAction::MoveCursor && Method != ECursorMoveMethod::ScreenPosition && AnyTextSelected())
+	if ( Args.GetAction() == ECursorAction::MoveCursor && Args.GetMoveMethod() != ECursorMoveMethod::ScreenPosition && AnyTextSelected( ) )
 	{
-		if (Method == ECursorMoveMethod::CharacterHorizontal)
+		if (Args.IsHorizontalMovement())
 		{
 			// If we're moving the cursor horizontally, we just snap to the start or end of the selection rather than 
 			// move the cursor by the normal movement rules
@@ -816,21 +816,10 @@ FReply SMultiLineEditableText::MoveCursor( ECursorMoveMethod::Type Method, const
 		}
 
 		// Work out which edge of the selection we need to start at
-		bool bSnapToSelectionStart = false;
-		switch (Method)
-		{
-		case ECursorMoveMethod::CharacterHorizontal:
-		case ECursorMoveMethod::Word:
-			bSnapToSelectionStart = Direction.X < 0.0f;
-			break;
-	
-		case ECursorMoveMethod::CharacterVertical:
-			bSnapToSelectionStart = Direction.Y < 0.0f;
-			break;
+		bool bSnapToSelectionStart =
+			Args.GetMoveMethod() == ECursorMoveMethod::Cardinal &&
+			(Args.GetMoveDirection().X < 0.0f || Args.GetMoveDirection().Y < 0.0f);
 
-		default:
-			break;
-		}
 
 		// Adjust the current cursor position - also set the new cursor position so that the bAllowMoveCursor == false case is handled
 		const FTextSelection Selection(SelectionStart.GetValue(), CursorPosition);
@@ -838,7 +827,7 @@ FReply SMultiLineEditableText::MoveCursor( ECursorMoveMethod::Type Method, const
 		NewCursorPosition = CursorPosition;
 
 		// If we're snapping to a word boundary, but the selection was already at a word boundary, don't let the cursor move any more
-		if (Method == ECursorMoveMethod::Word && IsAtWordStart(NewCursorPosition))
+		if (Args.GetGranularity() == ECursorMoveGranularity::Word && IsAtWordStart(NewCursorPosition))
 		{
 			bAllowMoveCursor = false;
 		}
@@ -848,46 +837,49 @@ FReply SMultiLineEditableText::MoveCursor( ECursorMoveMethod::Type Method, const
 	bool bUpdatePreferredCursorScreenOffsetInLine = false;
 	if (bAllowMoveCursor)
 	{
-		switch (Method)
+		if ( Args.GetMoveMethod() == ECursorMoveMethod::Cardinal )
 		{
-		case ECursorMoveMethod::CharacterHorizontal:
-			NewCursorPosition = TranslatedLocation( CursorPosition, Direction.X );
-			bUpdatePreferredCursorScreenOffsetInLine = true;
-			break;
-
-		case ECursorMoveMethod::Word:
-			NewCursorPosition = ScanForWordBoundary( CursorPosition, Direction.X );
-			bUpdatePreferredCursorScreenOffsetInLine = true;
-			break;
-	
-		case ECursorMoveMethod::CharacterVertical:
-			TranslateLocationVertical( CursorPosition, Direction.Y, NewCursorPosition, NewCursorAlignment );
-			break;
-	
-		case ECursorMoveMethod::ScreenPosition:
+			if ( Args.GetGranularity() == ECursorMoveGranularity::Character )
 			{
-				ETextHitPoint HitPoint = ETextHitPoint::WithinText;
-				NewCursorPosition = TextLayout->GetTextLocationAt( Direction * TextLayout->GetScale(), &HitPoint );
-				bUpdatePreferredCursorScreenOffsetInLine = true;
-
-				// Moving with the mouse behaves a bit differently to moving with the keyboard, as clicking at the end of a wrapped line needs to place the cursor there
-				// rather than at the start of the next line (which is tricky since they have the same index according to GetTextLocationAt!).
-				// We use the HitPoint to work this out and then adjust the cursor position accordingly
-				if (HitPoint == ETextHitPoint::RightGutter)
+				if ( Args.IsHorizontalMovement() )
+				{
+					NewCursorPosition = TranslatedLocation( CursorPosition, Args.GetMoveDirection().X );
+					bUpdatePreferredCursorScreenOffsetInLine = true;
+				}
+				else
 				{
-					NewCursorPosition = FTextLocation(NewCursorPosition, -1);
-					NewCursorAlignment = ECursorAlignment::Right;
+					TranslateLocationVertical( CursorPosition, Args.GetMoveDirection().Y, Args.GetGeometryScale(), NewCursorPosition, NewCursorAlignment );
 				}
 			}
-			break;
+			else
+			{
+				checkSlow( Args.GetGranularity() == ECursorMoveGranularity::Word );
+				NewCursorPosition = ScanForWordBoundary( CursorPosition, Args.GetMoveDirection().X );
+				bUpdatePreferredCursorScreenOffsetInLine = true;
+			}
+		}
+		else if ( Args.GetMoveMethod() == ECursorMoveMethod::ScreenPosition )
+		{
+			ETextHitPoint HitPoint = ETextHitPoint::WithinText;
+			NewCursorPosition = TextLayout->GetTextLocationAt( Args.GetLocalPosition() * Args.GetGeometryScale(), &HitPoint );
+			bUpdatePreferredCursorScreenOffsetInLine = true;
 
-		default:
+			// Moving with the mouse behaves a bit differently to moving with the keyboard, as clicking at the end of a wrapped line needs to place the cursor there
+			// rather than at the start of the next line (which is tricky since they have the same index according to GetTextLocationAt!).
+			// We use the HitPoint to work this out and then adjust the cursor position accordingly
+			if ( HitPoint == ETextHitPoint::RightGutter )
+			{
+				NewCursorPosition = FTextLocation( NewCursorPosition, -1 );
+				NewCursorAlignment = ECursorAlignment::Right;
+			}
+		}
+		else
+		{
 			checkSlow(false, "Unknown ECursorMoveMethod value");
-			break;
 		}
 	}
 
-	if (Action == ECursorAction::SelectText)
+	if (Args.GetAction() == ECursorAction::SelectText)
 	{
 		// We are selecting text. Just remember where the selection started.
 		// The cursor is implicitly the other endpoint.
@@ -1055,7 +1047,7 @@ void SMultiLineEditableText::UpdatePreferredCursorScreenOffsetInLine()
 	PreferredCursorScreenOffsetInLine = TextLayout->GetLocationAt(CursorInfo.GetCursorInteractionLocation(), CursorInfo.GetCursorAlignment() == ECursorAlignment::Right).X;
 }
 
-void SMultiLineEditableText::JumpTo(ETextLocation::Type JumpLocation, ECursorAction::Type Action)
+void SMultiLineEditableText::JumpTo(ETextLocation JumpLocation, ECursorAction Action)
 {
 	switch (JumpLocation)
 	{
@@ -1215,9 +1207,11 @@ bool SMultiLineEditableText::SelectAllTextWhenFocused()
 	return bSelectAllTextWhenFocused.Get();
 }
 
-void SMultiLineEditableText::SelectWordAt(const FVector2D& LocalPosition)
+void SMultiLineEditableText::SelectWordAt( const FGeometry& MyGeometry, const FVector2D& ScreenSpacePosition )
 {
-	FTextLocation InitialLocation = TextLayout->GetTextLocationAt(LocalPosition * TextLayout->GetScale());
+	const FVector2D LocalPosition = MyGeometry.AbsoluteToLocal( ScreenSpacePosition );
+
+	FTextLocation InitialLocation = TextLayout->GetTextLocationAt(LocalPosition * MyGeometry.Scale);
 	FTextSelection WordSelection = TextLayout->GetWordAt(InitialLocation);
 
 	FTextLocation WordStart = WordSelection.GetBeginning();
@@ -1261,7 +1255,7 @@ bool SMultiLineEditableText::AnyTextSelected() const
 	return SelectionPosition != CursorInteractionPosition;
 }
 
-bool SMultiLineEditableText::IsTextSelectedAt(const FVector2D& LocalPosition) const
+bool SMultiLineEditableText::IsTextSelectedAt( const FGeometry& MyGeometry, const FVector2D& ScreenSpacePosition ) const
 {
 	const FTextLocation CursorInteractionPosition = CursorInfo.GetCursorInteractionLocation();
 	const FTextLocation SelectionPosition = SelectionStart.Get(CursorInteractionPosition);
@@ -1271,7 +1265,8 @@ bool SMultiLineEditableText::IsTextSelectedAt(const FVector2D& LocalPosition) co
 		return false;
 	}
 
-	const FTextLocation ClickedPosition = TextLayout->GetTextLocationAt(LocalPosition * TextLayout->GetScale());
+	const FVector2D LocalPosition = MyGeometry.AbsoluteToLocal(ScreenSpacePosition);
+	const FTextLocation ClickedPosition = TextLayout->GetTextLocationAt(LocalPosition * MyGeometry.Scale);
 
 	FTextLocation SelectionLocation = SelectionStart.Get(CursorInteractionPosition);
 	FTextSelection Selection(SelectionLocation, CursorInteractionPosition);
@@ -1920,7 +1915,7 @@ FTextLocation SMultiLineEditableText::TranslatedLocation( const FTextLocation& L
 	}
 }
 
-void SMultiLineEditableText::TranslateLocationVertical( const FTextLocation& Location, int8 Direction, FTextLocation& OutCursorPosition, TOptional<ECursorAlignment>& OutCursorAlignment ) const
+void SMultiLineEditableText::TranslateLocationVertical( const FTextLocation& Location, int8 Direction, float GeometryScale, FTextLocation& OutCursorPosition, TOptional<ECursorAlignment>& OutCursorAlignment ) const
 {
 	const TArray< FTextLayout::FLineView >& LineViews = TextLayout->GetLineViews();
 	const int32 NumberOfLineViews = LineViews.Num();
@@ -1934,7 +1929,7 @@ void SMultiLineEditableText::TranslateLocationVertical( const FTextLocation& Loc
 	
 	// Our horizontal position is the clamped version of whatever the user explicitly set with horizontal movement.
 	ETextHitPoint HitPoint = ETextHitPoint::WithinText;
-	OutCursorPosition = TextLayout->GetTextLocationAt( NewLineView, FVector2D( PreferredCursorScreenOffsetInLine, NewLineView.Offset.Y ) * TextLayout->GetScale(), &HitPoint );
+	OutCursorPosition = TextLayout->GetTextLocationAt( NewLineView, FVector2D( PreferredCursorScreenOffsetInLine, NewLineView.Offset.Y ) * GeometryScale, &HitPoint );
 
 	// PreferredCursorScreenOffsetInLine can cause the cursor to move to the right hand gutter, and it needs to be placed there
 	// rather than at the start of the next line (which is tricky since they have the same index according to GetTextLocationAt!).
diff --git a/Engine/Source/Runtime/Slate/Public/Widgets/Input/SEditableText.h b/Engine/Source/Runtime/Slate/Public/Widgets/Input/SEditableText.h
--- a/Engine/Source/Runtime/Slate/Public/Widgets/Input/SEditableText.h
+++ b/Engine/Source/Runtime/Slate/Public/Widgets/Input/SEditableText.h
@@ -290,17 +290,17 @@ public:
 	virtual void DeleteChar() override;
 	virtual bool CanTypeCharacter(const TCHAR CharInQuestion) const override;
 	virtual void TypeChar( const int32 Character ) override;
-	virtual FReply MoveCursor( ECursorMoveMethod::Type Method, const FVector2D& Direction, ECursorAction::Type Action ) override;
-	virtual void JumpTo(ETextLocation::Type JumpLocation, ECursorAction::Type Action) override;
+	virtual FReply MoveCursor( FMoveCursor Args ) override;
+	virtual void JumpTo(ETextLocation JumpLocation, ECursorAction Action) override;
 	virtual void ClearSelection() override;
 	virtual void SelectAllText() override;
 	virtual bool SelectAllTextWhenFocused() override;
-	virtual void SelectWordAt(const FVector2D& LocalPosition) override;
+	virtual void SelectWordAt( const FGeometry& MyGeometry, const FVector2D& ScreenSpacePosition ) override;
 	virtual void BeginDragSelection() override;
 	virtual bool IsDragSelecting() const override;
 	virtual void EndDragSelection() override;
 	virtual bool AnyTextSelected() const override;
-	virtual bool IsTextSelectedAt(const FVector2D& LocalPosition) const override;
+	virtual bool IsTextSelectedAt( const FGeometry& MyGeometry, const FVector2D& ScreenSpacePosition ) const override;
 	virtual void SetWasFocusedByLastMouseDown( bool Value ) override;
 	virtual bool WasFocusedByLastMouseDown() const override;
 	virtual void SetHasDragSelectedSinceFocused( bool Value ) override;
@@ -367,11 +367,12 @@ protected:
 	/**
 	 * Determines which character was clicked on
 	 *
-	 * @param  InLocalCursorPosition	Click position (in local space)
+	 * @param  InLocalCursorPosition       Click position (in local space)
+	 * @param  GeometryScale               The DPI scale with which this widget is being drawn.
 	 *
 	 * @return  The clicked character index
 	 */
-	int32 FindClickedCharacterIndex( const FVector2D& InLocalCursorPosition ) const;
+	int32 FindClickedCharacterIndex( const FVector2D& InLocalCursorPosition, float GeometryScale ) const;
 
 	/** Find the closest word boundary */
 	int32 ScanForWordBoundary( const int32 Location, int8 Direction ) const; 
diff --git a/Engine/Source/Runtime/Slate/Public/Widgets/Text/ITextEditorWidget.h b/Engine/Source/Runtime/Slate/Public/Widgets/Text/ITextEditorWidget.h
--- a/Engine/Source/Runtime/Slate/Public/Widgets/Text/ITextEditorWidget.h
+++ b/Engine/Source/Runtime/Slate/Public/Widgets/Text/ITextEditorWidget.h
@@ -3,48 +3,104 @@
 #pragma once
 
 
-namespace ECursorMoveMethod
+enum class ECursorMoveMethod
 {
-	enum Type
-	{
-		/** Move one character at a time (e.g. arrow left, right). */
-		CharacterHorizontal,
-		/** Move vertically (e.g. arrow up, down) */
-		CharacterVertical,
-		/** Move one word at a time (e.g. ctrl is held down). */
-		Word,
-		/** Move the cursor to the correct character based on the given screen position. */
-		ScreenPosition,
-	};
-}
-
-
-namespace ECursorAction
+	/** Move in one of the cardinal directions e.g. arrow left, right, up, down */
+	Cardinal,
+	/** Move the cursor to the correct character based on the given screen position. */
+	ScreenPosition
+};
+
+enum class ECursorMoveGranularity
 {
-	enum Type
-	{
-		/** Just relocate the cursor. */
-		MoveCursor,
-		/** Move the cursor and select any text that it passes over. */
-		SelectText
-	};
-}
+	/** Move one character at a time (e.g. arrow left, right, up, down). */
+	Character,
+	/** Move one word at a time (e.g. ctrl is held down and arrow left/right/up/down). */
+	Word
+};
+
+enum class ECursorAction
+{
+	/** Just relocate the cursor. */
+	MoveCursor,
+	/** Move the cursor and select any text that it passes over. */
+	SelectText
+};
+
 
+enum class ETextLocation
+{
+	/** Jump to the beginning of text. (e.g. Ctrl+Home) */
+	BeginningOfDocument,
+	/** Jump to the end of text. (e.g. Ctrl+End) */
+	EndOfDocument,
+	/** Jump to the beginning of the line or beginning of text within a line (e.g. Home) */
+	BeginningOfLine,
+	/** Jump to the end of line (e.g. End) */
+	EndOfLine
+};
 
-namespace ETextLocation
+/**
+ * Argument to the ITextEditorWidget::Move(); it decouples performing
+ * cursor movement and text highlighting actions from event handling.
+ * FMoveCursor describes how the cursor can be moved via keyboard,
+ * mouse, and other mechanisms in the future.
+ */
+class FMoveCursor
 {
-	enum Type
-	{
-		/** Jump to the beginning of text. (e.g. Ctrl+Home) */
-		BeginningOfDocument,
-		/** Jump to the end of text. (e.g. Ctrl+End) */
-		EndOfDocument,
-		/** Jump to the beginning of the line or beginning of text within a line (e.g. Home) */
-		BeginningOfLine,
-		/** Jump to the end of line (e.g. End) */
-		EndOfLine
-	};
-}
+public:
+	/**
+	 * Creates a MoveCursor action that describes moving by a single character in any of the cardinal directions.
+	 *
+	 * @param Granularity    Move one character at a time, or on word boundaries (e.g. Ctrl is held down, or double-click mouse+drag)
+	 * @param Direction      Axis-aligned unit vector along which to move. e.g. Move right: (0,1) or Move Up: (-1, 0)
+	 * @param Action         Just move or also select text?
+	 */
+	static FMoveCursor Cardinal( ECursorMoveGranularity Granularity, FIntPoint Direction, ECursorAction Action );
+
+	/**
+	 * Creates a MoveCursor action that describes moving the text cursor by selecting an arbitrary coordinate on the screen.
+	 * e.g. User user touches a touch device screen or uses the mouse to point at text.
+	 *
+	 * @param LocalPosition          Position in the text widget where the user wants to move the cursor by pointing
+	 * @param GeometryScale          DPI Scale of the widget geometry.
+	 * @param Action                 Just move or also select text?
+	 */
+	static FMoveCursor ViaScreenPointer( FVector2D LocalPosition, float GeometryScale, ECursorAction Action );
+
+	/** @see ECursorMoveMethod */
+	ECursorMoveMethod GetMoveMethod() const;
+
+	/** Is the cursor moving up/down; Only valid for word and character movement methods. */
+	bool IsVerticalMovement() const;
+
+	/** Is the cursor moving left/right; Only valid for word and character movement methods. */
+	bool IsHorizontalMovement() const;
+
+	/** @see ECursorAction */
+	ECursorAction GetAction() const;
+
+	/** When using directional movement (i.e. Character or Word granularity; not screen position, which way to move.) */
+	FIntPoint GetMoveDirection() const;
+
+	/** Which position in the widget where the user touched when using ScreenPosition mode. */
+	FVector2D GetLocalPosition() const;
+
+	/** Move one character at a time, or one word? */
+	ECursorMoveGranularity GetGranularity() const;
+
+	/** Geometry Scale at the time of the event that caused this action */
+	float GetGeometryScale() const;
+
+private:
+	FMoveCursor( ECursorMoveGranularity InGranularity, ECursorMoveMethod InMethod, FVector2D InDirectionOrPosition, float InGeometryScale, ECursorAction InAction );
+
+	ECursorMoveGranularity Granularity;
+	ECursorMoveMethod Method;
+	FVector2D DirectionOrPosition;
+	ECursorAction Action;
+	float GeometryScale;
+};
 
 
 class SLATE_API ITextEditorWidget 
@@ -89,20 +145,17 @@ class SLATE_API ITextEditorWidget
 	 */
 	virtual void TypeChar( const int32 Character ) = 0;
 	
+
 	/**
-	 * Move the cursors, left, right, up down, one character or word at a time.
+	 * Move the cursor left, right, up down, one character or word at a time.
 	 * Optionally select text while moving.
-	 * 
-	 * @param  Method       Move horizontally, vertically, or whole words horizontally.
-	 * @param  Direction    Move forward/backward by number of units specified; usuatlly +/-1.
-	 * @param  Action       Just move or select text.
 	 */
-	virtual FReply MoveCursor( ECursorMoveMethod::Type Method, const FVector2D& Direction, ECursorAction::Type Action ) = 0;
+	virtual FReply MoveCursor( FMoveCursor Args ) = 0;
 	
 	/**
 	 * Jump the cursor to a special location: e.g. end of line, beginning of document, etc..
 	 */
-	virtual void JumpTo(ETextLocation::Type JumpLocation, ECursorAction::Type Action) = 0;
+	virtual void JumpTo(ETextLocation JumpLocation, ECursorAction Action) = 0;
 
 	/**
 	 * Clears text selection
@@ -116,7 +169,7 @@ class SLATE_API ITextEditorWidget
 
 	virtual bool SelectAllTextWhenFocused() = 0;
 
-	virtual void SelectWordAt(const FVector2D& LocalPosition) = 0;
+	virtual void SelectWordAt( const FGeometry& MyGeometry, const FVector2D& ScreenSpacePosition ) = 0;
 
 	virtual void BeginDragSelection() = 0;
 
@@ -126,7 +179,7 @@ class SLATE_API ITextEditorWidget
 
 	virtual bool AnyTextSelected() const = 0;
 
-	virtual bool IsTextSelectedAt(const FVector2D& LocalPosition) const = 0;
+	virtual bool IsTextSelectedAt( const FGeometry& MyGeometry, const FVector2D& ScreenSpacePosition ) const = 0;
 
 	virtual void SetWasFocusedByLastMouseDown( bool Value ) = 0;
 	virtual bool WasFocusedByLastMouseDown() const = 0;
diff --git a/Engine/Source/Runtime/Slate/Public/Widgets/Text/SMultiLineEditableText.h b/Engine/Source/Runtime/Slate/Public/Widgets/Text/SMultiLineEditableText.h
--- a/Engine/Source/Runtime/Slate/Public/Widgets/Text/SMultiLineEditableText.h
+++ b/Engine/Source/Runtime/Slate/Public/Widgets/Text/SMultiLineEditableText.h
@@ -292,16 +292,16 @@ private:
 	virtual void DeleteChar() override;
 	virtual bool CanTypeCharacter(const TCHAR CharInQuestion) const override;
 	virtual void TypeChar( const int32 Character ) override;
-	virtual FReply MoveCursor( ECursorMoveMethod::Type Method, const FVector2D& Direction, ECursorAction::Type Action ) override;
-	virtual void JumpTo(ETextLocation::Type JumpLocation, ECursorAction::Type Action) override;
+	virtual FReply MoveCursor( FMoveCursor Args ) override;
+	virtual void JumpTo(ETextLocation JumpLocation, ECursorAction Action) override;
 	virtual void SelectAllText() override;
 	virtual bool SelectAllTextWhenFocused() override;
-	virtual void SelectWordAt(const FVector2D& LocalPosition) override;
+	virtual void SelectWordAt( const FGeometry& MyGeometry, const FVector2D& ScreenSpacePosition ) override;
 	virtual void BeginDragSelection() override;
 	virtual bool IsDragSelecting() const override;
 	virtual void EndDragSelection() override;
 	virtual bool AnyTextSelected() const override;
-	virtual bool IsTextSelectedAt(const FVector2D& LocalPosition) const override;
+	virtual bool IsTextSelectedAt( const FGeometry& MyGeometry, const FVector2D& ScreenSpacePosition ) const override;
 	virtual void SetWasFocusedByLastMouseDown( bool Value ) override;
 	virtual bool WasFocusedByLastMouseDown() const override;
 	virtual void SetHasDragSelectedSinceFocused( bool Value ) override;
@@ -374,12 +374,13 @@ private:
 	/**
 	 * Given a location and a Direction to offset, return a new location.
 	 *
-	 * @param Location				Cursor location from which to offset
-	 * @param Direction				Positive means down, negative means up.
-	 * @param OutCursorPosition		Fill with the updated cursor position.
-	 * @param OutCursorAlignment	Optionally fill with a new cursor alignment (will be auto-calculated if not set).
+	 * @param Location              Cursor location from which to offset
+	 * @param Direction             Positive means down, negative means up.
+	 * @param GeometryScale         Geometry DPI scale at which the widget is being rendered
+	 * @param OutCursorPosition     Fill with the updated cursor position.
+	 * @param OutCursorAlignment    Optionally fill with a new cursor alignment (will be auto-calculated if not set).
 	 */
-	void TranslateLocationVertical( const FTextLocation& Location, int8 Direction, FTextLocation& OutCursorPosition, TOptional<ECursorAlignment>& OutCursorAlignment ) const;
+	void TranslateLocationVertical( const FTextLocation& Location, int8 Direction, float GeometryScale, FTextLocation& OutCursorPosition, TOptional<ECursorAlignment>& OutCursorAlignment ) const;
 
 	/** Find the closest word boundary */
 	FTextLocation ScanForWordBoundary( const FTextLocation& Location, int8 Direction ) const; 
