Slate Hittest grid now cleared immediately before we paint; actually responsive at low framerates.

Also fixes TTP 343004: Buttons are extremely difficult to click at low framerates

diff --git a/Engine/Source/Runtime/Slate/Private/Framework/Application/SlateApplication.cpp b/Engine/Source/Runtime/Slate/Private/Framework/Application/SlateApplication.cpp
--- a/Engine/Source/Runtime/Slate/Private/Framework/Application/SlateApplication.cpp
+++ b/Engine/Source/Runtime/Slate/Private/Framework/Application/SlateApplication.cpp
@@ -718,6 +718,12 @@ void FSlateApplication::PrivateDrawWindows( TSharedPtr<SWindow> DrawOnlyThisWind
 	{
 		SCOPE_CYCLE_COUNTER( STAT_SlateDrawWindowTime );
 
+		const bool bClearHittestGrid = !DrawOnlyThisWindow.IsValid();
+		if ( !SWidget::UseLegacyHittest( ) && bClearHittestGrid )
+		{
+			HittestGrid->BeginFrame( VirtualDesktopRect );
+		}
+
 		TSharedPtr<SWindow> ActiveModalWindow = GetActiveModalWindow(); 
 
 		if (ActiveModalWindow.IsValid())
@@ -859,11 +865,6 @@ void FSlateApplication::Tick()
 	{
 		SCOPE_CYCLE_COUNTER( STAT_SlateTickWindowAndChildren );
 
-		if ( !SWidget::UseLegacyHittest() )
-		{
-			HittestGrid->BeginFrame( VirtualDesktopRect );
-		}
-
 		if ( ActiveModalWindow.IsValid() )
 		{
 			// There is a modal window, and we just need to tick it.
