Added support for SMultiLineEditableText to edit rich-text

TTP# 336464 - Editor: Finish the Multiline Editable Text Block

Abstracted away the SetEditableText and GetEditableText functions from SMultiLineEditableText into "text marshallers" which handle converting text to and from a TextLayout.

There are three types of text marshallers currently implemented:
 * Plain-text
 * Rich-text
 * Syntax highlighting

Text marshallers also have the ability to inject formatting "live" (as the text changes), which is how the syntax highlighting marshaller works.

Added the ability for a run to query the information it was created with. This allows the rich-text marshaller to reconstruct the original rich-text from the styled runs.

To test this out, I've implemented a simple WYSIWYG rich-text editor demo with the following features:
 * Two SMultiLineEditableText widgets showing the same source text, one using a rich-text marshaller, and one using a syntax highlighter marshaller.
 * A toolbar to allow you to control the style of the selected text.
 * A button to allow you to insert a hyperlink into the document.

The demo also makes use of the meta-data stored in the runs (the same information used to reconstruct the original rich-text) to read the text style of whatever is currently under the cursor, live, as the cursor is moved.

Miscellaneous fixes:
 * Fixed an issue where deleting text that spanned multiple runs could leave the remaining runs in a bad state, leading to phantom text appearing (see FTextLayout::RemoveAt).
 * Fixed an issue where new-lines at the end of a rich-text document would be lost (see CalculateLineRanges).
 * Fixed an issue where \\r\\n line endings werenÆt being handled correctly by the rich-text parser (see CalculateLineRanges).
 * Fixed an issue where the rich-text parser would treat an empty run as plain-text (see FRichTextMarkupProcessing::ParseLineRanges).
 * Fixed an issue where inserting a line break when the cursor was at the end of a line containing multiple runs could sometimes fail (see FTextLayout::SplitLineAt).
 * Fixed mouse cursor movement not working correctly with a FSlateHyperlinkRun (see FSlateHyperlinkRun::GetTextIndexAt).

diff --git a/Engine/Source/Editor/EditorStyle/Private/SlateEditorStyle.cpp b/Engine/Source/Editor/EditorStyle/Private/SlateEditorStyle.cpp
--- a/Engine/Source/Editor/EditorStyle/Private/SlateEditorStyle.cpp
+++ b/Engine/Source/Editor/EditorStyle/Private/SlateEditorStyle.cpp
@@ -299,11 +299,11 @@ void FSlateEditorStyle::FStyle::SetupGeneralStyles()
 		const FCheckBoxStyle ToggleButtonStyle = FCheckBoxStyle()
 			.SetCheckBoxType(ESlateCheckBoxType::ToggleButton)
 			.SetUncheckedImage( FSlateNoResource() )
-			.SetUncheckedPressedImage( BOX_BRUSH( "Common/RoundedSelection_16x", 4.0f/16.0f, SelectionColor_Pressed ) )
 			.SetUncheckedHoveredImage( BOX_BRUSH( "Common/RoundedSelection_16x", 4.0f/16.0f, SelectionColor ) )
+			.SetUncheckedPressedImage( BOX_BRUSH( "Common/RoundedSelection_16x", 4.0f/16.0f, SelectionColor_Pressed ) )
 			.SetCheckedImage( BOX_BRUSH( "Common/RoundedSelection_16x",  4.0f/16.0f, SelectionColor_Pressed ) )
-			.SetCheckedHoveredImage( BOX_BRUSH( "Common/RoundedSelection_16x",  4.0f/16.0f, SelectionColor_Pressed ) )
-			.SetUncheckedPressedImage( BOX_BRUSH( "Common/RoundedSelection_16x",  4.0f/16.0f, SelectionColor ) );
+			.SetCheckedHoveredImage( BOX_BRUSH( "Common/RoundedSelection_16x",  4.0f/16.0f, SelectionColor ) )
+			.SetCheckedPressedImage( BOX_BRUSH( "Common/RoundedSelection_16x",  4.0f/16.0f, SelectionColor_Pressed ) );
 		/* ... and set new style */
 		Set( "ToggleButtonCheckbox", ToggleButtonStyle );
 
diff --git a/Engine/Source/Editor/UnrealEd/Private/CreditsScreen.cpp b/Engine/Source/Editor/UnrealEd/Private/CreditsScreen.cpp
--- a/Engine/Source/Editor/UnrealEd/Private/CreditsScreen.cpp
+++ b/Engine/Source/Editor/UnrealEd/Private/CreditsScreen.cpp
@@ -19,7 +19,6 @@ void SCreditsScreen::Construct(const FArguments& InArgs)
 	FString CreditsText;
 	FFileHelper::LoadFileToString(CreditsText, *( FPaths::EngineContentDir() + TEXT("Editor/Credits.rt") ));
 	CreditsText.ReplaceInline(TEXT("%VERSION%"), *Version);
-	CreditsText.ReplaceInline(TEXT("\r\n"), TEXT("\n"));
 
 	ChildSlot
 	[
diff --git a/Engine/Source/Runtime/Core/Public/Containers/UnrealString.h b/Engine/Source/Runtime/Core/Public/Containers/UnrealString.h
--- a/Engine/Source/Runtime/Core/Public/Containers/UnrealString.h
+++ b/Engine/Source/Runtime/Core/Public/Containers/UnrealString.h
@@ -446,25 +446,31 @@ public:
 
 	FORCEINLINE void InsertAt(int32 Index, TCHAR Character)
 	{
-		if (Data.Num() == 0)
+		if (Character != 0)
 		{
-			*this += Character;
-		}
-		else
-		{
-			Data.Insert(Character, Index);
+			if (Data.Num() == 0)
+			{
+				*this += Character;
+			}
+			else
+			{
+				Data.Insert(Character, Index);
+			}
 		}
 	}
 
 	FORCEINLINE void InsertAt(int32 Index, const FString& Characters)
 	{
-		if (Data.Num() == 0)
+		if (Characters.Len())
 		{
-			*this += Characters;
-		}
-		else
-		{
-			Data.Insert(Characters.Data.GetData(), Characters.Len(), Index);
+			if (Data.Num() == 0)
+			{
+				*this += Characters;
+			}
+			else
+			{
+				Data.Insert(Characters.Data.GetData(), Characters.Len(), Index);
+			}
 		}
 	}
 
diff --git a/Engine/Source/Runtime/Slate/Private/Framework/Application/SlateApplication.cpp b/Engine/Source/Runtime/Slate/Private/Framework/Application/SlateApplication.cpp
--- a/Engine/Source/Runtime/Slate/Private/Framework/Application/SlateApplication.cpp
+++ b/Engine/Source/Runtime/Slate/Private/Framework/Application/SlateApplication.cpp
@@ -1561,18 +1561,19 @@ bool FSlateApplication::SetKeyboardFocus( const FWidgetPath& InFocusPath, const
 					}
 				}
 				
+				const FWeakWidgetPath OldFocusedWidgetPath = FocusedWidgetPath;
+				// Store a weak widget path to the widget that's taking focus
+				FocusedWidgetPath = FWeakWidgetPath( NewFocusPath );
+
 				if( OldFocusedWidget.IsValid() )
 				{
 					// Switch worlds for widgets in the old path
-					FScopedSwitchWorldHack SwitchWorld( FocusedWidgetPath.Window.Pin() );
+					FScopedSwitchWorldHack SwitchWorld( OldFocusedWidgetPath.Window.Pin() );
 
 					// Let previously-focused widget know that it's losing focus
 					OldFocusedWidget->OnKeyboardFocusLost( FKeyboardFocusEvent( InCause ) );
 				}
 
-				// Store a weak widget path to the widget that's taking focus
-				FocusedWidgetPath = FWeakWidgetPath( NewFocusPath );
-
 				if (bReflectorShowingFocus)
 				{
 					WidgetReflector->SetWidgetsToVisualize(NewFocusPath);
@@ -1616,22 +1617,23 @@ FModifierKeysState FSlateApplication::GetModifierKeys() const
 
 void FSlateApplication::ClearKeyboardFocus( const EKeyboardFocusCause::Type InCause )
 {
+	TSharedPtr< SWidget > OldFocusedWidget( GetKeyboardFocusedWidget() );
+	const FWeakWidgetPath OldFocusedWidgetPath = FocusedWidgetPath;
+	FocusedWidgetPath = FWeakWidgetPath();
+
 	// Let previously-focused widget know that it's losing focus
+	if( OldFocusedWidget.IsValid() )
 	{
-		TSharedPtr< SWidget > OldFocusedWidget( GetKeyboardFocusedWidget() );
-		if( OldFocusedWidget.IsValid() )
+		if ( OldFocusedWidgetPath.Window.IsValid() )
 		{
-			if ( FocusedWidgetPath.Window.IsValid() )
-			{
-				// Switch worlds for widgets in the current path
-				FScopedSwitchWorldHack SwitchWorld( FocusedWidgetPath.Window.Pin().ToSharedRef() );
+			// Switch worlds for widgets in the current path
+			FScopedSwitchWorldHack SwitchWorld( OldFocusedWidgetPath.Window.Pin().ToSharedRef() );
 
-				OldFocusedWidget->OnKeyboardFocusLost( FKeyboardFocusEvent( InCause ) );
-			}
-			else
-			{
-				OldFocusedWidget->OnKeyboardFocusLost( FKeyboardFocusEvent( InCause ) );
-			}
+			OldFocusedWidget->OnKeyboardFocusLost( FKeyboardFocusEvent( InCause ) );
+		}
+		else
+		{
+			OldFocusedWidget->OnKeyboardFocusLost( FKeyboardFocusEvent( InCause ) );
 		}
 	}
 
@@ -1642,8 +1644,6 @@ void FSlateApplication::ClearKeyboardFocus( const EKeyboardFocusCause::Type InCa
 	{
 		WidgetReflector->SetWidgetsToVisualize(FWidgetPath());
 	}
-
-	FocusedWidgetPath = FWeakWidgetPath();
 }
 
 
diff --git a/Engine/Source/Runtime/Slate/Private/Framework/Testing/TestStyle.cpp b/Engine/Source/Runtime/Slate/Private/Framework/Testing/TestStyle.cpp
--- a/Engine/Source/Runtime/Slate/Private/Framework/Testing/TestStyle.cpp
+++ b/Engine/Source/Runtime/Slate/Private/Framework/Testing/TestStyle.cpp
@@ -91,6 +91,8 @@ TSharedRef< ISlateStyle > FTestStyle::Create()
 	{
 		Style->Set( "RichText.Background", new BOX_BRUSH( "Testing/FlatColorSquare", FVector2D(1.0f, 1.0f), FMargin(0), FLinearColor(FColor(0xffeff3f3))) );
 
+		Style->Set( "RichText.RoundedBackground", new BOX_BRUSH( "Common/RoundedSelection_16x", 4.0f/16.0f, FLinearColor(FColor(0xffeff3f3)) ) );
+
 		const FTextBlockStyle NormalRichTextStyle = FTextBlockStyle(NormalText)
 			.SetFont( TTF_FONT( "Fonts/Roboto-Regular", 14 ) )
 			.SetColorAndOpacity(FLinearColor(FColor( 0xff2c3e50 )));
@@ -232,6 +234,119 @@ TSharedRef< ISlateStyle > FTestStyle::Create()
 			}
 		}
 
+		// Editor example
+		{
+			// Default text styles
+			{
+				const FTextBlockStyle RichTextNormal = FTextBlockStyle()
+					.SetFont(TTF_FONT("Fonts/Roboto-Regular", 11))
+					.SetColorAndOpacity(FSlateColor::UseForeground())
+					.SetShadowOffset(FVector2D::ZeroVector)
+					.SetShadowColorAndOpacity(FLinearColor::Black)
+					.SetHighlightColor(FLinearColor(0.02f, 0.3f, 0.0f))
+					.SetHighlightShape(BOX_BRUSH("Common/TextBlockHighlightShape", FMargin(3.f /8.f)));
+				Style->Set( "RichText.Editor.Text", RichTextNormal );
+
+				{
+					const FButtonStyle RichTextHyperlinkButton = FButtonStyle()
+						.SetNormal(BORDER_BRUSH("Old/HyperlinkDotted", FMargin(0,0,0,3/16.0f), FSlateColor::UseForeground() ) )
+						.SetPressed(FSlateNoResource() )
+						.SetHovered(BORDER_BRUSH("Old/HyperlinkUnderline", FMargin(0,0,0,3/16.0f), FSlateColor::UseForeground() ) );
+
+					const FHyperlinkStyle RichTextHyperlink = FHyperlinkStyle()
+						.SetUnderlineStyle(RichTextHyperlinkButton)
+						.SetTextStyle(RichTextNormal)
+						.SetPadding(FMargin(0.0f));
+					Style->Set( "RichText.Editor.Hyperlink", RichTextHyperlink );
+				}
+			}
+
+			// Toolbar
+			{
+				const FLinearColor NormalColor(FColor(0xffeff3f3));
+				const FLinearColor SelectedColor(FColor(0xffdbe4d5));
+				const FLinearColor HoverColor(FColor(0xffdbe4e4));
+				const FLinearColor DisabledColor(FColor(0xaaaaaa));
+				const FLinearColor TextColor(FColor(0xff2c3e50));
+
+				Style->Set( "RichText.Toolbar.HyperlinkImage", new IMAGE_BRUSH( "Testing/hyperlink", Icon16x16, TextColor ) );
+
+				Style->Set( "RichText.Toolbar.TextColor", TextColor );
+
+				Style->Set( "RichText.Toolbar.Text", FTextBlockStyle(NormalText)
+					.SetFont( TTF_FONT( "Fonts/Roboto-Regular", 10 ) )
+					.SetColorAndOpacity( TextColor )
+					);
+
+				Style->Set( "RichText.Toolbar.BoldText", FTextBlockStyle(NormalText)
+					.SetFont( TTF_FONT( "Fonts/Roboto-Bold", 10 ) )
+					.SetColorAndOpacity( TextColor )
+					);
+
+				Style->Set( "RichText.Toolbar.ItalicText", FTextBlockStyle(NormalText)
+					.SetFont( TTF_FONT( "Testing/Fonts/Roboto-Italic", 10 ) )
+					.SetColorAndOpacity( TextColor )
+					);
+
+				Style->Set( "RichText.Toolbar.Checkbox", FCheckBoxStyle()
+					.SetCheckBoxType(ESlateCheckBoxType::CheckBox)
+					.SetUncheckedImage( IMAGE_BRUSH( "Common/CheckBox", Icon16x16, FLinearColor::White ) )
+					.SetUncheckedHoveredImage( IMAGE_BRUSH( "Common/CheckBox", Icon16x16, HoverColor ) )
+					.SetUncheckedPressedImage( IMAGE_BRUSH( "Common/CheckBox_Hovered", Icon16x16, HoverColor ) )
+					.SetCheckedImage( IMAGE_BRUSH( "Common/CheckBox_Checked_Hovered", Icon16x16, FLinearColor::White ) )
+					.SetCheckedHoveredImage( IMAGE_BRUSH( "Common/CheckBox_Checked_Hovered", Icon16x16, HoverColor ) )
+					.SetCheckedPressedImage( IMAGE_BRUSH( "Common/CheckBox_Checked", Icon16x16, HoverColor ) )
+					.SetUndeterminedImage( IMAGE_BRUSH( "Common/CheckBox_Undetermined", Icon16x16, FLinearColor::White ) )
+					.SetUndeterminedHoveredImage( IMAGE_BRUSH( "Common/CheckBox_Undetermined_Hovered", Icon16x16, HoverColor ) )
+					.SetUndeterminedPressedImage( IMAGE_BRUSH( "Common/CheckBox_Undetermined_Hovered", Icon16x16, FLinearColor::White ) )
+					);
+
+				Style->Set( "RichText.Toolbar.ToggleButtonCheckbox", FCheckBoxStyle()
+					.SetCheckBoxType(ESlateCheckBoxType::ToggleButton)
+					.SetUncheckedImage( BOX_BRUSH("Testing/FlatColorSquare", FVector2D(1.0f, 1.0f), FMargin(1), NormalColor ) )
+					.SetUncheckedHoveredImage( BOX_BRUSH("Testing/FlatColorSquare", FVector2D(1.0f, 1.0f), FMargin(1), HoverColor ) )
+					.SetUncheckedPressedImage( BOX_BRUSH("Testing/FlatColorSquare", FVector2D(1.0f, 1.0f), FMargin(1), HoverColor ) )
+					.SetCheckedImage( BOX_BRUSH("Testing/FlatColorSquare",  FVector2D(1.0f, 1.0f), FMargin(1), SelectedColor ) )
+					.SetCheckedHoveredImage( BOX_BRUSH("Testing/FlatColorSquare",  FVector2D(1.0f, 1.0f), FMargin(1), HoverColor ) )
+					.SetCheckedPressedImage( BOX_BRUSH("Testing/FlatColorSquare",  FVector2D(1.0f, 1.0f), FMargin(1), HoverColor ) )
+					);
+
+				const FButtonStyle Button = FButtonStyle()
+					.SetNormal( BOX_BRUSH("Testing/FlatColorSquare", FVector2D(1.0f, 1.0f), FMargin(1), NormalColor ) )
+					.SetHovered( BOX_BRUSH("Testing/FlatColorSquare", FVector2D(1.0f, 1.0f), FMargin(1), HoverColor ) )
+					.SetPressed( BOX_BRUSH("Testing/FlatColorSquare",  FVector2D(1.0f, 1.0f), FMargin(1), SelectedColor ) )
+					.SetDisabled( BOX_BRUSH("Testing/FlatColorSquare",  FVector2D(1.0f, 1.0f), FMargin(1), DisabledColor ) )
+					.SetNormalPadding( FMargin( 2,2,2,2 ) )
+					.SetPressedPadding( FMargin( 2,3,2,1 ) );
+				Style->Set( "RichText.Toolbar.Button", Button );
+
+				const FComboButtonStyle ComboButton = FComboButtonStyle()
+					.SetButtonStyle(Button)
+					.SetDownArrowImage(IMAGE_BRUSH("Common/ComboArrow", Icon8x8))
+					.SetMenuBorderBrush(BOX_BRUSH( "Testing/FlatColorSquare", FVector2D(1.0f, 1.0f), FMargin(1), NormalColor ))
+					.SetMenuBorderPadding(FMargin(0.0f));
+				Style->Set( "RichText.Toolbar.ComboButton", ComboButton	);
+
+				{
+					const FButtonStyle ComboBoxButton = FButtonStyle()
+						.SetNormal( BOX_BRUSH("Testing/FlatColorSquare", FVector2D(1.0f, 1.0f), FMargin(1), FLinearColor::White ) )
+						.SetHovered( BOX_BRUSH("Testing/FlatColorSquare", FVector2D(1.0f, 1.0f), FMargin(1), FLinearColor::White ) )
+						.SetPressed( BOX_BRUSH("Testing/FlatColorSquare",  FVector2D(1.0f, 1.0f), FMargin(1), FLinearColor::White ) )
+						.SetDisabled( BOX_BRUSH("Testing/FlatColorSquare",  FVector2D(1.0f, 1.0f), FMargin(1), DisabledColor ) )
+						.SetNormalPadding( FMargin( 2,2,2,2 ) )
+						.SetPressedPadding( FMargin( 2,3,2,1 ) );
+
+					const FComboButtonStyle ComboBoxComboButton = FComboButtonStyle(ComboButton)
+						.SetButtonStyle(ComboBoxButton)
+						.SetMenuBorderPadding(FMargin(1.0));
+
+					Style->Set( "RichText.Toolbar.ComboBox", FComboBoxStyle()
+						.SetComboButtonStyle(ComboBoxComboButton)
+						);
+				}
+			}
+		}
+
 		// The War of the Worlds example
 		{
 			Style->Set( "TheWarOfTheWorlds.Text", FTextBlockStyle(NormalText)
diff --git a/Engine/Source/Runtime/Slate/Private/Framework/Text/PlainTextLayoutMarshaller.cpp b/Engine/Source/Runtime/Slate/Private/Framework/Text/PlainTextLayoutMarshaller.cpp
new file mode 100644
--- /dev/null
+++ b/Engine/Source/Runtime/Slate/Private/Framework/Text/PlainTextLayoutMarshaller.cpp
@@ -0,0 +1,54 @@
+// Copyright 1998-2014 Epic Games, Inc. All Rights Reserved.
+
+#include "SlatePrivatePCH.h"
+#include "PlainTextLayoutMarshaller.h"
+
+#if WITH_FANCY_TEXT
+
+TSharedRef< FPlainTextLayoutMarshaller > FPlainTextLayoutMarshaller::Create(FTextBlockStyle InDefaultTextStyle)
+{
+	return MakeShareable(new FPlainTextLayoutMarshaller(MoveTemp(InDefaultTextStyle)));
+}
+
+FPlainTextLayoutMarshaller::~FPlainTextLayoutMarshaller()
+{
+}
+
+void FPlainTextLayoutMarshaller::SetText(const FString& SourceString, FTextLayout& TargetTextLayout)
+{
+	TArray<FTextRange> LineRanges;
+	FTextRange::CalculateLineRangesFromString(SourceString, LineRanges);
+
+	for(const FTextRange& LineRange : LineRanges)
+	{
+		TSharedRef<FString> LineText = MakeShareable(new FString(SourceString.Mid(LineRange.BeginIndex, LineRange.Len())));
+
+		TArray<TSharedRef<IRun>> Runs;
+		Runs.Add(FSlateTextRun::Create(FRunInfo(), LineText, DefaultTextStyle));
+
+		TargetTextLayout.AddLine(LineText, Runs);
+	}
+}
+
+void FPlainTextLayoutMarshaller::GetText(FString& TargetString, const FTextLayout& SourceTextLayout)
+{
+	SourceTextLayout.GetAsText(TargetString);
+}
+
+void FPlainTextLayoutMarshaller::SetDefaultTextStyle(FTextBlockStyle InDefaultTextStyle)
+{
+	DefaultTextStyle = MoveTemp(InDefaultTextStyle);
+	MakeDirty();
+}
+
+const FTextBlockStyle& FPlainTextLayoutMarshaller::GetDefaultTextStyle() const
+{
+	return DefaultTextStyle;
+}
+
+FPlainTextLayoutMarshaller::FPlainTextLayoutMarshaller(FTextBlockStyle InDefaultTextStyle)
+	: DefaultTextStyle(MoveTemp(InDefaultTextStyle))
+{
+}
+
+#endif //WITH_FANCY_TEXT
diff --git a/Engine/Source/Runtime/Slate/Private/Framework/Text/RichTextLayoutMarshaller.cpp b/Engine/Source/Runtime/Slate/Private/Framework/Text/RichTextLayoutMarshaller.cpp
new file mode 100644
--- /dev/null
+++ b/Engine/Source/Runtime/Slate/Private/Framework/Text/RichTextLayoutMarshaller.cpp
@@ -0,0 +1,147 @@
+// Copyright 1998-2014 Epic Games, Inc. All Rights Reserved.
+
+#include "SlatePrivatePCH.h"
+#include "RichTextLayoutMarshaller.h"
+#include "RichTextMarkupProcessing.h"
+
+#if WITH_FANCY_TEXT
+
+TSharedRef< FRichTextLayoutMarshaller > FRichTextLayoutMarshaller::Create(TArray< TSharedRef< ITextDecorator > > InDecorators, const ISlateStyle* const InDecoratorStyleSet, FTextBlockStyle InDefaultTextStyle)
+{
+	return MakeShareable(new FRichTextLayoutMarshaller(FDefaultRichTextMarkupParser::Create(), FDefaultRichTextMarkupWriter::Create(), MoveTemp(InDecorators), InDecoratorStyleSet, MoveTemp(InDefaultTextStyle)));
+}
+
+TSharedRef< FRichTextLayoutMarshaller > FRichTextLayoutMarshaller::Create(TSharedPtr< IRichTextMarkupParser > InParser, TSharedPtr< IRichTextMarkupWriter > InWriter, TArray< TSharedRef< ITextDecorator > > InDecorators, const ISlateStyle* const InDecoratorStyleSet, FTextBlockStyle InDefaultTextStyle)
+{
+	return MakeShareable(new FRichTextLayoutMarshaller(MoveTemp(InParser), MoveTemp(InWriter), MoveTemp(InDecorators), InDecoratorStyleSet, MoveTemp(InDefaultTextStyle)));
+}
+
+FRichTextLayoutMarshaller::~FRichTextLayoutMarshaller()
+{
+}
+
+void FRichTextLayoutMarshaller::SetText(const FString& SourceString, FTextLayout& TargetTextLayout)
+{
+	TArray<FTextLineParseResults> LineParseResultsArray;
+	FString ProcessedString;
+	Parser->Process(LineParseResultsArray, SourceString, ProcessedString);
+
+	// Iterate through parsed line results and create processed lines with runs.
+	for (const FTextLineParseResults& LineParseResults : LineParseResultsArray)
+	{
+		TSharedRef<FString> ModelString = MakeShareable(new FString());
+		TArray< TSharedRef< IRun > > Runs;
+
+		for (const FTextRunParseResults& RunParseResult : LineParseResults.Runs)
+		{
+			TSharedPtr< ISlateRun > Run;
+
+			TSharedPtr< ITextDecorator > Decorator = TryGetDecorator(ProcessedString, RunParseResult);
+
+			if (Decorator.IsValid())
+			{
+				// Create run and update model string.
+				Run = Decorator->Create(RunParseResult, ProcessedString, ModelString, DecoratorStyleSet);
+			}
+			else
+			{
+				FRunInfo RunInfo(RunParseResult.Name);
+				for (const TPair<FString, FTextRange>& Pair : RunParseResult.MetaData)
+				{
+					RunInfo.MetaData.Add(Pair.Key, ProcessedString.Mid( Pair.Value.BeginIndex, Pair.Value.EndIndex - Pair.Value.BeginIndex));
+				}
+
+				const FTextBlockStyle* TextBlockStyle;
+				FTextRange ModelRange;
+				ModelRange.BeginIndex = ModelString->Len();
+				if (!(RunParseResult.Name.IsEmpty()) && DecoratorStyleSet->HasWidgetStyle< FTextBlockStyle >(FName(*RunParseResult.Name)))
+				{
+					*ModelString += ProcessedString.Mid(RunParseResult.ContentRange.BeginIndex, RunParseResult.ContentRange.EndIndex - RunParseResult.ContentRange.BeginIndex);
+					TextBlockStyle = &(DecoratorStyleSet->GetWidgetStyle< FTextBlockStyle >( FName(*RunParseResult.Name) ));
+				}
+				else
+				{
+					*ModelString += ProcessedString.Mid(RunParseResult.OriginalRange.BeginIndex, RunParseResult.OriginalRange.EndIndex - RunParseResult.OriginalRange.BeginIndex);
+					TextBlockStyle = &DefaultTextStyle;
+				}
+				ModelRange.EndIndex = ModelString->Len();
+
+				// Create run.
+				Run = FSlateTextRun::Create(RunInfo, ModelString, *TextBlockStyle, ModelRange);
+			}
+
+			Runs.Add( Run.ToSharedRef() );
+		}
+
+		TargetTextLayout.AddLine(ModelString, Runs);
+	}
+}
+
+void FRichTextLayoutMarshaller::GetText(FString& TargetString, const FTextLayout& SourceTextLayout)
+{
+	const auto& SourceLineModels = SourceTextLayout.GetLineModels();
+
+	TArray<IRichTextMarkupWriter::FRichTextLine> WriterLines;
+	WriterLines.Reserve(SourceLineModels.Num());
+
+	for (const FTextLayout::FLineModel& LineModel : SourceLineModels)
+	{
+		IRichTextMarkupWriter::FRichTextLine WriterLine;
+		WriterLine.Runs.Reserve(LineModel.Runs.Num());
+
+		for (const FTextLayout::FRunModel& RunModel : LineModel.Runs)
+		{
+			FString RunText;
+			RunModel.AppendTextTo(RunText);
+			WriterLine.Runs.Emplace(IRichTextMarkupWriter::FRichTextRun(RunModel.GetRun()->GetRunInfo(), RunText));
+		}
+
+		WriterLines.Add(WriterLine);
+	}
+
+	Writer->Write(WriterLines, TargetString);
+}
+
+void FRichTextLayoutMarshaller::SetDefaultTextStyle(FTextBlockStyle InDefaultTextStyle)
+{
+	DefaultTextStyle = MoveTemp(InDefaultTextStyle);
+	MakeDirty();
+}
+
+const FTextBlockStyle& FRichTextLayoutMarshaller::GetDefaultTextStyle() const
+{
+	return DefaultTextStyle;
+}
+
+FRichTextLayoutMarshaller::FRichTextLayoutMarshaller(TSharedPtr< IRichTextMarkupParser > InParser, TSharedPtr< IRichTextMarkupWriter > InWriter, TArray< TSharedRef< ITextDecorator > > InDecorators, const ISlateStyle* const InDecoratorStyleSet, FTextBlockStyle InDefaultTextStyle)
+	: Parser(MoveTemp(InParser))
+	, Writer(MoveTemp(InWriter))
+	, Decorators(MoveTemp(InDecorators))
+	, InlineDecorators()
+	, DecoratorStyleSet(InDecoratorStyleSet)
+	, DefaultTextStyle(MoveTemp(InDefaultTextStyle))
+{
+}
+
+TSharedPtr< ITextDecorator > FRichTextLayoutMarshaller::TryGetDecorator(const FString& Line, const FTextRunParseResults& TextRun) const
+{
+	for (auto DecoratorIter = InlineDecorators.CreateConstIterator(); DecoratorIter; ++DecoratorIter)
+	{
+		if ((*DecoratorIter)->Supports(TextRun, Line))
+		{
+			return *DecoratorIter;
+		}
+	}
+
+	for (auto DecoratorIter = Decorators.CreateConstIterator(); DecoratorIter; ++DecoratorIter)
+	{
+		if ((*DecoratorIter)->Supports(TextRun, Line))
+		{
+			return *DecoratorIter;
+		}
+	}
+
+	return nullptr;
+}
+
+#endif //WITH_FANCY_TEXT
diff --git a/Engine/Source/Runtime/Slate/Private/Framework/Text/RichTextMarkupProcessing.cpp b/Engine/Source/Runtime/Slate/Private/Framework/Text/RichTextMarkupProcessing.cpp
--- a/Engine/Source/Runtime/Slate/Private/Framework/Text/RichTextMarkupProcessing.cpp
+++ b/Engine/Source/Runtime/Slate/Private/Framework/Text/RichTextMarkupProcessing.cpp
@@ -7,16 +7,16 @@
 #include "RichTextMarkupProcessing.h"
 
 
-TSharedRef< FRichTextMarkupProcessing > FRichTextMarkupProcessing::Create()
+TSharedRef< FDefaultRichTextMarkupParser > FDefaultRichTextMarkupParser::Create()
 {
-	return MakeShareable( new FRichTextMarkupProcessing() );
+	return MakeShareable( new FDefaultRichTextMarkupParser() );
 }
 
-void FRichTextMarkupProcessing::Process(TArray<FTextLineParseResults>& Results, const FString& Input, FString& Output)
+void FDefaultRichTextMarkupParser::Process(TArray<FTextLineParseResults>& Results, const FString& Input, FString& Output)
 {
 #if UE_ENABLE_ICU
 	TArray<FTextRange> LineRanges;
-	CalculateLineRanges(Input, LineRanges);
+	FTextRange::CalculateLineRangesFromString(Input, LineRanges);
 	ParseLineRanges(Input, LineRanges, Results);
 	HandleEscapeSequences(Input, Results, Output);
 #else
@@ -28,48 +28,10 @@ void FRichTextMarkupProcessing::Process(TArray<FTextLineParseResults>& Results,
 #endif
 }
 
-void FRichTextMarkupProcessing::CalculateLineRanges(const FString& Input, TArray<FTextRange>& LineRanges) const
-{
-	// Iterate over each line break candidate, adding ranges from after the end of the last line added to before the newline or end of string.
-	FLineBreakIterator LBI(Input);
-	int32 RangeBegin = LBI.GetCurrentPosition();
-	for(;;)
-	{
-		const int32 BreakIndex = LBI.MoveToNext();
-
-		// If beginning or potential end is invalid, cease.
-		if(RangeBegin == INDEX_NONE || BreakIndex == INDEX_NONE)
-		{
-			break;
-		}
-
-		const int32 BreakingCharacterIndex = BreakIndex - 1;
-		if(BreakingCharacterIndex >= RangeBegin) // Valid index check.
-		{
-			const TCHAR BreakingCharacter = Input[BreakingCharacterIndex];
-			// If line break candidate is after a newline, add the range as a new line.
-			if(FChar::IsLinebreak(BreakingCharacter))
-			{
-				LineRanges.Emplace( FTextRange(RangeBegin, BreakingCharacterIndex) );
-			}
-			// If the line break candidate is the end of the string, add the range as a new line.
-			else if(BreakIndex == Input.Len())
-			{
-				LineRanges.Emplace( FTextRange(RangeBegin, BreakIndex) );
-			}
-			else
-			{
-				continue;
-			}
-			RangeBegin = BreakIndex; // The next line begins after the end of the current line.
-		}
-	}
-}
-
-void FRichTextMarkupProcessing::ParseLineRanges(const FString& Input, const TArray<FTextRange>& LineRanges, TArray<FTextLineParseResults>& LineParseResultsArray) const
+void FDefaultRichTextMarkupParser::ParseLineRanges(const FString& Input, const TArray<FTextRange>& LineRanges, TArray<FTextLineParseResults>& LineParseResultsArray) const
 {
 	// Special regular expression pattern for matching rich text markup elements. IE: <ElementName AttributeName="AttributeValue">Content</>
-	FRegexPattern ElementRegexPattern( TEXT("<([\\w\\d\\.]+)((?: (?:[\\w\\d\\.]+=(?>\".*?\")))+)?(?:(?:/>)|(?:>(.+?)</>))") );
+	FRegexPattern ElementRegexPattern( TEXT("<([\\w\\d\\.]+)((?: (?:[\\w\\d\\.]+=(?>\".*?\")))+)?(?:(?:/>)|(?:>(.*?)</>))") );
 	FRegexMatcher ElementRegexMatcher(ElementRegexPattern, Input);
 
 	// Parse line ranges, creating line parse results and run parse results.
@@ -171,7 +133,7 @@ void FRichTextMarkupProcessing::ParseLineRanges(const FString& Input, const TArr
 	}
 }
 
-void FRichTextMarkupProcessing::HandleEscapeSequences(const FString& Input, TArray<FTextLineParseResults>& LineParseResultsArray, FString& ConcatenatedUnescapedLines) const
+void FDefaultRichTextMarkupParser::HandleEscapeSequences(const FString& Input, TArray<FTextLineParseResults>& LineParseResultsArray, FString& ConcatenatedUnescapedLines) const
 {
 	// Modify original string to handle escape sequences that need to be replaced while updating run ranges.
 	for(int32 i = 0; i < LineParseResultsArray.Num(); ++i)
@@ -295,4 +257,105 @@ void FRichTextMarkupProcessing::HandleEscapeSequences(const FString& Input, TArr
 	}
 }
 
-#endif //WITH_FANCY_TEXT
\ No newline at end of file
+
+TSharedRef< FDefaultRichTextMarkupWriter > FDefaultRichTextMarkupWriter::Create()
+{
+	return MakeShareable( new FDefaultRichTextMarkupWriter() );
+}
+
+void FDefaultRichTextMarkupWriter::Write(const TArray<FRichTextLine>& InLines, FString& Output)
+{
+	for (int32 LineIndex = 0; LineIndex < InLines.Num(); ++LineIndex)
+	{
+		const FRichTextLine& Line = InLines[LineIndex];
+
+		// Append \n to the end of the previous line
+		if(LineIndex > 0)
+		{
+			Output.AppendChar('\n');
+		}
+
+		for (const FRichTextRun& Run : Line.Runs)
+		{
+			// Our rich-text format takes the form of <Name metakey1="metavalue1" metakey2="metavalue2">The Text</>
+			const bool bHasTag = !Run.Info.Name.IsEmpty();
+			if (bHasTag)
+			{
+				Output.AppendChar('<');
+
+				Output.Append(Run.Info.Name);
+
+				for(const TPair<FString, FString>& MetaDataEntry : Run.Info.MetaData)
+				{
+					Output.AppendChar(' ');
+					Output.Append(MetaDataEntry.Key);
+					Output.AppendChar('=');
+					Output.AppendChar('"');
+					Output.Append(MetaDataEntry.Value);
+					Output.AppendChar('"');
+				}
+
+				Output.AppendChar('>');
+			}
+
+			FString RunText = Run.Text;
+			EscapeText(RunText);
+			Output.Append(RunText);
+
+			if (bHasTag)
+			{
+				Output.Append(TEXT("</>"));
+			}
+		}
+	}
+}
+
+void FDefaultRichTextMarkupWriter::EscapeText(FString& TextToEscape)
+{
+	// List of characters that we have to escape to avoid accidental rich-text formatting
+	static const TPair<TCHAR, FString> EscapeCharacters[] = {
+		TPairInitializer<TCHAR, FString>('&', "&amp;"),
+		TPairInitializer<TCHAR, FString>('"', "&quot;"),
+		TPairInitializer<TCHAR, FString>('<', "&lt;"),
+		TPairInitializer<TCHAR, FString>('>', "&gt;"),
+	};
+
+	// First count the extra space needed
+	int32 EscapedStringLen = TextToEscape.Len();
+	for (int32 TextIndex = 0; TextIndex < TextToEscape.Len(); ++TextIndex)
+	{
+		const TCHAR CharToEscape = TextToEscape[TextIndex];
+		for (const auto& EscapeCharacter : EscapeCharacters)
+		{
+			if (CharToEscape == EscapeCharacter.Key)
+			{
+				EscapedStringLen += (EscapeCharacter.Value.Len() - 1);
+				break;
+			}
+		}
+	}
+
+	if (EscapedStringLen == TextToEscape.Len())
+	{
+		return;
+	}
+
+	TextToEscape.Reserve(EscapedStringLen);
+
+	// Then perform the actual escape (backwards to make iteration simpler)
+	for (int32 TextIndex = TextToEscape.Len() - 1; TextIndex >= 0; --TextIndex)
+	{
+		const TCHAR CharToEscape = TextToEscape[TextIndex];
+		for (const auto& EscapeCharacter : EscapeCharacters)
+		{
+			if (CharToEscape == EscapeCharacter.Key)
+			{
+				TextToEscape.RemoveAt(TextIndex, 1, false);
+				TextToEscape.InsertAt(TextIndex, EscapeCharacter.Value);
+				break;
+			}
+		}
+	}
+}
+
+#endif //WITH_FANCY_TEXT
diff --git a/Engine/Source/Runtime/Slate/Private/Framework/Text/RichTextMarkupProcessing.h b/Engine/Source/Runtime/Slate/Private/Framework/Text/RichTextMarkupProcessing.h
--- a/Engine/Source/Runtime/Slate/Private/Framework/Text/RichTextMarkupProcessing.h
+++ b/Engine/Source/Runtime/Slate/Private/Framework/Text/RichTextMarkupProcessing.h
@@ -4,23 +4,35 @@
 #if WITH_FANCY_TEXT
 
 #include "IRichTextMarkupParser.h"
+#include "IRichTextMarkupWriter.h"
 
-class FRichTextMarkupProcessing : public IRichTextMarkupParser
+class FDefaultRichTextMarkupParser : public IRichTextMarkupParser
 {
 public:
-	static TSharedRef< FRichTextMarkupProcessing > Create();
+	static TSharedRef< FDefaultRichTextMarkupParser > Create();
 
 public:
 	virtual void Process(TArray<FTextLineParseResults>& Results, const FString& Input, FString& Output) override;
 
 private:
-	FRichTextMarkupProcessing() {}
+	FDefaultRichTextMarkupParser() {}
 
-	void CalculateLineRanges(const FString& Input, TArray<FTextRange>& LineRanges) const;
 	void ParseLineRanges(const FString& Input, const TArray<FTextRange>& LineRanges, TArray<FTextLineParseResults>& LineParseResultsArray) const;
 	void HandleEscapeSequences(const FString& Input, TArray<FTextLineParseResults>& LineParseResultsArray, FString& ConcatenatedUnescapedLines) const;
 
 };
 
+class FDefaultRichTextMarkupWriter : public IRichTextMarkupWriter
+{
+public:
+	static TSharedRef< FDefaultRichTextMarkupWriter > Create();
+
+public:
+	virtual void Write(const TArray<FRichTextLine>& InLines, FString& Output) override;
+
+private:
+	FDefaultRichTextMarkupWriter() {}
+	static void EscapeText(FString& TextToEscape);
+};
 
-#endif //WITH_FANCY_TEXT
\ No newline at end of file
+#endif //WITH_FANCY_TEXT
diff --git a/Engine/Source/Runtime/Slate/Private/Framework/Text/SlateHyperlinkRun.cpp b/Engine/Source/Runtime/Slate/Private/Framework/Text/SlateHyperlinkRun.cpp
--- a/Engine/Source/Runtime/Slate/Private/Framework/Text/SlateHyperlinkRun.cpp
+++ b/Engine/Source/Runtime/Slate/Private/Framework/Text/SlateHyperlinkRun.cpp
@@ -6,14 +6,14 @@
 
 #include "SlateHyperlinkRun.h"
 
-TSharedRef< FSlateHyperlinkRun > FSlateHyperlinkRun::Create( const TSharedRef< const FString >& InText, const FHyperlinkStyle& InStyle, const TMap<FString,FString>& Metadata, FOnClick NavigateDelegate )
+TSharedRef< FSlateHyperlinkRun > FSlateHyperlinkRun::Create( const FRunInfo& InRunInfo, const TSharedRef< const FString >& InText, const FHyperlinkStyle& InStyle, FOnClick NavigateDelegate )
 {
-	return MakeShareable( new FSlateHyperlinkRun( InText, InStyle, Metadata, NavigateDelegate ) );
+	return MakeShareable( new FSlateHyperlinkRun( InRunInfo, InText, InStyle, NavigateDelegate ) );
 }
 
-TSharedRef< FSlateHyperlinkRun > FSlateHyperlinkRun::Create( const TSharedRef< const FString >& InText, const FHyperlinkStyle& InStyle, const TMap<FString,FString>& Metadata, FOnClick NavigateDelegate, const FTextRange& InRange )
+TSharedRef< FSlateHyperlinkRun > FSlateHyperlinkRun::Create( const FRunInfo& InRunInfo, const TSharedRef< const FString >& InText, const FHyperlinkStyle& InStyle, FOnClick NavigateDelegate, const FTextRange& InRange )
 {
-	return MakeShareable( new FSlateHyperlinkRun( InText, InStyle, Metadata, NavigateDelegate, InRange ) );
+	return MakeShareable( new FSlateHyperlinkRun( InRunInfo, InText, InStyle, NavigateDelegate, InRange ) );
 }
 
 FTextRange FSlateHyperlinkRun::GetTextRange() const 
@@ -68,7 +68,7 @@ TSharedRef< ILayoutBlock > FSlateHyperlinkRun::CreateBlock( int32 StartIndex, in
 
 void FSlateHyperlinkRun::OnNavigate()
 {
-	NavigateDelegate.Execute( Metadata );
+	NavigateDelegate.Execute( RunInfo.MetaData );
 }
 
 int32 FSlateHyperlinkRun::OnPaint( const FPaintArgs& Args, const FTextLayout::FLineView& Line, const TSharedRef< ILayoutBlock >& Block, const FTextBlockStyle& DefaultStyle, const FGeometry& AllottedGeometry, const FSlateRect& MyClippingRect, FSlateWindowElementList& OutDrawElements, int32 LayerId, const FWidgetStyle& InWidgetStyle, bool bParentEnabled ) const 
@@ -92,7 +92,31 @@ void FSlateHyperlinkRun::ArrangeChildren( const TSharedRef< ILayoutBlock >& Bloc
 
 int32 FSlateHyperlinkRun::GetTextIndexAt( const TSharedRef< ILayoutBlock >& Block, const FVector2D& Location, float Scale, ETextHitPoint* const OutHitPoint ) const
 {
-	return INDEX_NONE;
+	const FVector2D& BlockOffset = Block->GetLocationOffset();
+	const FVector2D& BlockSize = Block->GetSize();
+
+	const float Left = BlockOffset.X;
+	const float Top = BlockOffset.Y;
+	const float Right = BlockOffset.X + BlockSize.X;
+	const float Bottom = BlockOffset.Y + BlockSize.Y;
+
+	const bool ContainsPoint = Location.X >= Left && Location.X < Right && Location.Y >= Top && Location.Y < Bottom;
+
+	if ( !ContainsPoint )
+	{
+		return INDEX_NONE;
+	}
+
+	const FTextRange BlockRange = Block->GetTextRange();
+	const TSharedRef< FSlateFontMeasure > FontMeasure = FSlateApplication::Get().GetRenderer()->GetFontMeasureService();
+
+	const int32 Index = FontMeasure->FindCharacterIndexAtOffset( Text.Get(), BlockRange.BeginIndex, BlockRange.EndIndex, Style.TextStyle.Font, Location.X - BlockOffset.X, true, Scale );
+	if (OutHitPoint)
+	{
+		*OutHitPoint = (Index == BlockRange.EndIndex) ? ETextHitPoint::RightGutter : ETextHitPoint::WithinText;
+	}
+
+	return Index;
 }
 
 FVector2D FSlateHyperlinkRun::GetLocationAt( const TSharedRef< ILayoutBlock >& Block, int32 Offset, float Scale ) const
@@ -108,15 +132,15 @@ void FSlateHyperlinkRun::Move(const TSharedRef<FString>& NewText, const FTextRan
 
 TSharedRef<IRun> FSlateHyperlinkRun::Clone() const
 {
-	return FSlateHyperlinkRun::Create(Text, Style, Metadata, NavigateDelegate, Range);
+	return FSlateHyperlinkRun::Create(RunInfo, Text, Style, NavigateDelegate, Range);
 }
 
-void FSlateHyperlinkRun::AppendText(FString& AppendToText) const
+void FSlateHyperlinkRun::AppendTextTo(FString& AppendToText) const
 {
 	AppendToText.Append(**Text + Range.BeginIndex, Range.Len());
 }
 
-void FSlateHyperlinkRun::AppendText(FString& AppendToText, const FTextRange& PartialRange) const
+void FSlateHyperlinkRun::AppendTextTo(FString& AppendToText, const FTextRange& PartialRange) const
 {
 	check(Range.BeginIndex <= PartialRange.BeginIndex);
 	check(Range.EndIndex >= PartialRange.EndIndex);
@@ -124,11 +148,16 @@ void FSlateHyperlinkRun::AppendText(FString& AppendToText, const FTextRange& Par
 	AppendToText.Append(**Text + PartialRange.BeginIndex, PartialRange.Len());
 }
 
-FSlateHyperlinkRun::FSlateHyperlinkRun( const TSharedRef< const FString >& InText, const FHyperlinkStyle& InStyle, const FMetadata& InMetadata, FOnClick InNavigateDelegate ) 
-	: Text( InText )
+const FRunInfo& FSlateHyperlinkRun::GetRunInfo() const
+{
+	return RunInfo;
+}
+
+FSlateHyperlinkRun::FSlateHyperlinkRun( const FRunInfo& InRunInfo, const TSharedRef< const FString >& InText, const FHyperlinkStyle& InStyle, FOnClick InNavigateDelegate ) 
+	: RunInfo( InRunInfo )
+	, Text( InText )
 	, Range( 0, Text->Len() )
 	, Style( InStyle )
-	, Metadata( InMetadata )
 	, NavigateDelegate( InNavigateDelegate )
 	, ViewModel( MakeShareable( new FSlateHyperlinkRun::FWidgetViewModel() ) )
 	, Children()
@@ -136,11 +165,11 @@ FSlateHyperlinkRun::FSlateHyperlinkRun( const TSharedRef< const FString >& InTex
 
 }
 
-FSlateHyperlinkRun::FSlateHyperlinkRun( const TSharedRef< const FString >& InText, const FHyperlinkStyle& InStyle, const FMetadata& InMetadata, FOnClick InNavigateDelegate, const FTextRange& InRange ) 
-	: Text( InText )
+FSlateHyperlinkRun::FSlateHyperlinkRun( const FRunInfo& InRunInfo, const TSharedRef< const FString >& InText, const FHyperlinkStyle& InStyle, FOnClick InNavigateDelegate, const FTextRange& InRange ) 
+	: RunInfo( InRunInfo )
+	, Text( InText )
 	, Range( InRange )
 	, Style( InStyle )
-	, Metadata( InMetadata )
 	, NavigateDelegate( InNavigateDelegate )
 	, ViewModel( MakeShareable( new FSlateHyperlinkRun::FWidgetViewModel() ) )
 	, Children()
@@ -149,10 +178,10 @@ FSlateHyperlinkRun::FSlateHyperlinkRun( const TSharedRef< const FString >& InTex
 }
 
 FSlateHyperlinkRun::FSlateHyperlinkRun( const FSlateHyperlinkRun& Run ) 
-	: Text( Run.Text )
+	: RunInfo( Run.RunInfo )
+	, Text( Run.Text )
 	, Range( Run.Range )
 	, Style( Run.Style )
-	, Metadata( Run.Metadata )
 	, NavigateDelegate( Run.NavigateDelegate )
 	, ViewModel( MakeShareable( new FSlateHyperlinkRun::FWidgetViewModel() ) )
 	, Children()
diff --git a/Engine/Source/Runtime/Slate/Private/Framework/Text/SlateImageRun.cpp b/Engine/Source/Runtime/Slate/Private/Framework/Text/SlateImageRun.cpp
--- a/Engine/Source/Runtime/Slate/Private/Framework/Text/SlateImageRun.cpp
+++ b/Engine/Source/Runtime/Slate/Private/Framework/Text/SlateImageRun.cpp
@@ -6,28 +6,29 @@
 
 #include "SlateImageRun.h"
 
-TSharedRef< FSlateImageRun > FSlateImageRun::Create( const TSharedRef< const FString >& InText, const FSlateBrush* InImage, int16 InBaseline )
+TSharedRef< FSlateImageRun > FSlateImageRun::Create( const FRunInfo& InRunInfo, const TSharedRef< const FString >& InText, const FSlateBrush* InImage, int16 InBaseline )
 {
 	if ( InImage == nullptr)
 	{
 		InImage = FStyleDefaults::GetNoBrush();
 	}
 
-	return MakeShareable( new FSlateImageRun( InText, InImage, InBaseline ) );
+	return MakeShareable( new FSlateImageRun( InRunInfo, InText, InImage, InBaseline ) );
 }
 
-TSharedRef< FSlateImageRun > FSlateImageRun::Create( const TSharedRef< const FString >& InText, const FSlateBrush* InImage, int16 InBaseline, const FTextRange& InRange )
+TSharedRef< FSlateImageRun > FSlateImageRun::Create( const FRunInfo& InRunInfo, const TSharedRef< const FString >& InText, const FSlateBrush* InImage, int16 InBaseline, const FTextRange& InRange )
 {
 	if ( InImage == nullptr)
 	{
 		InImage = FStyleDefaults::GetNoBrush();
 	}
 
-	return MakeShareable( new FSlateImageRun( InText, InImage, InBaseline, InRange ) );
+	return MakeShareable( new FSlateImageRun( InRunInfo, InText, InImage, InBaseline, InRange ) );
 }
 
-FSlateImageRun::FSlateImageRun( const TSharedRef< const FString >& InText, const FSlateBrush* InImage, int16 InBaseline, const FTextRange& InRange ) 
-	: Text( InText )
+FSlateImageRun::FSlateImageRun( const FRunInfo& InRunInfo, const TSharedRef< const FString >& InText, const FSlateBrush* InImage, int16 InBaseline, const FTextRange& InRange ) 
+	: RunInfo( InRunInfo )
+	, Text( InText )
 	, Range( InRange )
 	, Image( InImage )
 	, Baseline( InBaseline )
@@ -35,8 +36,9 @@ FSlateImageRun::FSlateImageRun( const TSharedRef< const FString >& InText, const
 	check( Image != nullptr);
 }
 
-FSlateImageRun::FSlateImageRun( const TSharedRef< const FString >& InText, const FSlateBrush* InImage, int16 InBaseline ) 
-	: Text( InText )
+FSlateImageRun::FSlateImageRun( const FRunInfo& InRunInfo, const TSharedRef< const FString >& InText, const FSlateBrush* InImage, int16 InBaseline ) 
+	: RunInfo( InRunInfo )
+	, Text( InText )
 	, Range( 0, Text->Len() )
 	, Image( InImage )
 	, Baseline( InBaseline )
@@ -121,18 +123,23 @@ void FSlateImageRun::Move(const TSharedRef<FString>& NewText, const FTextRange&
 
 TSharedRef<IRun> FSlateImageRun::Clone() const
 {
-	return FSlateImageRun::Create(Text, Image, Baseline, Range);
+	return FSlateImageRun::Create(RunInfo, Text, Image, Baseline, Range);
 }
 
-void FSlateImageRun::AppendText(FString& Text) const
+void FSlateImageRun::AppendTextTo(FString& Text) const
 {
 	// Do nothing
 }
 
-void FSlateImageRun::AppendText(FString& AppendToText, const FTextRange& PartialRange) const
+void FSlateImageRun::AppendTextTo(FString& AppendToText, const FTextRange& PartialRange) const
 {
 	check(Range.BeginIndex <= PartialRange.BeginIndex);
 	check(Range.EndIndex >= PartialRange.EndIndex);
 }
 
-#endif //WITH_FANCY_TEXT
\ No newline at end of file
+const FRunInfo& FSlateImageRun::GetRunInfo() const
+{
+	return RunInfo;
+}
+
+#endif //WITH_FANCY_TEXT
diff --git a/Engine/Source/Runtime/Slate/Private/Framework/Text/SlateTextRun.cpp b/Engine/Source/Runtime/Slate/Private/Framework/Text/SlateTextRun.cpp
--- a/Engine/Source/Runtime/Slate/Private/Framework/Text/SlateTextRun.cpp
+++ b/Engine/Source/Runtime/Slate/Private/Framework/Text/SlateTextRun.cpp
@@ -6,14 +6,14 @@
 
 #include "SlateTextRun.h"
 
-TSharedRef< FSlateTextRun > FSlateTextRun::Create( const TSharedRef< const FString >& InText, const FTextBlockStyle& Style )
+TSharedRef< FSlateTextRun > FSlateTextRun::Create( const FRunInfo& InRunInfo, const TSharedRef< const FString >& InText, const FTextBlockStyle& Style )
 {
-	return MakeShareable( new FSlateTextRun( InText, Style ) );
+	return MakeShareable( new FSlateTextRun( InRunInfo, InText, Style ) );
 }
 
-TSharedRef< FSlateTextRun > FSlateTextRun::Create( const TSharedRef< const FString >& InText, const FTextBlockStyle& Style, const FTextRange& InRange )
+TSharedRef< FSlateTextRun > FSlateTextRun::Create( const FRunInfo& InRunInfo, const TSharedRef< const FString >& InText, const FTextBlockStyle& Style, const FTextRange& InRange )
 {
-	return MakeShareable( new FSlateTextRun( InText, Style, InRange ) );
+	return MakeShareable( new FSlateTextRun( InRunInfo, InText, Style, InRange ) );
 }
 
 FTextRange FSlateTextRun::GetTextRange() const 
@@ -177,15 +177,15 @@ void FSlateTextRun::Move(const TSharedRef<FString>& NewText, const FTextRange& N
 
 TSharedRef<IRun> FSlateTextRun::Clone() const
 {
-	return FSlateTextRun::Create(Text, Style, Range);
+	return FSlateTextRun::Create(RunInfo, Text, Style, Range);
 }
 
-void FSlateTextRun::AppendText(FString& AppendToText) const
+void FSlateTextRun::AppendTextTo(FString& AppendToText) const
 {
 	AppendToText.Append(**Text + Range.BeginIndex, Range.Len());
 }
 
-void FSlateTextRun::AppendText(FString& AppendToText, const FTextRange& PartialRange) const
+void FSlateTextRun::AppendTextTo(FString& AppendToText, const FTextRange& PartialRange) const
 {
 	check(Range.BeginIndex <= PartialRange.BeginIndex);
 	check(Range.EndIndex >= PartialRange.EndIndex);
@@ -193,8 +193,14 @@ void FSlateTextRun::AppendText(FString& AppendToText, const FTextRange& PartialR
 	AppendToText.Append(**Text + PartialRange.BeginIndex, PartialRange.Len());
 }
 
-FSlateTextRun::FSlateTextRun( const TSharedRef< const FString >& InText, const FTextBlockStyle& InStyle ) 
-	: Text( InText )
+const FRunInfo& FSlateTextRun::GetRunInfo() const
+{
+	return RunInfo;
+}
+
+FSlateTextRun::FSlateTextRun( const FRunInfo& InRunInfo, const TSharedRef< const FString >& InText, const FTextBlockStyle& InStyle ) 
+	: RunInfo( InRunInfo )
+	, Text( InText )
 	, Style( InStyle )
 	, Range( 0, Text->Len() )
 #if TEXT_LAYOUT_DEBUG
@@ -204,8 +210,9 @@ FSlateTextRun::FSlateTextRun( const TSharedRef< const FString >& InText, const F
 
 }
 
-FSlateTextRun::FSlateTextRun( const TSharedRef< const FString >& InText, const FTextBlockStyle& InStyle, const FTextRange& InRange ) 
-	: Text( InText )
+FSlateTextRun::FSlateTextRun( const FRunInfo& InRunInfo, const TSharedRef< const FString >& InText, const FTextBlockStyle& InStyle, const FTextRange& InRange ) 
+	: RunInfo( InRunInfo )
+	, Text( InText )
 	, Style( InStyle )
 	, Range( InRange )
 #if TEXT_LAYOUT_DEBUG
@@ -216,7 +223,8 @@ FSlateTextRun::FSlateTextRun( const TSharedRef< const FString >& InText, const F
 }
 
 FSlateTextRun::FSlateTextRun( const FSlateTextRun& Run ) 
-	: Text( Run.Text )
+	: RunInfo( Run.RunInfo )
+	, Text( Run.Text )
 	, Style( Run.Style )
 	, Range( Run.Range )
 #if TEXT_LAYOUT_DEBUG
diff --git a/Engine/Source/Runtime/Slate/Private/Framework/Text/SlateWidgetRun.cpp b/Engine/Source/Runtime/Slate/Private/Framework/Text/SlateWidgetRun.cpp
--- a/Engine/Source/Runtime/Slate/Private/Framework/Text/SlateWidgetRun.cpp
+++ b/Engine/Source/Runtime/Slate/Private/Framework/Text/SlateWidgetRun.cpp
@@ -6,14 +6,14 @@
 
 #include "SlateWidgetRun.h"
 
-TSharedRef< FSlateWidgetRun > FSlateWidgetRun::Create( const TSharedRef< const FString >& InText, const FWidgetRunInfo& InWidgetInfo )
+TSharedRef< FSlateWidgetRun > FSlateWidgetRun::Create( const FRunInfo& InRunInfo, const TSharedRef< const FString >& InText, const FWidgetRunInfo& InWidgetInfo )
 {
-	return MakeShareable( new FSlateWidgetRun( InText, InWidgetInfo ) );
+	return MakeShareable( new FSlateWidgetRun( InRunInfo, InText, InWidgetInfo ) );
 }
 
-TSharedRef< FSlateWidgetRun > FSlateWidgetRun::Create( const TSharedRef< const FString >& InText, const FWidgetRunInfo& InWidgetInfo, const FTextRange& InRange )
+TSharedRef< FSlateWidgetRun > FSlateWidgetRun::Create( const FRunInfo& InRunInfo, const TSharedRef< const FString >& InText, const FWidgetRunInfo& InWidgetInfo, const FTextRange& InRange )
 {
-	return MakeShareable( new FSlateWidgetRun( InText, InWidgetInfo, InRange ) );
+	return MakeShareable( new FSlateWidgetRun( InRunInfo, InText, InWidgetInfo, InRange ) );
 }
 
 FTextRange FSlateWidgetRun::GetTextRange() const 
@@ -85,22 +85,28 @@ void FSlateWidgetRun::Move(const TSharedRef<FString>& NewText, const FTextRange&
 
 TSharedRef<IRun> FSlateWidgetRun::Clone() const
 {
-	return FSlateWidgetRun::Create(Text, Info, Range);
+	return FSlateWidgetRun::Create(RunInfo, Text, Info, Range);
 }
 
-void FSlateWidgetRun::AppendText(FString& AppendToText) const
+void FSlateWidgetRun::AppendTextTo(FString& AppendToText) const
 {
 	//Do nothing
 }
 
-void FSlateWidgetRun::AppendText(FString& AppendToText, const FTextRange& PartialRange) const
+void FSlateWidgetRun::AppendTextTo(FString& AppendToText, const FTextRange& PartialRange) const
 {
 	check(Range.BeginIndex <= PartialRange.BeginIndex);
 	check(Range.EndIndex >= PartialRange.EndIndex);
 }
 
-FSlateWidgetRun::FSlateWidgetRun( const TSharedRef< const FString >& InText, const FWidgetRunInfo& InWidgetInfo ) 
-	: Text( InText )
+const FRunInfo& FSlateWidgetRun::GetRunInfo() const
+{
+	return RunInfo;
+}
+
+FSlateWidgetRun::FSlateWidgetRun( const FRunInfo& InRunInfo, const TSharedRef< const FString >& InText, const FWidgetRunInfo& InWidgetInfo ) 
+	: RunInfo( InRunInfo )
+	, Text( InText )
 	, Range( 0, Text->Len() )
 	, Info( InWidgetInfo )
 	, Children()
@@ -109,8 +115,9 @@ FSlateWidgetRun::FSlateWidgetRun( const TSharedRef< const FString >& InText, con
 	Children.Add( Info.Widget );
 }
 
-FSlateWidgetRun::FSlateWidgetRun( const TSharedRef< const FString >& InText, const FWidgetRunInfo& InWidgetInfo, const FTextRange& InRange ) 
-	: Text( InText )
+FSlateWidgetRun::FSlateWidgetRun( const FRunInfo& InRunInfo, const TSharedRef< const FString >& InText, const FWidgetRunInfo& InWidgetInfo, const FTextRange& InRange ) 
+	: RunInfo( InRunInfo )
+	, Text( InText )
 	, Range( InRange )
 	, Info( InWidgetInfo )
 	, Children()
@@ -120,7 +127,8 @@ FSlateWidgetRun::FSlateWidgetRun( const TSharedRef< const FString >& InText, con
 }
 
 FSlateWidgetRun::FSlateWidgetRun( const FSlateWidgetRun& Run ) 
-	: Text( Run.Text )
+	: RunInfo( Run.RunInfo )
+	, Text( Run.Text )
 	, Range( Run.Range )
 	, Info( Run.Info )
 	, Children()
diff --git a/Engine/Source/Runtime/Slate/Private/Framework/Text/SyntaxHighlighterTextLayoutMarshaller.cpp b/Engine/Source/Runtime/Slate/Private/Framework/Text/SyntaxHighlighterTextLayoutMarshaller.cpp
new file mode 100644
--- /dev/null
+++ b/Engine/Source/Runtime/Slate/Private/Framework/Text/SyntaxHighlighterTextLayoutMarshaller.cpp
@@ -0,0 +1,182 @@
+// Copyright 1998-2014 Epic Games, Inc. All Rights Reserved.
+
+#include "SlatePrivatePCH.h"
+#include "SyntaxHighlighterTextLayoutMarshaller.h"
+
+#if WITH_FANCY_TEXT
+
+FSyntaxHighlighterTextLayoutMarshaller::~FSyntaxHighlighterTextLayoutMarshaller()
+{
+}
+
+void FSyntaxHighlighterTextLayoutMarshaller::SetText(const FString& SourceString, FTextLayout& TargetTextLayout)
+{
+	if(bSyntaxHighlightingEnabled)
+	{
+		TArray<FSyntaxTokenizer::FTokenizedLine> TokenizedLines;
+		Tokenizer->Process(TokenizedLines, SourceString);
+		ParseTokens(SourceString, TargetTextLayout, TokenizedLines);
+	}
+	else
+	{
+		FPlainTextLayoutMarshaller::SetText(SourceString, TargetTextLayout);
+	}
+}
+
+bool FSyntaxHighlighterTextLayoutMarshaller::RequiresLiveUpdate() const
+{
+	return bSyntaxHighlightingEnabled;
+}
+
+void FSyntaxHighlighterTextLayoutMarshaller::EnableSyntaxHighlighting(const bool bEnable)
+{
+	bSyntaxHighlightingEnabled = bEnable;
+	MakeDirty();
+}
+
+bool FSyntaxHighlighterTextLayoutMarshaller::IsSyntaxHighlightingEnabled() const
+{
+	return bSyntaxHighlightingEnabled;
+}
+
+FSyntaxHighlighterTextLayoutMarshaller::FSyntaxHighlighterTextLayoutMarshaller(TSharedPtr< FSyntaxTokenizer > InTokenizer, FTextBlockStyle InDefaultTextStyle)
+	: FPlainTextLayoutMarshaller(MoveTemp(InDefaultTextStyle))
+	, Tokenizer(MoveTemp(InTokenizer))
+	, bSyntaxHighlightingEnabled(true)
+{
+}
+
+
+TSharedRef< FRichTextSyntaxHighlighterTextLayoutMarshaller > FRichTextSyntaxHighlighterTextLayoutMarshaller::Create(const FSyntaxTextStyle& InSyntaxTextStyle)
+{
+	TArray<FSyntaxTokenizer::FRule> TokenizerRules;
+	TokenizerRules.Emplace(FSyntaxTokenizer::FRule(TEXT("</>")));
+	TokenizerRules.Emplace(FSyntaxTokenizer::FRule(TEXT("<")));
+	TokenizerRules.Emplace(FSyntaxTokenizer::FRule(TEXT(">")));
+	TokenizerRules.Emplace(FSyntaxTokenizer::FRule(TEXT("=")));
+	TokenizerRules.Emplace(FSyntaxTokenizer::FRule(TEXT("\"")));
+
+	return MakeShareable(new FRichTextSyntaxHighlighterTextLayoutMarshaller(FSyntaxTokenizer::Create(TokenizerRules), InSyntaxTextStyle));
+}
+
+FRichTextSyntaxHighlighterTextLayoutMarshaller::~FRichTextSyntaxHighlighterTextLayoutMarshaller()
+{
+}
+
+void FRichTextSyntaxHighlighterTextLayoutMarshaller::ParseTokens(const FString& SourceString, FTextLayout& TargetTextLayout, TArray<FSyntaxTokenizer::FTokenizedLine> TokenizedLines)
+{
+	enum class EParseState : uint8
+	{
+		LookingForNode,
+		LookingForNodeName,
+		LookingForNodeAttributeKey,
+		LookingForNodeAttribueValueBegin,
+		LookingForNodeAttribueValueBody,
+	};
+
+	// Parse the tokens, generating the styled runs for each line
+	EParseState ParseState = EParseState::LookingForNode;
+	for(const FSyntaxTokenizer::FTokenizedLine& TokenizedLine : TokenizedLines)
+	{
+		TSharedRef<FString> ModelString = MakeShareable(new FString());
+		TArray< TSharedRef< IRun > > Runs;
+
+		for(const FSyntaxTokenizer::FToken& Token : TokenizedLine.Tokens)
+		{
+			const FString TokenText = SourceString.Mid(Token.Range.BeginIndex, Token.Range.Len());
+
+			const FTextRange ModelRange(ModelString->Len(), ModelString->Len() + TokenText.Len());
+			ModelString->Append(TokenText);
+
+			FRunInfo RunInfo(TEXT("SyntaxHighlight.Normal"));
+			FTextBlockStyle TextBlockStyle = SyntaxTextStyle.NormalTextStyle;
+
+			const bool bIsWhitespace = FString(TokenText).TrimTrailing().IsEmpty();
+			if(!bIsWhitespace)
+			{
+				bool bHasMatchedSyntax = false;
+				if(Token.Type == FSyntaxTokenizer::ETokenType::Syntax)
+				{
+					if(ParseState == EParseState::LookingForNode && TokenText == TEXT("<"))
+					{
+						RunInfo.Name = TEXT("SyntaxHighlight.Node");
+						TextBlockStyle = SyntaxTextStyle.NodeTextStyle;
+						ParseState = EParseState::LookingForNodeName;
+						bHasMatchedSyntax = true;
+					}
+					else if(ParseState == EParseState::LookingForNodeAttributeKey && TokenText == TEXT(">"))
+					{
+						RunInfo.Name = TEXT("SyntaxHighlight.Node");
+						TextBlockStyle = SyntaxTextStyle.NodeTextStyle;
+						ParseState = EParseState::LookingForNode;
+					}
+					else if(ParseState == EParseState::LookingForNode && TokenText == TEXT("</>"))
+					{
+						RunInfo.Name = TEXT("SyntaxHighlight.Node");
+						TextBlockStyle = SyntaxTextStyle.NodeTextStyle;
+						// No state change
+						bHasMatchedSyntax = true;
+					}
+					else if(ParseState == EParseState::LookingForNodeAttributeKey && TokenText == TEXT("="))
+					{
+						RunInfo.Name = TEXT("SyntaxHighlight.NodeAttribueAssignment");
+						TextBlockStyle = SyntaxTextStyle.NodeAttribueAssignmentTextStyle;
+						ParseState = EParseState::LookingForNodeAttribueValueBegin;
+						bHasMatchedSyntax = true;
+					}
+					else if(ParseState == EParseState::LookingForNodeAttribueValueBegin && TokenText == TEXT("\""))
+					{
+						RunInfo.Name = TEXT("SyntaxHighlight.NodeAttributeValue");
+						TextBlockStyle = SyntaxTextStyle.NodeAttributeValueTextStyle;
+						ParseState = EParseState::LookingForNodeAttribueValueBody;
+						bHasMatchedSyntax = true;
+					}
+					else if(ParseState == EParseState::LookingForNodeAttribueValueBody && TokenText == TEXT("\""))
+					{
+						RunInfo.Name = TEXT("SyntaxHighlight.NodeAttributeValue");
+						TextBlockStyle = SyntaxTextStyle.NodeAttributeValueTextStyle;
+						ParseState = EParseState::LookingForNodeAttributeKey;
+						bHasMatchedSyntax = true;
+					}
+				}
+				
+				// It's possible that we fail to match a syntax token if we're in a state where it isn't parsed
+				// In this case, we treat it as a literal token
+				if(Token.Type == FSyntaxTokenizer::ETokenType::Literal || !bHasMatchedSyntax)
+				{
+					if(ParseState == EParseState::LookingForNodeName)
+					{
+						RunInfo.Name = TEXT("SyntaxHighlight.Node");
+						TextBlockStyle = SyntaxTextStyle.NodeTextStyle;
+						ParseState = EParseState::LookingForNodeAttributeKey;
+					}
+					else if(ParseState == EParseState::LookingForNodeAttributeKey)
+					{
+						RunInfo.Name = TEXT("SyntaxHighlight.NodeAttributeKey");
+						TextBlockStyle = SyntaxTextStyle.NodeAttributeKeyTextStyle;
+						// No state change, a key can be multiple tokens - consume until we find an equals
+					}
+					else if(ParseState == EParseState::LookingForNodeAttribueValueBody)
+					{
+						RunInfo.Name = TEXT("SyntaxHighlight.NodeAttributeValue");
+						TextBlockStyle = SyntaxTextStyle.NodeAttributeValueTextStyle;
+						// No state change, a value can be multiple tokens - consume until we find a closing quote
+					}
+				}
+			}
+
+			TSharedRef< ISlateRun > Run = FSlateTextRun::Create(RunInfo, ModelString, TextBlockStyle, ModelRange);
+			Runs.Add(Run);
+		}
+
+		TargetTextLayout.AddLine(ModelString, Runs);
+	}
+}
+
+FRichTextSyntaxHighlighterTextLayoutMarshaller::FRichTextSyntaxHighlighterTextLayoutMarshaller(TSharedPtr< FSyntaxTokenizer > InTokenizer, const FSyntaxTextStyle& InSyntaxTextStyle)
+	: FSyntaxHighlighterTextLayoutMarshaller(MoveTemp(InTokenizer), InSyntaxTextStyle.NormalTextStyle)
+	, SyntaxTextStyle(InSyntaxTextStyle)
+{
+}
+
+#endif //WITH_FANCY_TEXT
diff --git a/Engine/Source/Runtime/Slate/Private/Framework/Text/SyntaxTokenizer.cpp b/Engine/Source/Runtime/Slate/Private/Framework/Text/SyntaxTokenizer.cpp
new file mode 100644
--- /dev/null
+++ b/Engine/Source/Runtime/Slate/Private/Framework/Text/SyntaxTokenizer.cpp
@@ -0,0 +1,85 @@
+// Copyright 1998-2014 Epic Games, Inc. All Rights Reserved.
+
+#include "SlatePrivatePCH.h"
+#include "SyntaxTokenizer.h"
+#include "WordBreakIterator.h"
+
+TSharedRef< FSyntaxTokenizer > FSyntaxTokenizer::Create(TArray<FRule> InRules)
+{
+	return MakeShareable(new FSyntaxTokenizer(MoveTemp(InRules)));
+}
+
+FSyntaxTokenizer::~FSyntaxTokenizer()
+{
+}
+
+void FSyntaxTokenizer::Process(TArray<FTokenizedLine>& OutTokenizedLines, const FString& Input)
+{
+#if UE_ENABLE_ICU
+	TArray<FTextRange> LineRanges;
+	FTextRange::CalculateLineRangesFromString(Input, LineRanges);
+	TokenizeLineRanges(Input, LineRanges, OutTokenizedLines);
+#else
+	FTokenizedLine FakeTokenizedLine;
+	FakeTokenizedLine.Range = FTextRange(0, Input.Len());
+	FakeTokenizedLine.Tokens.Emplace(FToken(ETokenType::Literal, FakeTokenizedLine.Range));
+	OutTokenizedLines.Add(FakeTokenizedLine);
+#endif
+}
+
+FSyntaxTokenizer::FSyntaxTokenizer(TArray<FRule> InRules)
+	: Rules(MoveTemp(InRules))
+{
+}
+
+void FSyntaxTokenizer::TokenizeLineRanges(const FString& Input, const TArray<FTextRange>& LineRanges, TArray<FTokenizedLine>& OutTokenizedLines)
+{
+	// Tokenize line ranges
+	FWordBreakIterator WBI(Input);
+	for(const FTextRange& LineRange : LineRanges)
+	{
+		FTokenizedLine TokenizedLine;
+		TokenizedLine.Range = LineRange;
+
+		if(TokenizedLine.Range.IsEmpty())
+		{
+			TokenizedLine.Tokens.Emplace(FToken(ETokenType::Literal, TokenizedLine.Range));
+		}
+		else
+		{
+			int32 CurrentOffset = LineRange.BeginIndex;
+			while(CurrentOffset < LineRange.EndIndex)
+			{
+				// First check for a match against any syntax token rules
+				bool bHasMatchedSyntax = false;
+				for(const FRule& Rule : Rules)
+				{
+					if(FCString::Strncmp(&Input[CurrentOffset], *Rule.MatchText, Rule.MatchText.Len()) == 0)
+					{
+						const int32 SyntaxTokenEnd = CurrentOffset + Rule.MatchText.Len();
+						TokenizedLine.Tokens.Emplace(FToken(ETokenType::Syntax, FTextRange(CurrentOffset, SyntaxTokenEnd)));
+
+						check(SyntaxTokenEnd <= LineRange.EndIndex);
+
+						bHasMatchedSyntax = true;
+						CurrentOffset = SyntaxTokenEnd;
+						break;
+					}
+				}
+
+				if(bHasMatchedSyntax)
+				{
+					continue;
+				}
+
+				// If none matched, consume the character(s) as text
+				const int32 NextWordBoundary = WBI.MoveToCandidateAfter(CurrentOffset);
+				const int32 TextTokenEnd = (NextWordBoundary == INDEX_NONE) ? LineRange.EndIndex : FMath::Min(NextWordBoundary, LineRange.EndIndex);
+				TokenizedLine.Tokens.Emplace(FToken(ETokenType::Literal, FTextRange(CurrentOffset, TextTokenEnd)));
+				CurrentOffset = TextTokenEnd;
+			}
+		}
+
+		OutTokenizedLines.Add(TokenizedLine);
+	}
+}
diff --git a/Engine/Source/Runtime/Slate/Private/Framework/Text/Tests/RichTextMarkupProcessingTest.cpp b/Engine/Source/Runtime/Slate/Private/Framework/Text/Tests/RichTextMarkupProcessingTest.cpp
--- a/Engine/Source/Runtime/Slate/Private/Framework/Text/Tests/RichTextMarkupProcessingTest.cpp
+++ b/Engine/Source/Runtime/Slate/Private/Framework/Text/Tests/RichTextMarkupProcessingTest.cpp
@@ -101,7 +101,7 @@ namespace
 bool FRTFProcessingTest::RunTest (const FString& Parameters)
 {
 #if UE_ENABLE_ICU
-	TSharedRef< FRichTextMarkupProcessing > RichTextMarkupProcessing = FRichTextMarkupProcessing::Create();
+	TSharedRef< FDefaultRichTextMarkupParser > RichTextMarkupProcessing = FDefaultRichTextMarkupParser::Create();
 
 	// Text
 	{
diff --git a/Engine/Source/Runtime/Slate/Private/Framework/Text/TextDecorators.cpp b/Engine/Source/Runtime/Slate/Private/Framework/Text/TextDecorators.cpp
--- a/Engine/Source/Runtime/Slate/Private/Framework/Text/TextDecorators.cpp
+++ b/Engine/Source/Runtime/Slate/Private/Framework/Text/TextDecorators.cpp
@@ -18,25 +18,25 @@ FWidgetDecorator::FWidgetDecorator( FString InRunName, const FCreateWidget& InCr
 
 }
 
-bool FWidgetDecorator::Supports( const FTextRunParseResults& RunInfo, const FString& Text ) const
+bool FWidgetDecorator::Supports( const FTextRunParseResults& RunParseResult, const FString& Text ) const
 {
-	return ( RunInfo.Name == RunName );
+	return ( RunParseResult.Name == RunName );
 }
 
-TSharedRef< ISlateRun > FWidgetDecorator::Create( const FTextRunParseResults& Run, const FString& OriginalText, const TSharedRef< FString >& InOutModelText, const ISlateStyle* Style )
+TSharedRef< ISlateRun > FWidgetDecorator::Create( const FTextRunParseResults& RunParseResult, const FString& OriginalText, const TSharedRef< FString >& InOutModelText, const ISlateStyle* Style )
 {
 	FTextRange ModelRange;
 	ModelRange.BeginIndex = InOutModelText->Len();
 	*InOutModelText += TEXT('\x200B'); // Zero-Width Breaking Space
 	ModelRange.EndIndex = InOutModelText->Len();
 
-	FTextRunInfo RunInfo( Run.Name, FText::FromString( OriginalText.Mid( Run.ContentRange.BeginIndex, Run.ContentRange.EndIndex - Run.ContentRange.BeginIndex ) ) );
-	for(TPair<FString, FTextRange> Pair : Run.MetaData)
+	FTextRunInfo RunInfo( RunParseResult.Name, FText::FromString( OriginalText.Mid( RunParseResult.ContentRange.BeginIndex, RunParseResult.ContentRange.EndIndex - RunParseResult.ContentRange.BeginIndex ) ) );
+	for(const TPair<FString, FTextRange>& Pair : RunParseResult.MetaData)
 	{
-		RunInfo.MetaData[Pair.Key] = OriginalText.Mid( Pair.Value.BeginIndex, Pair.Value.EndIndex - Pair.Value.BeginIndex);
+		RunInfo.MetaData.Add(Pair.Key, OriginalText.Mid( Pair.Value.BeginIndex, Pair.Value.EndIndex - Pair.Value.BeginIndex));
 	}
 
-	return FSlateWidgetRun::Create( InOutModelText, CreateWidgetDelegate.Execute( RunInfo, Style ), ModelRange );
+	return FSlateWidgetRun::Create( RunInfo, InOutModelText, CreateWidgetDelegate.Execute( RunInfo, Style ), ModelRange );
 }
 
 TSharedRef< FImageDecorator > FImageDecorator::Create( FString InRunName, const ISlateStyle* const InStyle )
@@ -51,20 +51,26 @@ FImageDecorator::FImageDecorator( FString InRunName, const ISlateStyle* const In
 
 }
 
-bool FImageDecorator::Supports( const FTextRunParseResults& RunInfo, const FString& Text ) const
+bool FImageDecorator::Supports( const FTextRunParseResults& RunParseResult, const FString& Text ) const
 {
-	return ( RunInfo.Name == RunName );
+	return ( RunParseResult.Name == RunName );
 }
 
-TSharedRef< ISlateRun > FImageDecorator::Create( const FTextRunParseResults& RunInfo, const FString& OriginalText, const TSharedRef< FString >& InOutModelText, const ISlateStyle* Style )
+TSharedRef< ISlateRun > FImageDecorator::Create( const FTextRunParseResults& RunParseResult, const FString& OriginalText, const TSharedRef< FString >& InOutModelText, const ISlateStyle* Style )
 {
-	const FTextRange* const BrushNameRange = RunInfo.MetaData.Find( TEXT("src") );
+	const FTextRange* const BrushNameRange = RunParseResult.MetaData.Find( TEXT("src") );
 
 	FTextRange ModelRange;
 	ModelRange.BeginIndex = InOutModelText->Len();
 	*InOutModelText += TEXT('\x200B'); // Zero-Width Breaking Space
 	ModelRange.EndIndex = InOutModelText->Len();
 
+	FRunInfo RunInfo( RunParseResult.Name );
+	for(const TPair<FString, FTextRange>& Pair : RunParseResult.MetaData)
+	{
+		RunInfo.MetaData.Add(Pair.Key, OriginalText.Mid( Pair.Value.BeginIndex, Pair.Value.EndIndex - Pair.Value.BeginIndex));
+	}
+
 	if ( BrushNameRange != NULL )
 	{
 		const FString BrushNameString = OriginalText.Mid(BrushNameRange->BeginIndex, BrushNameRange->EndIndex - BrushNameRange->BeginIndex);
@@ -77,14 +83,14 @@ TSharedRef< ISlateRun > FImageDecorator::Create( const FTextRunParseResults& Run
 		if ( Style->HasWidgetStyle< FInlineTextImageStyle >( BrushName ) )
 		{
 			const FInlineTextImageStyle& ImageStyle = Style->GetWidgetStyle< FInlineTextImageStyle >( BrushName );
-			return FSlateImageRun::Create( InOutModelText, &ImageStyle.Image, ImageStyle.Baseline, ModelRange );
+			return FSlateImageRun::Create( RunInfo, InOutModelText, &ImageStyle.Image, ImageStyle.Baseline, ModelRange );
 		}
 
-		return FSlateImageRun::Create( InOutModelText, Style->GetBrush( BrushName ), 0, ModelRange );
+		return FSlateImageRun::Create( RunInfo, InOutModelText, Style->GetBrush( BrushName ), 0, ModelRange );
 	}
 
 	const FInlineTextImageStyle& ImageStyle = FInlineTextImageStyle::GetDefault();
-	return FSlateImageRun::Create( InOutModelText, &ImageStyle.Image, ImageStyle.Baseline, ModelRange );
+	return FSlateImageRun::Create( RunInfo, InOutModelText, &ImageStyle.Image, ImageStyle.Baseline, ModelRange );
 }
 
 TSharedRef< FHyperlinkDecorator > FHyperlinkDecorator::Create( FString Id, const FSlateHyperlinkRun::FOnClick& NavigateDelegate )
@@ -99,23 +105,23 @@ FHyperlinkDecorator::FHyperlinkDecorator( FString InId, const FSlateHyperlinkRun
 
 }
 
-bool FHyperlinkDecorator::Supports( const FTextRunParseResults& RunInfo, const FString& Text ) const
+bool FHyperlinkDecorator::Supports( const FTextRunParseResults& RunParseResult, const FString& Text ) const
 {
-	const FTextRange* const MetaDataIdRange = RunInfo.MetaData.Find( TEXT("id") );
+	const FTextRange* const MetaDataIdRange = RunParseResult.MetaData.Find( TEXT("id") );
 	FString MetaDataId;
 	if(MetaDataIdRange)
 	{
 		MetaDataId = Text.Mid(MetaDataIdRange->BeginIndex, MetaDataIdRange->EndIndex - MetaDataIdRange->BeginIndex);
 	}
 
-	return ( RunInfo.Name == TEXT("a") && MetaDataIdRange != NULL && MetaDataId == Id );
+	return ( RunParseResult.Name == TEXT("a") && MetaDataIdRange != NULL && MetaDataId == Id );
 }
 
-TSharedRef< ISlateRun > FHyperlinkDecorator::Create( const FTextRunParseResults& RunInfo, const FString& OriginalText, const TSharedRef< FString >& InOutModelText, const ISlateStyle* Style )
+TSharedRef< ISlateRun > FHyperlinkDecorator::Create( const FTextRunParseResults& RunParseResult, const FString& OriginalText, const TSharedRef< FString >& InOutModelText, const ISlateStyle* Style )
 {
 	FString StyleName = TEXT("Hyperlink");
 
-	const FTextRange* const MetaDataStyleNameRange = RunInfo.MetaData.Find( TEXT("style") );
+	const FTextRange* const MetaDataStyleNameRange = RunParseResult.MetaData.Find( TEXT("style") );
 	if ( MetaDataStyleNameRange != NULL )
 	{
 		const FString MetaDataStyleName = OriginalText.Mid(MetaDataStyleNameRange->BeginIndex, MetaDataStyleNameRange->EndIndex - MetaDataStyleNameRange->BeginIndex);
@@ -124,7 +130,7 @@ TSharedRef< ISlateRun > FHyperlinkDecorator::Create( const FTextRunParseResults&
 
 	FTextRange ModelRange;
 	ModelRange.BeginIndex = InOutModelText->Len();
-	*InOutModelText += OriginalText.Mid(RunInfo.ContentRange.BeginIndex, RunInfo.ContentRange.EndIndex - RunInfo.ContentRange.BeginIndex);
+	*InOutModelText += OriginalText.Mid(RunParseResult.ContentRange.BeginIndex, RunParseResult.ContentRange.EndIndex - RunParseResult.ContentRange.BeginIndex);
 	ModelRange.EndIndex = InOutModelText->Len();
 
 	if ( !Style->HasWidgetStyle<FHyperlinkStyle>( FName( *StyleName ) ) )
@@ -132,13 +138,12 @@ TSharedRef< ISlateRun > FHyperlinkDecorator::Create( const FTextRunParseResults&
 		Style = &FCoreStyle::Get();
 	}
 
-	TMap<FString, FString> MetaData;
-	for(const TPair<FString, FTextRange>& Pair : RunInfo.MetaData)
+	FRunInfo RunInfo( RunParseResult.Name );
+	for(const TPair<FString, FTextRange>& Pair : RunParseResult.MetaData)
 	{
-		MetaData.Add(Pair.Key, OriginalText.Mid( Pair.Value.BeginIndex, Pair.Value.EndIndex - Pair.Value.BeginIndex));
-		//MetaData[Pair.Key] = OriginalText.Mid( Pair.Value.BeginIndex, Pair.Value.EndIndex - Pair.Value.BeginIndex);
+		RunInfo.MetaData.Add(Pair.Key, OriginalText.Mid( Pair.Value.BeginIndex, Pair.Value.EndIndex - Pair.Value.BeginIndex));
 	}
 
-	return FSlateHyperlinkRun::Create( InOutModelText, Style->GetWidgetStyle<FHyperlinkStyle>( FName( *StyleName ) ), MetaData, NavigateDelegate, ModelRange );
+	return FSlateHyperlinkRun::Create( RunInfo, InOutModelText, Style->GetWidgetStyle<FHyperlinkStyle>( FName( *StyleName ) ), NavigateDelegate, ModelRange );
 }
 #endif //WITH_FANCY_TEXT
\ No newline at end of file
diff --git a/Engine/Source/Runtime/Slate/Private/Framework/Text/TextLayout.cpp b/Engine/Source/Runtime/Slate/Private/Framework/Text/TextLayout.cpp
--- a/Engine/Source/Runtime/Slate/Private/Framework/Text/TextLayout.cpp
+++ b/Engine/Source/Runtime/Slate/Private/Framework/Text/TextLayout.cpp
@@ -882,7 +882,7 @@ bool FTextLayout::InsertAt(const FTextLocation& Location, TCHAR Character)
 	LineModel.HasWrappingInformation = false;
 
 	bool RunIsAfterInsertLocation = false;
-	for (int RunIndex = 0; RunIndex < LineModel.Runs.Num(); RunIndex++)
+	for (int32 RunIndex = 0; RunIndex < LineModel.Runs.Num(); RunIndex++)
 	{
 		FRunModel& RunModel = LineModel.Runs[ RunIndex ];
 		const FTextRange& RunRange = RunModel.GetTextRange();
@@ -907,7 +907,7 @@ bool FTextLayout::InsertAt(const FTextLocation& Location, TCHAR Character)
 	return true;
 }
 
-bool FTextLayout::InsertAt( const FTextLocation& Location, const FString& Text )
+bool FTextLayout::InsertAt(const FTextLocation& Location, const FString& Text)
 {
 	const int32 InsertLocation = Location.GetOffset();
 	const int32 LineIndex = Location.GetLineIndex();
@@ -923,7 +923,7 @@ bool FTextLayout::InsertAt( const FTextLocation& Location, const FString& Text )
 	LineModel.HasWrappingInformation = false;
 
 	bool RunIsAfterInsertLocation = false;
-	for (int RunIndex = 0; RunIndex < LineModel.Runs.Num(); RunIndex++)
+	for (int32 RunIndex = 0; RunIndex < LineModel.Runs.Num(); RunIndex++)
 	{
 		FRunModel& RunModel = LineModel.Runs[ RunIndex ];
 		const FTextRange& RunRange = RunModel.GetTextRange();
@@ -948,6 +948,71 @@ bool FTextLayout::InsertAt( const FTextLocation& Location, const FString& Text )
 	return true;
 }
 
+bool FTextLayout::InsertAt(const FTextLocation& Location, TSharedRef<IRun> InRun, const bool bAlwaysKeepRightRun)
+{
+	const int32 InsertLocation = Location.GetOffset();
+	const int32 LineIndex = Location.GetLineIndex();
+
+	if (!LineModels.IsValidIndex(LineIndex))
+	{
+		return false;
+	}
+
+	FLineModel& LineModel = LineModels[LineIndex];
+
+	FString NewRunText;
+	InRun->AppendTextTo(NewRunText);
+
+	LineModel.Text->InsertAt(InsertLocation, NewRunText);
+	LineModel.HasWrappingInformation = false;
+
+	bool RunIsAfterInsertLocation = false;
+	for (int32 RunIndex = 0; RunIndex < LineModel.Runs.Num(); RunIndex++)
+	{
+		const TSharedRef<IRun> Run = LineModel.Runs[RunIndex].GetRun();
+		const FTextRange& RunRange = Run->GetTextRange();
+
+		const bool bIsLastRun = RunIndex == LineModel.Runs.Num() - 1;
+		if (RunRange.Contains(InsertLocation) || (bIsLastRun && !RunIsAfterInsertLocation))
+		{
+			check(RunIsAfterInsertLocation == false);
+			RunIsAfterInsertLocation = true;
+
+			// This run contains the insertion point, so we need to split it
+			const TSharedRef<IRun> LeftRun = Run;
+			const TSharedRef<IRun> RightRun = Run->Clone();
+
+			// Update the text ranges for all of the runs
+			const int32 InsertLocationEnd = InsertLocation + NewRunText.Len();
+			LeftRun->SetTextRange( FTextRange( RunRange.BeginIndex, InsertLocation ) );
+			InRun->Move( LineModel.Text, FTextRange( InsertLocation, InsertLocationEnd ) );
+			RightRun->SetTextRange( FTextRange( InsertLocationEnd, RunRange.EndIndex + NewRunText.Len() ) );
+
+			// Insert the new runs at the correct place, and then skip over these new array entries
+			const bool LeftRunHasText = !LeftRun->GetTextRange().IsEmpty();
+			const bool RightRunHasText = !RightRun->GetTextRange().IsEmpty();
+			if (!LeftRunHasText)
+			{
+				LineModel.Runs.RemoveAt(RunIndex--);
+			}
+			LineModel.Runs.Insert(InRun, ++RunIndex);
+			if (RightRunHasText || bAlwaysKeepRightRun)
+			{
+				LineModel.Runs.Insert(RightRun, ++RunIndex);
+			}
+		}
+		else if (RunIsAfterInsertLocation)
+		{
+			FTextRange NewRange = RunRange;
+			NewRange.Offset(NewRunText.Len());
+			Run->SetTextRange(NewRange);
+		}
+	}
+
+	DirtyFlags |= EDirtyState::Layout;
+	return true;
+}
+
 bool FTextLayout::JoinLineWithNextLine(int32 LineIndex)
 {
 	if (!LineModels.IsValidIndex(LineIndex) || !LineModels.IsValidIndex(LineIndex + 1))
@@ -1018,7 +1083,8 @@ bool FTextLayout::SplitLineAt(const FTextLocation& Location)
 		const TSharedRef<IRun> Run = LineModel.Runs[RunIndex].GetRun();
 		const FTextRange& RunRange = Run->GetTextRange();
 		
-		if ((LineModel.Runs.Num() == 1 && RunRange.InclusiveContains(BreakLocation)) || RunRange.Contains(BreakLocation))
+		const bool bIsLastRun = RunIndex == LineModel.Runs.Num() - 1;
+		if (RunRange.Contains(BreakLocation) || (bIsLastRun && RunIsToTheLeftOfTheBreakLocation))
 		{
 			check(RunIsToTheLeftOfTheBreakLocation == true);
 			RunIsToTheLeftOfTheBreakLocation = false;
@@ -1080,32 +1146,58 @@ bool FTextLayout::RemoveAt( const FTextLocation& Location, int32 Count )
 	LineModel.Text->RemoveAt(RemoveLocation, Count);
 	LineModel.HasWrappingInformation = false;
 
-	bool RunIsAfterRemoveLocation = false;
-	for (int RunIndex = 0; RunIndex < LineModel.Runs.Num(); RunIndex++)
+	const FTextRange RemoveTextRange(RemoveLocation, RemoveLocation + Count);
+	for (int32 RunIndex = LineModel.Runs.Num() - 1; RunIndex >= 0; --RunIndex)
 	{
 		FRunModel& RunModel = LineModel.Runs[RunIndex];
-		const FTextRange& RunRange = RunModel.GetTextRange();
+		const FTextRange RunRange = RunModel.GetTextRange();
 
-		if (RunRange.Contains(RemoveLocation))
+		const FTextRange IntersectedRangeToRemove = RunRange.Intersect(RemoveTextRange);
+		if (IntersectedRangeToRemove.IsEmpty() && RunRange.EndIndex > RemoveTextRange.EndIndex)
 		{
-			check(RunIsAfterRemoveLocation == false);
-			RunIsAfterRemoveLocation = true;
-
-			FTextRange NewRange(RunRange.BeginIndex, RunRange.EndIndex - Count);
-			if (NewRange.IsEmpty() && LineModel.Runs.Num() > 1)
+			// The whole run is contained to the right of the removal range, just adjust its range by the amount of text that was removed
+			FTextRange NewRange = RunRange;
+			NewRange.Offset(-Count);
+			RunModel.SetTextRange(NewRange);
+		}
+		else if (!IntersectedRangeToRemove.IsEmpty())
+		{
+			const int32 RunLength = RunRange.Len();
+			const int32 IntersectedLength = IntersectedRangeToRemove.Len();
+			if (RunLength == IntersectedLength)
 			{
-				LineModel.Runs.RemoveAt(RunIndex--);
+				// The text for this entire run has been removed
+				if (LineModel.Runs.Num() > 1)
+				{
+					LineModel.Runs.RemoveAt(RunIndex);
+				}
+				else
+				{
+					RunModel.SetTextRange(FTextRange(0, 0));
+				}
+			}
+			else if (RunRange.BeginIndex > RemoveTextRange.BeginIndex)
+			{
+				// Some of this run has been removed, and this run is the right hand part of the removal
+				// So we need to adjust the range so that we start at the removal point since the text has been removed from the beginning of this run
+				const FTextRange NewRange(RemoveTextRange.BeginIndex, RunRange.EndIndex - Count);
+				check(!NewRange.IsEmpty());
+				RunModel.SetTextRange(NewRange);
 			}
 			else
 			{
+				// Some of this run has been removed, and this run is the left hand part of the removal
+				// So we need to adjust the range by the amount of text that has been removed from the end of this run
+				const FTextRange NewRange(RunRange.BeginIndex, RunRange.EndIndex - IntersectedLength);
+				check(!NewRange.IsEmpty());
 				RunModel.SetTextRange(NewRange);
 			}
-		}
-		else if (RunIsAfterRemoveLocation)
-		{
-			FTextRange NewRange = RunRange;
-			NewRange.Offset(-Count);
-			RunModel.SetTextRange(NewRange);
+
+			if (RunRange.BeginIndex <= RemoveTextRange.BeginIndex)
+			{
+				// break since we don't need to process the runs to the left of the removal point
+				break;
+			}
 		}
 	}
 
@@ -1196,7 +1288,7 @@ void FTextLayout::GetAsTextAndOffsets(FString* const OutDisplayText, FTextOffset
 			const FRunModel& Run = LineModel.Runs[RunIndex];
 			if (OutDisplayText)
 			{
-				Run.AppendText(*OutDisplayText);
+				Run.AppendTextTo(*OutDisplayText);
 			}
 			LineLength += Run.GetTextRange().Len();
 		}
@@ -1221,7 +1313,7 @@ void GetRangeAsTextFromLine(FString& DisplayText, const FTextLayout::FLineModel&
 
 		if (!IntersectRange.IsEmpty())
 		{
-			Run.AppendText(DisplayText, IntersectRange);
+			Run.AppendTextTo(DisplayText, IntersectRange);
 		}
 		else if (RunRange.BeginIndex > Range.EndIndex)
 		{
@@ -1473,14 +1565,14 @@ void FTextLayout::FRunModel::ClearCache()
 	MeasuredRangeSizes.Empty();
 }
 
-void FTextLayout::FRunModel::AppendText(FString& Text) const
+void FTextLayout::FRunModel::AppendTextTo(FString& Text) const
 {
-	Run->AppendText(Text);
+	Run->AppendTextTo(Text);
 }
 
-void FTextLayout::FRunModel::AppendText(FString& Text, const FTextRange& Range) const
+void FTextLayout::FRunModel::AppendTextTo(FString& Text, const FTextRange& Range) const
 {
-	Run->AppendText(Text, Range);
+	Run->AppendTextTo(Text, Range);
 }
 
 TSharedRef< ILayoutBlock > FTextLayout::FRunModel::CreateBlock( const FBlockDefinition& BlockDefine, float Scale ) const
diff --git a/Engine/Source/Runtime/Slate/Private/Framework/Text/TextRange.cpp b/Engine/Source/Runtime/Slate/Private/Framework/Text/TextRange.cpp
new file mode 100644
--- /dev/null
+++ b/Engine/Source/Runtime/Slate/Private/Framework/Text/TextRange.cpp
@@ -0,0 +1,35 @@
+// Copyright 1998-2014 Epic Games, Inc. All Rights Reserved.
+
+#include "SlatePrivatePCH.h"
+#include "TextRange.h"
+
+void FTextRange::CalculateLineRangesFromString(const FString& Input, TArray<FTextRange>& LineRanges)
+{
+	int32 LineBeginIndex = 0;
+
+	// Loop through splitting at new-lines
+	const TCHAR* const InputStart = *Input;
+	for(const TCHAR* CurrentChar = InputStart; CurrentChar && *CurrentChar; ++CurrentChar)
+	{
+		// Handle a chain of \r\n slightly differently to stop the FChar::IsLinebreak adding two separate new-lines
+		const bool bIsWindowsNewLine = (*CurrentChar == '\r' && *(CurrentChar + 1) == '\n');
+		if(bIsWindowsNewLine || FChar::IsLinebreak(*CurrentChar))
+		{
+			const int32 LineEndIndex = (CurrentChar - InputStart);
+			check(LineEndIndex >= LineBeginIndex);
+			LineRanges.Emplace(FTextRange(LineBeginIndex, LineEndIndex));
+
+			if(bIsWindowsNewLine)
+			{
+				++CurrentChar; // skip the \n of the \r\n chain
+			}
+			LineBeginIndex = (CurrentChar - InputStart) + 1; // The next line begins after the end of the current line
+		}
+	}
+
+	// Process any remaining string after the last new-line
+	if(LineBeginIndex <= Input.Len())
+	{
+		LineRanges.Emplace(FTextRange(LineBeginIndex, Input.Len()));
+	}
+}
diff --git a/Engine/Source/Runtime/Slate/Private/Widgets/Input/SMultiLineEditableTextBox.cpp b/Engine/Source/Runtime/Slate/Private/Widgets/Input/SMultiLineEditableTextBox.cpp
--- a/Engine/Source/Runtime/Slate/Private/Widgets/Input/SMultiLineEditableTextBox.cpp
+++ b/Engine/Source/Runtime/Slate/Private/Widgets/Input/SMultiLineEditableTextBox.cpp
@@ -40,10 +40,12 @@ void SMultiLineEditableTextBox::Construct( const FArguments& InArgs )
 				[
 					SAssignNew( EditableText, SMultiLineEditableText )
 					.Text( InArgs._Text )
+					.Marshaller( InArgs._Marshaller )
 					.Font( Font )
 					.IsReadOnly( InArgs._IsReadOnly )
 					.OnTextChanged( InArgs._OnTextChanged )
 					.OnTextCommitted( InArgs._OnTextCommitted )
+					.OnCursorMoved( InArgs._OnCursorMoved )
 					.Justification(InArgs._Justification)
 					.LineHeightPercentage(InArgs._LineHeightPercentage)
 					.Margin(InArgs._Margin)
@@ -151,5 +153,34 @@ FReply SMultiLineEditableTextBox::OnKeyboardFocusReceived( const FGeometry& MyGe
 	return Reply;
 }
 
+FText SMultiLineEditableTextBox::GetSelectedText() const
+{
+	return EditableText->GetSelectedText();
+}
+
+void SMultiLineEditableTextBox::InsertTextAtCursor(const FText& InText)
+{
+	EditableText->InsertTextAtCursor(InText);
+}
+
+void SMultiLineEditableTextBox::InsertTextAtCursor(const FString& InString)
+{
+	EditableText->InsertTextAtCursor(InString);
+}
+
+void SMultiLineEditableTextBox::InsertRunAtCursor(TSharedRef<IRun> InRun)
+{
+	EditableText->InsertRunAtCursor(InRun);
+}
+
+void SMultiLineEditableTextBox::ApplyToSelection(const FRunInfo& InRunInfo, const FTextBlockStyle& InStyle)
+{
+	EditableText->ApplyToSelection(InRunInfo, InStyle);
+}
+
+TSharedPtr<const IRun> SMultiLineEditableTextBox::GetRunUnderCursor() const
+{
+	return EditableText->GetRunUnderCursor();
+}
 
-#endif //WITH_FANCY_TEXT
\ No newline at end of file
+#endif //WITH_FANCY_TEXT
diff --git a/Engine/Source/Runtime/Slate/Private/Widgets/Testing/STestSuite.cpp b/Engine/Source/Runtime/Slate/Private/Widgets/Testing/STestSuite.cpp
--- a/Engine/Source/Runtime/Slate/Private/Widgets/Testing/STestSuite.cpp
+++ b/Engine/Source/Runtime/Slate/Private/Widgets/Testing/STestSuite.cpp
@@ -5,6 +5,7 @@
 #include "SLayoutExample.h"
 #include "SWidgetGallery.h"
 #include "TestStyle.h"
+#include "SyntaxHighlighterTextLayoutMarshaller.h"
 
 
 #define LOCTEXT_NAMESPACE "STestSuite"
@@ -958,16 +959,11 @@ private:
 };
 
 
-class SRichTextTest : public SCompoundWidget
+struct RichTextHelper
 {
-public:
-	SLATE_BEGIN_ARGS( SRichTextTest ){}
-	SLATE_END_ARGS()
-
-	BEGIN_SLATE_FUNCTION_BUILD_OPTIMIZATION
-	void Construct(const FArguments& InArgs)
+	static FText GetAliceInWonderland()
 	{
-		FText AliceInWonderland = FText::FromString( 
+		return FText::FromString(
 			TEXT("The <a id=\"browser\" href=\"http://en.wikipedia.org/wiki/Dormouse_(Alice%27s_Adventures_in_Wonderland)\" style=\"RichText.Interactive.Text.Hyperlink\">Dormouse</> had closed its eyes by this time, and was going off into a doze; but, on being pinched by the Hatter, it woke up again with a little shriek, and went on: '<RichText.Interactive.Text.Dialogue>-that begins with an M, such as </><a id=\"browser\" href=\"http://en.wikipedia.org/wiki/Mousetrap_(weapon)\" style=\"RichText.Interactive.Text.DialogueHyperlink\">mouse-traps</><RichText.Interactive.Text.Dialogue>, and the moon, and memory, and muchness-you know you say things are \"much of a muchness\"-did you ever see such a thing as a drawing of a muchness?</>'")
 			TEXT("\n\n")
 			TEXT("'<RichText.Interactive.Text.Dialogue>Really, now you ask me,</>' said <a id=\"browser\" href=\"http://en.wikipedia.org/wiki/Alice_(Alice%27s_Adventures_in_Wonderland)\" style=\"RichText.Interactive.Text.Hyperlink\">Alice</>, very much confused, '<RichText.Interactive.Text.Dialogue>I don't think-</>'")
@@ -987,9 +983,12 @@ public:
 			TEXT("'<RichText.Interactive.Text.Dialogue>I couldn't help it,</>' said Five, in a sulky tone; '<RichText.Interactive.Text.Dialogue>Seven jogged my elbow.</>'")
 			TEXT("\n\n")
 			TEXT("On which Seven looked up and said, '<RichText.Interactive.Text.Dialogue>That's right, Five! Always lay the blame on others!</>'")
-			TEXT("\n\n") );
+			TEXT("\n\n"));
+	}
 
-		FText TheWarOfTheWorlds_Part1 = FText::FromString( 
+	static FText GetTheWarOfTheWorlds_Part1()
+	{
+		return FText::FromString(
 			TEXT("When we had finished eating we went softly upstairs to my study, and I looked again out of the open window.  In one night the valley had become a valley of ashes.  The fires had dwindled now.  Where flames had been there were now streamers of smoke; but the countless ruins of shattered and gutted houses and blasted and blackened trees that the night had hidden stood out now gaunt and terrible in the pitiless light of dawn.  Yet here and there some object had had the luck to escape--a white railway signal here, the end of a greenhouse there, white and fresh amid the wreckage.  Never before in the history of warfare had destruction been so indiscriminate and so universal. And shining with the growing light of the east, three of the metallic giants stood about the pit, their cowls rotating as though they were surveying the desolation they had made.")
 			TEXT("\n\n")
 			TEXT("It seemed to me that the pit had been enlarged, and ever and again puffs of vivid green vapour streamed up and out of it towards the brightening dawn--streamed up, whirled, broke, and vanished.")
@@ -1009,9 +1008,12 @@ public:
 			TEXT("We went down the lane, by the body of the man in black, sodden now from the overnight hail, and broke into the woods at the foot of the hill.  We pushed through these towards the railway without meeting a soul.  The woods across the line were but the scarred and blackened ruins of woods; for the most part the trees had fallen, but a certain proportion still stood, dismal grey stems, with dark brown foliage instead of green.")
 			TEXT("\n\n")
 			TEXT("On our side the fire had done no more than scorch the nearer trees; it had failed to secure its footing.  In one place the woodmen had been at work on Saturday; trees, felled and freshly trimmed, lay in a clearing, with heaps of sawdust by the sawing-machine and its engine. Hard by was a temporary hut, deserted.  There was not a breath of wind this morning, and everything was strangely still.  Even the birds were hushed, and as we hurried along I and the artilleryman talked in whispers and looked now and again over our shoulders.  Once or twice we stopped to listen.")
-			TEXT("\n\n") );
+			TEXT("\n\n"));
+	}
 
-		FText TheWarOfTheWorlds_Part2 = FText::FromString( 
+	static FText GetTheWarOfTheWorlds_Part2()
+	{
+		return FText::FromString(
 			TEXT("And beyond, over the blue hills that rise southward of the river, the glittering Martians went to and fro, calmly and methodically spreading their poison cloud over this patch of country and then over that, laying it again with their steam jets when it had served its purpose, and taking possession of the conquered country.  They do not seem to have aimed at extermination so much as at complete demoralisation and the destruction of any opposition.  They exploded any stores of powder they came upon, cut every telegraph, and wrecked the railways here and there.  They were hamstringing mankind.  They seemed in no hurry to extend the field of their operations, and did not come beyond the central part of London all that day.  It is possible that a very considerable number of people in London stuck to their houses through Monday morning.  Certain it is that many died at home suffocated by the Black Smoke.")
 			TEXT("\n\n")
 			TEXT("Until about midday the Pool of London was an astonishing scene. Steamboats and shipping of all sorts lay there, tempted by the enormous sums of money offered by fugitives, and it is said that many who swam out to these vessels were thrust off with boathooks and drowned.  About one o'clock in the afternoon the thinning remnant of a cloud of the black vapour appeared between the arches of Blackfriars Bridge.  At that the Pool became a scene of mad confusion, fighting, and collision, and for some time a multitude of boats and barges jammed in the northern arch of the Tower Bridge, and the sailors and lightermen had to fight savagely against the people who swarmed upon them from the riverfront.  People were actually clambering down the piers of the bridge from above.")
@@ -1027,10 +1029,12 @@ public:
 			TEXT("On Wednesday the three fugitives--they had passed the night in a field of unripe wheat--reached Chelmsford, and there a body of the inhabitants, calling itself the Committee of Public Supply, seized the pony as provisions, and would give nothing in exchange for it but the promise of a share in it the next day.  Here there were rumours of Martians at Epping, and news of the destruction of Waltham Abbey Powder Mills in a vain attempt to blow up one of the invaders.")
 			TEXT("\n\n")
 			TEXT("People were watching for Martians here from the church towers.  My brother, very luckily for him as it chanced, preferred to push on at once to the coast rather than wait for food, although all three of them were very hungry.  By midday they passed through Tillingham, which, strangely enough, seemed to be quite silent and deserted, save for a few furtive plunderers hunting for food.  Near Tillingham they suddenly came in sight of the sea, and the most amazing crowd of shipping of all sorts that it is possible to imagine.")
-			TEXT("\n\n") );
-
+			TEXT("\n\n"));
+	}
 
-		FText TheWarOfTheWorlds_Part3 = FText::FromString( 
+	static FText GetTheWarOfTheWorlds_Part3()
+	{
+		return FText::FromString(
 			TEXT("They saw the gaunt figures separating and rising out of the water as they retreated shoreward, and one of them raised the camera-like generator of the Heat-Ray.  He held it pointing obliquely downward, and a bank of steam sprang from the water at its touch.  It must have driven through the iron of the ship's side like a white-hot iron rod through paper.")
 			TEXT("\n\n")
 			TEXT("But no one heeded that very much.  At the sight of the Martian's collapse the captain on the bridge yelled inarticulately, and all the crowding passengers on the steamer's stern shouted together.  And then they yelled again.  For, surging out beyond the white tumult, drove something long and black, the flames streaming from its middle parts, its ventilators and funnels spouting fire.")
@@ -1046,9 +1050,12 @@ public:
 			TEXT("Then suddenly out of the golden haze of the sunset came the vibration of guns, and a form of black shadows moving.  Everyone struggled to the rail of the steamer and peered into the blinding furnace of the west, but nothing was to be distinguished clearly.  A mass of smoke rose slanting and barred the face of the sun.  The steamboat throbbed on its way through an interminable suspense.")
 			TEXT("\n\n")
 			TEXT("The sun sank into grey clouds, the sky flushed and darkened, the evening star trembled into sight.  It was deep twilight when the captain cried out and pointed.  My brother strained his eyes. Something rushed up into the sky out of the greyness--rushed slantingly upward and very swiftly into the luminous clearness above the clouds in the western sky; something flat and broad, and very large, that swept round in a vast curve, grew smaller, sank slowly, and vanished again into the grey mystery of the night.  And as it flew it rained down darkness upon the land.")
-			TEXT("\n\n") );
+			TEXT("\n\n"));
+	}
 
-		FText AroundTheWorldIn80Days_Rainbow = FText::FromString( 
+	static FText GetAroundTheWorldIn80Days_Rainbow()
+	{
+		return FText::FromString(
 			TEXT("<Rainbow.Text.Red>\"</><Rainbow.Text.Orange>I</> <Rainbow.Text.Yellow>know</> <Rainbow.Text.Green>it;</> <Rainbow.Text.Blue>I</> <Rainbow.Text.Red>don't</> <Rainbow.Text.Orange>blame</> <Rainbow.Text.Yellow>you.</>  <Rainbow.Text.Green>We</> <Rainbow.Text.Blue>start</> <Rainbow.Text.Red>for</> <Rainbow.Text.Orange>Dover</> <Rainbow.Text.Yellow>and</> <Rainbow.Text.Green>Calais</> <Rainbow.Text.Blue>in</> <Rainbow.Text.Red>ten</> <Rainbow.Text.Orange>minutes.</>\"")
 			TEXT("\n\n")
 			TEXT("<Rainbow.Text.Yellow>A</> <Rainbow.Text.Green>puzzled</> <Rainbow.Text.Blue>grin</> <Rainbow.Text.Red>overspread</> <Rainbow.Text.Orange>Passepartout's</> <Rainbow.Text.Yellow>round</> <Rainbow.Text.Green>face;</> <Rainbow.Text.Blue>clearly</> <Rainbow.Text.Red>he</> <Rainbow.Text.Orange>had</> <Rainbow.Text.Yellow>not</> <Rainbow.Text.Green>comprehended</> <Rainbow.Text.Blue>his</> <Rainbow.Text.Red>master.</>")
@@ -1067,9 +1074,54 @@ public:
 			TEXT("<Rainbow.Text.Orange>\"</><Rainbow.Text.Yellow>We'll</> <Rainbow.Text.Green>have</> <Rainbow.Text.Blue>no</> <Rainbow.Text.Red>trunks;</> <Rainbow.Text.Orange>only</> <Rainbow.Text.Yellow>a</> <Rainbow.Text.Green>carpet-bag,</> <Rainbow.Text.Blue>with</> <Rainbow.Text.Red>two</> <Rainbow.Text.Yellow>shirts</> <Rainbow.Text.Green>and</> <Rainbow.Text.Blue>three</> <Rainbow.Text.Red>pairs</> <Rainbow.Text.Orange>of</> <Rainbow.Text.Yellow>stockings</> <Rainbow.Text.Green>for</> <Rainbow.Text.Blue>me,</> <Rainbow.Text.Red>and</> <Rainbow.Text.Orange>the</> <Rainbow.Text.Yellow>same</> <Rainbow.Text.Green>for</> <Rainbow.Text.Blue>you.</>  <Rainbow.Text.Red>We'll</> <Rainbow.Text.Orange>buy</> <Rainbow.Text.Yellow>our</> <Rainbow.Text.Green>clothes</> <Rainbow.Text.Blue>on</> <Rainbow.Text.Red>the</> <Rainbow.Text.Orange>way.</>  <Rainbow.Text.Yellow>Bring</> <Rainbow.Text.Green>down</> <Rainbow.Text.Blue>my</> <Rainbow.Text.Red>mackintosh</> <Rainbow.Text.Orange>and</> <Rainbow.Text.Yellow>traveling-cloak,</> <Rainbow.Text.Green>and</> <Rainbow.Text.Blue>some</> <Rainbow.Text.Red>stout</> <Rainbow.Text.Yellow>shoes,</> <Rainbow.Text.Green>though</> <Rainbow.Text.Blue>we</> <Rainbow.Text.Red>shall</> <Rainbow.Text.Orange>do</> <Rainbow.Text.Yellow>little</> <Rainbow.Text.Green>walking.</>  <Rainbow.Text.Blue>Make haste!</><Rainbow.Text.Red>\"</>")
 			TEXT("\n\n")
 			TEXT("<Rainbow.Text.Orange>Passepartout</> <Rainbow.Text.Yellow>tried</> <Rainbow.Text.Green>to</> <Rainbow.Text.Blue>reply,</> <Rainbow.Text.Red>but</> <Rainbow.Text.Orange>could</> <Rainbow.Text.Yellow>not.</>  <Rainbow.Text.Green>He</> <Rainbow.Text.Blue>went</> <Rainbow.Text.Red>out,</> <Rainbow.Text.Orange>mounted</> <Rainbow.Text.Yellow>to</> <Rainbow.Text.Green>his</> <Rainbow.Text.Blue>own</> <Rainbow.Text.Red>room,</> <Rainbow.Text.Orange>fell</> <Rainbow.Text.Yellow>into</> <Rainbow.Text.Green>a</> <Rainbow.Text.Blue>chair,</> <Rainbow.Text.Red>and</> <Rainbow.Text.Orange>muttered:</> <Rainbow.Text.Yellow>\"</><Rainbow.Text.Green>That's</> <Rainbow.Text.Blue>good,</> <Rainbow.Text.Red>that</> <Rainbow.Text.Orange>is!</> <Rainbow.Text.Yellow>And</> <Rainbow.Text.Green>I,</> <Rainbow.Text.Blue>who</> <Rainbow.Text.Red>wanted</> <Rainbow.Text.Orange>to</> <Rainbow.Text.Yellow>remain</> <Rainbow.Text.Green>quiet!</><Rainbow.Text.Blue>\"</>")
-			TEXT("\n\n") );
+			TEXT("\n\n"));
+	}
+
+	static void OnBrowserLinkClicked(const FSlateHyperlinkRun::FMetadata& Metadata, TSharedRef<SWidget> ParentWidget)
+	{
+		const FString* url = Metadata.Find(TEXT("href"));
+
+		if(url)
+		{
+			FPlatformProcess::LaunchURL(**url, nullptr, nullptr);
+		}
+		else
+		{
+			SpawnProClickerPopUp(LOCTEXT("FailedToFindUrlPopUpMessage", "Sorry this hyperlink is not <RichText.Tagline.TextHighlight>configured incorrectly</>!"), ParentWidget);
+		}
+	}
 
+	static void SpawnProClickerPopUp(const FText& Text, TSharedRef<SWidget> ParentWidget)
+	{
+		TSharedRef<SWidget> Widget =
+			SNew(SBorder).Padding(10).BorderImage(FTestStyle::Get().GetBrush("RichText.Tagline.Background"))
+			[
+				SNew(SRichTextBlock)
+				.Text(Text)
+				.TextStyle(FTestStyle::Get(), "RichText.Tagline.Text")
+				.DecoratorStyleSet(&FTestStyle::Get())
+				.Justification(ETextJustify::Center)
+			];
 
+
+		FSlateApplication::Get().PushMenu(
+			ParentWidget, // Parent widget should be TestSuite, not the menu thats open or it will be closed when the menu is dismissed
+			Widget,
+			FSlateApplication::Get().GetCursorPos(), // summon location
+			FPopupTransitionEffect(FPopupTransitionEffect::ContextMenu)
+			);
+	}
+};
+
+class SRichTextTest : public SCompoundWidget
+{
+public:
+	SLATE_BEGIN_ARGS( SRichTextTest ){}
+	SLATE_END_ARGS()
+
+	BEGIN_SLATE_FUNCTION_BUILD_OPTIMIZATION
+	void Construct(const FArguments& InArgs)
+	{
 		WrapWidth = 600;
 		bShouldWrap = true;
 		LineHeight = 1.0f;
@@ -1321,14 +1373,14 @@ public:
 										.BorderImage( FTestStyle::Get().GetBrush( "RichText.Background" ) )
 										[
 											SAssignNew( InteractiveRichText, SRichTextBlock )
-											.Text( AliceInWonderland )
+											.Text( RichTextHelper::GetAliceInWonderland() )
 											.TextStyle( FTestStyle::Get(), "RichText.Interactive.Text" )
 											.DecoratorStyleSet( &FTestStyle::Get() )
 											.Margin( this, &SRichTextTest::GetRichTextMargin )
 											.WrapTextAt( this, &SRichTextTest::GetRichTextWrapWidthValue )
 											.Justification( this, &SRichTextTest::JustificationGetSelected )
 											.LineHeightPercentage( this, &SRichTextTest::GetLineHeight )
-											+ SRichTextBlock::HyperlinkDecorator( TEXT("browser"), this, &SRichTextTest::OnBrowserLinkClicked )
+											+ SRichTextBlock::HyperlinkDecorator( TEXT("browser"), FSlateHyperlinkRun::FOnClick::CreateStatic( &RichTextHelper::OnBrowserLinkClicked, AsShared() ) )
 										]
 									]
 								]
@@ -1359,7 +1411,7 @@ public:
 						+ SHorizontalBox::Slot()
 						[
 							SNew( SRichTextBlock)
-							.Text( TheWarOfTheWorlds_Part1 )
+							.Text( RichTextHelper::GetTheWarOfTheWorlds_Part1() )
 							.TextStyle( FTestStyle::Get(), "TheWarOfTheWorlds.Text" )
 							.DecoratorStyleSet( &FTestStyle::Get() )
 							.WrapTextAt( 266 )
@@ -1368,7 +1420,7 @@ public:
 						+ SHorizontalBox::Slot() .Padding( 25, 0 )
 						[
 							SNew( SRichTextBlock)
-							.Text( TheWarOfTheWorlds_Part2 )
+							.Text( RichTextHelper::GetTheWarOfTheWorlds_Part2() )
 							.TextStyle( FTestStyle::Get(), "TheWarOfTheWorlds.Text" )
 							.DecoratorStyleSet( &FTestStyle::Get() )
 							.WrapTextAt( 266 )
@@ -1377,7 +1429,7 @@ public:
 						+ SHorizontalBox::Slot()
 						[
 							SNew( SRichTextBlock)
-							.Text( TheWarOfTheWorlds_Part3 )
+							.Text( RichTextHelper::GetTheWarOfTheWorlds_Part3() )
 							.TextStyle( FTestStyle::Get(), "TheWarOfTheWorlds.Text" )
 							.DecoratorStyleSet( &FTestStyle::Get() )
 							.WrapTextAt( 266 )
@@ -1406,7 +1458,7 @@ public:
 						+ SHorizontalBox::Slot()
 						[
 							SNew( SRichTextBlock)
-							.Text( AroundTheWorldIn80Days_Rainbow )
+							.Text( RichTextHelper::GetAroundTheWorldIn80Days_Rainbow() )
 							.TextStyle( FTestStyle::Get(), "Rainbow.Text" )
 							.DecoratorStyleSet( &FTestStyle::Get() )
 							.WrapTextAt( 600 )
@@ -1467,48 +1519,13 @@ public:
 
 	FReply OnWidgetDecoratorClicked()
 	{
-		SpawnProClickerPopUp( LOCTEXT("WidgetDecoratorExamplePopUpMessage", "I don't really do anything. <RichText.Tagline.TextHighlight>Sorry</>.") );
+		RichTextHelper::SpawnProClickerPopUp( LOCTEXT("WidgetDecoratorExamplePopUpMessage", "I don't really do anything. <RichText.Tagline.TextHighlight>Sorry</>."), AsShared() );
 		return FReply::Handled();
 	}
 
 	void OnHyperlinkDecoratorClicked( const FSlateHyperlinkRun::FMetadata& Metadata )
 	{
-		SpawnProClickerPopUp( LOCTEXT("HyperlinkDecoratorExamplePopUpMessage", "You're a <RichText.Tagline.TextHighlight>pro</> at clicking!") );
-	}
-
-	void OnBrowserLinkClicked( const FSlateHyperlinkRun::FMetadata& Metadata )
-	{
-		const FString* url = Metadata.Find( TEXT("href") );
-
-		if ( url != nullptr )
-		{
-			FPlatformProcess::LaunchURL( **url, nullptr, nullptr );
-		}
-		else
-		{
-			SpawnProClickerPopUp( LOCTEXT("FailedToFindUrlPopUpMessage", "Sorry this hyperlink is not <RichText.Tagline.TextHighlight>configured incorrectly</>!") );
-		}
-	}
-
-	void SpawnProClickerPopUp( const FText& Text )
-	{
-		TSharedRef<SWidget> Widget = 
-			SNew( SBorder ) .Padding(10) .BorderImage( FTestStyle::Get().GetBrush( "RichText.Tagline.Background" ) )
-			[
-				SNew( SRichTextBlock )
-				.Text( Text )
-				.TextStyle( FTestStyle::Get(), "RichText.Tagline.Text" )
-				.DecoratorStyleSet( &FTestStyle::Get() )
-				.Justification( ETextJustify::Center )
-			];
-
-
-		FSlateApplication::Get().PushMenu( 
-			AsShared(), // Parent widget should be TestSuite, not the menu thats open or it will be closed when the menu is dismissed
-			Widget,
-			FSlateApplication::Get().GetCursorPos(), // summon location
-			FPopupTransitionEffect( FPopupTransitionEffect::ContextMenu )
-		);
+		RichTextHelper::SpawnProClickerPopUp( LOCTEXT("HyperlinkDecoratorExamplePopUpMessage", "You're a <RichText.Tagline.TextHighlight>pro</> at clicking!"), AsShared() );
 	}
 
 	void OnTextChanged( const FText& NewText )
@@ -1887,6 +1904,826 @@ protected:
 
  };
 
+ #if WITH_FANCY_TEXT
+
+ /**
+  * This is used in conjunction with the TextStyle decorator to allow arbitrary styling of text within a rich-text editor
+  * This struct defines a set of known font families, as well as providing some utility functions for converting the text style to and from a text layout run
+  */
+struct FTextStyles
+{
+	/** Flags controlling which TTF or OTF font should be picked from the given font family */
+	struct EFontStyle
+	{
+		typedef uint8 Flags;
+		enum Flag
+		{
+			Regular = 0,
+			Bold = 1<<0,
+			Italic = 1<<1,
+		};
+	};
+
+	/** This struct defines a font family, which combines multiple TTF or OTF fonts into a single group, allowing text to be styled as bold or italic */
+	struct FFontFamily
+	{
+		FFontFamily(FText InDisplayName, FName InFamilyName, FName InRegularFont, FName InBoldFont, FName InItalicFont, FName InBoldItalicFont)
+			: DisplayName(InDisplayName)
+			, FamilyName(InFamilyName)
+			, RegularFont(InRegularFont)
+			, BoldFont(InBoldFont)
+			, ItalicFont(InItalicFont)
+			, BoldItalicFont(InBoldItalicFont)
+		{
+		}
+
+		/** Name used for this font family in the UI */
+		FText DisplayName;
+
+		/** Named used to identify this family from the TextStyle decorator */
+		FName FamilyName;
+
+		/** File paths to the fonts on disk */
+		FName RegularFont;
+		FName BoldFont;
+		FName ItalicFont;
+		FName BoldItalicFont;
+	};
+
+	/** Convert the given text style into run meta-information, so that valid source rich-text formatting can be generated for it */
+	static FRunInfo CreateRunInfo(const TSharedPtr<FFontFamily>& InFontFamily, const uint8 InFontSize, const EFontStyle::Flags InFontStyle, const FLinearColor& InFontColor)
+	{
+		FString FontStyleString;
+		if(InFontStyle == EFontStyle::Regular)
+		{
+			FontStyleString = TEXT("Regular");
+		}
+		else
+		{
+			if(InFontStyle & EFontStyle::Bold)
+			{
+				FontStyleString += "Bold";
+			}
+			if(InFontStyle & EFontStyle::Italic)
+			{
+				FontStyleString += "Italic";
+			}
+		}
+
+		FRunInfo RunInfo(TEXT("TextStyle"));
+		RunInfo.MetaData.Add(TEXT("FontFamily"), InFontFamily->FamilyName.ToString());
+		RunInfo.MetaData.Add(TEXT("FontSize"), FString::FromInt(InFontSize));
+		RunInfo.MetaData.Add(TEXT("FontStyle"), FontStyleString);
+		RunInfo.MetaData.Add(TEXT("FontColor"), InFontColor.ToString());
+		return RunInfo;
+	}
+
+	/** Explode some run meta-information back out into its component text style parts */
+	void ExplodeRunInfo(const FRunInfo& InRunInfo, TSharedPtr<FFontFamily>& OutFontFamily, uint8& OutFontSize, EFontStyle::Flags& OutFontStyle, FLinearColor& OutFontColor) const
+	{
+		check(AvailableFontFamilies.Num());
+
+		const FString* const FontFamilyString = InRunInfo.MetaData.Find(TEXT("FontFamily"));
+		if(FontFamilyString)
+		{
+			OutFontFamily = FindFontFamily(FName(**FontFamilyString));
+		}
+		if(!OutFontFamily.IsValid())
+		{
+			OutFontFamily = AvailableFontFamilies[0];
+		}
+
+		OutFontSize = 11;
+		const FString* const FontSizeString = InRunInfo.MetaData.Find(TEXT("FontSize"));
+		if(FontSizeString)
+		{
+			OutFontSize = static_cast<uint8>(FPlatformString::Atoi(**FontSizeString));
+		}
+
+		OutFontStyle = EFontStyle::Regular;
+		const FString* const FontStyleString = InRunInfo.MetaData.Find(TEXT("FontStyle"));
+		if(FontStyleString)
+		{
+			if(*FontStyleString == TEXT("Bold"))
+			{
+				 OutFontStyle = EFontStyle::Bold;
+			}
+			else if(*FontStyleString == TEXT("Italic"))
+			{
+				 OutFontStyle = EFontStyle::Italic;
+			}
+			else if(*FontStyleString == TEXT("BoldItalic"))
+			{
+				 OutFontStyle = EFontStyle::Bold | EFontStyle::Italic;
+			}
+		}
+
+		OutFontColor = FLinearColor::Black;
+		const FString* const FontColorString = InRunInfo.MetaData.Find(TEXT("FontColor"));
+		if(FontColorString && !OutFontColor.InitFromString(*FontColorString))
+		{
+			OutFontColor = FLinearColor::Black;
+		}
+	}
+
+	/** Convert the given text style into a text block style for use by Slate */
+	static FTextBlockStyle CreateTextBlockStyle(const TSharedPtr<FFontFamily>& InFontFamily, const uint8 InFontSize, const EFontStyle::Flags InFontStyle, const FLinearColor& InFontColor)
+	{
+		FName FontName = InFontFamily->RegularFont;
+		if((InFontStyle & EFontStyle::Bold) && (InFontStyle & EFontStyle::Italic))
+		{
+			FontName = InFontFamily->BoldItalicFont;
+		}
+		else if(InFontStyle & EFontStyle::Bold)
+		{
+			FontName = InFontFamily->BoldFont;
+		}
+		else if(InFontStyle & EFontStyle::Italic)
+		{
+			FontName = InFontFamily->ItalicFont;
+		}
+
+		FTextBlockStyle TextBlockStyle;
+		TextBlockStyle.SetFontName(FontName);
+		TextBlockStyle.SetFontSize(InFontSize);
+		TextBlockStyle.SetColorAndOpacity(InFontColor);
+		return TextBlockStyle;
+	}
+
+	/** Convert the given run meta-information into a text block style for use by Slate */
+	FTextBlockStyle CreateTextBlockStyle(const FRunInfo& InRunInfo) const
+	{
+		TSharedPtr<FFontFamily> FontFamily;
+		uint8 FontSize;
+		EFontStyle::Flags FontStyle;
+		FLinearColor FontColor;
+		ExplodeRunInfo(InRunInfo, FontFamily, FontSize, FontStyle, FontColor);
+		return CreateTextBlockStyle(FontFamily, FontSize, FontStyle, FontColor);
+	}
+
+	/** Try and find a font family with the given name */
+	TSharedPtr<FFontFamily> FindFontFamily(const FName InFamilyName) const
+	{
+		const TSharedPtr<FFontFamily>* const FoundFontFamily = AvailableFontFamilies.FindByPredicate([InFamilyName](TSharedPtr<FFontFamily>& Entry) -> bool
+		{
+			return Entry->FamilyName == InFamilyName;
+		});
+		return (FoundFontFamily) ? *FoundFontFamily : nullptr;
+	}
+
+	TArray<TSharedPtr<FFontFamily>> AvailableFontFamilies;
+};
+
+/**
+ * This is a custom decorator used to allow arbitrary styling of text within a rich-text editor
+ * This is required since normal text styling can only work with known styles from a given Slate style-set
+ */
+class FTextStyleDecorator : public ITextDecorator
+{
+public:
+
+	static TSharedRef<FTextStyleDecorator> Create(FTextStyles* const InTextStyles)
+	{
+		return MakeShareable(new FTextStyleDecorator(InTextStyles));
+	}
+
+	virtual ~FTextStyleDecorator()
+	{
+	}
+
+	virtual bool Supports(const FTextRunParseResults& RunParseResult, const FString& Text) const override
+	{
+		return (RunParseResult.Name == TEXT("TextStyle"));
+	}
+
+	virtual TSharedRef<ISlateRun> Create(const FTextRunParseResults& RunParseResult, const FString& OriginalText, const TSharedRef<FString>& InOutModelText, const ISlateStyle* Style) override
+	{
+		FRunInfo RunInfo(RunParseResult.Name);
+		for(const TPair<FString, FTextRange>& Pair : RunParseResult.MetaData)
+		{
+			RunInfo.MetaData.Add(Pair.Key, OriginalText.Mid(Pair.Value.BeginIndex, Pair.Value.EndIndex - Pair.Value.BeginIndex));
+		}
+
+		FTextRange ModelRange;
+		ModelRange.BeginIndex = InOutModelText->Len();
+		*InOutModelText += OriginalText.Mid(RunParseResult.ContentRange.BeginIndex, RunParseResult.ContentRange.EndIndex - RunParseResult.ContentRange.BeginIndex);
+		ModelRange.EndIndex = InOutModelText->Len();
+
+		return FSlateTextRun::Create(RunInfo, InOutModelText, TextStyles->CreateTextBlockStyle(RunInfo), ModelRange);
+	}
+
+private:
+
+	FTextStyleDecorator(FTextStyles* const InTextStyles)
+		: TextStyles(InTextStyles)
+	{
+	}
+
+	FTextStyles* TextStyles;
+};
+
+class SRichTextEditTest : public SCompoundWidget
+{
+public:
+	SLATE_BEGIN_ARGS( SRichTextEditTest ){}
+	SLATE_END_ARGS()
+
+	/**
+	 * Construct the widget
+	 *
+	 * @param InArgs   Declaration from which to construct the widget
+	 */
+	BEGIN_SLATE_FUNCTION_BUILD_OPTIMIZATION
+	void Construct(const FArguments& InArgs)
+	{
+		// Define and add the "Roboto" font family
+		TextStyles.AvailableFontFamilies.Emplace(MakeShareable(new FTextStyles::FFontFamily(
+			LOCTEXT("RobotoFontFamily", "Roboto"), 
+			TEXT("Roboto"), 
+			FName(*(FPaths::EngineContentDir() / TEXT("Slate/Fonts/Roboto-Regular.ttf"))),
+			FName(*(FPaths::EngineContentDir() / TEXT("Slate/Fonts/Roboto-Bold.ttf"))),
+			FName(*(FPaths::EngineContentDir() / TEXT("Slate/Testing/Fonts/Roboto-Italic.ttf"))),
+			FName(*(FPaths::EngineContentDir() / TEXT("Slate/Testing/Fonts/Roboto-BoldItalic.ttf")))
+			)));
+
+		// Set some sensible defaults (these also match the default text style of "RichText.Editor.Text"
+		ActiveFontFamily = TextStyles.AvailableFontFamilies[0];
+		FontSize = 11;
+		FontStyle = FTextStyles::EFontStyle::Regular;
+		FontColor = FLinearColor::Black;
+
+		// Define the example text
+		// Both of our demo editors display this same text, but the text marshallers display it differently
+		RichEditableText = FText::FromString(
+			TEXT("<TextStyle FontFamily=\"Roboto\" FontSize=\"11\" FontStyle=\"Regular\" FontColor=\"(R=0.000000,G=0.000000,B=0.000000,A=1.000000)\">This is some regular text!</>")
+			TEXT("\n")
+			TEXT("<TextStyle FontFamily=\"Roboto\" FontSize=\"11\" FontStyle=\"Italic\" FontColor=\"(R=0.000000,G=0.000000,B=0.000000,A=1.000000)\">This is some italic text!</>")
+			TEXT("\n")
+			TEXT("<TextStyle FontFamily=\"Roboto\" FontSize=\"11\" FontStyle=\"Bold\" FontColor=\"(R=0.000000,G=0.000000,B=0.000000,A=1.000000)\">This is some bold text!</>")
+			TEXT("\n")
+			TEXT("<TextStyle FontFamily=\"Roboto\" FontSize=\"11\" FontStyle=\"BoldItalic\" FontColor=\"(R=0.000000,G=0.000000,B=0.000000,A=1.000000)\">This is some bold-italic text!</>")
+			TEXT("\n")
+			TEXT("<TextStyle FontFamily=\"Roboto\" FontSize=\"26\" FontStyle=\"BoldItalic\" FontColor=\"(R=1.000000,G=0.457315,B=0.164297,A=1.000000)\">This </><TextStyle FontFamily=\"Roboto\" FontSize=\"21\" FontStyle=\"Italic\" FontColor=\"(R=0.331120,G=0.348415,B=1.000000,A=1.000000)\">is </><TextStyle FontFamily=\"Roboto\" FontSize=\"21\" FontStyle=\"Bold\" FontColor=\"(R=1.000000,G=0.195252,B=0.137910,A=1.000000)\">some </><TextStyle FontFamily=\"Roboto\" FontSize=\"18\" FontStyle=\"BoldItalic\" FontColor=\"(R=1.000000,G=0.167594,B=0.566043,A=1.000000)\">crazy </><TextStyle FontFamily=\"Roboto\" FontSize=\"14\" FontStyle=\"Italic\" FontColor=\"(R=0.175489,G=0.695000,B=0.169137,A=1.000000)\">text!</>")
+			TEXT("\n")
+			TEXT("<TextStyle FontFamily=\"Roboto\" FontSize=\"11\" FontStyle=\"Regular\" FontColor=\"(R=0.000000,G=0.000000,B=0.000000,A=1.000000)\">This is a hyperlink: </><a id=\"browser\" href=\"https://www.unrealengine.com/\" style=\"RichText.Editor.Hyperlink\">Unreal Engine</><TextStyle FontFamily=\"Roboto\" FontSize=\"11\" FontStyle=\"Regular\" FontColor=\"(R=0.000000,G=0.000000,B=0.000000,A=1.000000)\"></>")
+			);
+
+		// The rich-text marshaller will be used with the rich-text editor...
+		TSharedRef<FRichTextLayoutMarshaller> RichTextMarshaller = FRichTextLayoutMarshaller::Create(
+			TArray<TSharedRef<ITextDecorator>>(), 
+			&FTestStyle::Get(), 
+			FTestStyle::Get().GetWidgetStyle<FTextBlockStyle>("RichText.Editor.Text")
+			);
+
+		// ... so we also need to add some decorators to handle the things we want to demo
+		OnHyperlinkClicked = FSlateHyperlinkRun::FOnClick::CreateStatic(&RichTextHelper::OnBrowserLinkClicked, AsShared());
+		RichTextMarshaller->AppendInlineDecorator(FHyperlinkDecorator::Create(TEXT("browser"), OnHyperlinkClicked));
+		RichTextMarshaller->AppendInlineDecorator(FTextStyleDecorator::Create(&TextStyles));
+
+		// The syntax highlighter marshaller is self contained, so doesn't need any extra configuration
+		SyntaxHighlighterMarshaller = FRichTextSyntaxHighlighterTextLayoutMarshaller::Create(
+			FRichTextSyntaxHighlighterTextLayoutMarshaller::FSyntaxTextStyle()
+			);
+
+		this->ChildSlot
+		[
+			SNew(SScrollBox)
+			+SScrollBox::Slot()
+			[
+				SNew(SBorder)
+				.BorderImage(FTestStyle::Get().GetBrush("RichText.Background"))
+				[
+					SNew(SVerticalBox)
+
+					// Demo title
+					+SVerticalBox::Slot()
+					.AutoHeight()
+					.Padding(0)
+					[
+						SNew(SBorder)
+						.BorderImage(FTestStyle::Get().GetBrush("RichText.Tagline.Background"))
+						.Padding(0)
+						[
+							SNew(SRichTextBlock)
+							.Text(LOCTEXT("RichTextEdit.Header01", "This WYSIWYG editor has been created to show the flexibility of Slate's <RichText.Tagline.TextHighlight>SMultiLineEditableText</>"))
+							.TextStyle(FTestStyle::Get(), "RichText.Tagline.Text")
+							.DecoratorStyleSet(&FTestStyle::Get())
+							.AutoWrapText(true)
+							.Justification(ETextJustify::Center)
+							.Margin(FMargin(20))
+						]
+					]
+
+					// Demo description
+					+SVerticalBox::Slot()
+					.AutoHeight()
+					.Padding(0)
+					[
+						SNew(SRichTextBlock)
+						.Text(LOCTEXT("RichTextEdit.HowItWorks", "<RichText.TextHighlight>SMultiLineEditableText</> utilizes the power of text marshallers to control how the source text is displayed, and the two editors below are demonstrating the <RichText.TextHighlight>FRichTextLayoutMarshaller</> and <RichText.TextHighlight>FRichTextSyntaxHighlighterTextLayoutMarshaller</> implementations.\n\nYou can use the toolbar to control the style of the text in the rich-text editor, or you can directly edit the source text (demonstrating live injection of style information) to update the rich-text view."))
+						.TextStyle(FTestStyle::Get(), "RichText.Text")
+						.DecoratorStyleSet(&FTestStyle::Get())
+						.AutoWrapText(true)
+						.Margin(FMargin(40, 20))
+					]
+
+					// Rich-text editor
+					+SVerticalBox::Slot()
+					.AutoHeight()
+					.Padding(FMargin(25, 0))
+					[
+						SNew(SBorder)
+						.Padding(5.0f)
+						.BorderImage(FCoreStyle::Get().GetBrush("BoxShadow"))
+						[
+							SNew(SBorder)
+							.BorderImage(FTestStyle::Get().GetBrush("RichText.Interactive.Details.Background"))
+							.Padding(FMargin(4))
+							[
+								// Rich-text editor toolbar
+								SNew(SVerticalBox)
+
+								+SVerticalBox::Slot()
+								.AutoHeight()
+								.Padding(FMargin(0.0f, 0.0f, 0.0f, 4.0f))
+								[
+									SNew(SBorder)
+									.BorderImage(FTestStyle::Get().GetBrush("RichText.RoundedBackground"))
+									.Padding(FMargin(4))
+									[
+										SNew(SHorizontalBox)
+
+										+SHorizontalBox::Slot()
+										.AutoWidth()
+										[
+											SNew(SComboBox<TSharedPtr<FTextStyles::FFontFamily>>)
+											.ComboBoxStyle(FTestStyle::Get(), "RichText.Toolbar.ComboBox")
+											.OptionsSource(&TextStyles.AvailableFontFamilies)
+											.OnSelectionChanged(this, &SRichTextEditTest::OnActiveFontFamilyChanged)
+											.OnGenerateWidget(this, &SRichTextEditTest::GenerateFontFamilyComboEntry)
+											.InitiallySelectedItem(ActiveFontFamily)
+											[
+												SNew(SBox)
+												.Padding(FMargin(0.0f, 0.0f, 2.0f, 0.0f))
+												[
+													SNew(STextBlock)
+													.Text(this, &SRichTextEditTest::GetActiveFontFamilyName)
+												]
+											]
+										]
+
+										+SHorizontalBox::Slot()
+										.AutoWidth()
+										[
+											SNew(SBox)
+											.WidthOverride(24)
+											[
+												SNew(SNumericEntryBox<uint8>)
+												.Value(this, &SRichTextEditTest::GetFontSize)
+												.OnValueCommitted(this, &SRichTextEditTest::SetFontSize)
+											]
+										]
+
+										+SHorizontalBox::Slot()
+										.AutoWidth()
+										.Padding(FMargin(4.0f, 0.0f, 0.0f, 0.0f))
+										[
+											SNew(SCheckBox)
+											.Style(FTestStyle::Get(), "RichText.Toolbar.ToggleButtonCheckbox")
+											.IsChecked(this, &SRichTextEditTest::IsFontStyleBold)
+											.OnCheckStateChanged(this, &SRichTextEditTest::OnFontStyleBoldChanged)
+											[
+												SNew(SBox)
+												.WidthOverride(24)
+												.HAlign(HAlign_Center)
+												.VAlign(VAlign_Center)
+												[
+													SNew(STextBlock)
+													.TextStyle(FTestStyle::Get(), "RichText.Toolbar.BoldText")
+													.Text(LOCTEXT("BoldLabel", "B"))
+												]
+											]
+										]
+
+										+SHorizontalBox::Slot()
+										.AutoWidth()
+										[
+											SNew(SCheckBox)
+											.Style(FTestStyle::Get(), "RichText.Toolbar.ToggleButtonCheckbox")
+											.IsChecked(this, &SRichTextEditTest::IsFontStyleItalic)
+											.OnCheckStateChanged(this, &SRichTextEditTest::OnFontStyleItalicChanged)
+											[
+												SNew(SBox)
+												.WidthOverride(24)
+												.HAlign(HAlign_Center)
+												.VAlign(VAlign_Center)
+												[
+													SNew(STextBlock)
+													.TextStyle(FTestStyle::Get(), "RichText.Toolbar.ItalicText")
+													.Text(LOCTEXT("ItalicLabel", "I"))
+												]
+											]
+										]
+
+										+SHorizontalBox::Slot()
+										.Padding(FMargin(4.0f, 0.0f, 0.0f, 0.0f))
+										.AutoWidth()
+										[
+											SNew(SButton)
+											.ButtonStyle(FTestStyle::Get(), "RichText.Toolbar.Button")
+											.OnClicked(this, &SRichTextEditTest::OpenFontColorPicker)
+											[
+												SNew(SOverlay)
+
+												+SOverlay::Slot()
+												.Padding(FMargin(0.0f, 0.0f, 0.0f, 4.0f))
+												.HAlign(HAlign_Center)
+												.VAlign(VAlign_Bottom)
+												[
+													SNew(STextBlock)
+													.TextStyle(FTestStyle::Get(), "RichText.Toolbar.BoldText")
+													.Text(LOCTEXT("ColorLabel", "A"))
+												]
+
+												+SOverlay::Slot()
+												.HAlign(HAlign_Center)
+												.VAlign(VAlign_Bottom)
+												[
+													SNew(SColorBlock)
+													.Color(this, &SRichTextEditTest::GetFontColor)
+													.Size(FVector2D(20.0f, 6.0f))
+												]
+											]
+										]
+
+										+SHorizontalBox::Slot()
+										.Padding(FMargin(4.0f, 0.0f, 0.0f, 0.0f))
+										.AutoWidth()
+										[
+											SAssignNew(HyperlinkComboButton, SComboButton)
+											.ComboButtonStyle(FTestStyle::Get(), "RichText.Toolbar.ComboButton")
+											.HasDownArrow(false)
+											.OnComboBoxOpened(this, &SRichTextEditTest::HandleHyperlinkComboOpened)
+											.ButtonContent()
+											[
+												SNew(SBox)
+												.WidthOverride(20)
+												.HAlign(HAlign_Center)
+												.VAlign(VAlign_Center)
+												[
+													SNew(SImage)
+													.Image(FTestStyle::Get().GetBrush("RichText.Toolbar.HyperlinkImage"))
+												]
+											]
+											.MenuContent()
+											[
+												SNew(SGridPanel)
+												.FillColumn(1, 1.0f)
+
+												+SGridPanel::Slot(0, 0)
+												.HAlign(HAlign_Right)
+												.Padding(FMargin(2.0f))
+												[
+													SNew(STextBlock)
+													.TextStyle(FTestStyle::Get(), "RichText.Toolbar.Text")
+													.Text(LOCTEXT("HyperlinkNameLabel", "Name:"))
+												]
+												+SGridPanel::Slot(1, 0)
+												.Padding(FMargin(2.0f))
+												[
+													SNew(SBox)
+													.WidthOverride(300)
+													[
+														SAssignNew(HyperlinkNameTextBox, SEditableTextBox)
+													]
+												]
+
+												+SGridPanel::Slot(0, 1)
+												.HAlign(HAlign_Right)
+												.Padding(FMargin(2.0f))
+												[
+													SNew(STextBlock)
+													.TextStyle(FTestStyle::Get(), "RichText.Toolbar.Text")
+													.Text(LOCTEXT("HyperlinkURLLabel", "URL:"))
+												]
+												+SGridPanel::Slot(1, 1)
+												.Padding(FMargin(2.0f))
+												[
+													SNew(SBox)
+													.WidthOverride(300)
+													[
+														SAssignNew(HyperlinkURLTextBox, SEditableTextBox)
+													]
+												]
+
+												+SGridPanel::Slot(0, 2)
+												.HAlign(HAlign_Right)
+												.Padding(FMargin(2.0f))
+												.ColumnSpan(2)
+												[
+													SNew(SButton)
+													.ButtonStyle(FTestStyle::Get(), "RichText.Toolbar.Button")
+													.OnClicked(this, &SRichTextEditTest::HandleInsertHyperlinkClicked)
+													[
+														SNew(STextBlock)
+														.TextStyle(FTestStyle::Get(), "RichText.Toolbar.Text")
+														.Text(LOCTEXT("HyperlinkInsertLabel", "Insert Hyperlink"))
+													]
+												]
+											]
+										]
+									]
+								]
+
+								+SVerticalBox::Slot()
+								.AutoHeight()
+								.Padding(FMargin(0.0f, 0.0f, 0.0f, 0.0f))
+								[
+									SAssignNew(RichEditableTextBox, SMultiLineEditableTextBox)
+									.Font(FTestStyle::Get().GetWidgetStyle<FTextBlockStyle>("RichText.Editor.Text").Font)
+									.Text(this, &SRichTextEditTest::GetRichEditableText)
+									.OnTextChanged(this, &SRichTextEditTest::HandleRichEditableTextChanged)
+									.OnTextCommitted(this, &SRichTextEditTest::HandleRichEditableTextCommitted)
+									.OnCursorMoved(this, &SRichTextEditTest::HandleRichEditableTextCursorMoved)
+									.Marshaller(RichTextMarshaller)
+									.AutoWrapText(true)
+									.Margin(4)
+									.LineHeightPercentage(1.1f)
+								]
+							]
+						]
+					]
+
+					// Source text editor
+					+SVerticalBox::Slot()
+					.AutoHeight()
+					.Padding(FMargin(25, 0, 25, 25))
+					[
+						SNew(SBorder)
+						.Padding(5.0f)
+						.BorderImage(FCoreStyle::Get().GetBrush("BoxShadow"))
+						[
+							SNew(SBorder)
+							.BorderImage(FTestStyle::Get().GetBrush("RichText.Interactive.Details.Background"))
+							.Padding(FMargin(4))
+							[
+								SNew(SVerticalBox)
+
+								+SVerticalBox::Slot()
+								.AutoHeight()
+								.Padding(FMargin(0.0f, 0.0f, 0.0f, 4.0f))
+								[
+									SNew(SMultiLineEditableTextBox)
+									.Text(this, &SRichTextEditTest::GetRichEditableText)
+									.OnTextChanged(this, &SRichTextEditTest::HandleRichEditableTextChanged)
+									.OnTextCommitted(this, &SRichTextEditTest::HandleRichEditableTextCommitted)
+									.Marshaller(SyntaxHighlighterMarshaller)
+									.AutoWrapText(true)
+									.Margin(4)
+									.LineHeightPercentage(1.1f)
+								]
+
+								+SVerticalBox::Slot()
+								.AutoHeight()
+								[
+									SNew(SBorder)
+									.BorderImage(FTestStyle::Get().GetBrush("RichText.RoundedBackground"))
+									.Padding(FMargin(4))
+									.HAlign(HAlign_Right)
+									[
+										SNew(SCheckBox)
+										.Style(FTestStyle::Get(), "RichText.Toolbar.Checkbox")
+										.IsChecked(this, &SRichTextEditTest::IsEnableSyntaxHighlightingChecked)
+										.OnCheckStateChanged(this, &SRichTextEditTest::OnEnableSyntaxHighlightingChanged)
+										[
+											SNew(STextBlock)
+											.TextStyle(FTestStyle::Get(), "RichText.Toolbar.Text")
+											.Text(LOCTEXT("EnableSyntaxHighlightingLabel", "Enable Syntax Highlighting"))
+										]
+									]
+								]
+							]
+						]
+					]
+				]
+			]
+		];
+	}
+	END_SLATE_FUNCTION_BUILD_OPTIMIZATION
+
+protected:
+	
+	FText GetRichEditableText() const
+	{
+		return RichEditableText;
+	}
+
+	void HandleRichEditableTextChanged(const FText& Text)
+	{
+		RichEditableText = Text;
+	}
+
+	void HandleRichEditableTextCommitted(const FText& Text, ETextCommit::Type Type)
+	{
+		RichEditableText = Text;
+	}
+
+	void HandleRichEditableTextCursorMoved(const FTextLocation& NewCursorPosition )
+	{
+		// We can use GetRunUnderCursor to query the style of the text under the cursor
+		// We can then use this to update the toolbar
+		TSharedPtr<const IRun> Run = RichEditableTextBox->GetRunUnderCursor();
+		if(Run.IsValid() && Run->GetRunInfo().Name == TEXT("TextStyle"))
+		{
+			TextStyles.ExplodeRunInfo(Run->GetRunInfo(), ActiveFontFamily, FontSize, FontStyle, FontColor);
+		}
+	}
+
+	FText GetActiveFontFamilyName() const
+	{
+		return ActiveFontFamily->DisplayName;
+	}
+
+	void OnActiveFontFamilyChanged(TSharedPtr<FTextStyles::FFontFamily> NewValue, ESelectInfo::Type)
+	{
+		ActiveFontFamily = NewValue;
+		StyleSelectedText();
+	}
+
+	TSharedRef<SWidget> GenerateFontFamilyComboEntry(TSharedPtr<FTextStyles::FFontFamily> SourceEntry)
+	{
+		return SNew(STextBlock).Text(SourceEntry->DisplayName);
+	}
+
+	TOptional<uint8> GetFontSize() const
+	{
+		return FontSize;
+	}
+
+	void SetFontSize(uint8 NewValue, ETextCommit::Type)
+	{		
+		FontSize = NewValue;
+		StyleSelectedText();
+	}
+
+	ESlateCheckBoxState::Type IsFontStyleBold() const
+	{
+		return (FontStyle & FTextStyles::EFontStyle::Bold) ? ESlateCheckBoxState::Checked : ESlateCheckBoxState::Unchecked;
+	}
+
+	void OnFontStyleBoldChanged(ESlateCheckBoxState::Type InState)
+	{
+		if(InState == ESlateCheckBoxState::Checked)
+		{
+			FontStyle |= FTextStyles::EFontStyle::Bold;
+		}
+		else
+		{
+			FontStyle &= ~FTextStyles::EFontStyle::Bold;
+		}
+		StyleSelectedText();
+	}
+
+	ESlateCheckBoxState::Type IsFontStyleItalic() const
+	{
+		return (FontStyle & FTextStyles::EFontStyle::Italic) ? ESlateCheckBoxState::Checked : ESlateCheckBoxState::Unchecked;
+	}
+
+	void OnFontStyleItalicChanged(ESlateCheckBoxState::Type InState)
+	{
+		if(InState == ESlateCheckBoxState::Checked)
+		{
+			FontStyle |= FTextStyles::EFontStyle::Italic;
+		}
+		else
+		{
+			FontStyle &= ~FTextStyles::EFontStyle::Italic;
+		}
+		StyleSelectedText();
+	}
+
+	FLinearColor GetFontColor() const
+	{
+		return FontColor;
+	}
+
+	void SetFontColor(FLinearColor NewValue)
+	{		
+		FontColor = NewValue;
+		StyleSelectedText();
+	}
+
+	FReply OpenFontColorPicker()
+	{
+		FColorPickerArgs PickerArgs;
+		PickerArgs.bOnlyRefreshOnMouseUp = true;
+		PickerArgs.ParentWidget = AsShared();
+		PickerArgs.bUseAlpha = false;
+		PickerArgs.bOnlyRefreshOnOk = false;
+		PickerArgs.OnColorCommitted = FOnLinearColorValueChanged::CreateSP(this, &SRichTextEditTest::SetFontColor);
+		PickerArgs.OnColorPickerCancelled = FOnColorPickerCancelled::CreateSP(this, &SRichTextEditTest::SetFontColor);
+		PickerArgs.InitialColorOverride = FontColor;
+
+		OpenColorPicker(PickerArgs);
+
+		return FReply::Handled();
+	}
+
+	void StyleSelectedText()
+	{
+		// Apply the current style to the selected text
+		// If no text is selected, then a new (empty) run will be inserted with the appropriate style
+		const FRunInfo RunInfo = TextStyles.CreateRunInfo(ActiveFontFamily, FontSize, FontStyle, FontColor);
+		const FTextBlockStyle TextBlockStyle = TextStyles.CreateTextBlockStyle(ActiveFontFamily, FontSize, FontStyle, FontColor);
+		RichEditableTextBox->ApplyToSelection(RunInfo, TextBlockStyle);
+		FSlateApplication::Get().SetKeyboardFocus(RichEditableTextBox, EKeyboardFocusCause::SetDirectly);
+	}
+
+	void HandleHyperlinkComboOpened()
+	{
+		// Read any currently selected text, and use this as the default name of the hyperlink
+		FString SelectedText = RichEditableTextBox->GetSelectedText().ToString();
+		for(int32 SelectedTextIndex = 0; SelectedTextIndex < SelectedText.Len(); ++SelectedTextIndex)
+		{
+			if(FChar::IsLinebreak(SelectedText[SelectedTextIndex]))
+			{
+				SelectedText = SelectedText.Left(SelectedTextIndex);
+				break;
+			}
+		}
+		HyperlinkNameTextBox->SetText(FText::FromString(SelectedText));
+
+		// We can use GetRunUnderCursor to query whether the cursor is currently over a hyperlink
+		// If it is, we can use that as the default URL for the hyperlink
+		TSharedPtr<const IRun> Run = RichEditableTextBox->GetRunUnderCursor();
+		if(Run.IsValid() && Run->GetRunInfo().Name == TEXT("a"))
+		{
+			const FString* const URLUnderCursor = Run->GetRunInfo().MetaData.Find(TEXT("href"));
+			HyperlinkURLTextBox->SetText((URLUnderCursor) ? FText::FromString(*URLUnderCursor) : FText());
+		}
+		else
+		{
+			HyperlinkURLTextBox->SetText(FText());
+		}
+	}
+
+	FReply HandleInsertHyperlinkClicked()
+	{
+		HyperlinkComboButton->SetIsOpen(false);
+
+		const FText& Name = HyperlinkNameTextBox->GetText();
+		const FText& URL = HyperlinkURLTextBox->GetText();
+
+		// Create the correct meta-information for this run, so that valid source rich-text formatting can be generated for it
+		FRunInfo RunInfo(TEXT("a"));
+		RunInfo.MetaData.Add(TEXT("id"), TEXT("browser"));
+		RunInfo.MetaData.Add(TEXT("href"), URL.ToString());
+		RunInfo.MetaData.Add(TEXT("style"), TEXT("RichText.Editor.Hyperlink"));
+
+		// Create the new run, and then insert it at the cursor position
+		TSharedRef<FSlateHyperlinkRun> HyperlinkRun = FSlateHyperlinkRun::Create(
+			RunInfo, 
+			MakeShareable(new FString(Name.ToString())), 
+			FTestStyle::Get().GetWidgetStyle<FHyperlinkStyle>(FName(TEXT("RichText.Editor.Hyperlink"))), 
+			OnHyperlinkClicked
+			);
+		RichEditableTextBox->InsertRunAtCursor(HyperlinkRun);
+
+		return FReply::Handled();
+	}
+
+	ESlateCheckBoxState::Type IsEnableSyntaxHighlightingChecked() const
+	{
+		return (SyntaxHighlighterMarshaller->IsSyntaxHighlightingEnabled()) ? ESlateCheckBoxState::Checked : ESlateCheckBoxState::Unchecked;
+	}
+
+	void OnEnableSyntaxHighlightingChanged(ESlateCheckBoxState::Type InState)
+	{
+		SyntaxHighlighterMarshaller->EnableSyntaxHighlighting(InState == ESlateCheckBoxState::Checked);
+	}
+
+protected:
+
+	FText RichEditableText;
+	TSharedPtr<SMultiLineEditableTextBox> RichEditableTextBox;
+
+	TSharedPtr<FSyntaxHighlighterTextLayoutMarshaller> SyntaxHighlighterMarshaller;
+
+	FSlateHyperlinkRun::FOnClick OnHyperlinkClicked;
+	TSharedPtr<SComboButton> HyperlinkComboButton;
+	TSharedPtr<SEditableTextBox> HyperlinkNameTextBox;
+	TSharedPtr<SEditableTextBox> HyperlinkURLTextBox;
+
+	FTextStyles TextStyles;
+
+	TSharedPtr<FTextStyles::FFontFamily> ActiveFontFamily;
+	uint8 FontSize;
+	FTextStyles::EFontStyle::Flags FontStyle;
+	FLinearColor FontColor;
+};
+
+#endif //WITH_FANCY_TEXT
+
 /** Demonstrates the brokenness of our current approach to trading smoothness for sharpness. */
  class SLayoutRoundingTest : public SCompoundWidget
  {
@@ -3291,16 +4128,13 @@ TSharedRef<SDockTab> SpawnTab(const FSpawnTabArgs& Args, FName TabIdentifier)
 			MakeLayoutExample()
 		];
 	}
+#if WITH_FANCY_TEXT
 	else if (TabIdentifier == FName(TEXT("RichTextTab")))
 	{
 		return SNew(SDockTab)
 			. Label( LOCTEXT("RichTextTestTab", "Rich Text") )
 			[
-				#if WITH_FANCY_TEXT
 				SNew( SRichTextTest )
-				#else
-				SNew( SSpacer )
-				#endif //WITH_FANCY_TEXT
 			];
 	}
 	else if ( TabIdentifier == FName( TEXT( "MultiLineEditTab" ) ) )
@@ -3315,6 +4149,7 @@ TSharedRef<SDockTab> SpawnTab(const FSpawnTabArgs& Args, FName TabIdentifier)
 				#endif //WITH_FANCY_TEXT
 			];
 	}
+#endif //WITH_FANCY_TEXT
 	else if (TabIdentifier == FName(TEXT("EditableTextTab")))
 	{
 		return SNew(SDockTab)
@@ -3324,6 +4159,17 @@ TSharedRef<SDockTab> SpawnTab(const FSpawnTabArgs& Args, FName TabIdentifier)
 			SNew( STextEditTest )
 		];
 	}
+#if WITH_FANCY_TEXT
+	else if (TabIdentifier == FName(TEXT("RichEditableTextTab")))
+	{
+		return SNew(SDockTab)
+			. Label( LOCTEXT("RichEditableTextTestTab", "Rich Editable Text") )
+			. ToolTipText( LOCTEXT( "RichEditableTextTestTabToolTip", "Switches to the Rich Editable Text tab, where you can test the various rich editable text features." ) )
+			[
+				SNew( SRichTextEditTest )
+			];
+	}
+#endif //WITH_FANCY_TEXT
 	else if (TabIdentifier == FName("LayoutRoundingTab"))
 	{
 		return SNew(SDockTab)
@@ -3461,9 +4307,12 @@ TSharedRef<SDockTab> SpawnTestSuite1( const FSpawnTabArgs& Args )
 				#endif
 				->AddTab("LayoutRoundingTab", ETabState::OpenedTab)
 				->AddTab("EditableTextTab", ETabState::OpenedTab)
+#if WITH_FANCY_TEXT
 				->AddTab("RichTextTab", ETabState::OpenedTab)
 				->AddTab("MultiLineEditTab", ETabState::OpenedTab)
-			)
+				->AddTab("RichEditableTextTab", ETabState::OpenedTab)
+#endif //WITH_FANCY_TEXT
+				)
 		)
 	)
 #if PLATFORM_SUPPORTS_MULTIPLE_NATIVE_WINDOWS
@@ -3510,6 +4359,10 @@ TSharedRef<SDockTab> SpawnTestSuite1( const FSpawnTabArgs& Args )
 		.SetDisplayName( NSLOCTEXT( "TestSuite1", "MultiLineEditTab", "Multiline Edit" ) )
 		.SetGroup( TestSuiteMenu::SuiteTabs );
 
+	TestSuite1TabManager->RegisterTabSpawner( "RichEditableTextTab", FOnSpawnTab::CreateStatic( &SpawnTab, FName("RichEditableTextTab") ) )
+		.SetDisplayName( NSLOCTEXT("TestSuite1", "RichEditableTextTab", "Rich Editable Text Test") )
+		.SetGroup(TestSuiteMenu::SuiteTabs);
+
 	TestSuite1TabManager->RegisterTabSpawner( "LayoutRoundingTab", FOnSpawnTab::CreateStatic( &SpawnTab, FName("LayoutRoundingTab") ) )
 		.SetDisplayName( NSLOCTEXT("TestSuite1", "LayoutRoundingTab", "Layout Rounding") )
 		.SetGroup(TestSuiteMenu::SuiteTabs);
diff --git a/Engine/Source/Runtime/Slate/Private/Widgets/Text/SMultiLineEditableText.cpp b/Engine/Source/Runtime/Slate/Private/Widgets/Text/SMultiLineEditableText.cpp
--- a/Engine/Source/Runtime/Slate/Private/Widgets/Text/SMultiLineEditableText.cpp
+++ b/Engine/Source/Runtime/Slate/Private/Widgets/Text/SMultiLineEditableText.cpp
@@ -6,6 +6,7 @@
 
 #include "TextEditHelper.h"
 #include "SlateWordWrapper.h"
+#include "PlainTextLayoutMarshaller.h"
 
 void SMultiLineEditableText::FCursorInfo::SetCursorLocationAndCalculateAlignment(const TSharedPtr<FTextLayout>& TextLayout, const FTextLocation& InCursorPosition)
 {
@@ -226,6 +227,13 @@ void SMultiLineEditableText::Construct( const FArguments& InArgs )
 
 	OnTextChanged = InArgs._OnTextChanged;
 	OnTextCommitted = InArgs._OnTextCommitted;
+	OnCursorMoved = InArgs._OnCursorMoved;
+
+	Marshaller = InArgs._Marshaller;
+	if (!Marshaller.IsValid())
+	{
+		Marshaller = FPlainTextLayoutMarshaller::Create(TextStyle);
+	}
 
 	CursorLineHighlighter = FCursorLineHighlighter::Create(&CursorInfo);
 	TextCompositionHighlighter = FTextCompositionHighlighter::Create();
@@ -319,66 +327,19 @@ bool SMultiLineEditableText::SetEditableText(const FText& TextToSet, const bool
 {
 	const FText EditedText = GetEditableText();
 
+	Marshaller->ClearDirty();
+
 	const bool HasTextChanged = !EditedText.ToString().Equals(TextToSet.ToString(), ESearchCase::CaseSensitive);
 	if (HasTextChanged || bForce)
 	{
 		const FString& TextToSetString = TextToSet.ToString();
-		const int32 TextToSetLength = TextToSetString.Len();
 
 		ClearSelection();
 		TextLayout->ClearLines();
 
-		if(TextToSetLength)
+		if(TextToSetString.Len())
 		{
-			// Iterate over each line break candidate, adding ranges from after the end of the last line added to before the newline or end of string.
-			FLineBreakIterator LBI(TextToSetString);
-			int32 RangeBegin = LBI.GetCurrentPosition();
-			TSharedPtr< FString > LineText;
-
-			for (;;)
-			{
-				const int32 BreakIndex = LBI.MoveToNext();
-
-				// If beginning or potential end is invalid, cease.
-				if (RangeBegin == INDEX_NONE || BreakIndex == INDEX_NONE)
-				{
-					break;
-				}
-
-				const int32 BreakingCharacterIndex = BreakIndex - 1;
-				if (BreakingCharacterIndex >= RangeBegin) // Valid index check.
-				{
-					const TCHAR BreakingCharacter = TextToSetString[BreakingCharacterIndex];
-					// If line break candidate is after a newline, add the range as a new line.
-					if (FChar::IsLinebreak(BreakingCharacter))
-					{
-						LineText = MakeShareable(new FString(BreakingCharacterIndex - RangeBegin, (*TextToSetString) + RangeBegin));
-					}
-					// If the line break candidate is the end of the string, add the range as a new line.
-					else if (BreakIndex == TextToSetLength)
-					{
-						LineText = MakeShareable(new FString(BreakIndex - RangeBegin, (*TextToSetString) + RangeBegin));
-					}
-					else
-					{
-						continue;
-					}
-
-					if (LineText.IsValid())
-					{
-						// Create run.
-						TArray< TSharedRef< IRun > > Runs;
-						Runs.Add(FSlateTextRun::Create(LineText.ToSharedRef(), TextStyle));
-
-						TextLayout->AddLine(LineText.ToSharedRef(), Runs);
-
-						// Reset the LineText so we don't add it multiple times.
-						LineText.Reset();
-					}
-
-					RangeBegin = BreakIndex; // The next line begins after the end of the current line.
-				}
-			}
+			Marshaller->SetText(TextToSetString, *TextLayout);
 		}
 		else
 		{
@@ -386,7 +347,7 @@ bool SMultiLineEditableText::SetEditableText(const FText& TextToSet, const bool
 
 			// Create an empty run
 			TArray< TSharedRef< IRun > > Runs;
-			Runs.Add(FSlateTextRun::Create(LineText.ToSharedRef(), TextStyle));
+			Runs.Add(FSlateTextRun::Create(FRunInfo(), LineText.ToSharedRef(), TextStyle));
 
 			TextLayout->AddLine(LineText.ToSharedRef(), Runs);
 		}
@@ -399,9 +360,23 @@ bool SMultiLineEditableText::SetEditableText(const FText& TextToSet, const bool
 
 FText SMultiLineEditableText::GetEditableText() const
 {
-	FText EditedText;
-	TextLayout->GetAsText(EditedText);
-	return EditedText;
+	FString EditedText;
+	Marshaller->GetText(EditedText, *TextLayout);
+	return FText::FromString(EditedText);
+}
+
+void SMultiLineEditableText::ForceRefreshTextLayout(const FText& CurrentText)
+{
+	// Marshallers shouldn't inject any visible characters into the text, but SetEditableText will clear the current selection
+	// so we need to back that up here so we don't cause the cursor to jump around
+	const TOptional<FTextLocation> OldSelectionStart = SelectionStart;
+	const FCursorInfo OldCursorInfo = CursorInfo;
+
+	SetEditableText(CurrentText, true);
+
+	SelectionStart = OldSelectionStart;
+	CursorInfo = OldCursorInfo;
+	UpdateCursorHighlight();
 }
 
 FReply SMultiLineEditableText::OnKeyboardFocusReceived( const FGeometry& MyGeometry, const FKeyboardFocusEvent& InKeyboardFocusEvent )
@@ -434,11 +409,6 @@ void SMultiLineEditableText::OnKeyboardFocusLost( const FKeyboardFocusEvent& InK
 		}
 
 		// When focus is lost let anyone who is interested that text was committed
-		if (InKeyboardFocusEvent.GetCause() != EKeyboardFocusCause::WindowActivate) //Don't clear selection when activating a new window or cant copy and paste on right click )
-		{
-			ClearSelection();
-		}
-		
 		// See if user explicitly tabbed away or moved focus
 		ETextCommit::Type TextAction;
 		switch (InKeyboardFocusEvent.GetCause())
@@ -463,7 +433,7 @@ void SMultiLineEditableText::OnKeyboardFocusLost( const FKeyboardFocusEvent& InK
 		const FText EditedText = GetEditableText();
 
 		OnTextCommitted.ExecuteIfBound(EditedText, TextAction);
-		RemoveCursorHighlight();
+		UpdateCursorHighlight();
 	}
 }
 
@@ -505,9 +475,18 @@ void SMultiLineEditableText::FinishChangingText()
 
 		// Let outsiders know that the text content has been changed
 		OnTextChanged.ExecuteIfBound(EditedText);
+		
+		OnCursorMoved.ExecuteIfBound(CursorInfo.GetCursorInteractionLocation());
 
 		// Update the desired cursor position, since typing will have moved it
 		UpdatePreferredCursorScreenOffsetInLine();
+
+		// If the marshaller we're using requires live text updates (eg, because it injects formatting into the source text)
+		// then we need to force a SetEditableText here so that it can update the new editable text with any extra formatting
+		if (Marshaller->RequiresLiveUpdate())
+		{
+			ForceRefreshTextLayout(EditedText);
+		}
 	}
 }
 
@@ -540,7 +519,6 @@ void SMultiLineEditableText::BackspaceChar()
 		FTextLocation FinalCursorLocation = CursorInteractionPosition;
 
 		const TArray< FTextLayout::FLineModel >& Lines = TextLayout->GetLineModels();
-		const FTextLayout::FLineModel& Line = Lines[CursorInteractionPosition.GetLineIndex()];
 
 		//If we are at the very beginning of the line...
 		if (CursorInteractionPosition.GetOffset() == 0)
@@ -548,26 +526,17 @@ void SMultiLineEditableText::BackspaceChar()
 			//And if the current line isn't the very first line then...
 			if (CursorInteractionPosition.GetLineIndex() > 0)
 			{
-				int32 CachePreviousLinesCurrentLength = Lines[CursorInteractionPosition.GetLineIndex() - 1].Text->Len();
-				if (TextLayout->JoinLineWithNextLine(CursorInteractionPosition.GetLineIndex() - 1))
+				const int32 PreviousLineIndex = CursorInteractionPosition.GetLineIndex() - 1;
+				const int32 CachePreviousLinesCurrentLength = Lines[PreviousLineIndex].Text->Len();
+				if (TextLayout->JoinLineWithNextLine(PreviousLineIndex))
 				{
 					//Update the cursor so it appears at the end of the previous line,
 					//as we're going to delete the imaginary \n separating them
-					FinalCursorLocation = FTextLocation(CursorInteractionPosition.GetLineIndex() - 1, CachePreviousLinesCurrentLength);
+					FinalCursorLocation = FTextLocation(PreviousLineIndex, CachePreviousLinesCurrentLength);
 				}
 			}
 			//else do nothing as the FinalCursorLocation is already correct
 		}
-		//else if (CursorInteractionPosition.GetOffset() == Line.Text->Len())
-		//{
-		//	//And if the current line isn't the very last line then...
-		//	if (Lines.IsValidIndex(CursorInteractionPosition.GetLineIndex() + 1))
-		//	{
-		//		//do nothing to the cursor as the FinalCursorLocation is already correct
-		//		TextLayout->JoinLineWithNextLine(CursorInteractionPosition.GetLineIndex());
-		//	}
-		//	//else do nothing as the FinalCursorLocation is already correct
-		//}
 		else
 		{
 			// Delete character to the left of the caret
@@ -774,7 +743,7 @@ void SMultiLineEditableText::DeleteSelectedText()
 			{
 				const TSharedRef< FString > EmptyText = MakeShareable(new FString());
 				TArray< TSharedRef< IRun > > Runs;
-				Runs.Add(FSlateTextRun::Create(EmptyText, TextStyle));
+				Runs.Add(FSlateTextRun::Create(FRunInfo(), EmptyText, TextStyle));
 
 				TextLayout->AddLine(EmptyText, Runs);
 			}
@@ -903,6 +872,8 @@ FReply SMultiLineEditableText::MoveCursor( FMoveCursor Args )
 		CursorInfo.SetCursorLocationAndCalculateAlignment(TextLayout, NewCursorPosition);
 	}
 
+	OnCursorMoved.ExecuteIfBound(CursorInfo.GetCursorInteractionLocation());
+
 	if (bUpdatePreferredCursorScreenOffsetInLine)
 	{
 		UpdatePreferredCursorScreenOffsetInLine();
@@ -1014,6 +985,7 @@ void SMultiLineEditableText::UpdateCursorHighlight()
 		}
 	}
 	
+	if ( bHasKeyboardFocus )
 	{
 		// The cursor mode uses the literal position rather than the interaction position
 		const FTextLocation CursorPosition = CursorInfo.GetCursorLocation();
@@ -1074,6 +1046,7 @@ void SMultiLineEditableText::JumpTo(ETextLocation JumpLocation, ECursorAction Ac
 				}
 
 				CursorInfo.SetCursorLocationAndCalculateAlignment(TextLayout, NewCursorPosition);
+				OnCursorMoved.ExecuteIfBound(CursorInfo.GetCursorInteractionLocation());
 				UpdatePreferredCursorScreenOffsetInLine();
 				UpdateCursorHighlight();
 			}
@@ -1095,6 +1068,7 @@ void SMultiLineEditableText::JumpTo(ETextLocation JumpLocation, ECursorAction Ac
 			}
 
 			CursorInfo.SetCursorLocationAndCalculateAlignment(TextLayout, NewCursorPosition);
+			OnCursorMoved.ExecuteIfBound(CursorInfo.GetCursorInteractionLocation());
 			UpdatePreferredCursorScreenOffsetInLine();
 			UpdateCursorHighlight();
 		}
@@ -1123,6 +1097,7 @@ void SMultiLineEditableText::JumpTo(ETextLocation JumpLocation, ECursorAction Ac
 				}
 
 				CursorInfo.SetCursorLocationAndAlignment(NewCursorPosition, ECursorAlignment::Right);
+				OnCursorMoved.ExecuteIfBound(CursorInfo.GetCursorInteractionLocation());
 				UpdatePreferredCursorScreenOffsetInLine();
 				UpdateCursorHighlight();
 			}
@@ -1148,6 +1123,7 @@ void SMultiLineEditableText::JumpTo(ETextLocation JumpLocation, ECursorAction Ac
 				}
 
 				CursorInfo.SetCursorLocationAndCalculateAlignment(TextLayout, NewCursorPosition);
+				OnCursorMoved.ExecuteIfBound(CursorInfo.GetCursorInteractionLocation());
 				UpdatePreferredCursorScreenOffsetInLine();
 				UpdateCursorHighlight();
 			}
@@ -1161,6 +1137,171 @@ void SMultiLineEditableText::ClearSelection()
 	SelectionStart = TOptional<FTextLocation>();
 }
 
+FText SMultiLineEditableText::GetSelectedText() const
+{
+	FText Text;
+
+	if (AnyTextSelected())
+	{
+		const FTextLocation CursorInteractionPosition = CursorInfo.GetCursorInteractionLocation();
+		FTextLocation SelectionLocation = SelectionStart.Get(CursorInteractionPosition);
+		FTextSelection Selection(SelectionLocation, CursorInteractionPosition);
+
+		FString SelectedText;
+		TextLayout->GetSelectionAsText(SelectedText, Selection);
+
+		Text = FText::FromString(SelectedText);
+	}
+
+	return Text;
+}
+
+void SMultiLineEditableText::InsertTextAtCursor(const FText& InText)
+{
+	InsertTextAtCursor(InText.ToString());
+}
+
+void SMultiLineEditableText::InsertTextAtCursor(const FString& InString)
+{
+	StartChangingText();
+
+	DeleteSelectedText();
+	InsertTextAtCursorImpl(InString);
+
+	FinishChangingText();
+}
+
+void SMultiLineEditableText::InsertRunAtCursor(TSharedRef<IRun> InRun)
+{
+	StartChangingText();
+
+	DeleteSelectedText();
+
+	const FTextLocation CursorInteractionPosition = CursorInfo.GetCursorInteractionLocation();
+	TextLayout->InsertAt(CursorInteractionPosition, InRun, true); // true to preserve the run after the insertion point, even if it's empty - this preserves the text formatting correctly
+
+	// Move the cursor along since we've inserted some new text
+	FString InRunText;
+	InRun->AppendTextTo(InRunText);
+
+	const TArray<FTextLayout::FLineModel>& Lines = TextLayout->GetLineModels();
+	const FTextLayout::FLineModel& Line = Lines[CursorInteractionPosition.GetLineIndex()];
+	const FTextLocation FinalCursorLocation = FTextLocation(CursorInteractionPosition.GetLineIndex(), FMath::Min(CursorInteractionPosition.GetOffset() + InRunText.Len(), Line.Text->Len()));
+
+	CursorInfo.SetCursorLocationAndCalculateAlignment(TextLayout, FinalCursorLocation);
+	UpdateCursorHighlight();
+
+	FinishChangingText();
+}
+
+void SMultiLineEditableText::ApplyToSelection(const FRunInfo& InRunInfo, const FTextBlockStyle& InStyle)
+{
+	StartChangingText();
+
+	const FTextLocation CursorInteractionPosition = CursorInfo.GetCursorInteractionLocation();
+	const FTextLocation SelectionLocation = SelectionStart.Get(CursorInteractionPosition);
+	const FTextSelection Selection(SelectionLocation, CursorInteractionPosition);
+
+	const int32 SelectionBeginningLineIndex = Selection.GetBeginning().GetLineIndex();
+	const int32 SelectionBeginningLineOffset = Selection.GetBeginning().GetOffset();
+
+	const int32 SelectionEndLineIndex = Selection.GetEnd().GetLineIndex();
+	const int32 SelectionEndLineOffset = Selection.GetEnd().GetOffset();
+
+	if (SelectionBeginningLineIndex == SelectionEndLineIndex)
+	{
+		TSharedRef<FString> SelectedText = MakeShareable(new FString);
+		TextLayout->GetSelectionAsText(*SelectedText, Selection);
+		
+		TextLayout->RemoveAt(Selection.GetBeginning(), SelectionEndLineOffset - SelectionBeginningLineOffset);
+			
+		TSharedRef<IRun> StyledRun = FSlateTextRun::Create(InRunInfo, SelectedText, InStyle);
+		TextLayout->InsertAt(Selection.GetBeginning(), StyledRun);
+	}
+	else
+	{
+		const TArray< FTextLayout::FLineModel >& Lines = TextLayout->GetLineModels();
+
+		{
+			const FTextLayout::FLineModel& Line = Lines[SelectionBeginningLineIndex];
+
+			const FTextLocation LineStartLocation(SelectionBeginningLineIndex, SelectionBeginningLineOffset);
+			const FTextLocation LineEndLocation(SelectionBeginningLineIndex, Line.Text->Len());
+
+			TSharedRef<FString> SelectedText = MakeShareable(new FString);
+			TextLayout->GetSelectionAsText(*SelectedText, FTextSelection(LineStartLocation, LineEndLocation));
+
+			TextLayout->RemoveAt(LineStartLocation, LineEndLocation.GetOffset() - LineStartLocation.GetOffset());
+			
+			TSharedRef<IRun> StyledRun = FSlateTextRun::Create(InRunInfo, SelectedText, InStyle);
+			TextLayout->InsertAt(LineStartLocation, StyledRun);
+		}
+			
+		for (int32 LineIndex = SelectionBeginningLineIndex + 1; LineIndex < SelectionEndLineIndex; ++LineIndex)
+		{
+			const FTextLayout::FLineModel& Line = Lines[LineIndex];
+				
+			const FTextLocation LineStartLocation(LineIndex, 0);
+			const FTextLocation LineEndLocation(LineIndex, Line.Text->Len());
+
+			TSharedRef<FString> SelectedText = MakeShareable(new FString);
+			TextLayout->GetSelectionAsText(*SelectedText, FTextSelection(LineStartLocation, LineEndLocation));
+
+			TextLayout->RemoveAt(LineStartLocation, LineEndLocation.GetOffset() - LineStartLocation.GetOffset());
+			
+			TSharedRef<IRun> StyledRun = FSlateTextRun::Create(InRunInfo, SelectedText, InStyle);
+			TextLayout->InsertAt(LineStartLocation, StyledRun);
+		}
+
+		{
+			const FTextLayout::FLineModel& Line = Lines[SelectionEndLineIndex];
+
+			const FTextLocation LineStartLocation(SelectionEndLineIndex, 0);
+			const FTextLocation LineEndLocation(SelectionEndLineIndex, SelectionEndLineOffset);
+
+			TSharedRef<FString> SelectedText = MakeShareable(new FString);
+			TextLayout->GetSelectionAsText(*SelectedText, FTextSelection(LineStartLocation, LineEndLocation));
+
+			TextLayout->RemoveAt(LineStartLocation, LineEndLocation.GetOffset() - LineStartLocation.GetOffset());
+			
+			TSharedRef<IRun> StyledRun = FSlateTextRun::Create(InRunInfo, SelectedText, InStyle);
+			TextLayout->InsertAt(LineStartLocation, StyledRun);
+		}
+	}
+
+	SelectionStart = SelectionLocation;
+	CursorInfo.SetCursorLocationAndCalculateAlignment(TextLayout, CursorInteractionPosition);
+
+	UpdatePreferredCursorScreenOffsetInLine();
+	UpdateCursorHighlight();
+
+	FinishChangingText();
+}
+
+TSharedPtr<const IRun> SMultiLineEditableText::GetRunUnderCursor() const
+{
+	const TArray<FTextLayout::FLineModel>& Lines = TextLayout->GetLineModels();
+	
+	const FTextLocation CursorInteractionLocation = CursorInfo.GetCursorInteractionLocation();
+	if (Lines.IsValidIndex(CursorInteractionLocation.GetLineIndex()))
+	{
+		const FTextLayout::FLineModel& LineModel = Lines[CursorInteractionLocation.GetLineIndex()];
+		for (int32 RunIndex = 0; RunIndex < LineModel.Runs.Num(); ++RunIndex)
+		{
+			const FTextLayout::FRunModel& RunModel = LineModel.Runs[RunIndex];
+			const FTextRange RunRange = RunModel.GetTextRange();
+
+			const bool bIsLastRun = RunIndex == LineModel.Runs.Num() - 1;
+			if (RunRange.Contains(CursorInteractionLocation.GetOffset()) || bIsLastRun)
+			{
+				return RunModel.GetRun();
+			}
+		}
+	}
+
+	return nullptr;
+}
+
 bool SMultiLineEditableText::CanExecuteSelectAll() const
 {
 	bool bCanExecute = true;
@@ -1440,20 +1581,19 @@ void SMultiLineEditableText::PasteTextFromClipboard()
 	{
 		StartChangingText(); 
 
-		// Delete currently selected text, if any
 		DeleteSelectedText();
 
 		// Paste from the clipboard
 		FString PastedText;
 		FPlatformMisc::ClipboardPaste(PastedText);
 
-		InsertTextAtCursor(PastedText);
+		InsertTextAtCursorImpl(PastedText);
 
 		FinishChangingText();
 	}
 }
 
-void SMultiLineEditableText::InsertTextAtCursor(const FString& InString)
+void SMultiLineEditableText::InsertTextAtCursorImpl(const FString& InString)
 {
 	if (InString.Len())
 	{
@@ -1521,6 +1661,8 @@ void SMultiLineEditableText::Undo()
 				CursorInfo = UndoState.CursorInfo.CreateUndo();
 				SelectionStart = UndoState.SelectionStart;
 
+				OnCursorMoved.ExecuteIfBound(CursorInfo.GetCursorInteractionLocation());
+
 				UpdateCursorHighlight();
 			}
 
@@ -1554,6 +1696,8 @@ void SMultiLineEditableText::Redo()
 				CursorInfo.RestoreFromUndo(UndoState.CursorInfo);
 				SelectionStart = UndoState.SelectionStart;
 
+				OnCursorMoved.ExecuteIfBound(CursorInfo.GetCursorInteractionLocation());
+
 				UpdateCursorHighlight();
 			}
 
@@ -1694,21 +1838,29 @@ void SMultiLineEditableText::Tick( const FGeometry& AllottedGeometry, const doub
 
 	const bool bShouldAppearFocused = HasKeyboardFocus() || ContextMenuWindow.IsValid();
 
-	if (!bShouldAppearFocused && BoundText.IsBound())
+	if (!bShouldAppearFocused)
 	{
+		bool bHasSetText = false;
+
 		const FText& TextToSet = BoundText.Get(FText::GetEmpty());
-		if (!BoundTextLastTick.IdenticalTo(TextToSet))
+		if (BoundText.IsBound() && !BoundTextLastTick.IdenticalTo(TextToSet))
 		{
 			// The pointer used by the bound text has changed, however the text may still be the same - check that now
 			if (!BoundTextLastTick.ToString().Equals(TextToSet.ToString(), ESearchCase::CaseSensitive))
 			{
 				// The source text has changed, so update the internal editable text
+				bHasSetText = true;
 				SetEditableText(TextToSet);
 			}
 
 			// Update this even if the text is lexically identical, as it will update the pointer compared by IdenticalTo for the next Tick
 			BoundTextLastTick = TextToSet;
 		}
+
+		if (!bHasSetText && Marshaller->IsDirty())
+		{
+			ForceRefreshTextLayout(TextToSet);
+		}
 	}
 }
 
@@ -2096,6 +2248,7 @@ void SMultiLineEditableText::FTextInputMethodContext::SetSelectionRange(const ui
 			break;
 		}
 
+		OwningWidgetPtr->OnCursorMoved.ExecuteIfBound(OwningWidgetPtr->CursorInfo.GetCursorInteractionLocation());
 		OwningWidgetPtr->UpdateCursorHighlight();
 	}
 }
@@ -2122,7 +2275,7 @@ void SMultiLineEditableText::FTextInputMethodContext::SetTextInRange(const uint3
 		// We do this as a select, delete, and insert as it's the simplest way to keep the text layout correct
 		SetSelectionRange(BeginIndex, Length, ITextInputMethodContext::ECaretPosition::Beginning);
 		OwningWidgetPtr->DeleteSelectedText();
-		OwningWidgetPtr->InsertTextAtCursor(InString);
+		OwningWidgetPtr->InsertTextAtCursorImpl(InString);
 
 		// Has the text changed?
 		const FText EditedText = OwningWidgetPtr->GetEditableText();
diff --git a/Engine/Source/Runtime/Slate/Private/Widgets/Text/SRichTextBlock.cpp b/Engine/Source/Runtime/Slate/Private/Widgets/Text/SRichTextBlock.cpp
--- a/Engine/Source/Runtime/Slate/Private/Widgets/Text/SRichTextBlock.cpp
+++ b/Engine/Source/Runtime/Slate/Private/Widgets/Text/SRichTextBlock.cpp
@@ -13,27 +13,31 @@
 
 void SRichTextBlock::Construct( const FArguments& InArgs )
 {
+	TextStyle = *InArgs._TextStyle;
 	WrapTextAt = InArgs._WrapTextAt;
 	AutoWrapText = InArgs._AutoWrapText;
 	Margin = InArgs._Margin;
 	LineHeightPercentage = InArgs._LineHeightPercentage;
 	Justification = InArgs._Justification;
-	TagStyleSet = InArgs._DecoratorStyleSet;
-	TextStyle = *InArgs._TextStyle;
 
 	CachedAutoWrapTextWidth = 0.0f;
 
-	Parser = InArgs._Parser;
-	if ( !Parser.IsValid() )
 	{
-		Parser = FRichTextMarkupProcessing::Create();
+		TSharedPtr<IRichTextMarkupParser> Parser = InArgs._Parser;
+		if ( !Parser.IsValid() )
+		{
+			Parser = FDefaultRichTextMarkupParser::Create();
+		}
+
+		Marshaller = FRichTextLayoutMarshaller::Create(Parser, nullptr, InArgs._Decorators, InArgs._DecoratorStyleSet, TextStyle);
+		for ( const TSharedRef< ITextDecorator >& Decorator : InArgs.InlineDecorators )
+		{
+			Marshaller->AppendInlineDecorator( Decorator );
+		}
 	}
 
 	TextHighlighter = FSlateTextHighlightRunRenderer::Create();
 
-	Decorators.Append( InArgs._Decorators );
-	Decorators.Append( InArgs.InlineDecorators );
-
 	TextLayout = FSlateTextLayout::Create();
 
 	SetText( InArgs._Text );
@@ -49,21 +53,24 @@ void SRichTextBlock::Tick( const FGeometry& AllottedGeometry, const double InCur
 {
 	TextLayout->SetScale( AllottedGeometry.Scale );
 
-	if (BoundText.IsBound())
+	bool bRequiresTextUpdate = false;
+	const FText& TextToSet = GetText();
+	if (BoundText.IsBound() && !BoundTextLastTick.IdenticalTo(TextToSet))
 	{
-		const FText& TextToSet = GetText();
-		if (!BoundTextLastTick.IdenticalTo(TextToSet))
+		// The pointer used by the bound text has changed, however the text may still be the same - check that now
+		if (!BoundTextLastTick.ToString().Equals(TextToSet.ToString(), ESearchCase::CaseSensitive))
 		{
-			// The pointer used by the bound text has changed, however the text may still be the same - check that now
-			if (!BoundTextLastTick.ToString().Equals(TextToSet.ToString(), ESearchCase::CaseSensitive))
-			{
-				// The source text has changed, so update the internal editable text
-				SetText(TextToSet);
-			}
-
-			// Update this even if the text is lexically identical, as it will update the pointer compared by IdenticalTo for the next Tick
-			BoundTextLastTick = TextToSet;
+			// The source text has changed, so update the internal editable text
+			bRequiresTextUpdate = true;
 		}
+
+		// Update this even if the text is lexically identical, as it will update the pointer compared by IdenticalTo for the next Tick
+		BoundTextLastTick = TextToSet;
+	}
+
+	if (bRequiresTextUpdate || Marshaller->IsDirty())
+	{
+		SetText(TextToSet);
 	}
 }
 
@@ -155,22 +162,6 @@ void SRichTextBlock::OnArrangeChildren( const FGeometry& AllottedGeometry, FArra
 	}
 }
 
-TSharedPtr< ITextDecorator > SRichTextBlock::TryGetDecorator( const TArray< TSharedRef< ITextDecorator > >& InDecorators, const FString& InText, const FTextRunParseResults& TextRun ) const
-{
-	if (InDecorators.Num() > 0)
-	{
-		for (auto DecoratorIter = InDecorators.CreateConstIterator(); DecoratorIter; ++DecoratorIter)
-		{
-			if ((*DecoratorIter)->Supports(TextRun, InText))
-			{
-				return *DecoratorIter;
-			}
-		}
-	}
-
-	return nullptr;
-}
-
 void SRichTextBlock::SetText( const TAttribute<FText>& InTextAttr )
 {
 	BoundText = InTextAttr;
@@ -183,59 +174,9 @@ void SRichTextBlock::SetText( const TAttribute<FText>& InTextAttr )
 		BoundTextLastTick = InText;
 	}
 
-	const FString& InString = InText.ToString();
-	TArray<FTextLineParseResults> LineParseResultsArray;
-	FString ProcessedString;
-	Parser->Process(LineParseResultsArray, InString, ProcessedString);
-
 	TextLayout->ClearLines();
 
-	// Iterate through parsed line results and create processed lines with runs.
-	for (int LineIndex = 0; LineIndex < LineParseResultsArray.Num(); LineIndex++)
-	{
-		const FTextLineParseResults& LineParseResults = LineParseResultsArray[ LineIndex ];
-
-		TSharedRef<FString> ModelString = MakeShareable(new FString());
-		TArray< TSharedRef< IRun > > Runs;
-
-		for (int RunIndex = 0; RunIndex < LineParseResults.Runs.Num(); RunIndex++)
-		{
-			const FTextRunParseResults& RunParseResult = LineParseResults.Runs[ RunIndex ];
-			TSharedPtr< ISlateRun > Run;
-
-			TSharedPtr< ITextDecorator > Decorator = TryGetDecorator( Decorators, ProcessedString, RunParseResult );
-
-			if ( Decorator.IsValid() )
-			{
-				// Create run and update model string.
-				Run = Decorator->Create( RunParseResult, ProcessedString, ModelString, TagStyleSet );
-			}
-			else
-			{
-				const FTextBlockStyle* TextBlockStyle;
-				FTextRange ModelRange;
-				ModelRange.BeginIndex = ModelString->Len();
-				if(!(RunParseResult.Name.IsEmpty()) && TagStyleSet->HasWidgetStyle< FTextBlockStyle >( FName(*RunParseResult.Name) ))
-				{
-					*ModelString += ProcessedString.Mid(RunParseResult.ContentRange.BeginIndex, RunParseResult.ContentRange.EndIndex - RunParseResult.ContentRange.BeginIndex);
-					TextBlockStyle = &(TagStyleSet->GetWidgetStyle< FTextBlockStyle >( FName(*RunParseResult.Name) ));
-				}
-				else
-				{
-					*ModelString += ProcessedString.Mid(RunParseResult.OriginalRange.BeginIndex, RunParseResult.OriginalRange.EndIndex - RunParseResult.OriginalRange.BeginIndex);
-					TextBlockStyle = &(TextStyle);
-				}
-				ModelRange.EndIndex = ModelString->Len();
-
-				// Create run.
-				Run = FSlateTextRun::Create( ModelString, *TextBlockStyle, ModelRange );
-			}
-
-			Runs.Add( Run.ToSharedRef() );
-		}
-
-		TextLayout->AddLine( ModelString, Runs );
-	}
+	Marshaller->SetText(InText.ToString(), *TextLayout);
 }
 
 void SRichTextBlock::SetHighlightText( const FText& InHighlightText )
@@ -245,15 +186,15 @@ void SRichTextBlock::SetHighlightText( const FText& InHighlightText )
 		return;
 	}
 
-	TextHighlights.Empty();
 	HighlightText = InHighlightText;
 
-	const FString HighlightTextString = HighlightText.ToString();
+	const FString& HighlightTextString = HighlightText.ToString();
 	const int32 HighlightTextLength = HighlightTextString.Len();
 
 	const TArray< FTextLayout::FLineModel >& LineModels = TextLayout->GetLineModels();
 
-	for (int LineIndex = 0; LineIndex < LineModels.Num(); LineIndex++)
+	TArray< FTextRunRenderer > TextHighlights;
+	for (int32 LineIndex = 0; LineIndex < LineModels.Num(); LineIndex++)
 	{
 		const FTextLayout::FLineModel& LineModel = LineModels[ LineIndex ];
 
diff --git a/Engine/Source/Runtime/Slate/Public/Framework/Text/BaseTextLayoutMarshaller.h b/Engine/Source/Runtime/Slate/Public/Framework/Text/BaseTextLayoutMarshaller.h
new file mode 100644
--- /dev/null
+++ b/Engine/Source/Runtime/Slate/Public/Framework/Text/BaseTextLayoutMarshaller.h
@@ -0,0 +1,54 @@
+// Copyright 1998-2014 Epic Games, Inc. All Rights Reserved.
+#pragma once
+
+#if WITH_FANCY_TEXT
+
+#include "ITextLayoutMarshaller.h"
+
+/**
+ * Base class implementing some common functionality for all text marshallers
+ */
+class SLATE_API FBaseTextLayoutMarshaller : public ITextLayoutMarshaller
+{
+public:
+
+	// BEGIN ITextLayoutMarshaller interface
+	virtual bool RequiresLiveUpdate() const override
+	{
+		return false;
+	}
+
+	virtual void MakeDirty() override
+	{
+		bIsDirty = true;
+	}
+
+	virtual void ClearDirty() override
+	{
+		bIsDirty = false;
+	}
+
+	virtual bool IsDirty() const override
+	{
+		return bIsDirty;
+	}
+	// END ITextLayoutMarshaller interface
+
+protected:
+
+	FBaseTextLayoutMarshaller()
+		: bIsDirty(false)
+	{
+	}
+
+	virtual ~FBaseTextLayoutMarshaller()
+	{
+	}
+
+private:
+
+	bool bIsDirty;
+
+};
+
+#endif //WITH_FANCY_TEXT
diff --git a/Engine/Source/Runtime/Slate/Public/Framework/Text/IRichTextMarkupWriter.h b/Engine/Source/Runtime/Slate/Public/Framework/Text/IRichTextMarkupWriter.h
new file mode 100644
--- /dev/null
+++ b/Engine/Source/Runtime/Slate/Public/Framework/Text/IRichTextMarkupWriter.h
@@ -0,0 +1,43 @@
+// Copyright 1998-2014 Epic Games, Inc. All Rights Reserved.
+#pragma once
+
+#if WITH_FANCY_TEXT
+
+/**
+ */
+class SLATE_API IRichTextMarkupWriter
+{
+public:
+
+	/** A single run in a rich-text line, contains the meta-information that was parsed out of it, as well as the lexical contents of the run */
+	struct FRichTextRun
+	{
+		FRichTextRun(const FRunInfo& InInfo, FString InText)
+			: Info(InInfo)
+			, Text(MoveTemp(InText))
+		{
+		}
+
+		FRunInfo Info;
+		FString Text;
+	};
+
+	/** A single line in a rich-text document, comprising of a number of styled runs */
+	struct FRichTextLine
+	{
+		TArray<FRichTextRun> Runs;
+	};
+
+	/**
+	 * Virtual destructor
+	 */
+	virtual ~IRichTextMarkupWriter() {}
+
+	/**
+	 * Write the provided array of line and run info, producing an output string containing the markup needed to persist the run layouts
+	 */
+	virtual void Write(const TArray<FRichTextLine>& InLines, FString& Output) = 0;
+
+};
+
+#endif //WITH_FANCY_TEXT
diff --git a/Engine/Source/Runtime/Slate/Public/Framework/Text/IRun.h b/Engine/Source/Runtime/Slate/Public/Framework/Text/IRun.h
--- a/Engine/Source/Runtime/Slate/Public/Framework/Text/IRun.h
+++ b/Engine/Source/Runtime/Slate/Public/Framework/Text/IRun.h
@@ -1,6 +1,25 @@
 // Copyright 1998-2014 Epic Games, Inc. All Rights Reserved.
 #pragma once
 
+struct SLATE_API FRunInfo
+{
+	FRunInfo()
+		: Name()
+		, MetaData()
+	{
+	}
+
+	FRunInfo( FString InName )
+		: Name( MoveTemp(InName) )
+		, MetaData()
+	{
+
+	}
+
+	FString Name;
+	TMap< FString, FString > MetaData;
+};
+
 class SLATE_API IRun
 {
 public:
@@ -26,7 +45,9 @@ public:
 	virtual void Move(const TSharedRef<FString>& NewText, const FTextRange& NewRange) = 0;
 	virtual TSharedRef<IRun> Clone() const = 0;
 
-	virtual void AppendText(FString& Text) const = 0;
-	virtual void AppendText(FString& Text, const FTextRange& Range) const = 0;
+	virtual void AppendTextTo(FString& Text) const = 0;
+	virtual void AppendTextTo(FString& Text, const FTextRange& Range) const = 0;
+
+	virtual const FRunInfo& GetRunInfo() const = 0;
 
 };
diff --git a/Engine/Source/Runtime/Slate/Public/Framework/Text/ITextDecorator.h b/Engine/Source/Runtime/Slate/Public/Framework/Text/ITextDecorator.h
--- a/Engine/Source/Runtime/Slate/Public/Framework/Text/ITextDecorator.h
+++ b/Engine/Source/Runtime/Slate/Public/Framework/Text/ITextDecorator.h
@@ -50,19 +50,16 @@ public:
 	TArray< FTextRunParseResults > Runs;
 };
 
-struct SLATE_API FTextRunInfo
+struct SLATE_API FTextRunInfo : FRunInfo
 {
 	FTextRunInfo( FString InName, const FText& InContent )
-		: Name( InName )
+		: FRunInfo( MoveTemp(InName) )
 		, Content( InContent )
-		, MetaData()
 	{
 
 	}
 
-	FString Name;
 	FText Content;
-	TMap< FString, FString > MetaData;
 };
 
 class SLATE_API ITextDecorator
diff --git a/Engine/Source/Runtime/Slate/Public/Framework/Text/ITextLayoutMarshaller.h b/Engine/Source/Runtime/Slate/Public/Framework/Text/ITextLayoutMarshaller.h
new file mode 100644
--- /dev/null
+++ b/Engine/Source/Runtime/Slate/Public/Framework/Text/ITextLayoutMarshaller.h
@@ -0,0 +1,51 @@
+// Copyright 1998-2014 Epic Games, Inc. All Rights Reserved.
+#pragma once
+
+#if WITH_FANCY_TEXT
+
+class FTextLayout;
+
+/**
+ * Interface used to get/set the raw text to/from a text layout
+ */
+class SLATE_API ITextLayoutMarshaller
+{
+public:
+
+	virtual ~ITextLayoutMarshaller() {}
+
+	/**
+	 * Populate the text layout from the string
+	 */
+	virtual void SetText(const FString& SourceString, FTextLayout& TargetTextLayout) = 0;
+
+	/**
+	 * Populate the string from the text layout
+	 */
+	virtual void GetText(FString& TargetString, const FTextLayout& SourceTextLayout) = 0;
+
+	/**
+	 * Return true if the marshaller requires the text be updated "live" (eg, because it inserts formatting directly into the source text)
+	 * Returning true will cause SetText to be called every time the source text is changed, which is costly, but required for things like syntax highlighting
+	 */
+	virtual bool RequiresLiveUpdate() const = 0;
+
+	/**
+	 * Mark this text layout as dirty (eg, because some settings have changed), so that it will cause SetText to be called on the next Tick
+	 */
+	virtual void MakeDirty() = 0;
+
+	/**
+	 * Mark this text layout as clean once it's been updated
+	 */
+	virtual void ClearDirty() = 0;
+
+	/**
+	 * Is this text layout dirty? (eg, because some settings have changed)
+	 * If so, you should call SetText on it to update the text layout
+	 */
+	virtual bool IsDirty() const = 0;
+
+};
+
+#endif //WITH_FANCY_TEXT
diff --git a/Engine/Source/Runtime/Slate/Public/Framework/Text/PlainTextLayoutMarshaller.h b/Engine/Source/Runtime/Slate/Public/Framework/Text/PlainTextLayoutMarshaller.h
new file mode 100644
--- /dev/null
+++ b/Engine/Source/Runtime/Slate/Public/Framework/Text/PlainTextLayoutMarshaller.h
@@ -0,0 +1,35 @@
+// Copyright 1998-2014 Epic Games, Inc. All Rights Reserved.
+#pragma once
+
+#if WITH_FANCY_TEXT
+
+#include "BaseTextLayoutMarshaller.h"
+
+/**
+ * Get/set the raw text to/from a text layout as plain text
+ */
+class SLATE_API FPlainTextLayoutMarshaller : public FBaseTextLayoutMarshaller
+{
+public:
+
+	static TSharedRef< FPlainTextLayoutMarshaller > Create(FTextBlockStyle InDefaultTextStyle);
+
+	virtual ~FPlainTextLayoutMarshaller();
+	
+	// ITextLayoutMarshaller
+	virtual void SetText(const FString& SourceString, FTextLayout& TargetTextLayout) override;
+	virtual void GetText(FString& TargetString, const FTextLayout& SourceTextLayout) override;
+
+	void SetDefaultTextStyle(FTextBlockStyle InDefaultTextStyle);
+	const FTextBlockStyle& GetDefaultTextStyle() const;
+
+protected:
+
+	FPlainTextLayoutMarshaller(FTextBlockStyle InDefaultTextStyle);
+
+	/** Default style used by the TextLayout */
+	FTextBlockStyle DefaultTextStyle;
+
+};
+
+#endif //WITH_FANCY_TEXT
diff --git a/Engine/Source/Runtime/Slate/Public/Framework/Text/RichTextLayoutMarshaller.h b/Engine/Source/Runtime/Slate/Public/Framework/Text/RichTextLayoutMarshaller.h
new file mode 100644
--- /dev/null
+++ b/Engine/Source/Runtime/Slate/Public/Framework/Text/RichTextLayoutMarshaller.h
@@ -0,0 +1,65 @@
+// Copyright 1998-2014 Epic Games, Inc. All Rights Reserved.
+#pragma once
+
+#if WITH_FANCY_TEXT
+
+#include "BaseTextLayoutMarshaller.h"
+#include "ITextDecorator.h"
+#include "IRichTextMarkupParser.h"
+#include "IRichTextMarkupWriter.h"
+
+/**
+ * Get/set the raw text to/from a text layout as rich text
+ */
+class SLATE_API FRichTextLayoutMarshaller : public FBaseTextLayoutMarshaller
+{
+public:
+
+	static TSharedRef< FRichTextLayoutMarshaller > Create(TArray< TSharedRef< ITextDecorator > > InDecorators, const ISlateStyle* const InDecoratorStyleSet, FTextBlockStyle InDefaultTextStyle);
+	static TSharedRef< FRichTextLayoutMarshaller > Create(TSharedPtr< IRichTextMarkupParser > InParser, TSharedPtr< IRichTextMarkupWriter > InWriter, TArray< TSharedRef< ITextDecorator > > InDecorators, const ISlateStyle* const InDecoratorStyleSet, FTextBlockStyle InDefaultTextStyle);
+
+	virtual ~FRichTextLayoutMarshaller();
+
+	// ITextLayoutMarshaller
+	virtual void SetText(const FString& SourceString, FTextLayout& TargetTextLayout) override;
+	virtual void GetText(FString& TargetString, const FTextLayout& SourceTextLayout) override;
+
+	void SetDefaultTextStyle(FTextBlockStyle InDefaultTextStyle);
+	const FTextBlockStyle& GetDefaultTextStyle() const;
+
+	/**
+	 * Append an inline decorator to this marshaller
+	 */
+	inline FRichTextLayoutMarshaller& AppendInlineDecorator(const TSharedRef< ITextDecorator >& DecoratorToAdd)
+	{
+		InlineDecorators.Add(DecoratorToAdd);
+		return *this;
+	}
+
+protected:
+
+	FRichTextLayoutMarshaller(TSharedPtr< IRichTextMarkupParser > InParser, TSharedPtr< IRichTextMarkupWriter > InWriter, TArray< TSharedRef< ITextDecorator > > InDecorators, const ISlateStyle* const InDecoratorStyleSet, FTextBlockStyle InDefaultTextStyle);
+
+	TSharedPtr< ITextDecorator > TryGetDecorator(const FString& Line, const FTextRunParseResults& TextRun) const;
+
+	/** The parser used to resolve any markup used in the provided string. */
+	TSharedPtr< IRichTextMarkupParser > Parser;
+
+	/** The writer used to recreate any markup needed to preserve the styled runs. */
+	TSharedPtr< IRichTextMarkupWriter > Writer;
+
+	/** Any decorators that should be used while parsing the text. */
+	TArray< TSharedRef< ITextDecorator > > Decorators;
+
+	/** Additional decorators can be appended inline. Inline decorators get precedence over decorators not specified inline. */
+	TArray< TSharedRef< ITextDecorator > > InlineDecorators;
+
+	/** The style set used for looking up styles used by decorators */
+	const ISlateStyle* DecoratorStyleSet;
+
+	/** Default style used by the TextLayout */
+	FTextBlockStyle DefaultTextStyle;
+
+};
+
+#endif //WITH_FANCY_TEXT
diff --git a/Engine/Source/Runtime/Slate/Public/Framework/Text/SlateHyperlinkRun.h b/Engine/Source/Runtime/Slate/Public/Framework/Text/SlateHyperlinkRun.h
--- a/Engine/Source/Runtime/Slate/Public/Framework/Text/SlateHyperlinkRun.h
+++ b/Engine/Source/Runtime/Slate/Public/Framework/Text/SlateHyperlinkRun.h
@@ -30,9 +30,9 @@ public:
 
 public:
 
-	static TSharedRef< FSlateHyperlinkRun > Create( const TSharedRef< const FString >& InText, const FHyperlinkStyle& InStyle, const FMetadata& Metadata, FOnClick NavigateDelegate );
+	static TSharedRef< FSlateHyperlinkRun > Create( const FRunInfo& InRunInfo, const TSharedRef< const FString >& InText, const FHyperlinkStyle& InStyle, FOnClick NavigateDelegate );
 																																	 
-	static TSharedRef< FSlateHyperlinkRun > Create( const TSharedRef< const FString >& InText, const FHyperlinkStyle& InStyle, const FMetadata& Metadata, FOnClick NavigateDelegate, const FTextRange& InRange );
+	static TSharedRef< FSlateHyperlinkRun > Create( const FRunInfo& InRunInfo, const TSharedRef< const FString >& InText, const FHyperlinkStyle& InStyle, FOnClick NavigateDelegate, const FTextRange& InRange );
 
 public:
 
@@ -68,14 +68,16 @@ public:
 	virtual void Move(const TSharedRef<FString>& NewText, const FTextRange& NewRange) override;
 	virtual TSharedRef<IRun> Clone() const override;
 
-	virtual void AppendText(FString& AppendToText) const override;
-	virtual void AppendText(FString& AppendToText, const FTextRange& PartialRange) const override;
+	virtual void AppendTextTo(FString& AppendToText) const override;
+	virtual void AppendTextTo(FString& AppendToText, const FTextRange& PartialRange) const override;
+
+	virtual const FRunInfo& GetRunInfo() const override;
 
 private:
 
-	FSlateHyperlinkRun( const TSharedRef< const FString >& InText, const FHyperlinkStyle& InStyle, const FMetadata& InMetadata, FOnClick InNavigateDelegate );
+	FSlateHyperlinkRun( const FRunInfo& InRunInfo, const TSharedRef< const FString >& InText, const FHyperlinkStyle& InStyle, FOnClick InNavigateDelegate );
 																										 
-	FSlateHyperlinkRun( const TSharedRef< const FString >& InText, const FHyperlinkStyle& InStyle, const FMetadata& InMetadata, FOnClick InNavigateDelegate, const FTextRange& InRange );
+	FSlateHyperlinkRun( const FRunInfo& InRunInfo, const TSharedRef< const FString >& InText, const FHyperlinkStyle& InStyle, FOnClick InNavigateDelegate, const FTextRange& InRange );
 
 	FSlateHyperlinkRun( const FSlateHyperlinkRun& Run );
 
@@ -85,10 +87,10 @@ private:
 
 private:
 
+	FRunInfo RunInfo;
 	TSharedRef< const FString > Text;
 	FTextRange Range;
 	FHyperlinkStyle Style;
-	TMap<FString,FString> Metadata;
 	FOnClick NavigateDelegate;
 
 	TSharedRef< FWidgetViewModel > ViewModel;
diff --git a/Engine/Source/Runtime/Slate/Public/Framework/Text/SlateImageRun.h b/Engine/Source/Runtime/Slate/Public/Framework/Text/SlateImageRun.h
--- a/Engine/Source/Runtime/Slate/Public/Framework/Text/SlateImageRun.h
+++ b/Engine/Source/Runtime/Slate/Public/Framework/Text/SlateImageRun.h
@@ -7,8 +7,8 @@ class SLATE_API FSlateImageRun : public ISlateRun, public TSharedFromThis< FSlat
 {
 public:
 
-	static TSharedRef< FSlateImageRun > Create( const TSharedRef< const FString >& InText, const FSlateBrush* InImage, int16 InBaseline );
-	static TSharedRef< FSlateImageRun > Create( const TSharedRef< const FString >& InText, const FSlateBrush* InImage, int16 InBaseline, const FTextRange& InRange );
+	static TSharedRef< FSlateImageRun > Create( const FRunInfo& InRunInfo, const TSharedRef< const FString >& InText, const FSlateBrush* InImage, int16 InBaseline );
+	static TSharedRef< FSlateImageRun > Create( const FRunInfo& InRunInfo, const TSharedRef< const FString >& InText, const FSlateBrush* InImage, int16 InBaseline, const FTextRange& InRange );
 
 public:
 
@@ -43,17 +43,20 @@ public:
 	virtual void Move(const TSharedRef<FString>& NewText, const FTextRange& NewRange) override;
 	virtual TSharedRef<IRun> Clone() const override;
 
-	virtual void AppendText(FString& Text) const override;
-	virtual void AppendText(FString& AppendToText, const FTextRange& PartialRange) const override;
+	virtual void AppendTextTo(FString& Text) const override;
+	virtual void AppendTextTo(FString& AppendToText, const FTextRange& PartialRange) const override;
+
+	virtual const FRunInfo& GetRunInfo() const override;
 
 private:
 
-	FSlateImageRun( const TSharedRef< const FString >& InText, const FSlateBrush* InImage, int16 InBaseline );
+	FSlateImageRun( const FRunInfo& InRunInfo, const TSharedRef< const FString >& InText, const FSlateBrush* InImage, int16 InBaseline );
 
-	FSlateImageRun( const TSharedRef< const FString >& InText, const FSlateBrush* InImage, int16 InBaseline, const FTextRange& InRange );
+	FSlateImageRun( const FRunInfo& InRunInfo, const TSharedRef< const FString >& InText, const FSlateBrush* InImage, int16 InBaseline, const FTextRange& InRange );
 
 private:
 
+	FRunInfo RunInfo;
 	TSharedRef< const FString > Text;
 	FTextRange Range;
 	const FSlateBrush* Image;
diff --git a/Engine/Source/Runtime/Slate/Public/Framework/Text/SlateTextRun.h b/Engine/Source/Runtime/Slate/Public/Framework/Text/SlateTextRun.h
--- a/Engine/Source/Runtime/Slate/Public/Framework/Text/SlateTextRun.h
+++ b/Engine/Source/Runtime/Slate/Public/Framework/Text/SlateTextRun.h
@@ -7,9 +7,9 @@ class SLATE_API FSlateTextRun : public ISlateRun, public TSharedFromThis< FSlate
 {
 public:
 
-	static TSharedRef< FSlateTextRun > Create( const TSharedRef< const FString >& InText, const FTextBlockStyle& Style );
+	static TSharedRef< FSlateTextRun > Create( const FRunInfo& InRunInfo, const TSharedRef< const FString >& InText, const FTextBlockStyle& Style );
 
-	static TSharedRef< FSlateTextRun > Create( const TSharedRef< const FString >& InText, const FTextBlockStyle& Style, const FTextRange& InRange );
+	static TSharedRef< FSlateTextRun > Create( const FRunInfo& InRunInfo, const TSharedRef< const FString >& InText, const FTextBlockStyle& Style, const FTextRange& InRange );
 
 public:
 
@@ -41,19 +41,22 @@ public:
 	virtual void Move(const TSharedRef<FString>& NewText, const FTextRange& NewRange) override;
 	virtual TSharedRef<IRun> Clone() const override;
 
-	virtual void AppendText(FString& Text) const override;
-	virtual void AppendText(FString& AppendToText, const FTextRange& PartialRange) const override;
+	virtual void AppendTextTo(FString& Text) const override;
+	virtual void AppendTextTo(FString& AppendToText, const FTextRange& PartialRange) const override;
+
+	virtual const FRunInfo& GetRunInfo() const override;
 
 private:
 
-	FSlateTextRun( const TSharedRef< const FString >& InText, const FTextBlockStyle& InStyle );
+	FSlateTextRun( const FRunInfo& InRunInfo, const TSharedRef< const FString >& InText, const FTextBlockStyle& InStyle );
 
-	FSlateTextRun( const TSharedRef< const FString >& InText, const FTextBlockStyle& InStyle, const FTextRange& InRange );
+	FSlateTextRun( const FRunInfo& InRunInfo, const TSharedRef< const FString >& InText, const FTextBlockStyle& InStyle, const FTextRange& InRange );
 
 	FSlateTextRun( const FSlateTextRun& Run );
 
 private:
 
+	FRunInfo RunInfo;
 	TSharedRef< const FString > Text;
 	FTextBlockStyle Style;
 	FTextRange Range;
diff --git a/Engine/Source/Runtime/Slate/Public/Framework/Text/SlateWidgetRun.h b/Engine/Source/Runtime/Slate/Public/Framework/Text/SlateWidgetRun.h
--- a/Engine/Source/Runtime/Slate/Public/Framework/Text/SlateWidgetRun.h
+++ b/Engine/Source/Runtime/Slate/Public/Framework/Text/SlateWidgetRun.h
@@ -30,9 +30,9 @@ public:
 		}
 	};
 
-	static TSharedRef< FSlateWidgetRun > Create( const TSharedRef< const FString >& InText, const FWidgetRunInfo& InWidgetInfo );
+	static TSharedRef< FSlateWidgetRun > Create( const FRunInfo& InRunInfo, const TSharedRef< const FString >& InText, const FWidgetRunInfo& InWidgetInfo );
 
-	static TSharedRef< FSlateWidgetRun > Create( const TSharedRef< const FString >& InText, const FWidgetRunInfo& InWidgetInfo, const FTextRange& InRange );
+	static TSharedRef< FSlateWidgetRun > Create( const FRunInfo& InRunInfo, const TSharedRef< const FString >& InText, const FWidgetRunInfo& InWidgetInfo, const FTextRange& InRange );
 
 	virtual ~FSlateWidgetRun() {}
 
@@ -63,19 +63,22 @@ public:
 	virtual void Move(const TSharedRef<FString>& NewText, const FTextRange& NewRange) override;
 	virtual TSharedRef<IRun> Clone() const override;
 
-	virtual void AppendText(FString& Text) const override;
-	virtual void AppendText(FString& AppendToText, const FTextRange& PartialRange) const override;
+	virtual void AppendTextTo(FString& Text) const override;
+	virtual void AppendTextTo(FString& AppendToText, const FTextRange& PartialRange) const override;
+
+	virtual const FRunInfo& GetRunInfo() const override;
 
 private:
 
-	FSlateWidgetRun( const TSharedRef< const FString >& InText, const FWidgetRunInfo& InWidgetInfo );
+	FSlateWidgetRun( const FRunInfo& InRunInfo, const TSharedRef< const FString >& InText, const FWidgetRunInfo& InWidgetInfo );
 
-	FSlateWidgetRun( const TSharedRef< const FString >& InText, const FWidgetRunInfo& InWidgetInfo, const FTextRange& InRange );
+	FSlateWidgetRun( const FRunInfo& InRunInfo, const TSharedRef< const FString >& InText, const FWidgetRunInfo& InWidgetInfo, const FTextRange& InRange );
 
 	FSlateWidgetRun( const FSlateWidgetRun& Run );
 
 private:
 
+	FRunInfo RunInfo;
 	TSharedRef< const FString > Text;
 	FTextRange Range;
 	FWidgetRunInfo Info;
diff --git a/Engine/Source/Runtime/Slate/Public/Framework/Text/SyntaxHighlighterTextLayoutMarshaller.h b/Engine/Source/Runtime/Slate/Public/Framework/Text/SyntaxHighlighterTextLayoutMarshaller.h
new file mode 100644
--- /dev/null
+++ b/Engine/Source/Runtime/Slate/Public/Framework/Text/SyntaxHighlighterTextLayoutMarshaller.h
@@ -0,0 +1,89 @@
+// Copyright 1998-2014 Epic Games, Inc. All Rights Reserved.
+#pragma once
+
+#if WITH_FANCY_TEXT
+
+#include "PlainTextLayoutMarshaller.h"
+#include "ITextDecorator.h"
+#include "SyntaxTokenizer.h"
+
+/**
+ * Get/set the raw text to/from a text layout, and also inject syntax highlighting
+ */
+class SLATE_API FSyntaxHighlighterTextLayoutMarshaller : public FPlainTextLayoutMarshaller
+{
+public:
+
+	virtual ~FSyntaxHighlighterTextLayoutMarshaller();
+
+	// ITextLayoutMarshaller
+	virtual void SetText(const FString& SourceString, FTextLayout& TargetTextLayout) override;
+	virtual bool RequiresLiveUpdate() const override;
+
+	void EnableSyntaxHighlighting(const bool bEnable);
+	bool IsSyntaxHighlightingEnabled() const;
+
+protected:
+
+	virtual void ParseTokens(const FString& SourceString, FTextLayout& TargetTextLayout, TArray<FSyntaxTokenizer::FTokenizedLine> TokenizedLines) = 0;
+
+	FSyntaxHighlighterTextLayoutMarshaller(TSharedPtr< FSyntaxTokenizer > InTokenizer, FTextBlockStyle InDefaultTextStyle);
+
+	/** Tokenizer used to style the text */
+	TSharedPtr< FSyntaxTokenizer > Tokenizer;
+
+	/** True if syntax highlighting is enabled, false to fallback to plain text */
+	bool bSyntaxHighlightingEnabled;
+
+};
+
+/**
+ * Get/set the raw text to/from a text layout, and also inject syntax highlighting for our rich-text markup
+ */
+class SLATE_API FRichTextSyntaxHighlighterTextLayoutMarshaller : public FSyntaxHighlighterTextLayoutMarshaller
+{
+public:
+
+	struct FSyntaxTextStyle
+	{
+		FSyntaxTextStyle()
+			: NormalTextStyle(FCoreStyle::Get().GetWidgetStyle<FTextBlockStyle>("SyntaxHighlight.Normal"))
+			, NodeTextStyle(FCoreStyle::Get().GetWidgetStyle<FTextBlockStyle>("SyntaxHighlight.Node"))
+			, NodeAttributeKeyTextStyle(FCoreStyle::Get().GetWidgetStyle<FTextBlockStyle>("SyntaxHighlight.NodeAttributeKey"))
+			, NodeAttribueAssignmentTextStyle(FCoreStyle::Get().GetWidgetStyle<FTextBlockStyle>("SyntaxHighlight.NodeAttribueAssignment"))
+			, NodeAttributeValueTextStyle(FCoreStyle::Get().GetWidgetStyle<FTextBlockStyle>("SyntaxHighlight.NodeAttributeValue"))
+		{
+		}
+
+		FSyntaxTextStyle(const FTextBlockStyle& InNormalTextStyle, const FTextBlockStyle& InNodeTextStyle, const FTextBlockStyle& InNodeAttributeKeyTextStyle, const FTextBlockStyle& InNodeAttribueAssignmentTextStyle, const FTextBlockStyle& InNodeAttributeValueTextStyle)
+			: NormalTextStyle(InNormalTextStyle)
+			, NodeTextStyle(InNodeTextStyle)
+			, NodeAttributeKeyTextStyle(InNodeAttributeKeyTextStyle)
+			, NodeAttribueAssignmentTextStyle(InNodeAttribueAssignmentTextStyle)
+			, NodeAttributeValueTextStyle(InNodeAttributeValueTextStyle)
+		{
+		}
+
+		FTextBlockStyle NormalTextStyle;
+		FTextBlockStyle NodeTextStyle;
+		FTextBlockStyle NodeAttributeKeyTextStyle;
+		FTextBlockStyle NodeAttribueAssignmentTextStyle;
+		FTextBlockStyle NodeAttributeValueTextStyle;
+	};
+
+	static TSharedRef< FRichTextSyntaxHighlighterTextLayoutMarshaller > Create(const FSyntaxTextStyle& InSyntaxTextStyle);
+
+	virtual ~FRichTextSyntaxHighlighterTextLayoutMarshaller();
+
+protected:
+
+	virtual void ParseTokens(const FString& SourceString, FTextLayout& TargetTextLayout, TArray<FSyntaxTokenizer::FTokenizedLine> TokenizedLines) override;
+
+	FRichTextSyntaxHighlighterTextLayoutMarshaller(TSharedPtr< FSyntaxTokenizer > InTokenizer, const FSyntaxTextStyle& InSyntaxTextStyle);
+
+	/** Styles used to display the text */
+	FSyntaxTextStyle SyntaxTextStyle;
+
+};
+
+#endif //WITH_FANCY_TEXT
diff --git a/Engine/Source/Runtime/Slate/Public/Framework/Text/SyntaxTokenizer.h b/Engine/Source/Runtime/Slate/Public/Framework/Text/SyntaxTokenizer.h
new file mode 100644
--- /dev/null
+++ b/Engine/Source/Runtime/Slate/Public/Framework/Text/SyntaxTokenizer.h
@@ -0,0 +1,70 @@
+// Copyright 1998-2014 Epic Games, Inc. All Rights Reserved.
+#pragma once
+
+/**
+ * Tokenize the text based upon the given rule set
+ */
+class SLATE_API FSyntaxTokenizer
+{
+public:
+
+	/** Denotes the type of token */
+	enum class ETokenType : uint8
+	{
+		/** Syntax token matching the tokenizing rules provided */
+		Syntax,
+		/** String token containing some literal text */
+		Literal,
+	};
+
+	/** A token referencing a range in the original text */
+	struct FToken
+	{
+		FToken(const ETokenType InType, const FTextRange& InRange)
+			: Type(InType)
+			, Range(InRange)
+		{
+		}
+
+		ETokenType Type;
+		FTextRange Range;
+	};
+
+	/** A line containing a series of tokens */
+	struct FTokenizedLine
+	{
+		FTextRange Range;
+		TArray<FToken> Tokens;
+	};
+
+	/** Rule used to match syntax token types */
+	struct FRule
+	{
+		FRule(FString InMatchText)
+			: MatchText(MoveTemp(InMatchText))
+		{
+		}
+
+		FString MatchText;
+	};
+
+	/** 
+	 * Create a new tokenizer which will use the given rules to match syntax tokens
+	 * @param InRules Rules to control the tokenizer, processed in-order so the most greedy matches must come first
+	 */
+	static TSharedRef< FSyntaxTokenizer > Create(TArray<FRule> InRules);
+
+	virtual ~FSyntaxTokenizer();
+
+	void Process(TArray<FTokenizedLine>& OutTokenizedLines, const FString& Input);
+
+private:
+
+	FSyntaxTokenizer(TArray<FRule> InRules);
+
+	void TokenizeLineRanges(const FString& Input, const TArray<FTextRange>& LineRanges, TArray<FTokenizedLine>& OutTokenizedLines);
+
+	/** Rules to control the tokenizer, processed in-order so the most greedy matches must come first */
+	TArray<FRule> Rules;
+
+};
diff --git a/Engine/Source/Runtime/Slate/Public/Framework/Text/TextDecorators.h b/Engine/Source/Runtime/Slate/Public/Framework/Text/TextDecorators.h
--- a/Engine/Source/Runtime/Slate/Public/Framework/Text/TextDecorators.h
+++ b/Engine/Source/Runtime/Slate/Public/Framework/Text/TextDecorators.h
@@ -18,9 +18,9 @@ public:
 
 public:
 
-	virtual bool Supports( const FTextRunParseResults& RunInfo, const FString& Text ) const override;
+	virtual bool Supports( const FTextRunParseResults& RunParseResult, const FString& Text ) const override;
 
-	virtual TSharedRef< ISlateRun > Create( const FTextRunParseResults& RunInfo, const FString& OriginalText, const TSharedRef< FString >& InOutModelText, const ISlateStyle* Style ) override;
+	virtual TSharedRef< ISlateRun > Create( const FTextRunParseResults& RunParseResult, const FString& OriginalText, const TSharedRef< FString >& InOutModelText, const ISlateStyle* Style ) override;
 
 private:
 
@@ -41,9 +41,9 @@ public:
 
 public:
 
-	virtual bool Supports( const FTextRunParseResults& RunInfo, const FString& Text ) const override;
+	virtual bool Supports( const FTextRunParseResults& RunParseResult, const FString& Text ) const override;
 
-	virtual TSharedRef< ISlateRun > Create( const FTextRunParseResults& RunInfo, const FString& OriginalText, const TSharedRef< FString >& InOutModelText, const ISlateStyle* Style ) override;
+	virtual TSharedRef< ISlateRun > Create( const FTextRunParseResults& RunParseResult, const FString& OriginalText, const TSharedRef< FString >& InOutModelText, const ISlateStyle* Style ) override;
 
 private:
 
@@ -64,9 +64,9 @@ public:
 
 public:
 
-	virtual bool Supports( const FTextRunParseResults& RunInfo, const FString& Text ) const override;
+	virtual bool Supports( const FTextRunParseResults& RunParseResult, const FString& Text ) const override;
 
-	virtual TSharedRef< ISlateRun > Create( const FTextRunParseResults& RunInfo, const FString& OriginalText, const TSharedRef< FString >& InOutModelText, const ISlateStyle* Style ) override;
+	virtual TSharedRef< ISlateRun > Create( const FTextRunParseResults& RunParseResult, const FString& OriginalText, const TSharedRef< FString >& InOutModelText, const ISlateStyle* Style ) override;
 
 private:
 
diff --git a/Engine/Source/Runtime/Slate/Public/Framework/Text/TextLayout.h b/Engine/Source/Runtime/Slate/Public/Framework/Text/TextLayout.h
--- a/Engine/Source/Runtime/Slate/Public/Framework/Text/TextLayout.h
+++ b/Engine/Source/Runtime/Slate/Public/Framework/Text/TextLayout.h
@@ -188,8 +188,8 @@ public:
 
 		void ClearCache();
 
-		void AppendText(FString& Text) const;
-		void AppendText(FString& Text, const FTextRange& Range) const;
+		void AppendTextTo(FString& Text) const;
+		void AppendTextTo(FString& Text, const FTextRange& Range) const;
 
 	private:
 
@@ -351,9 +351,11 @@ public:
 
 	bool InsertAt(const FTextLocation& Location, TCHAR Character);
 
-	bool InsertAt( const FTextLocation& Location, const FString& Text );
+	bool InsertAt(const FTextLocation& Location, const FString& Text);
 
-	bool RemoveAt( const FTextLocation& Location, int32 Count = 1 );
+	bool InsertAt(const FTextLocation& Location, TSharedRef<IRun> InRun, const bool bAlwaysKeepRightRun = false);
+
+	bool RemoveAt(const FTextLocation& Location, int32 Count = 1);
 
 	bool RemoveLine(int32 LineIndex);
 
diff --git a/Engine/Source/Runtime/Slate/Public/Framework/Text/TextRange.h b/Engine/Source/Runtime/Slate/Public/Framework/Text/TextRange.h
--- a/Engine/Source/Runtime/Slate/Public/Framework/Text/TextRange.h
+++ b/Engine/Source/Runtime/Slate/Public/Framework/Text/TextRange.h
@@ -34,6 +34,11 @@ struct SLATE_API FTextRange
 		return Intersected;
 	}
 
+	/**
+	 * Produce an array of line ranges from the given text, breaking at any new-line characters
+	 */
+	static void CalculateLineRangesFromString(const FString& Input, TArray<FTextRange>& LineRanges);
+
 	int32 BeginIndex;
 	int32 EndIndex;
-};
\ No newline at end of file
+};
diff --git a/Engine/Source/Runtime/Slate/Public/Widgets/Input/SMultiLineEditableTextBox.h b/Engine/Source/Runtime/Slate/Public/Widgets/Input/SMultiLineEditableTextBox.h
--- a/Engine/Source/Runtime/Slate/Public/Widgets/Input/SMultiLineEditableTextBox.h
+++ b/Engine/Source/Runtime/Slate/Public/Widgets/Input/SMultiLineEditableTextBox.h
@@ -14,6 +14,7 @@ public:
 
 	SLATE_BEGIN_ARGS( SMultiLineEditableTextBox )
 		: _Style(&FCoreStyle::Get().GetWidgetStyle< FEditableTextBoxStyle >("NormalEditableTextBox"))
+		, _Marshaller()
 		, _Text()
 		, _HintText()
 		, _Font()
@@ -39,6 +40,9 @@ public:
 		/** The styling of the textbox */
 		SLATE_STYLE_ARGUMENT( FEditableTextBoxStyle, Style )
 
+		/** The marshaller used to get/set the raw text to/from the text layout. */
+		SLATE_ARGUMENT(TSharedPtr< ITextLayoutMarshaller >, Marshaller)
+
 		/** Sets the text content for this editable text box widget */
  		SLATE_ATTRIBUTE( FText, Text )
 
@@ -84,6 +88,9 @@ public:
 		/** Called whenever the text is committed.  This happens when the user presses enter or the text box loses focus. */
 		SLATE_EVENT( FOnTextCommitted, OnTextCommitted )
 
+		/** Called when the cursor is moved within the text area */
+		SLATE_EVENT( SMultiLineEditableText::FOnCursorMoved, OnCursorMoved )
+
 		/** Whether text wraps onto a new line when it's length exceeds this width; if this value is zero or negative, no wrapping occurs. */
 		SLATE_ATTRIBUTE( float, WrapTextAt )
 
@@ -146,6 +153,22 @@ public:
 	virtual bool HasKeyboardFocus() const override;
 	virtual FReply OnKeyboardFocusReceived( const FGeometry& MyGeometry, const FKeyboardFocusEvent& InKeyboardFocusEvent ) override;
 
+	/** Get the currently selected text */
+	FText GetSelectedText() const;
+
+	/** Insert the given text at the current cursor position, correctly taking into account new line characters */
+	void InsertTextAtCursor(const FText& InText);
+	void InsertTextAtCursor(const FString& InString);
+
+	/** Insert the given run at the current cursor position */
+	void InsertRunAtCursor(TSharedRef<IRun> InRun);
+
+	/** Apply the given style to the currently selected text (or insert a new run at the current cursor position if no text is selected) */
+	void ApplyToSelection(const FRunInfo& InRunInfo, const FTextBlockStyle& InStyle);
+
+	/** Get the run currently under the cursor, or null if there is no run currently under the cursor */
+	TSharedPtr<const IRun> GetRunUnderCursor() const;
+
 protected:
 
 	/** Editable text widget */
diff --git a/Engine/Source/Runtime/Slate/Public/Widgets/Text/SMultiLineEditableText.h b/Engine/Source/Runtime/Slate/Public/Widgets/Text/SMultiLineEditableText.h
--- a/Engine/Source/Runtime/Slate/Public/Widgets/Text/SMultiLineEditableText.h
+++ b/Engine/Source/Runtime/Slate/Public/Widgets/Text/SMultiLineEditableText.h
@@ -6,13 +6,18 @@
 
 #include "ITextEditorWidget.h"
 #include "ITextInputMethodSystem.h"
+#include "ITextLayoutMarshaller.h"
 
 /** An editable text widget that supports multiple lines and soft word-wrapping. */
 class SLATE_API SMultiLineEditableText : public SWidget, public ITextEditorWidget
 {
 public:
+	/** Called when the cursor is moved within the text area */
+	DECLARE_DELEGATE_OneParam( FOnCursorMoved, const FTextLocation& );
+
 	SLATE_BEGIN_ARGS( SMultiLineEditableText )
 		: _Text()
+		, _Marshaller()
 		, _WrapTextAt( 0.0f )
 		, _AutoWrapText(false)
 		, _TextStyle( &FCoreStyle::Get().GetWidgetStyle<FTextBlockStyle>( "NormalText" ) )
@@ -23,11 +28,15 @@ public:
 		, _IsReadOnly(false)
 		, _OnTextChanged()
 		, _OnTextCommitted()
+		, _OnCursorMoved()
 		, _ContextMenuExtender()
 	{}
 		/** The initial text that will appear in the widget. */
 		SLATE_ATTRIBUTE(FText, Text)
 
+		/** The marshaller used to get/set the raw text to/from the text layout. */
+		SLATE_ARGUMENT(TSharedPtr< ITextLayoutMarshaller >, Marshaller)
+
 		/** Whether text wraps onto a new line when it's length exceeds this width; if this value is zero or negative, no wrapping occurs. */
 		SLATE_ATTRIBUTE(float, WrapTextAt)
 
@@ -60,6 +69,9 @@ public:
 		/** Called whenever the text is committed.  This happens when the user presses enter or the text box loses focus. */
 		SLATE_EVENT(FOnTextCommitted, OnTextCommitted)
 
+		/** Called when the cursor is moved within the text area */
+		SLATE_EVENT(FOnCursorMoved, OnCursorMoved)
+
 		/** Menu extender for the right-click context menu */
 		SLATE_EVENT(FMenuExtensionDelegate, ContextMenuExtender)
 
@@ -86,6 +98,22 @@ public:
 	virtual bool GetIsReadOnly() const override;
 	virtual void ClearSelection() override;
 
+	/** Get the currently selected text */
+	FText GetSelectedText() const;
+
+	/** Insert the given text at the current cursor position, correctly taking into account new line characters */
+	void InsertTextAtCursor(const FText& InText);
+	void InsertTextAtCursor(const FString& InString);
+
+	/** Insert the given run at the current cursor position */
+	void InsertRunAtCursor(TSharedRef<IRun> InRun);
+
+	/** Apply the given style to the currently selected text (or insert a new run at the current cursor position if no text is selected) */
+	void ApplyToSelection(const FRunInfo& InRunInfo, const FTextBlockStyle& InStyle);
+
+	/** Get the run currently under the cursor, or null if there is no run currently under the cursor */
+	TSharedPtr<const IRun> GetRunUnderCursor() const;
+
 private:
 	
 	/** 
@@ -361,7 +389,7 @@ private:
 	bool DoesClipboardHaveAnyText() const;
 
 	/** Insert the given text at the current cursor position, correctly taking into account new line characters */
-	void InsertTextAtCursor(const FString& InString);
+	void InsertTextAtCursorImpl(const FString& InString);
 
 	/**
 	 * Given a location and a Direction to offset, return a new location.
@@ -434,6 +462,11 @@ private:
 	 */
 	FText GetEditableText() const;
 
+	/**
+	 * Force the text layout to be updated from the marshaller
+	 */
+	void ForceRefreshTextLayout(const FText& CurrentText);
+
 private:
 
 	/** The text displayed in this text block */
@@ -442,6 +475,9 @@ private:
 	/** The state of BoundText last Tick() (only used when BoundText is bound to a delegate providing the source text) */
 	FText BoundTextLastTick;
 
+	/** The marshaller used to get/set the BoundText text to/from the text layout. */
+	TSharedPtr< ITextLayoutMarshaller > Marshaller;
+
 	/** In control of the layout and wrapping of the BoundText */
 	TSharedPtr< FSlateTextLayout > TextLayout;
 
@@ -510,6 +546,9 @@ private:
 	/** Called whenever the text is committed.  This happens when the user presses enter or the text box loses focus. */
 	FOnTextCommitted OnTextCommitted;
 
+	/** Called when the cursor is moved within the text area */
+	FOnCursorMoved OnCursorMoved;
+
 	/** Menu extender for right-click context menu */
 	TSharedPtr<FExtender> MenuExtender;
 
diff --git a/Engine/Source/Runtime/Slate/Public/Widgets/Text/SRichTextBlock.h b/Engine/Source/Runtime/Slate/Public/Widgets/Text/SRichTextBlock.h
--- a/Engine/Source/Runtime/Slate/Public/Widgets/Text/SRichTextBlock.h
+++ b/Engine/Source/Runtime/Slate/Public/Widgets/Text/SRichTextBlock.h
@@ -3,6 +3,8 @@
 
 #if WITH_FANCY_TEXT
 
+#include "RichTextLayoutMarshaller.h"
+
 /**
  * A rich static text widget. 
  * Through the use of markup and text decorators, text with different styles, embedded image and widgets can be achieved.
@@ -127,10 +129,6 @@ public:
 	void SetHighlightText( const FText& InHighlightText );
 
 
-private:
-
-	TSharedPtr< ITextDecorator > TryGetDecorator( const TArray< TSharedRef< ITextDecorator > >& Decorators, const FString& Line, const FTextRunParseResults& TextRun ) const;
-
 private:
 
 	/** The text displayed in this text block */
@@ -139,15 +137,17 @@ private:
 	/** The state of BoundText last Tick() (only used when BoundText is bound to a delegate providing the source text) */
 	FText BoundTextLastTick;
 
+	/** The marshaller used to set the Text text in the text layout. */
+	TSharedPtr< FRichTextLayoutMarshaller > Marshaller;
+
 	TSharedPtr< FSlateTextLayout > TextLayout;
 
+	/** Default style used by the TextLayout */
+	FTextBlockStyle TextStyle;
+
 	FText HighlightText;
-	TArray< FTextRunRenderer > TextHighlights;
 	TSharedPtr< ISlateRunRenderer> TextHighlighter;
 
-	const ISlateStyle* TagStyleSet;
-	FTextBlockStyle TextStyle;
-	
 	/** Whether text wraps onto a new line when it's length exceeds this width; if this value is zero or negative, no wrapping occurs. */
 	TAttribute<float> WrapTextAt;
 	
@@ -160,10 +160,6 @@ private:
 	TAttribute< FMargin > Margin;
 	TAttribute< ETextJustify::Type > Justification; 
 	TAttribute< float > LineHeightPercentage;
-
-	TArray< TSharedRef< class ITextDecorator > > Decorators;
-
-	TSharedPtr< class IRichTextMarkupParser > Parser;
 };
 
 #endif //WITH_FANCY_TEXT
\ No newline at end of file
diff --git a/Engine/Source/Runtime/SlateCore/Private/Styling/CoreStyle.cpp b/Engine/Source/Runtime/SlateCore/Private/Styling/CoreStyle.cpp
--- a/Engine/Source/Runtime/SlateCore/Private/Styling/CoreStyle.cpp
+++ b/Engine/Source/Runtime/SlateCore/Private/Styling/CoreStyle.cpp
@@ -340,11 +340,11 @@ TSharedRef<ISlateStyle> FCoreStyle::Create( const FName& InStyleSetName )
 		const FCheckBoxStyle ToggleButtonStyle = FCheckBoxStyle()
 			.SetCheckBoxType(ESlateCheckBoxType::ToggleButton)
 			.SetUncheckedImage( FSlateNoResource() )
-			.SetUncheckedPressedImage( BOX_BRUSH("Common/RoundedSelection_16x", 4.0f/16.0f, SelectionColor_Pressed ) )
 			.SetUncheckedHoveredImage( BOX_BRUSH("Common/RoundedSelection_16x", 4.0f/16.0f, SelectionColor ) )
+			.SetUncheckedPressedImage( BOX_BRUSH("Common/RoundedSelection_16x", 4.0f/16.0f, SelectionColor_Pressed ) )
 			.SetCheckedImage( BOX_BRUSH("Common/RoundedSelection_16x",  4.0f/16.0f, SelectionColor_Pressed ) )
-			.SetCheckedHoveredImage( BOX_BRUSH("Common/RoundedSelection_16x",  4.0f/16.0f, SelectionColor_Pressed ) )
-			.SetUncheckedPressedImage( BOX_BRUSH("Common/RoundedSelection_16x",  4.0f/16.0f, SelectionColor ) );
+			.SetCheckedHoveredImage( BOX_BRUSH("Common/RoundedSelection_16x",  4.0f/16.0f, SelectionColor ) )
+			.SetCheckedPressedImage( BOX_BRUSH("Common/RoundedSelection_16x",  4.0f/16.0f, SelectionColor_Pressed ) );
 		Style->Set( "ToggleButtonCheckbox", ToggleButtonStyle );
 
 		/* A radio button is actually just a SCheckBox box with different images */
@@ -1111,6 +1111,24 @@ TSharedRef<ISlateStyle> FCoreStyle::Create( const FName& InStyleSetName )
 		Style->Set("Wizard.NextIcon", new IMAGE_BRUSH("Icons/NextIcon", Icon8x8));
 	}
 
+	// Syntax highlighting
+	{
+		const FTextBlockStyle MonospacedText = FTextBlockStyle()
+			.SetFont(TTF_FONT("Fonts/DroidSansMono", 9))
+			.SetColorAndOpacity(FSlateColor::UseForeground())
+			.SetShadowOffset(FVector2D::ZeroVector)
+			.SetShadowColorAndOpacity(FLinearColor::Black)
+			.SetHighlightColor(FLinearColor(0.02f, 0.3f, 0.0f))
+			.SetHighlightShape(BOX_BRUSH("Common/TextBlockHighlightShape", FMargin(3.f/8.f))
+			);
+
+		Style->Set("SyntaxHighlight.Normal", MonospacedText);
+		Style->Set("SyntaxHighlight.Node", FTextBlockStyle(MonospacedText).SetColorAndOpacity(FLinearColor(FColor(0xff006ab4)))); // blue
+		Style->Set("SyntaxHighlight.NodeAttributeKey", FTextBlockStyle(MonospacedText).SetColorAndOpacity(FLinearColor(FColor(0xffb40000)))); // red
+		Style->Set("SyntaxHighlight.NodeAttribueAssignment", FTextBlockStyle(MonospacedText).SetColorAndOpacity(FLinearColor(FColor(0xffb2b400)))); // yellow
+		Style->Set("SyntaxHighlight.NodeAttributeValue", FTextBlockStyle(MonospacedText).SetColorAndOpacity(FLinearColor(FColor(0xffb46100)))); // orange
+	}
+
 	return Style;
 }
 
