Some multi-line selection and deletion fixes

- Fixed a crash when cutting to the start of a line, as it was removing the entire line from the text layout rather than removing the content but leaving the line.
- Fixed the JumpTo code resetting the selection rather than adding to it.

diff --git a/Engine/Source/Runtime/Slate/Private/Widgets/Text/SMultiLineEditableText.cpp b/Engine/Source/Runtime/Slate/Private/Widgets/Text/SMultiLineEditableText.cpp
--- a/Engine/Source/Runtime/Slate/Private/Widgets/Text/SMultiLineEditableText.cpp
+++ b/Engine/Source/Runtime/Slate/Private/Widgets/Text/SMultiLineEditableText.cpp
@@ -727,15 +727,7 @@ void SMultiLineEditableText::DeleteSelectedText()
 			}
 
 			const FTextLayout::FLineModel& BeginLine = Lines[SelectionBeginningLineIndex];
-
-			if (SelectionBeginningLineOffset == 0)
-			{
-				TextLayout->RemoveLine(SelectionBeginningLineIndex);
-			}
-			else
-			{
-				TextLayout->RemoveAt(FTextLocation(SelectionBeginningLineIndex, SelectionBeginningLineOffset), BeginLine.Text->Len() - SelectionBeginningLineOffset);
-			}
+			TextLayout->RemoveAt(FTextLocation(SelectionBeginningLineIndex, SelectionBeginningLineOffset), BeginLine.Text->Len() - SelectionBeginningLineOffset);
 
 			TextLayout->JoinLineWithNextLine(SelectionBeginningLineIndex);
 
@@ -1038,7 +1030,10 @@ void SMultiLineEditableText::JumpTo(ETextLocation JumpLocation, ECursorAction Ac
 
 				if (Action == ECursorAction::SelectText)
 				{
-					this->SelectionStart = OldCursorPosition;
+					if (!SelectionStart.IsSet())
+					{
+						this->SelectionStart = OldCursorPosition;
+					}
 				}
 				else
 				{
@@ -1060,7 +1055,10 @@ void SMultiLineEditableText::JumpTo(ETextLocation JumpLocation, ECursorAction Ac
 
 			if (Action == ECursorAction::SelectText)
 			{
-				this->SelectionStart = OldCursorPosition;
+				if (!SelectionStart.IsSet())
+				{
+					this->SelectionStart = OldCursorPosition;
+				}
 			}
 			else
 			{
@@ -1089,7 +1087,10 @@ void SMultiLineEditableText::JumpTo(ETextLocation JumpLocation, ECursorAction Ac
 
 				if (Action == ECursorAction::SelectText)
 				{
-					this->SelectionStart = OldCursorPosition;
+					if (!SelectionStart.IsSet())
+					{
+						this->SelectionStart = OldCursorPosition;
+					}
 				}
 				else
 				{
@@ -1115,7 +1116,10 @@ void SMultiLineEditableText::JumpTo(ETextLocation JumpLocation, ECursorAction Ac
 
 				if (Action == ECursorAction::SelectText)
 				{
-					this->SelectionStart = OldCursorPosition;
+					if (!SelectionStart.IsSet())
+					{
+						this->SelectionStart = OldCursorPosition;
+					}
 				}
 				else
 				{
