Better implementation of FMacMenu's menuHasKeyEquivalent:

diff --git a/Engine/Source/Runtime/Slate/Private/Framework/MultiBox/Mac/MacMenu.cpp b/Engine/Source/Runtime/Slate/Private/Framework/MultiBox/Mac/MacMenu.cpp
--- a/Engine/Source/Runtime/Slate/Private/Framework/MultiBox/Mac/MacMenu.cpp
+++ b/Engine/Source/Runtime/Slate/Private/Framework/MultiBox/Mac/MacMenu.cpp
@@ -75,7 +75,72 @@
 
 - (BOOL)menuHasKeyEquivalent:(NSMenu*)menu forEvent:(NSEvent*)event target:(id*)target action:(SEL*)action
 {
-	return NO;
+	// @todo: FSlateMacMenu::GetMenuBuilderWidget is currently expensive, so to avoid a framerate drop when using, for example, WASD keys, we always return false if no modifier keys are pressed.
+	// That's OK for now, because for a rare case where a menu item has a shortcut without any modifiers, Slate will handle it, but it needs to be fixed.
+	if ([event modifierFlags] == 0 || ![menu isKindOfClass:[FMacMenu class]])
+	{
+		return NO;
+	}
+	else
+	{
+		FMacMenu* MacMenu = (FMacMenu*)menu;
+
+		FText WindowLabel = NSLOCTEXT("MainMenu", "WindowMenu", "Window");
+		bool bIsWindowMenu = (WindowLabel.ToString().Compare(FString([MacMenu title])) == 0);
+
+		unichar Character = toupper([[event characters] characterAtIndex:0]);
+
+		if (bIsWindowMenu)
+		{
+			return ([event modifierFlags] & NSCommandKeyMask) && (Character == 'W' || Character == 'M');
+		}
+		else
+		{
+			TSharedRef< SWidget > Widget = FSlateMacMenu::GetMenuBuilderWidget(MacMenu);
+
+			TSharedRef<const FMultiBox> MultiBox = StaticCastSharedRef<SMultiBoxWidget>(Widget)->GetMultiBox();
+			const TArray<TSharedRef<const FMultiBlock>>& MenuBlocks = MultiBox->GetBlocks();
+
+			for (int32 Index = 0; Index < MenuBlocks.Num(); Index++)
+			{
+				EMultiBlockType::Type Type = MenuBlocks[Index]->GetType();
+
+				if (Type == EMultiBlockType::MenuEntry)
+				{
+					TSharedRef<const FMenuEntryBlock> Block = StaticCastSharedRef<const FMenuEntryBlock>(MenuBlocks[Index]);
+					if (Block->GetAction().IsValid())
+					{
+						const TSharedRef<const FInputGesture>& Gesture = Block->GetAction()->GetActiveGesture();
+						if (!Gesture->GetKeyText().IsEmpty() && Gesture->GetKeyText().ToString()[0] == Character)
+						{
+							uint32 EventModifiers = [event modifierFlags] & (NSCommandKeyMask | NSShiftKeyMask | NSAlternateKeyMask | NSControlKeyMask);
+							uint32 GestureModfiers = 0;
+							if (Gesture->bCtrl)
+							{
+								GestureModfiers |= NSCommandKeyMask;
+							}
+							if (Gesture->bShift)
+							{
+								GestureModfiers |= NSShiftKeyMask;
+							}
+							if (Gesture->bAlt)
+							{
+								GestureModfiers |= NSAlternateKeyMask;
+							}
+							if (Gesture->bCmd)
+							{
+								GestureModfiers |= NSControlKeyMask;
+							}
+
+							return EventModifiers == GestureModfiers;
+						}
+					}
+				}
+			}
+		}
+
+		return NO;
+	}
 }
 
 - ( TSharedPtr< const FMenuEntryBlock >& )GetMenuEntryBlock
@@ -258,13 +323,7 @@ void FSlateMacMenu::UpdateMenu(FMacMenu* Menu)
 		[Menu addItem:[NSMenuItem separatorItem]];
 	}
 
-	const bool bShouldCloseWindowAfterMenuSelection = true;
-	FMenuBuilder MenuBuilder(bShouldCloseWindowAfterMenuSelection, [Menu GetMenuEntryBlock]->GetActionList(), [Menu GetMenuEntryBlock]->Extender);
-	{
-		// Have the menu fill its contents
-		[Menu GetMenuEntryBlock]->EntryBuilder.ExecuteIfBound(MenuBuilder);
-	}
-	TSharedRef< SWidget > Widget = MenuBuilder.MakeWidget();
+	TSharedRef< SWidget > Widget = GetMenuBuilderWidget(Menu);
 
 	TSharedRef<const FMultiBox> MultiBox = StaticCastSharedRef<SMultiBoxWidget>(Widget)->GetMultiBox();
 	const TArray<TSharedRef<const FMultiBlock>>& MenuBlocks = MultiBox->GetBlocks();
@@ -357,3 +416,15 @@ int32 FSlateMacMenu::GetMenuItemState(const TSharedRef<const class FMenuEntryBlo
 
 	return bIsChecked ? NSOnState : NSOffState;
 }
+
+TSharedRef<SWidget> FSlateMacMenu::GetMenuBuilderWidget(FMacMenu* Menu)
+{
+	// @todo: ideally this should be cached and updated only when the menu contents change. See menuHasKeyEquivalent.
+	const bool bShouldCloseWindowAfterMenuSelection = true;
+	FMenuBuilder MenuBuilder(bShouldCloseWindowAfterMenuSelection, [Menu GetMenuEntryBlock]->GetActionList(), [Menu GetMenuEntryBlock]->Extender);
+	{
+		// Have the menu fill its contents
+		[Menu GetMenuEntryBlock]->EntryBuilder.ExecuteIfBound(MenuBuilder);
+	}
+	return MenuBuilder.MakeWidget();
+}
diff --git a/Engine/Source/Runtime/Slate/Private/Framework/MultiBox/Mac/MacMenu.h b/Engine/Source/Runtime/Slate/Private/Framework/MultiBox/Mac/MacMenu.h
--- a/Engine/Source/Runtime/Slate/Private/Framework/MultiBox/Mac/MacMenu.h
+++ b/Engine/Source/Runtime/Slate/Private/Framework/MultiBox/Mac/MacMenu.h
@@ -28,4 +28,5 @@ public:
 	static void UpdateMenu(FMacMenu* InMenu);
 	static bool IsMenuItemEnabled(const TSharedRef<const class FMenuEntryBlock>& Block);
 	static int32 GetMenuItemState(const TSharedRef<const class FMenuEntryBlock>& Block);
+	static TSharedRef<SWidget> GetMenuBuilderWidget(FMacMenu* InMenu);
 };
