Merging using UE4-To-UE4-LauncherReleases

Fixed ability to add multiple empty slots ignoring buttons enabled status.
#OOP-453 - Hitting the 'Enter' key allows users to add multiple engine slots.

#proj Runtime.Slate

#branch UE4-LauncherReleases

#Change - Added an IsEnabled check before sending OnKeyDown message to a widget.

diff --git a/Engine/Source/Runtime/Slate/Private/Framework/Application/SlateApplication.cpp b/Engine/Source/Runtime/Slate/Private/Framework/Application/SlateApplication.cpp
--- a/Engine/Source/Runtime/Slate/Private/Framework/Application/SlateApplication.cpp
+++ b/Engine/Source/Runtime/Slate/Private/Framework/Application/SlateApplication.cpp
@@ -3018,6 +3018,7 @@ bool FSlateApplication::ProcessKeyDownEvent( FKeyboardEvent& InKeyboardEvent )
 			for( int32 WidgetIndex = 0; !Reply.IsEventHandled() && WidgetIndex < EventPath.Widgets.Num(); ++WidgetIndex )
 			{
 				FArrangedWidget& SomeWidgetGettingEvent = EventPath.Widgets( WidgetIndex );
+				if (SomeWidgetGettingEvent.Widget->IsEnabled())
 				{
 					Reply = SomeWidgetGettingEvent.Widget->OnPreviewKeyDown( SomeWidgetGettingEvent.Geometry, InKeyboardEvent ).SetHandler(SomeWidgetGettingEvent.Widget);
 					ProcessReply(EventPath, Reply, NULL, NULL);
@@ -3030,6 +3031,7 @@ bool FSlateApplication::ProcessKeyDownEvent( FKeyboardEvent& InKeyboardEvent )
 			for( int32 WidgetIndex = EventPath.Widgets.Num()-1; !Reply.IsEventHandled() && WidgetIndex >= 0; --WidgetIndex )
 			{
 				FArrangedWidget& SomeWidgetGettingEvent = EventPath.Widgets( WidgetIndex );
+				if (SomeWidgetGettingEvent.Widget->IsEnabled())
 				{
 					Reply = SomeWidgetGettingEvent.Widget->OnKeyDown( SomeWidgetGettingEvent.Geometry, InKeyboardEvent ).SetHandler(SomeWidgetGettingEvent.Widget);
 					ProcessReply(EventPath, Reply, NULL, NULL);
