Slate/Tabs/Editor - In the case where the tab area doesn't manage the window, and there's no parent owner tab, we now no longer allow dragging out the last tab from the well.  This will prevent orphaning the main frame.

diff --git a/Engine/Source/Runtime/Slate/Private/Framework/Docking/SDockingArea.cpp b/Engine/Source/Runtime/Slate/Private/Framework/Docking/SDockingArea.cpp
--- a/Engine/Source/Runtime/Slate/Private/Framework/Docking/SDockingArea.cpp
+++ b/Engine/Source/Runtime/Slate/Private/Framework/Docking/SDockingArea.cpp
@@ -10,6 +10,7 @@ void SDockingArea::Construct( const FArguments& InArgs, const TSharedRef<FTabMan
 	InTabManager->GetPrivateApi().OnDockAreaCreated( SharedThis(this) );
 
 	bManageParentWindow = InArgs._ShouldManageParentWindow;
+	bShouldAllowRemovalOfLastTab = InArgs._ShouldAllowRemovalOfLastTab;
 	bIsOverlayVisible = false;
 
 	bIsCenterTargetVisible = false;
@@ -311,6 +312,16 @@ TSharedRef<FTabManager> SDockingArea::GetTabManager() const
 	return MyTabManager.Pin().ToSharedRef();
 }
 
+bool SDockingArea::GetShouldAllowRemovalOfLastTab() const
+{
+	return bShouldAllowRemovalOfLastTab;
+}
+
+void SDockingArea::SetShouldAllowRemovalOfLastTab(bool Allow)
+{
+	bShouldAllowRemovalOfLastTab = Allow;
+}
+
 SDockingNode::ECleanupRetVal SDockingArea::CleanUpNodes()
 {
 	SDockingNode::ECleanupRetVal ReturnValue = SDockingSplitter::CleanUpNodes();
diff --git a/Engine/Source/Runtime/Slate/Private/Framework/Docking/SDockingArea.h b/Engine/Source/Runtime/Slate/Private/Framework/Docking/SDockingArea.h
--- a/Engine/Source/Runtime/Slate/Private/Framework/Docking/SDockingArea.h
+++ b/Engine/Source/Runtime/Slate/Private/Framework/Docking/SDockingArea.h
@@ -15,6 +15,7 @@ public:
 		: _ParentWindow()
 		, _ShouldManageParentWindow(true)
 		, _InitialContent()
+		, _ShouldAllowRemovalOfLastTab(true)
 		{
 			// Visible by default, but don't absorb clicks
 			_Visibility = EVisibility::SelfHitTestInvisible;
@@ -28,6 +29,19 @@ public:
 		/** True if this docking area should close the parent window when the last tab in this docking area goes away */
 		SLATE_ARGUMENT( bool, ShouldManageParentWindow )
 
+		/**
+		 * Should the dock area allow removing the last tab in the dock area.  By default this is true, but in cases like the main 
+		 * tab window, we don't want the last tab to ever be dragged out.
+		 */
+		SLATE_ARGUMENT(bool, ShouldAllowRemovalOfLastTab)
+
+
+		/** Should we ignore dragging the last tab in this well? */
+		bool IgnoreDraggingLastTabInWell() const;
+
+		/** Should we ignore dragging the last tab in this well? */
+		TAttribute<bool> bIgnoreDraggingLastTabInWell;
+
 		 /**
 		 * What to put into the DockArea initially. Usually a TabStack, so that some tabs can be added to it.
 		 */
@@ -75,6 +89,10 @@ public:
 
 	TSharedRef<FTabManager> GetTabManager() const;
 
+	bool GetShouldAllowRemovalOfLastTab() const;
+
+	void SetShouldAllowRemovalOfLastTab(bool Allow);
+
 protected:
 	virtual SDockingNode::ECleanupRetVal CleanUpNodes() override;
 
@@ -125,4 +143,7 @@ private:
 
 	/** True when the last tab has been pulled from this area, meaning that this DockArea will not be necessary once that tab finds a new home. */
 	bool bCleanUpUponTabRelocation;
+
+	/** True normally, prevents dragging out the last tab of a dock area if false. */
+	bool bShouldAllowRemovalOfLastTab;
 };
diff --git a/Engine/Source/Runtime/Slate/Private/Framework/Docking/TabManager.cpp b/Engine/Source/Runtime/Slate/Private/Framework/Docking/TabManager.cpp
--- a/Engine/Source/Runtime/Slate/Private/Framework/Docking/TabManager.cpp
+++ b/Engine/Source/Runtime/Slate/Private/Framework/Docking/TabManager.cpp
@@ -1079,16 +1079,19 @@ TSharedRef<SDockingNode> FTabManager::RestoreArea_Helper( const TSharedRef<FLayo
 		}
 		else
 		{
-			SAssignNew( NewDockAreaWidget, SDockingArea, SharedThis(this), NodeAsArea.ToSharedRef() )
+			SAssignNew(NewDockAreaWidget, SDockingArea, SharedThis(this), NodeAsArea.ToSharedRef())
 
 				// We only want to set a parent window on this dock area, if we need to have title area content
 				// embedded within it.  SDockingArea assumes that if it has a parent window set, then it needs to have
 				// title area content 
-				.ParentWindow( bEmbedTitleAreaContent ? ParentWindow : TSharedPtr<SWindow>() )
+				.ParentWindow(bEmbedTitleAreaContent ? ParentWindow : TSharedPtr<SWindow>())
 
 				// Never manage these windows, even if a parent window is set.  The owner will take care of
 				// destroying these windows.
-				.ShouldManageParentWindow( false );
+				.ShouldManageParentWindow(false)
+
+				// Don't allow the removal of the last tab, because it will leave an orphaned main frame.
+				.ShouldAllowRemovalOfLastTab(OwnerTabPtr.IsValid() ? true : false);
 
 			RestoreSplitterContent( NodeAsArea.ToSharedRef(), NewDockAreaWidget.ToSharedRef(), ParentWindow );
 		}
diff --git a/Engine/Source/Runtime/Slate/Private/Widgets/Docking/SDockTab.cpp b/Engine/Source/Runtime/Slate/Private/Widgets/Docking/SDockTab.cpp
--- a/Engine/Source/Runtime/Slate/Private/Widgets/Docking/SDockTab.cpp
+++ b/Engine/Source/Runtime/Slate/Private/Widgets/Docking/SDockTab.cpp
@@ -29,6 +29,14 @@ FReply SDockTab::OnMouseButtonDown( const FGeometry& MyGeometry, const FPointerE
 			TotalDraggedDistance = 0;
 			ActivateInParent(ETabActivationCause::UserClickedOnTab);
 
+			// If we're the last tab and we're disallowing the removal of the last tab, ignore the attempt to drag.
+			TSharedPtr<SDockingTabWell> ParentWell = ParentPtr.Pin();
+			TSharedPtr<SDockingArea> DockArea = GetDockArea();
+			if ( ParentWell.IsValid() && DockArea.IsValid() && ParentWell->GetNumTabs() == 1 && !DockArea->GetShouldAllowRemovalOfLastTab() )
+			{
+				return FReply::Unhandled();
+			}
+
 			return FReply::Handled().DetectDrag( SharedThis(this), EKeys::LeftMouseButton );
 		}
 		else if ( MouseEvent.GetEffectingButton() == EKeys::MiddleMouseButton )
