Fixed overlay issues with the asset install progress UI in the marketplace. Adjusted the hyperlink decorator to allow blank ids. Made the unreal marketplace no longer a button and swap it out on the details page with a proper back hyperlink. Adjusted the layout of the details header to reduce whitespace.

diff --git a/Engine/Source/Runtime/Slate/Private/Framework/Text/TextDecorators.cpp b/Engine/Source/Runtime/Slate/Private/Framework/Text/TextDecorators.cpp
--- a/Engine/Source/Runtime/Slate/Private/Framework/Text/TextDecorators.cpp
+++ b/Engine/Source/Runtime/Slate/Private/Framework/Text/TextDecorators.cpp
@@ -114,7 +114,7 @@ bool FHyperlinkDecorator::Supports( const FTextRunParseResults& RunParseResult,
 		MetaDataId = Text.Mid(MetaDataIdRange->BeginIndex, MetaDataIdRange->EndIndex - MetaDataIdRange->BeginIndex);
 	}
 
-	return ( RunParseResult.Name == TEXT("a") && MetaDataIdRange != NULL && MetaDataId == Id );
+	return (RunParseResult.Name == TEXT("a") && MetaDataId == Id);
 }
 
 TSharedRef< ISlateRun > FHyperlinkDecorator::Create( const FTextRunParseResults& RunParseResult, const FString& OriginalText, const TSharedRef< FString >& InOutModelText, const ISlateStyle* Style )
