Slate jittering issue fixes
#OPP-652 A scroll bar briefly flashes on screen in the Library when an engine slot is added or removed
#OPP-573 The library screen flashes erratically when a new engine slot is added or removed

#proj CommunityPortalServices

#branch UE4

#change - Tweaked engine slots to unify the size of them to one value. Helps with not causing jittering on creation. Also added a control to disable the scroll bar.

diff --git a/Engine/Source/Runtime/Slate/Public/Widgets/Views/STileView.h b/Engine/Source/Runtime/Slate/Public/Widgets/Views/STileView.h
--- a/Engine/Source/Runtime/Slate/Public/Widgets/Views/STileView.h
+++ b/Engine/Source/Runtime/Slate/Public/Widgets/Views/STileView.h
@@ -31,6 +31,7 @@ public:
 		, _SelectionMode(ESelectionMode::Multi)
 		, _ClearSelectionOnClick(true)
 		, _ExternalScrollbar()
+		, _ScrollbarVisibility(EVisibility::Visible)
 		{}
 
 		SLATE_EVENT( FOnGenerateRow, OnGenerateTile )
@@ -55,6 +56,8 @@ public:
 
 		SLATE_ARGUMENT( TSharedPtr<SScrollBar>, ExternalScrollbar )
 
+		SLATE_ATTRIBUTE(EVisibility, ScrollbarVisibility)
+
 	SLATE_END_ARGS()
 
 	/**
@@ -104,6 +107,10 @@ public:
 		{
 			// Make the TableView
 			this->ConstructChildren(InArgs._ItemWidth, InArgs._ItemHeight, TSharedPtr<SHeaderRow>(), InArgs._ExternalScrollbar);
+			if (ScrollBar.IsValid())
+			{
+				ScrollBar->SetUserVisibility(InArgs._ScrollbarVisibility);
+			}
 		}
 	}
 
