Expose PlatformApplication's IsMouseAttached to SlateApplication

diff --git a/Engine/Source/Runtime/Slate/Public/Framework/Application/SlateApplication.h b/Engine/Source/Runtime/Slate/Public/Framework/Application/SlateApplication.h
--- a/Engine/Source/Runtime/Slate/Public/Framework/Application/SlateApplication.h
+++ b/Engine/Source/Runtime/Slate/Public/Framework/Application/SlateApplication.h
@@ -457,6 +457,11 @@ public:
 	 */
 	bool IsUsingTrackpad() const { return PlatformApplication.IsValid() ? PlatformApplication->IsUsingTrackpad() : false; }
 
+	/**
+	 * @return True if there is a mouse device attached
+	 */
+	bool IsMouseAttached() const { return PlatformApplication.IsValid() ? PlatformApplication->IsMouseAttached() : false; }
+
 	/**
 	 * Sets the widget reflector.
 	 *
