UE4: Possible CIS fix

diff --git a/Engine/Source/Runtime/Slate/Public/Widgets/Views/STileView.h b/Engine/Source/Runtime/Slate/Public/Widgets/Views/STileView.h
--- a/Engine/Source/Runtime/Slate/Public/Widgets/Views/STileView.h
+++ b/Engine/Source/Runtime/Slate/Public/Widgets/Views/STileView.h
@@ -107,9 +107,9 @@ public:
 		{
 			// Make the TableView
 			this->ConstructChildren(InArgs._ItemWidth, InArgs._ItemHeight, TSharedPtr<SHeaderRow>(), InArgs._ExternalScrollbar);
-			if (ScrollBar.IsValid())
+			if (this->ScrollBar.IsValid())
 			{
-				ScrollBar->SetUserVisibility(InArgs._ScrollbarVisibility);
+				this->ScrollBar->SetUserVisibility(InArgs._ScrollbarVisibility);
 			}
 		}
 	}
