Added support for dynamic images in marketplace rich text. Including an img tag with a src of a url or ems image name will load that image in-line with the text. Extended the decorators to include the TextLayout they are creating a run for so we can do some dynamic dirtying when the image loads in after the initial flow.

Took another pass on the purchase flow dialogs.
Tweaks to colors through out the launcher.

diff --git a/Engine/Source/Runtime/Slate/Private/Framework/Text/RichTextLayoutMarshaller.cpp b/Engine/Source/Runtime/Slate/Private/Framework/Text/RichTextLayoutMarshaller.cpp
--- a/Engine/Source/Runtime/Slate/Private/Framework/Text/RichTextLayoutMarshaller.cpp
+++ b/Engine/Source/Runtime/Slate/Private/Framework/Text/RichTextLayoutMarshaller.cpp
@@ -41,7 +41,7 @@ void FRichTextLayoutMarshaller::SetText(const FString& SourceString, FTextLayout
 			if (Decorator.IsValid())
 			{
 				// Create run and update model string.
-				Run = Decorator->Create(RunParseResult, ProcessedString, ModelString, DecoratorStyleSet);
+				Run = Decorator->Create(TargetTextLayout.AsShared(), RunParseResult, ProcessedString, ModelString, DecoratorStyleSet);
 			}
 			else
 			{
diff --git a/Engine/Source/Runtime/Slate/Private/Framework/Text/TextDecorators.cpp b/Engine/Source/Runtime/Slate/Private/Framework/Text/TextDecorators.cpp
--- a/Engine/Source/Runtime/Slate/Private/Framework/Text/TextDecorators.cpp
+++ b/Engine/Source/Runtime/Slate/Private/Framework/Text/TextDecorators.cpp
@@ -23,7 +23,7 @@ bool FWidgetDecorator::Supports( const FTextRunParseResults& RunParseResult, con
 	return ( RunParseResult.Name == RunName );
 }
 
-TSharedRef< ISlateRun > FWidgetDecorator::Create( const FTextRunParseResults& RunParseResult, const FString& OriginalText, const TSharedRef< FString >& InOutModelText, const ISlateStyle* Style )
+TSharedRef< ISlateRun > FWidgetDecorator::Create(const TSharedRef<class FTextLayout>& TextLayout, const FTextRunParseResults& RunParseResult, const FString& OriginalText, const TSharedRef< FString >& InOutModelText, const ISlateStyle* Style)
 {
 	FTextRange ModelRange;
 	ModelRange.BeginIndex = InOutModelText->Len();
@@ -56,7 +56,7 @@ bool FImageDecorator::Supports( const FTextRunParseResults& RunParseResult, cons
 	return ( RunParseResult.Name == RunName );
 }
 
-TSharedRef< ISlateRun > FImageDecorator::Create( const FTextRunParseResults& RunParseResult, const FString& OriginalText, const TSharedRef< FString >& InOutModelText, const ISlateStyle* Style )
+TSharedRef< ISlateRun > FImageDecorator::Create(const TSharedRef<class FTextLayout>& TextLayout, const FTextRunParseResults& RunParseResult, const FString& OriginalText, const TSharedRef< FString >& InOutModelText, const ISlateStyle* Style)
 {
 	const FTextRange* const BrushNameRange = RunParseResult.MetaData.Find( TEXT("src") );
 
@@ -117,7 +117,7 @@ bool FHyperlinkDecorator::Supports( const FTextRunParseResults& RunParseResult,
 	return (RunParseResult.Name == TEXT("a") && MetaDataId == Id);
 }
 
-TSharedRef< ISlateRun > FHyperlinkDecorator::Create( const FTextRunParseResults& RunParseResult, const FString& OriginalText, const TSharedRef< FString >& InOutModelText, const ISlateStyle* Style )
+TSharedRef< ISlateRun > FHyperlinkDecorator::Create(const TSharedRef<class FTextLayout>& TextLayout, const FTextRunParseResults& RunParseResult, const FString& OriginalText, const TSharedRef< FString >& InOutModelText, const ISlateStyle* Style)
 {
 	FString StyleName = TEXT("Hyperlink");
 
diff --git a/Engine/Source/Runtime/Slate/Private/Framework/Text/TextLayout.cpp b/Engine/Source/Runtime/Slate/Private/Framework/Text/TextLayout.cpp
--- a/Engine/Source/Runtime/Slate/Private/Framework/Text/TextLayout.cpp
+++ b/Engine/Source/Runtime/Slate/Private/Framework/Text/TextLayout.cpp
@@ -617,6 +617,46 @@ void FTextLayout::UpdateHighlights()
 	DirtyFlags &= ~EDirtyState::Highlights;
 }
 
+void FTextLayout::DirtyRunLayout(const TSharedRef<const IRun>& Run)
+{
+	for (int LineModelIndex = 0; LineModelIndex < LineModels.Num(); LineModelIndex++)
+	{
+		FLineModel& LineModel = LineModels[LineModelIndex];
+
+		if (LineModel.HasWrappingInformation)
+		{
+			for (int RunIndex = 0; RunIndex < LineModel.Runs.Num(); RunIndex++)
+			{
+				if (LineModel.Runs[RunIndex].GetRun() == Run)
+				{
+					LineModel.Runs[RunIndex].ClearCache();
+					break;
+				}
+			}
+		}
+	}
+
+	DirtyFlags |= EDirtyState::Layout;
+}
+
+void FTextLayout::DirtyLayout()
+{
+	for (int LineModelIndex = 0; LineModelIndex < LineModels.Num(); LineModelIndex++)
+	{
+		FLineModel& LineModel = LineModels[LineModelIndex];
+
+		if (LineModel.HasWrappingInformation)
+		{
+			for (int RunIndex = 0; RunIndex < LineModel.Runs.Num(); RunIndex++)
+			{
+				LineModel.Runs[RunIndex].ClearCache();
+			}
+		}
+	}
+
+	DirtyFlags |= EDirtyState::Layout;
+}
+
 void FTextLayout::ClearRunRenderers()
 {
 	for (int32 Index = 0; Index < LineModels.Num(); Index++)
diff --git a/Engine/Source/Runtime/Slate/Private/Widgets/Testing/STestSuite.cpp b/Engine/Source/Runtime/Slate/Private/Widgets/Testing/STestSuite.cpp
--- a/Engine/Source/Runtime/Slate/Private/Widgets/Testing/STestSuite.cpp
+++ b/Engine/Source/Runtime/Slate/Private/Widgets/Testing/STestSuite.cpp
@@ -2096,7 +2096,7 @@ public:
 		return (RunParseResult.Name == TEXT("TextStyle"));
 	}
 
-	virtual TSharedRef<ISlateRun> Create(const FTextRunParseResults& RunParseResult, const FString& OriginalText, const TSharedRef<FString>& InOutModelText, const ISlateStyle* Style) override
+	virtual TSharedRef<ISlateRun> Create(const TSharedRef<FTextLayout>& TextLayout, const FTextRunParseResults& RunParseResult, const FString& OriginalText, const TSharedRef<FString>& InOutModelText, const ISlateStyle* Style) override
 	{
 		FRunInfo RunInfo(RunParseResult.Name);
 		for(const TPair<FString, FTextRange>& Pair : RunParseResult.MetaData)
diff --git a/Engine/Source/Runtime/Slate/Private/Widgets/Text/SRichTextBlock.cpp b/Engine/Source/Runtime/Slate/Private/Widgets/Text/SRichTextBlock.cpp
--- a/Engine/Source/Runtime/Slate/Private/Widgets/Text/SRichTextBlock.cpp
+++ b/Engine/Source/Runtime/Slate/Private/Widgets/Text/SRichTextBlock.cpp
@@ -219,5 +219,9 @@ void SRichTextBlock::SetHighlightText( const FText& InHighlightText )
 	TextLayout->SetRunRenderers( TextHighlights );
 }
 
+void SRichTextBlock::Refresh()
+{
+	TextLayout->DirtyLayout();
+}
 
 #endif //WITH_FANCY_TEXT
\ No newline at end of file
diff --git a/Engine/Source/Runtime/Slate/Public/Framework/Text/ITextDecorator.h b/Engine/Source/Runtime/Slate/Public/Framework/Text/ITextDecorator.h
--- a/Engine/Source/Runtime/Slate/Public/Framework/Text/ITextDecorator.h
+++ b/Engine/Source/Runtime/Slate/Public/Framework/Text/ITextDecorator.h
@@ -68,7 +68,7 @@ public:
 
 	virtual bool Supports( const FTextRunParseResults& RunInfo, const FString& Text ) const = 0;
 
-	virtual TSharedRef< ISlateRun > Create( const FTextRunParseResults& RunInfo, const FString& OriginalText, const TSharedRef< FString >& ModelText, const ISlateStyle* Style ) = 0;
+	virtual TSharedRef< ISlateRun > Create(const TSharedRef<class FTextLayout>& TextLayout, const FTextRunParseResults& RunInfo, const FString& OriginalText, const TSharedRef< FString >& ModelText, const ISlateStyle* Style ) = 0;
 };
 
 #endif //WITH_FANCY_TEXT
\ No newline at end of file
diff --git a/Engine/Source/Runtime/Slate/Public/Framework/Text/TextDecorators.h b/Engine/Source/Runtime/Slate/Public/Framework/Text/TextDecorators.h
--- a/Engine/Source/Runtime/Slate/Public/Framework/Text/TextDecorators.h
+++ b/Engine/Source/Runtime/Slate/Public/Framework/Text/TextDecorators.h
@@ -20,7 +20,7 @@ public:
 
 	virtual bool Supports( const FTextRunParseResults& RunParseResult, const FString& Text ) const override;
 
-	virtual TSharedRef< ISlateRun > Create( const FTextRunParseResults& RunParseResult, const FString& OriginalText, const TSharedRef< FString >& InOutModelText, const ISlateStyle* Style ) override;
+	virtual TSharedRef< ISlateRun > Create(const TSharedRef<class FTextLayout>& TextLayout, const FTextRunParseResults& RunParseResult, const FString& OriginalText, const TSharedRef< FString >& InOutModelText, const ISlateStyle* Style) override;
 
 private:
 
@@ -43,7 +43,7 @@ public:
 
 	virtual bool Supports( const FTextRunParseResults& RunParseResult, const FString& Text ) const override;
 
-	virtual TSharedRef< ISlateRun > Create( const FTextRunParseResults& RunParseResult, const FString& OriginalText, const TSharedRef< FString >& InOutModelText, const ISlateStyle* Style ) override;
+	virtual TSharedRef< ISlateRun > Create(const TSharedRef<class FTextLayout>& TextLayout, const FTextRunParseResults& RunParseResult, const FString& OriginalText, const TSharedRef< FString >& InOutModelText, const ISlateStyle* Style) override;
 
 private:
 
@@ -66,7 +66,7 @@ public:
 
 	virtual bool Supports( const FTextRunParseResults& RunParseResult, const FString& Text ) const override;
 
-	virtual TSharedRef< ISlateRun > Create( const FTextRunParseResults& RunParseResult, const FString& OriginalText, const TSharedRef< FString >& InOutModelText, const ISlateStyle* Style ) override;
+	virtual TSharedRef< ISlateRun > Create(const TSharedRef<class FTextLayout>& TextLayout, const FTextRunParseResults& RunParseResult, const FString& OriginalText, const TSharedRef< FString >& InOutModelText, const ISlateStyle* Style) override;
 
 private:
 
diff --git a/Engine/Source/Runtime/Slate/Public/Framework/Text/TextLayout.h b/Engine/Source/Runtime/Slate/Public/Framework/Text/TextLayout.h
--- a/Engine/Source/Runtime/Slate/Public/Framework/Text/TextLayout.h
+++ b/Engine/Source/Runtime/Slate/Public/Framework/Text/TextLayout.h
@@ -126,6 +126,7 @@ public:
 };
 
 class SLATE_API FTextLayout
+	: public TSharedFromThis<FTextLayout>
 {
 public:
 
@@ -338,6 +339,10 @@ public:
 
 	virtual void UpdateHighlights();
 
+	void DirtyRunLayout(const TSharedRef<const IRun>& Run);
+
+	void DirtyLayout();
+
 	int32 GetLineViewIndexForTextLocation(const TArray< FTextLayout::FLineView >& LineViews, const FTextLocation& Location, const bool bPerformInclusiveBoundsCheck);
 
 	/**
diff --git a/Engine/Source/Runtime/Slate/Public/Widgets/Text/SRichTextBlock.h b/Engine/Source/Runtime/Slate/Public/Widgets/Text/SRichTextBlock.h
--- a/Engine/Source/Runtime/Slate/Public/Widgets/Text/SRichTextBlock.h
+++ b/Engine/Source/Runtime/Slate/Public/Widgets/Text/SRichTextBlock.h
@@ -128,6 +128,10 @@ public:
 
 	void SetHighlightText( const FText& InHighlightText );
 
+	/**
+	 * Causes the text to reflow it's layout
+	 */
+	void Refresh();
 
 private:
 
