Close all but this tab for assets all over the place.

diff --git a/Engine/Source/Editor/MainFrame/Private/Frame/MainFrameHandler.h b/Engine/Source/Editor/MainFrame/Private/Frame/MainFrameHandler.h
--- a/Engine/Source/Editor/MainFrame/Private/Frame/MainFrameHandler.h
+++ b/Engine/Source/Editor/MainFrame/Private/Frame/MainFrameHandler.h
@@ -237,6 +237,8 @@ public:
 	{
 		TSharedRef<FGlobalTabmanager> GlobalTabManager = FGlobalTabmanager::Get();
 		
+		GlobalTabManager->SetMainTab(MainTab.ToSharedRef());
+
 		// Persistent layouts should get stored using the specified method.
 		GlobalTabManager->SetOnPersistLayout(FTabManager::FOnPersistLayout::CreateRaw(this, &FMainFrameHandler::HandleTabManagerPersistLayout));
 		
diff --git a/Engine/Source/Runtime/Slate/Private/Framework/Docking/SDockingTabStack.cpp b/Engine/Source/Runtime/Slate/Private/Framework/Docking/SDockingTabStack.cpp
--- a/Engine/Source/Runtime/Slate/Private/Framework/Docking/SDockingTabStack.cpp
+++ b/Engine/Source/Runtime/Slate/Private/Framework/Docking/SDockingTabStack.cpp
@@ -457,16 +457,15 @@ void SDockingTabStack::CloseAllButForegroundTab(ETabsToClose TabsToClose)
 		{
 			const TSharedRef<SDockTab>& Tab = TabWell->GetTabs()[DestroyIndex];
 
-			const bool bCanClose = (TabsToClose == CloseDocumentsAndTools) || (Tab->GetTabRole() == ETabRole::DocumentTab);
+			const bool bCanClose = 
+				(TabsToClose == CloseAllTabs) ||
+				(TabsToClose == CloseDocumentTabs && Tab->GetTabRole() == ETabRole::DocumentTab) ||
+				(TabsToClose == CloseDocumentAndMajorTabs && (Tab->GetTabRole() == ETabRole::DocumentTab || Tab->GetTabRole() == ETabRole::MajorTab));
 
-			if ((Tab == ForegroundTab) || !bCanClose)
+			if ((Tab == ForegroundTab) || !bCanClose || !Tab->RequestCloseTab())
 			{
 				++DestroyIndex;
 			}
-			else
-			{
-				Tab->RequestCloseTab();
-			}
 		}
 	}
 }
@@ -538,9 +537,9 @@ TSharedRef<SWidget> SDockingTabStack::MakeContextMenu()
 
 			// If the active tab is a document tab, and there is more than one open in this tab well, offer to close the others
 			TSharedPtr<SDockTab> ForegroundTabPtr = TabWell->GetForegroundTab();
-			if (ForegroundTabPtr.IsValid() && (ForegroundTabPtr->GetTabRole() == ETabRole::DocumentTab) && (TabWell->GetNumTabs() > 1))
+			if (ForegroundTabPtr.IsValid() && (ForegroundTabPtr->GetTabRole() == ETabRole::DocumentTab || ForegroundTabPtr->GetTabRole() == ETabRole::MajorTab) && (TabWell->GetNumTabs() > 1))
 			{
-				const ETabsToClose TabsToClose = CloseDocumentTabs;
+				const ETabsToClose TabsToClose = CloseDocumentAndMajorTabs;
 
 				MenuBuilder.AddMenuEntry(
 					LOCTEXT("CloseOtherTabs", "Close Other Tabs"),
diff --git a/Engine/Source/Runtime/Slate/Private/Framework/Docking/SDockingTabStack.h b/Engine/Source/Runtime/Slate/Private/Framework/Docking/SDockingTabStack.h
--- a/Engine/Source/Runtime/Slate/Private/Framework/Docking/SDockingTabStack.h
+++ b/Engine/Source/Runtime/Slate/Private/Framework/Docking/SDockingTabStack.h
@@ -114,7 +114,8 @@ protected:
 	enum ETabsToClose
 	{
 		CloseDocumentTabs,
-		CloseDocumentsAndTools
+		CloseDocumentAndMajorTabs,
+		CloseAllTabs
 	};
 	/**
 	 * Close all the background tabs.
diff --git a/Engine/Source/Runtime/Slate/Private/Framework/Docking/TabManager.cpp b/Engine/Source/Runtime/Slate/Private/Framework/Docking/TabManager.cpp
--- a/Engine/Source/Runtime/Slate/Private/Framework/Docking/TabManager.cpp
+++ b/Engine/Source/Runtime/Slate/Private/Framework/Docking/TabManager.cpp
@@ -492,6 +492,16 @@ void FTabManager::UpdateMainMenu(bool const bForce)
 #endif
 }
 
+void FTabManager::SetMainTab(const TSharedRef<const SDockTab>& InTab)
+{
+	MainNonCloseableTab = InTab;
+}
+
+bool FTabManager::IsTabCloseable(const TSharedRef<const SDockTab>& InTab) const
+{
+	return !(MainNonCloseableTab.Pin() == InTab);
+}
+
 TSharedPtr<FTabManager::FStack> FTabManager::FLayoutNode::AsStack()
 {
 	return TSharedPtr<FTabManager::FStack>();
diff --git a/Engine/Source/Runtime/Slate/Private/Widgets/Docking/SDockTab.cpp b/Engine/Source/Runtime/Slate/Private/Widgets/Docking/SDockTab.cpp
--- a/Engine/Source/Runtime/Slate/Private/Widgets/Docking/SDockTab.cpp
+++ b/Engine/Source/Runtime/Slate/Private/Widgets/Docking/SDockTab.cpp
@@ -525,6 +525,7 @@ void SDockTab::Construct( const FArguments& InArgs )
 					.ButtonStyle( CloseButtonStyle )
 					.OnClicked( this, &SDockTab::OnCloseButtonClicked )
 					.ContentPadding( 0 )
+					.Visibility(this, &SDockTab::HandleIsCloseButtonVisible)
 					[
 						SNew(SSpacer)
 						.Size( CloseButtonStyle->Normal.ImageSize )
@@ -714,13 +715,18 @@ FReply SDockTab::OnCloseButtonClicked()
 	return FReply::Handled();
 }
 
+EVisibility SDockTab::HandleIsCloseButtonVisible() const
+{
+	return MyTabManager.Pin()->IsTabCloseable(SharedThis(this)) ? EVisibility::Visible : EVisibility::Hidden;
+}
+
 bool SDockTab::CanCloseTab() const
 {
-	const bool bCanCloseTabNow = !OnCanCloseTab.IsBound() || OnCanCloseTab.Execute(); 
+	const bool bCanCloseTabNow = MyTabManager.Pin()->IsTabCloseable(SharedThis(this)) && (!OnCanCloseTab.IsBound() || OnCanCloseTab.Execute());
 	return bCanCloseTabNow;
 }
 
-void SDockTab::RequestCloseTab()
+bool SDockTab::RequestCloseTab()
 {	
 	this->PersistVisualState();
 	// The tab can be closed if the delegate is not bound or if the delegate call indicates we cannot close it
@@ -729,6 +735,7 @@ void SDockTab::RequestCloseTab()
 	{
 		RemoveTabFromParent();
 	}
+	return bCanCloseTabNow;
 }
 
 void SDockTab::PersistVisualState()
diff --git a/Engine/Source/Runtime/Slate/Public/Framework/Docking/TabManager.h b/Engine/Source/Runtime/Slate/Public/Framework/Docking/TabManager.h
--- a/Engine/Source/Runtime/Slate/Public/Framework/Docking/TabManager.h
+++ b/Engine/Source/Runtime/Slate/Public/Framework/Docking/TabManager.h
@@ -597,6 +597,12 @@ class SLATE_API FTabManager : public TSharedFromThis<FTabManager>
 		 */
 		void UpdateMainMenu(bool const bForce);
 
+		/** Provide a tab that will be the main tab and cannot be closed. */
+		void SetMainTab(const TSharedRef<const SDockTab>& InTab);
+
+		/** @return if the provided tab can be closed. */
+		bool IsTabCloseable(const TSharedRef<const SDockTab>& InTab) const;
+
 	protected:
 		void InvokeTabForMenu( FName TabId );
 
@@ -702,6 +708,9 @@ class SLATE_API FTabManager : public TSharedFromThis<FTabManager>
 		 * We will ignore that request them.
 		 */
 		bool bIsSavingVisualState;
+
+		/** The main tab, this tab cannot be closed. */
+		TWeakPtr<const SDockTab> MainNonCloseableTab;
 };
 
 
@@ -860,4 +869,5 @@ private:
 
 	/** Keeps track of the running-maximum number of unique parent windows in all dock areas and sub-managers during this session */
 	int32 AllAreasWindowMaxCount;
+
 };
