#UE4 const correctness for SListView::GetNumItemsSelected and ::GetSelectedItems

diff --git a/Engine/Source/Runtime/Slate/Public/Widgets/Views/SListView.h b/Engine/Source/Runtime/Slate/Public/Widgets/Views/SListView.h
--- a/Engine/Source/Runtime/Slate/Public/Widgets/Views/SListView.h
+++ b/Engine/Source/Runtime/Slate/Public/Widgets/Views/SListView.h
@@ -979,7 +979,7 @@ public:
 	 *
 	 * @return Number of selected items.
 	 */
-	int32 GetNumItemsSelected()
+	int32 GetNumItemsSelected() const
 	{
 		return SelectedItems.Num();
 	}
@@ -989,7 +989,7 @@ public:
 	 *
 	 * @return	List of selected item indices (in no particular order)
 	 */
-	TArray< ItemType > GetSelectedItems()
+	TArray< ItemType > GetSelectedItems() const
 	{
 		TArray< ItemType > SelectedItemArray;
 		SelectedItemArray.Empty( SelectedItems.Num() );
