Missed a file.

diff --git a/Engine/Source/Runtime/Slate/Public/Widgets/Docking/SDockTab.h b/Engine/Source/Runtime/Slate/Public/Widgets/Docking/SDockTab.h
--- a/Engine/Source/Runtime/Slate/Public/Widgets/Docking/SDockTab.h
+++ b/Engine/Source/Runtime/Slate/Public/Widgets/Docking/SDockTab.h
@@ -187,7 +187,7 @@ public:
 	bool CanCloseTab() const;
 
 	/** Requests that the tab be closed.  Tabs may prevent closing depending on their state */	
-	void RequestCloseTab();
+	bool RequestCloseTab();
 
 	/** A chance for the tab's content to save any internal layout info */
 	void PersistVisualState();
@@ -288,6 +288,9 @@ protected:
 	/** Specify the TabId that was used to spawn this tab. */
 	void SetLayoutIdentifier( const FTabId& TabId );
 
+	/** @return if the close button should be visible. */
+	EVisibility HandleIsCloseButtonVisible() const;
+
 protected:
 
 	/** The tab manager that created this tab. */
