Implemented scrolling for multi-line editable text

SMultiLineEditableText can now be provided with two scrollbars that it uses to manage its scroll offset. These scrollbars are typically provided by SMultiLineEditableTextBox, however you can also provide your own external scrollbars if required.

FTextLayout is now aware of its current view area and scroll offset, and will adjust the positions of any lines as the scroll offset is updated.

SMultiLineEditableText will always try and keep the cursor in view as long as you're interacting with the text area.

Also fixed some issues with SMultiLineEditableText and SRichTextBlock doubly compensating for the justification performed by FTextLayout. This could cause picking issues when lines contained lots of trailing whitespace, and is now handled by FTextLayout being aware that it needs to update its justified flow when its visible area changes size.

diff --git a/Engine/Source/Editor/EditorStyle/Private/SlateEditorStyle.cpp b/Engine/Source/Editor/EditorStyle/Private/SlateEditorStyle.cpp
--- a/Engine/Source/Editor/EditorStyle/Private/SlateEditorStyle.cpp
+++ b/Engine/Source/Editor/EditorStyle/Private/SlateEditorStyle.cpp
@@ -427,12 +427,24 @@ void FSlateEditorStyle::FStyle::SetupGeneralStyles()
 		
 	}
 
+	// Scrollbar
+	const FScrollBarStyle ScrollBar = FScrollBarStyle()
+		.SetVerticalBackgroundImage( IMAGE_BRUSH( "Common/Scrollbar_Background_Vertical", FVector2D(8,8) ) )
+		.SetHorizontalBackgroundImage( IMAGE_BRUSH( "Common/Scrollbar_Background_Horizontal", FVector2D(8,8) ) )
+		.SetNormalThumbImage( BOX_BRUSH( "Common/Scrollbar_Thumb", FMargin(4.f/16.f) ) )
+		.SetDraggedThumbImage( BOX_BRUSH( "Common/Scrollbar_Thumb", FMargin(4.f/16.f) ) )
+		.SetHoveredThumbImage( BOX_BRUSH( "Common/Scrollbar_Thumb", FMargin(4.f/16.f) ) );
+	{
+		Set( "Scrollbar", ScrollBar );
+	}
+
 	// EditableTextBox
 	NormalEditableTextBoxStyle = FEditableTextBoxStyle()
 		.SetBackgroundImageNormal( BOX_BRUSH( "Common/TextBox", FMargin(4.0f/16.0f) ) )
 		.SetBackgroundImageHovered( BOX_BRUSH( "Common/TextBox_Hovered", FMargin(4.0f/16.0f) ) )
 		.SetBackgroundImageFocused( BOX_BRUSH( "Common/TextBox_Hovered", FMargin(4.0f/16.0f) ) )
-		.SetBackgroundImageReadOnly( BOX_BRUSH( "Common/TextBox_ReadOnly", FMargin(4.0f/16.0f) ) );
+		.SetBackgroundImageReadOnly( BOX_BRUSH( "Common/TextBox_ReadOnly", FMargin(4.0f/16.0f) ) )
+		.SetScrollBarStyle( ScrollBar );
 	{
 		Set( "EditableTextBox.Background.Normal", new BOX_BRUSH( "Common/TextBox", FMargin(4.0f/16.0f) ) );
 		Set( "EditableTextBox.Background.Hovered", new BOX_BRUSH( "Common/TextBox_Hovered", FMargin(4.0f/16.0f) ) );
@@ -450,7 +462,8 @@ void FSlateEditorStyle::FStyle::SetupGeneralStyles()
 			.SetBackgroundImageNormal( *SpecialEditableTextImageNormal )
 			.SetBackgroundImageHovered( BOX_BRUSH( "Common/TextBox_Special_Hovered", FMargin(8.0f/32.0f) ) )
 			.SetBackgroundImageFocused( BOX_BRUSH( "Common/TextBox_Special_Hovered", FMargin(8.0f/32.0f) ) )
-			.SetBackgroundImageReadOnly( BOX_BRUSH( "Common/TextBox_ReadOnly", FMargin(4.0f/16.0f) ) );
+			.SetBackgroundImageReadOnly( BOX_BRUSH( "Common/TextBox_ReadOnly", FMargin(4.0f/16.0f) ) )
+			.SetScrollBarStyle( ScrollBar );
 		Set( "SpecialEditableTextBox", SpecialEditableTextBoxStyle );
 
 		Set( "SearchBox.ActiveBorder", new BOX_BRUSH( "Common/TextBox_Special_Active", FMargin(8.0f/32.0f) ) );
@@ -511,7 +524,8 @@ void FSlateEditorStyle::FStyle::SetupGeneralStyles()
 			.SetBackgroundImageNormal( BOX_BRUSH( "Common/TextBox", FMargin(4.0f/16.0f) ) )
 			.SetBackgroundImageHovered( BOX_BRUSH( "Common/TextBox_Hovered", FMargin(4.0f/16.0f) ) )
 			.SetBackgroundImageFocused( BOX_BRUSH( "Common/TextBox_Hovered", FMargin(4.0f/16.0f) ) )
-			.SetBackgroundImageReadOnly( BOX_BRUSH( "Common/TextBox_ReadOnly", FMargin(4.0f/16.0f) ) );
+			.SetBackgroundImageReadOnly( BOX_BRUSH( "Common/TextBox_ReadOnly", FMargin(4.0f/16.0f) ) )
+			.SetScrollBarStyle( ScrollBar );
 
 		FInlineEditableTextBlockStyle InlineEditableTextBlockStyle = FInlineEditableTextBlockStyle()
 			.SetTextStyle(InlineEditableTextBlockReadOnly)
@@ -588,18 +602,6 @@ void FSlateEditorStyle::FStyle::SetupGeneralStyles()
 	}
 #endif // WITH_EDITOR || IS_PROGRAM
 
-	// Scrollbar
-	{
-		Set( "Scrollbar", FScrollBarStyle()
-			.SetVerticalBackgroundImage( IMAGE_BRUSH( "Common/Scrollbar_Background_Vertical", FVector2D(8,8) ) )
-			.SetHorizontalBackgroundImage( IMAGE_BRUSH( "Common/Scrollbar_Background_Horizontal", FVector2D(8,8) ) )
-			.SetNormalThumbImage( BOX_BRUSH( "Common/Scrollbar_Thumb", FMargin(4.f/16.f) ) )
-			.SetDraggedThumbImage( BOX_BRUSH( "Common/Scrollbar_Thumb", FMargin(4.f/16.f) ) )
-			.SetHoveredThumbImage( BOX_BRUSH( "Common/Scrollbar_Thumb", FMargin(4.f/16.f) ) )
-			);
-	}
-
-
 	// Expandable button
 	{
 		Set( "ExpandableButton.Background", new BOX_BRUSH( "Common/Button", 8.0f/32.0f ) );
@@ -2916,6 +2918,8 @@ void FSlateEditorStyle::FStyle::SetupProfilerStyle()
 	
 void FSlateEditorStyle::FStyle::SetupGraphEditorStyles()
 {
+	const FScrollBarStyle ScrollBar = GetWidgetStyle<FScrollBarStyle>( "ScrollBar" );
+
 	// Graph Editor
 #if WITH_EDITOR || IS_PROGRAM
 	{
@@ -2995,7 +2999,8 @@ void FSlateEditorStyle::FStyle::SetupGraphEditorStyles()
 					.SetBackgroundImageNormal( BOX_BRUSH( "Common/TextBox", FMargin(4.0f/16.0f) ) )
 					.SetBackgroundImageHovered( BOX_BRUSH( "Common/TextBox_Hovered", FMargin(4.0f/16.0f) ) )
 					.SetBackgroundImageFocused( BOX_BRUSH( "Common/TextBox_Hovered", FMargin(4.0f/16.0f) ) )
-					.SetBackgroundImageReadOnly( BOX_BRUSH( "Common/TextBox_ReadOnly", FMargin(4.0f/16.0f) ) );
+					.SetBackgroundImageReadOnly( BOX_BRUSH( "Common/TextBox_ReadOnly", FMargin(4.0f/16.0f) ) )
+					.SetScrollBarStyle( ScrollBar );
 				Set( "Graph.StateNode.NodeTitleEditableText", GraphStateNodeTitleEditableText );
 
 				Set( "Graph.StateNode.NodeTitleInlineEditableText", FInlineEditableTextBlockStyle()
@@ -3182,7 +3187,8 @@ void FSlateEditorStyle::FStyle::SetupGraphEditorStyles()
 				.SetBackgroundImageNormal( BOX_BRUSH( "Common/TextBox", FMargin(4.0f/16.0f) ) )
 				.SetBackgroundImageHovered( BOX_BRUSH( "Common/TextBox_Hovered", FMargin(4.0f/16.0f) ) )
 				.SetBackgroundImageFocused( BOX_BRUSH( "Common/TextBox_Hovered", FMargin(4.0f/16.0f) ) )
-				.SetBackgroundImageReadOnly( BOX_BRUSH( "Common/TextBox_ReadOnly", FMargin(4.0f/16.0f) ) );
+				.SetBackgroundImageReadOnly( BOX_BRUSH( "Common/TextBox_ReadOnly", FMargin(4.0f/16.0f) ) )
+				.SetScrollBarStyle( ScrollBar );
 			Set( "Graph.Node.NodeTitleEditableText", GraphNodeTitleEditableText );
 
 			Set( "Graph.Node.NodeTitleInlineEditableText", FInlineEditableTextBlockStyle()
@@ -3209,7 +3215,8 @@ void FSlateEditorStyle::FStyle::SetupGraphEditorStyles()
 				.SetBackgroundImageNormal( BOX_BRUSH( "Common/TextBox", FMargin(4.0f/16.0f) ) )
 				.SetBackgroundImageHovered( BOX_BRUSH( "Common/TextBox_Hovered", FMargin(4.0f/16.0f) ) )
 				.SetBackgroundImageFocused( BOX_BRUSH( "Common/TextBox_Hovered", FMargin(4.0f/16.0f) ) )
-				.SetBackgroundImageReadOnly( BOX_BRUSH( "Common/TextBox_ReadOnly", FMargin(4.0f/16.0f) ) );
+				.SetBackgroundImageReadOnly( BOX_BRUSH( "Common/TextBox_ReadOnly", FMargin(4.0f/16.0f) ) )
+				.SetScrollBarStyle( ScrollBar );
 			Set( "Graph.CommentBlock.TitleEditableText", GraphCommentBlockTitleEditableText );
 
 			Set( "Graph.CommentBlock.TitleInlineEditableText", FInlineEditableTextBlockStyle()
@@ -3458,6 +3465,7 @@ void FSlateEditorStyle::FStyle::SetupGraphEditorStyles()
 					.SetBackgroundImageHovered( BOX_BRUSH( "Graph/CommonWidgets/TextBox_Hovered", FMargin(4.0f/16.0f) ) )
 					.SetBackgroundImageFocused( BOX_BRUSH( "Graph/CommonWidgets/TextBox_Hovered", FMargin(4.0f/16.0f) ) )
 					.SetBackgroundImageReadOnly( BOX_BRUSH( "Graph/CommonWidgets/TextBox", FMargin(4.0f/16.0f) ) )
+					.SetScrollBarStyle( ScrollBar )
 					);
 			}
 
@@ -3468,6 +3476,7 @@ void FSlateEditorStyle::FStyle::SetupGraphEditorStyles()
 					.SetBackgroundImageHovered( BOX_BRUSH( "Graph/CommonWidgets/TextBox_Hovered", FMargin(4.0f/16.0f) ) )
 					.SetBackgroundImageFocused( BOX_BRUSH( "Graph/CommonWidgets/TextBox_Hovered", FMargin(4.0f/16.0f) ) )
 					.SetBackgroundImageReadOnly( BOX_BRUSH( "Graph/CommonWidgets/TextBox", FMargin(4.0f/16.0f) ) )
+					.SetScrollBarStyle( ScrollBar )
 					.SetForegroundColor( FLinearColor::White )
 					.SetBackgroundColor( FLinearColor::Blue )
 					);
diff --git a/Engine/Source/Runtime/Slate/Private/Framework/Text/SlateTextLayout.cpp b/Engine/Source/Runtime/Slate/Private/Framework/Text/SlateTextLayout.cpp
--- a/Engine/Source/Runtime/Slate/Private/Framework/Text/SlateTextLayout.cpp
+++ b/Engine/Source/Runtime/Slate/Private/Framework/Text/SlateTextLayout.cpp
@@ -27,11 +27,11 @@ FChildren* FSlateTextLayout::GetChildren()
 
 void FSlateTextLayout::ArrangeChildren( const FGeometry& AllottedGeometry, FArrangedChildren& ArrangedChildren ) const
 {
-	for (int LineIndex = 0; LineIndex < LineViews.Num(); LineIndex++)
+	for (int32 LineIndex = 0; LineIndex < LineViews.Num(); LineIndex++)
 	{
 		const FTextLayout::FLineView& LineView = LineViews[ LineIndex ];
 
-		for (int BlockIndex = 0; BlockIndex < LineView.Blocks.Num(); BlockIndex++)
+		for (int32 BlockIndex = 0; BlockIndex < LineView.Blocks.Num(); BlockIndex++)
 		{
 			const TSharedRef< ILayoutBlock > Block = LineView.Blocks[ BlockIndex ];
 			const TSharedRef< ISlateRun > Run = StaticCastSharedRef< ISlateRun >( Block->GetRun() );
@@ -51,6 +51,14 @@ int32 FSlateTextLayout::OnPaint( const FPaintArgs& Args, const FTextBlockStyle&
 
 	for (const FTextLayout::FLineView& LineView : LineViews)
 	{
+		// Is this line visible?
+		const FSlateRect LineViewRect(AllottedGeometry.AbsolutePosition + LineView.Offset, AllottedGeometry.AbsolutePosition + LineView.Offset + LineView.Size);
+		const FSlateRect VisibleLineView = ClippingRect.IntersectionWith(LineViewRect);
+		if (VisibleLineView.IsEmpty())
+		{
+			continue;
+		}
+
 		// Render any underlays for this line
 		const int32 HighestUnderlayLayerId = OnPaintHighlights( Args, LineView, LineView.UnderlayHighlights, DefaultTextStyle, AllottedGeometry, ClippingRect, OutDrawElements, LayerId, InWidgetStyle, bParentEnabled );
 
@@ -126,16 +134,16 @@ void FSlateTextLayout::AggregateChildren()
 {
 	Children.Empty();
 	const TArray< FLineModel >& LineModels = GetLineModels();
-	for (int LineModelIndex = 0; LineModelIndex < LineModels.Num(); LineModelIndex++)
+	for (int32 LineModelIndex = 0; LineModelIndex < LineModels.Num(); LineModelIndex++)
 	{
 		const FLineModel& LineModel = LineModels[ LineModelIndex ];
-		for (int RunIndex = 0; RunIndex < LineModel.Runs.Num(); RunIndex++)
+		for (int32 RunIndex = 0; RunIndex < LineModel.Runs.Num(); RunIndex++)
 		{
 			const FRunModel& LineRun = LineModel.Runs[ RunIndex ];
 			const TSharedRef< ISlateRun > SlateRun = StaticCastSharedRef< ISlateRun >( LineRun.GetRun() );
 
 			const TArray< TSharedRef<SWidget> >& RunChildren = SlateRun->GetChildren();
-			for (int ChildIndex = 0; ChildIndex < RunChildren.Num(); ChildIndex++)
+			for (int32 ChildIndex = 0; ChildIndex < RunChildren.Num(); ChildIndex++)
 			{
 				const TSharedRef< SWidget >& Child = RunChildren[ ChildIndex ];
 				Children.Add( Child );
diff --git a/Engine/Source/Runtime/Slate/Private/Framework/Text/TextLayout.cpp b/Engine/Source/Runtime/Slate/Private/Framework/Text/TextLayout.cpp
--- a/Engine/Source/Runtime/Slate/Private/Framework/Text/TextLayout.cpp
+++ b/Engine/Source/Runtime/Slate/Private/Framework/Text/TextLayout.cpp
@@ -217,12 +217,14 @@ void FTextLayout::CreateLineViewBlocks( int32 LineModelIndex, const int32 StopIn
 	}
 
 	FVector2D LineSize( ForceInitToZero );
-	FVector2D CurrentOffset( Margin.Left, Margin.Top + DrawSize.Y );
+	
+	// Use a negative scroll offset since positive scrolling moves things negatively in screen space
+	FVector2D CurrentOffset( Margin.Left - ScrollOffset.X, Margin.Top + DrawSize.Y - ScrollOffset.Y );
 
 	if ( OutSoftLine.Num() > 0 )
 	{
 		float CurrentHorizontalPos = 0.0f;
-		for (int Index = 0; Index < OutSoftLine.Num(); Index++)
+		for (int32 Index = 0; Index < OutSoftLine.Num(); Index++)
 		{
 			const TSharedRef< ILayoutBlock > Block = OutSoftLine[ Index ];
 			const TSharedRef< IRun > Run = Block->GetRun();
@@ -263,63 +265,27 @@ void FTextLayout::JustifyLayout()
 		return;
 	}
 
-	const float LayoutWidthNoMargin = DrawSize.X - Margin.GetTotalSpaceAlong<Orient_Horizontal>() * Scale;
+	const float LayoutWidthNoMargin = FMath::Max(DrawSize.X - Margin.GetTotalSpaceAlong<Orient_Horizontal>() * Scale, ViewSize.X);
 
-	for (int LineViewIndex = 0; LineViewIndex < LineViews.Num(); LineViewIndex++)
-	{
-		FLineView& LineView = LineViews[ LineViewIndex ];
-		const float ExtraSpace = LayoutWidthNoMargin - LineView.Size.X;
-
-		FVector2D BlockOffset( ForceInitToZero );
-		if ( Justification == ETextJustify::Center )
-		{
-			BlockOffset = FVector2D( ExtraSpace / 2, 0 );
-		}
-		else if ( Justification == ETextJustify::Right )
-		{
-			BlockOffset = FVector2D( ExtraSpace, 0 );
-		}
-
-		for (int BlockIndex = 0; BlockIndex < LineView.Blocks.Num(); BlockIndex++)
-		{
-			const TSharedRef< ILayoutBlock >& Block = LineView.Blocks[ BlockIndex ];
-			Block->SetLocationOffset( Block->GetLocationOffset() + BlockOffset );
-		}
-	}
-}
-
-void FTextLayout::JustifyHighlights()
-{
-	if ( Justification == ETextJustify::Left )
-	{
-		return;
-	}
-
-	const float LayoutWidthNoMargin = DrawSize.X - Margin.GetTotalSpaceAlong<Orient_Horizontal>() * Scale;
-
-	for (int LineViewIndex = 0; LineViewIndex < LineViews.Num(); LineViewIndex++)
+	for (FLineView& LineView : LineViews)
 	{
-		FLineView& LineView = LineViews[ LineViewIndex ];
 		const float ExtraSpace = LayoutWidthNoMargin - LineView.Size.X;
 
-		float HighlightOffset = 0;
+		FVector2D OffsetAdjustment( ForceInitToZero );
 		if ( Justification == ETextJustify::Center )
 		{
-			HighlightOffset = ExtraSpace / 2;
+			OffsetAdjustment = FVector2D( ExtraSpace * 0.5f, 0 );
 		}
 		else if ( Justification == ETextJustify::Right )
 		{
-			HighlightOffset = ExtraSpace;
+			OffsetAdjustment = FVector2D( ExtraSpace, 0 );
 		}
 
-		for (FLineViewHighlight& Highlight : LineView.UnderlayHighlights)
-		{
-			Highlight.OffsetX += HighlightOffset;
-		}
+		LineView.Offset += OffsetAdjustment;
 
-		for (FLineViewHighlight& Highlight : LineView.OverlayHighlights)
+		for (const TSharedRef< ILayoutBlock >& Block : LineView.Blocks)
 		{
-			Highlight.OffsetX += HighlightOffset;
+			Block->SetLocationOffset( Block->GetLocationOffset() + OffsetAdjustment );
 		}
 	}
 }
@@ -334,7 +300,7 @@ void FTextLayout::FlowLayout()
 
 	TArray< TSharedRef< ILayoutBlock > > SoftLine;
 
-	for (int LineModelIndex = 0; LineModelIndex < LineModels.Num(); LineModelIndex++)
+	for (int32 LineModelIndex = 0; LineModelIndex < LineModels.Num(); LineModelIndex++)
 	{
 		const FLineModel& LineModel = LineModels[ LineModelIndex ];
 
@@ -359,7 +325,7 @@ void FTextLayout::FlowLayout()
 		}
 		else
 		{
-			for (int BreakIndex = 0; BreakIndex < LineModel.BreakCandidates.Num(); BreakIndex++)
+			for (int32 BreakIndex = 0; BreakIndex < LineModel.BreakCandidates.Num(); BreakIndex++)
 			{
 				const FBreakCandidate& Break = LineModel.BreakCandidates[ BreakIndex ];
 
@@ -487,10 +453,10 @@ void FTextLayout::FlowHighlights()
 
 void FTextLayout::EndLayout()
 {
-	for (int LineModelIndex = 0; LineModelIndex < LineModels.Num(); LineModelIndex++)
+	for (int32 LineModelIndex = 0; LineModelIndex < LineModels.Num(); LineModelIndex++)
 	{
 		FLineModel& LineModel = LineModels[ LineModelIndex ];
-		for (int RunIndex = 0; RunIndex < LineModel.Runs.Num(); RunIndex++)
+		for (int32 RunIndex = 0; RunIndex < LineModel.Runs.Num(); RunIndex++)
 		{
 			LineModel.Runs[ RunIndex ].EndLayout();
 		}
@@ -499,10 +465,10 @@ void FTextLayout::EndLayout()
 
 void FTextLayout::BeginLayout()
 {
-	for (int LineModelIndex = 0; LineModelIndex < LineModels.Num(); LineModelIndex++)
+	for (int32 LineModelIndex = 0; LineModelIndex < LineModels.Num(); LineModelIndex++)
 	{
 		FLineModel& LineModel = LineModels[ LineModelIndex ];
-		for (int RunIndex = 0; RunIndex < LineModel.Runs.Num(); RunIndex++)
+		for (int32 RunIndex = 0; RunIndex < LineModel.Runs.Num(); RunIndex++)
 		{
 			LineModel.Runs[ RunIndex ].BeginLayout();
 		}
@@ -523,7 +489,7 @@ void FTextLayout::CreateWrappingCache()
 		LineBreakIterator = FBreakIterator::CreateLineBreakIterator();
 	}
 
-	for (int LineModelIndex = 0; LineModelIndex < LineModels.Num(); LineModelIndex++)
+	for (int32 LineModelIndex = 0; LineModelIndex < LineModels.Num(); LineModelIndex++)
 	{
 		FLineModel& LineModel = LineModels[ LineModelIndex ];
 
@@ -532,7 +498,7 @@ void FTextLayout::CreateWrappingCache()
 			LineModel.BreakCandidates.Empty();
 			LineModel.HasWrappingInformation = true;
 
-			for (int RunIndex = 0; RunIndex < LineModel.Runs.Num(); RunIndex++)
+			for (int32 RunIndex = 0; RunIndex < LineModel.Runs.Num(); RunIndex++)
 			{
 				LineModel.Runs[ RunIndex ].ClearCache();
 			}
@@ -556,7 +522,7 @@ void FTextLayout::CreateWrappingCache()
 
 void FTextLayout::ClearWrappingCache()
 {
-	for (int LineModelIndex = 0; LineModelIndex < LineModels.Num(); LineModelIndex++)
+	for (int32 LineModelIndex = 0; LineModelIndex < LineModels.Num(); LineModelIndex++)
 	{
 		FLineModel& LineModel = LineModels[ LineModelIndex ];
 		LineModel.HasWrappingInformation = false;
@@ -570,8 +536,11 @@ FTextLayout::FTextLayout()
 	, Scale( 1.0f )
 	, WrappingWidth( 0 )
 	, Margin()
+	, Justification( ETextJustify::Left )
 	, LineHeightPercentage( 1.0f )
 	, DrawSize( ForceInitToZero )
+	, ViewSize( ForceInitToZero )
+	, ScrollOffset( ForceInitToZero )
 	, LineBreakIterator() // Initialized in FTextLayout::CreateWrappingCache if no custom iterator is provided
 	, WordBreakIterator(FBreakIterator::CreateWordBreakIterator())
 {
@@ -612,7 +581,6 @@ void FTextLayout::UpdateLayout()
 void FTextLayout::UpdateHighlights()
 {
 	FlowHighlights();
-	JustifyHighlights();
 
 	DirtyFlags &= ~EDirtyState::Highlights;
 }
@@ -673,7 +641,7 @@ void FTextLayout::SetRunRenderers( const TArray< FTextRunRenderer >& Renderers )
 {
 	ClearRunRenderers();
 
-	for (int Index = 0; Index < Renderers.Num(); Index++)
+	for (int32 Index = 0; Index < Renderers.Num(); Index++)
 	{
 		AddRunRenderer( Renderers[ Index ] );
 	}
@@ -687,7 +655,7 @@ void FTextLayout::AddRunRenderer( const FTextRunRenderer& Renderer )
 
 	// Renderers needs to be in order and not overlap
 	bool bWasInserted = false;
-	for (int Index = 0; Index < LineModel.RunRenderers.Num() && !bWasInserted; Index++)
+	for (int32 Index = 0; Index < LineModel.RunRenderers.Num() && !bWasInserted; Index++)
 	{
 		if ( LineModel.RunRenderers[ Index ].Range.BeginIndex > Renderer.Range.BeginIndex )
 		{
@@ -727,7 +695,7 @@ void FTextLayout::SetLineHighlights( const TArray< FTextLineHighlight >& Highlig
 {
 	ClearLineHighlights();
 
-	for (int Index = 0; Index < Highlights.Num(); Index++)
+	for (int32 Index = 0; Index < Highlights.Num(); Index++)
 	{
 		AddLineHighlight( Highlights[ Index ] );
 	}
@@ -742,7 +710,7 @@ void FTextLayout::AddLineHighlight( const FTextLineHighlight& Highlight )
 
 	// Try and maintain a stable sorted z-order - highlights with the same z-order should just render in the order they were added
 	bool bWasInserted = false;
-	for (int Index = 0; Index < LineModel.LineHighlights.Num() && !bWasInserted; Index++)
+	for (int32 Index = 0; Index < LineModel.LineHighlights.Num() && !bWasInserted; Index++)
 	{
 		if ( LineModel.LineHighlights[ Index ].ZOrder > Highlight.ZOrder )
 		{
@@ -1089,7 +1057,7 @@ bool FTextLayout::JoinLineWithNextLine(int32 LineIndex)
 	LineModel.HasWrappingInformation = false;
 
 	//Iterate through all of the next lines runs and bring them over to the current line
-	for (int RunIndex = 0; RunIndex < NextLineModel.Runs.Num(); RunIndex++)
+	for (int32 RunIndex = 0; RunIndex < NextLineModel.Runs.Num(); RunIndex++)
 	{
 		const TSharedRef<IRun> Run = NextLineModel.Runs[RunIndex].GetRun();
 		FTextRange NewRange = Run->GetTextRange();
@@ -1128,7 +1096,7 @@ bool FTextLayout::SplitLineAt(const FTextLocation& Location)
 	check(LeftLineModel.Text->Len() == BreakLocation);
 
 	bool RunIsToTheLeftOfTheBreakLocation = true;
-	for (int RunIndex = 0; RunIndex < LineModel.Runs.Num(); RunIndex++)
+	for (int32 RunIndex = 0; RunIndex < LineModel.Runs.Num(); RunIndex++)
 	{
 		const TSharedRef<IRun> Run = LineModel.Runs[RunIndex].GetRun();
 		const FTextRange& RunRange = Run->GetTextRange();
@@ -1272,7 +1240,7 @@ void FTextLayout::AddLine( const TSharedRef< FString >& Text, const TArray< TSha
 {
 	FLineModel LineModel( Text );
 
-	for (int Index = 0; Index < Runs.Num(); Index++)
+	for (int32 Index = 0; Index < Runs.Num(); Index++)
 	{
 		LineModel.Runs.Add( FRunModel( Runs[ Index ] ) );
 	}
@@ -1318,7 +1286,7 @@ void FTextLayout::GetAsTextAndOffsets(FString* const OutDisplayText, FTextOffset
 		OutTextOffsetLocations->OffsetData.Reserve(LineModels.Num());
 	}
 
-	for (int LineModelIndex = 0; LineModelIndex < LineModels.Num(); LineModelIndex++)
+	for (int32 LineModelIndex = 0; LineModelIndex < LineModels.Num(); LineModelIndex++)
 	{
 		const FLineModel& LineModel = LineModels[LineModelIndex];
 
@@ -1333,7 +1301,7 @@ void FTextLayout::GetAsTextAndOffsets(FString* const OutDisplayText, FTextOffset
 		}
 
 		int32 LineLength = 0;
-		for (int RunIndex = 0; RunIndex < LineModel.Runs.Num(); RunIndex++)
+		for (int32 RunIndex = 0; RunIndex < LineModel.Runs.Num(); RunIndex++)
 		{
 			const FRunModel& Run = LineModel.Runs[RunIndex];
 			if (OutDisplayText)
@@ -1354,7 +1322,7 @@ void FTextLayout::GetAsTextAndOffsets(FString* const OutDisplayText, FTextOffset
 
 void GetRangeAsTextFromLine(FString& DisplayText, const FTextLayout::FLineModel& LineModel, const FTextRange& Range)
 {
-	for (int RunIndex = 0; RunIndex < LineModel.Runs.Num(); RunIndex++)
+	for (int32 RunIndex = 0; RunIndex < LineModel.Runs.Num(); RunIndex++)
 	{
 		const FTextLayout::FRunModel& Run = LineModel.Runs[RunIndex];
 		const FTextRange& RunRange = Run.GetTextRange();
@@ -1392,7 +1360,7 @@ void FTextLayout::GetSelectionAsText(FString& DisplayText, const FTextSelection&
 		}
 		else
 		{
-			for (int LineIndex = SelectionBeginningLineIndex; LineIndex <= SelectionEndLineIndex; LineIndex++)
+			for (int32 LineIndex = SelectionBeginningLineIndex; LineIndex <= SelectionEndLineIndex; LineIndex++)
 			{
 				if (LineIndex == SelectionBeginningLineIndex)
 				{
@@ -1472,10 +1440,47 @@ FTextSelection FTextLayout::GetWordAt(const FTextLocation& Location) const
 	return FTextSelection(FTextLocation(LineIndex, CurrentBreak), FTextLocation(LineIndex, PreviousBreak));
 }
 
+void FTextLayout::SetVisibleRegion( const FVector2D& InViewSize, const FVector2D& InScrollOffset )
+{
+	if (ViewSize != InViewSize)
+	{
+		ViewSize = InViewSize;
+
+		if (Justification != ETextJustify::Left)
+		{
+			// If the view size has changed, we may need to update our positions based on our justification
+			DirtyFlags |= EDirtyState::Layout;
+		}
+	}
+	
+	if (ScrollOffset != InScrollOffset)
+	{
+		const FVector2D PreviousScrollOffset = ScrollOffset;
+		ScrollOffset = InScrollOffset;
+
+		// Use a negative scroll offset since positive scrolling moves things negatively in screen space
+		const FVector2D OffsetAdjustment = -(ScrollOffset - PreviousScrollOffset);
+
+		for (FLineView& LineView : LineViews)
+		{
+			LineView.Offset += OffsetAdjustment;
+
+			for (const TSharedRef< ILayoutBlock >& Block : LineView.Blocks)
+			{
+				Block->SetLocationOffset( Block->GetLocationOffset() + OffsetAdjustment );
+			}
+		}
+	}
+}
+
 void FTextLayout::SetLineBreakIterator( TSharedPtr<IBreakIterator> InLineBreakIterator )
 {
 	LineBreakIterator = InLineBreakIterator;
 	DirtyFlags |= EDirtyState::Layout;
+
+	// Changing the line break iterator will affect the wrapping information for *all lines*
+	// Clear out the entire cache so it gets regenerated on the text call to FlowLayout
+	ClearWrappingCache();
 }
 
 void FTextLayout::SetMargin( const FMargin& InMargin )
@@ -1497,11 +1502,11 @@ void FTextLayout::SetMargin( const FMargin& InMargin )
 	{
 		FVector2D OffsetAdjustment( Margin.Left - PreviousMargin.Left, Margin.Top - PreviousMargin.Top );
 		OffsetAdjustment *= Scale;
-		for (int LineViewIndex = 0; LineViewIndex < LineViews.Num(); LineViewIndex++)
+		for (int32 LineViewIndex = 0; LineViewIndex < LineViews.Num(); LineViewIndex++)
 		{
 			FLineView& LineView = LineViews[ LineViewIndex ];
 
-			for (int BlockIndex = 0; BlockIndex < LineView.Blocks.Num(); BlockIndex++)
+			for (int32 BlockIndex = 0; BlockIndex < LineView.Blocks.Num(); BlockIndex++)
 			{
 				const TSharedRef< ILayoutBlock >& Block = LineView.Blocks[ BlockIndex ];
 				Block->SetLocationOffset( Block->GetLocationOffset() + OffsetAdjustment );
@@ -1719,7 +1724,7 @@ TSharedRef< ILayoutBlock > FTextLayout::FRunModel::CreateBlock( const FBlockDefi
 			BlockSize += Run->Measure( SizeRange.BeginIndex, MeasuredRanges[ StartRangeIndex ].EndIndex, Scale );
 		}
 
-		for (int Index = StartRangeIndex + 1; Index < EndRangeIndex; Index++)
+		for (int32 Index = StartRangeIndex + 1; Index < EndRangeIndex; Index++)
 		{
 			BlockSize.X += MeasuredRangeSizes[ Index ].X;
 			BlockSize.Y = FMath::Max( MeasuredRangeSizes[ Index ].Y, BlockSize.Y );
@@ -1741,7 +1746,7 @@ TSharedRef< ILayoutBlock > FTextLayout::FRunModel::CreateBlock( const FBlockDefi
 	return Run->CreateBlock( BlockDefine.ActualRange.BeginIndex, BlockDefine.ActualRange.EndIndex, BlockSize, BlockDefine.Renderer );
 }
 
-int FTextLayout::FRunModel::BinarySearchForEndIndex( const TArray< FTextRange >& Ranges, int32 RangeBeginIndex, int32 EndIndex )
+int32 FTextLayout::FRunModel::BinarySearchForEndIndex( const TArray< FTextRange >& Ranges, int32 RangeBeginIndex, int32 EndIndex )
 {
 	int32 Min = RangeBeginIndex;
 	int32 Mid = 0;
@@ -1766,7 +1771,7 @@ int FTextLayout::FRunModel::BinarySearchForEndIndex( const TArray< FTextRange >&
 	return Mid;
 }
 
-int FTextLayout::FRunModel::BinarySearchForBeginIndex( const TArray< FTextRange >& Ranges, int32 BeginIndex )
+int32 FTextLayout::FRunModel::BinarySearchForBeginIndex( const TArray< FTextRange >& Ranges, int32 BeginIndex )
 {
 	int32 Min = 0;
 	int32 Mid = 0;
diff --git a/Engine/Source/Runtime/Slate/Private/Widgets/Input/SMultiLineEditableTextBox.cpp b/Engine/Source/Runtime/Slate/Private/Widgets/Input/SMultiLineEditableTextBox.cpp
--- a/Engine/Source/Runtime/Slate/Private/Widgets/Input/SMultiLineEditableTextBox.cpp
+++ b/Engine/Source/Runtime/Slate/Private/Widgets/Input/SMultiLineEditableTextBox.cpp
@@ -18,25 +18,66 @@ void SMultiLineEditableTextBox::Construct( const FArguments& InArgs )
 	BorderImageFocused = &InArgs._Style->BackgroundImageFocused;
 	BorderImageReadOnly = &InArgs._Style->BackgroundImageReadOnly;
 	TAttribute<FMargin> Padding = InArgs._Padding.IsSet() ? InArgs._Padding : InArgs._Style->Padding;
+	TAttribute<FMargin> HScrollBarPadding = InArgs._HScrollBarPadding.IsSet() ? InArgs._HScrollBarPadding : InArgs._Style->HScrollBarPadding;
+	TAttribute<FMargin> VScrollBarPadding = InArgs._VScrollBarPadding.IsSet() ? InArgs._VScrollBarPadding : InArgs._Style->VScrollBarPadding;
 	TAttribute<FSlateFontInfo> Font = InArgs._Font.IsSet() ? InArgs._Font : InArgs._Style->Font;
 	TAttribute<FSlateColor> ForegroundColor = InArgs._ForegroundColor.IsSet() ? InArgs._ForegroundColor : InArgs._Style->ForegroundColor;
 	TAttribute<FSlateColor> BackgroundColor = InArgs._BackgroundColor.IsSet() ? InArgs._BackgroundColor : InArgs._Style->BackgroundColor;
 	ReadOnlyForegroundColor = InArgs._ReadOnlyForegroundColor.IsSet() ? InArgs._ReadOnlyForegroundColor : InArgs._Style->ReadOnlyForegroundColor;
 
+	TSharedPtr<SWidget> HScrollBarWidget;
+	TSharedPtr<SScrollBar> HScrollBar = InArgs._HScrollBar;
+	if (!HScrollBar.IsValid())
+	{
+		// Create and use our own scrollbar
+		HScrollBar = SNew(SScrollBar)
+			.Style(&InArgs._Style->ScrollBarStyle)
+			.Orientation(Orient_Horizontal)
+			.Thickness(FVector2D(5.0f, 5.0f));
+		HScrollBarWidget = HScrollBar;
+	}
+	else
+	{
+		// User provided an external scrollbar, use a null widget
+		HScrollBarWidget = SNullWidget::NullWidget;
+	}
+
+	TSharedPtr<SWidget> VScrollBarWidget;
+	TSharedPtr<SScrollBar> VScrollBar = InArgs._VScrollBar;
+	if (!VScrollBar.IsValid())
+	{
+		// Create and use our own scrollbar
+		VScrollBar = SNew(SScrollBar)
+			.Style(&InArgs._Style->ScrollBarStyle)
+			.Orientation(Orient_Vertical)
+			.Thickness(FVector2D(5.0f, 5.0f));
+		VScrollBarWidget = VScrollBar;
+	}
+	else
+	{
+		// User provided an external scrollbar, use a null widget
+		VScrollBarWidget = SNullWidget::NullWidget;
+	}
+
 	SBorder::Construct( SBorder::FArguments()
 		.BorderImage( this, &SMultiLineEditableTextBox::GetBorderImage )
 		.BorderBackgroundColor( BackgroundColor )
 		.ForegroundColor( ForegroundColor )
-		.Padding( 0 )
+		.Padding( Padding )
 		[
 			SAssignNew( Box, SHorizontalBox)
+
 			+SHorizontalBox::Slot()
 			.VAlign(VAlign_Fill)
 			.HAlign(HAlign_Fill)
 			.FillWidth(1)
 			[
-				SNew( SBox )
-				.Padding( Padding )
+				SNew(SVerticalBox)
+
+				+SVerticalBox::Slot()
+				.VAlign(VAlign_Fill)
+				.HAlign(HAlign_Fill)
+				.FillHeight(1)
 				[
 					SAssignNew( EditableText, SMultiLineEditableText )
 					.Text( InArgs._Text )
@@ -51,8 +92,24 @@ void SMultiLineEditableTextBox::Construct( const FArguments& InArgs )
 					.Margin(InArgs._Margin)
 					.WrapTextAt(InArgs._WrapTextAt)
 					.AutoWrapText(InArgs._AutoWrapText)
+					.HScrollBar(HScrollBar)
+					.VScrollBar(VScrollBar)
+				]
+
+				+SVerticalBox::Slot()
+				.AutoHeight()
+				.Padding((InArgs._HScrollBar.IsValid()) ? FMargin(0) : HScrollBarPadding)
+				[
+					HScrollBarWidget.ToSharedRef()
 				]
 			]
+
+			+SHorizontalBox::Slot()
+			.AutoWidth()
+			.Padding((InArgs._VScrollBar.IsValid()) ? FMargin(0) : VScrollBarPadding)
+			[
+				VScrollBarWidget.ToSharedRef()
+			]
 		]
 	);
 
diff --git a/Engine/Source/Runtime/Slate/Private/Widgets/Testing/STestSuite.cpp b/Engine/Source/Runtime/Slate/Private/Widgets/Testing/STestSuite.cpp
--- a/Engine/Source/Runtime/Slate/Private/Widgets/Testing/STestSuite.cpp
+++ b/Engine/Source/Runtime/Slate/Private/Widgets/Testing/STestSuite.cpp
@@ -849,16 +849,17 @@ class SMultiLineEditingTest : public SCompoundWidget
 	{
 		MultilineEditableText = LOCTEXT( "MultiLineTextTest", "He has refused his Assent to Laws, the most wholesome and necessary for the public good.\nHe has forbidden his Governors to pass Laws of immediate and pressing importance, unless suspended in their operation till his Assent should be obtained; and when so suspended, he has utterly neglected to attend to them.\nHe has refused to pass other Laws for the accommodation of large districts of people, unless those people would relinquish the right of Representation in the Legislature, a right inestimable to them and formidable to tyrants only.\n\nHe has called together legislative bodies at places unusual, uncomfortable, and distant from the depository of their public Records, for the sole purpose of fatiguing them into compliance with his measures.\nHe has dissolved Representative Houses repeatedly, for opposing with manly firmness his invasions on the rights of the people.\nHe has refused for a long time, after such dissolutions, to cause others to be elected; whereby the Legislative powers, incapable of Annihilation, have returned to the People at large for their exercise; the State remaining in the mean time exposed to all the dangers of invasion from without, and convulsions within.\nHe has endeavoured to prevent the population of these States; for that purpose obstructing the Laws for Naturalization of Foreigners; refusing to pass others to encourage their migrations hither, and raising the conditions of new Appropriations of Lands.\n" );
 
-		TSharedRef<SScrollBar> FirstScrollBar = SNew( SScrollBar );
-		TSharedRef<SScrollBar> SecondScrollBar = SNew( SScrollBar );
-
 		this->ChildSlot
 		[
 			SNew(SVerticalBox)
-			+SVerticalBox::Slot().FillHeight( 1 )
+			+SVerticalBox::Slot()
+			.FillHeight( 1 )
 			[
 				SNew( SHorizontalBox )
-				+ SHorizontalBox::Slot().AutoWidth().VAlign( VAlign_Top ).Padding( 2 )
+				+ SHorizontalBox::Slot()
+				.AutoWidth()
+				.VAlign( VAlign_Top )
+				.Padding( 2 )
 				[
 					SNew( SVerticalBox )
 					+ SVerticalBox::Slot()
@@ -869,24 +870,18 @@ class SMultiLineEditingTest : public SCompoundWidget
 					]
 					+ SVerticalBox::Slot()
 					[
-						SNew( SHorizontalBox )
-						+ SHorizontalBox::Slot()
-						[
-							SNew( SMultiLineEditableTextBox )
-							.Text( this, &SMultiLineEditingTest::GetMultilineEditableText )
-							.Font( FSlateFontInfo( FPaths::EngineContentDir() / TEXT( "Slate/Fonts/Roboto-Regular.ttf" ), 12 ) )
-							.Justification( ETextJustify::Center )
-							.LineHeightPercentage( 2.0f )
-							.OnTextCommitted( this, &SMultiLineEditingTest::HandleMultilineEditableTextCommitted )
-							.WrapTextAt( 600.0f )
-						]
-						+ SHorizontalBox::Slot().AutoWidth()
-						[
-							FirstScrollBar
-						]
+						SNew( SMultiLineEditableTextBox )
+						.Text( this, &SMultiLineEditingTest::GetMultilineEditableText )
+						.Font( FSlateFontInfo( FPaths::EngineContentDir() / TEXT( "Slate/Fonts/Roboto-Regular.ttf" ), 12 ) )
+						.Justification( ETextJustify::Center )
+						.LineHeightPercentage( 2.0f )
+						.OnTextCommitted( this, &SMultiLineEditingTest::HandleMultilineEditableTextCommitted )
+						.WrapTextAt( 600.0f )
 					]
 				]
-				+ SHorizontalBox::Slot().FillWidth( 1 ).Padding( 2 )
+				+ SHorizontalBox::Slot()
+				.FillWidth( 1 )
+				.Padding( 2 )
 				[
 					SNew( SVerticalBox )
 					+ SVerticalBox::Slot()
@@ -896,29 +891,24 @@ class SMultiLineEditingTest : public SCompoundWidget
 						.Text( LOCTEXT( "MultiLineTextAutoWrap", "Multi-line editable text auto-wrapping" ) )
 					]
 					+ SVerticalBox::Slot()
-						[
-							SNew( SHorizontalBox )
-							+ SHorizontalBox::Slot()
-							[
-								SNew( SMultiLineEditableTextBox )
-								.Margin( 10 )
-								.Text( MultilineEditableText )
-								//.Justification(ETextJustify::Right)
-								.Font( FSlateFontInfo( FPaths::EngineContentDir() / TEXT( "Slate/Fonts/Roboto-Regular.ttf" ), 12 ) )
-								.AutoWrapText( true )
-								.HintText( LOCTEXT( "TypehereTextHint", "Type Here" ) )
-							]
-							+ SHorizontalBox::Slot().AutoWidth()
-								[
-									SecondScrollBar
-								]
-						]
+					[
+						SNew( SMultiLineEditableTextBox )
+						.Margin( 10 )
+						.Text( MultilineEditableText )
+						//.Justification(ETextJustify::Right)
+						.Font( FSlateFontInfo( FPaths::EngineContentDir() / TEXT( "Slate/Fonts/Roboto-Regular.ttf" ), 12 ) )
+						.AutoWrapText( true )
+						.HintText( LOCTEXT( "TypehereTextHint", "Type Here" ) )
+					]
 				]
 			]
-			+ SVerticalBox::Slot().AutoHeight()
+			+ SVerticalBox::Slot()
+			.AutoHeight()
 			[
 				SNew( SHorizontalBox )
-				+ SHorizontalBox::Slot().AutoWidth().VAlign( VAlign_Top ).Padding( 2 )
+				+ SHorizontalBox::Slot()
+				.VAlign( VAlign_Top )
+				.Padding( 2 )
 				[
 					SNew( SVerticalBox )
 					+ SVerticalBox::Slot()
@@ -928,16 +918,17 @@ class SMultiLineEditingTest : public SCompoundWidget
 						.Text( LOCTEXT( "MultiLineTextNoWrap", "Multi-line editable text without wrapping" ) )
 					]
 					+ SVerticalBox::Slot()
-						.AutoHeight()
+					.AutoHeight()
+					[
+						SNew( SHorizontalBox )
+						+ SHorizontalBox::Slot()
 						[
-							SNew( SHorizontalBox )
-							+ SHorizontalBox::Slot()
-							[
-								SNew( SMultiLineEditableTextBox )
-								.Font( FSlateFontInfo( FPaths::EngineContentDir() / TEXT( "Slate/Fonts/Roboto-Regular.ttf" ), 12 ) )
-								//.WrapTextAt(300.0f)
-							]
+							SNew( SMultiLineEditableTextBox )
+							.Font( FSlateFontInfo( FPaths::EngineContentDir() / TEXT( "Slate/Fonts/Roboto-Regular.ttf" ), 12 ) )
+							//.WrapTextAt(300.0f)
+							//.Justification( ETextJustify::Right )
 						]
+					]
 				]
 			]
 		];
@@ -1155,7 +1146,7 @@ public:
 						.BorderImage( FTestStyle::Get().GetBrush( "RichText.Tagline.Background" ) )
 						.Padding(0)
 						[
-							SNew( SRichTextBlock)
+							SNew( SRichTextBlock )
 							.Text( LOCTEXT("RichTextHeader05", "This is a text heavy page that has been created to show the performance and capabilities of Slate's <RichText.Tagline.TextHighlight>SRichTextBlock</>.") )
 							.TextStyle( FTestStyle::Get(), "RichText.Tagline.Text" )
 							.DecoratorStyleSet( &FTestStyle::Get() )
@@ -1167,7 +1158,7 @@ public:
 
 					+ SVerticalBox::Slot().AutoHeight() .HAlign(HAlign_Center) .Padding(20)			
 					[
-						SNew( SRichTextBlock)
+						SNew( SRichTextBlock )
 						.Text( LOCTEXT("RichText.HowItWorks", "<RichText.Header>What does it do?</>\n\nThe SRichTextBlock uses a concept called Decorators to introduce new font styles, images, animation and even whole interactive widgets inline with the text.\n\nSlate comes with a couple Decorators out of the box:\n\n        \u2022  <RichText.Text.Fancy>Text</> Decorator\n        \u2022  <img src=\"RichText.ImageDecorator\"/> Image Decorator\n        \u2022  <RichText.WidgetDecorator>Widget</> Decorator\n        \u2022  <a id=\"HyperlinkDecorator\" style=\"RichText.Hyperlink\">Hyperlink</> Decorator\n\n<RichText.Header>What about options?</>\n\nOf course you can always introduce your own Decorators by implementing ITextDecorator. This will give you full customization in how your text or widget is styled. \n\nBesides the power of Decorators the SRichTextBlock comes equipped with Margin support, Left-Center-Right Text Justification, a Line Height Scalar and Highlighting.\n\n<RichText.Header>How does the markup work?</>\n\nWell the markup parser is customizable so you can adjust the markup anyway you'd like by providing your own parser.  The parser that comes with Slate though uses a syntax very similar to xml.\n\n        \u2022  &lt;TextBlockStyleName>Your text content&lt;/>\n        \u2022  &lt;img src=\"SlateBrushStyleName\"/>\n        \u2022  &lt;a id=\"YourCustomId\"/>Your hyperlink text&lt;/>\n\nIf you ever want to use the markup syntax as actual text you can escape the markup using xml style escapes. For example:\n\n        \u2022  <    &amp;lt;\n        \u2022  >    &amp;gt;\n        \u2022  \"     &amp;quot;\n\nBut you only need to escape these characters when a set of them match actual syntax so this isn't generally an issue.\n\n<RichText.Header>Are there any catches?</>\n\nThere are still plenty of things the SRichTextBlock doesn't currently support. The most notable lacking feature is not having the ability to flow text around images or widgets. ") )
 						.TextStyle( FTestStyle::Get(), "RichText.Text" )
 						.DecoratorStyleSet( &FTestStyle::Get() )
@@ -1183,7 +1174,7 @@ public:
 						.BorderImage( FTestStyle::Get().GetBrush( "RichText.Tagline.Background" ) )
 						.Padding(0)
 						[
-							SNew( SRichTextBlock)
+							SNew( SRichTextBlock )
 							.Text( LOCTEXT("RichTextHeader01", "Here is an <RichText.Tagline.TextHighlight>interactive example</> of the different <RichText.Tagline.TextHighlight>SRichTextBlock</> features in action!") )
 							.TextStyle( FTestStyle::Get(), "RichText.Tagline.Text" )
 							.DecoratorStyleSet( &FTestStyle::Get() )
diff --git a/Engine/Source/Runtime/Slate/Private/Widgets/Text/SMultiLineEditableText.cpp b/Engine/Source/Runtime/Slate/Private/Widgets/Text/SMultiLineEditableText.cpp
--- a/Engine/Source/Runtime/Slate/Private/Widgets/Text/SMultiLineEditableText.cpp
+++ b/Engine/Source/Runtime/Slate/Private/Widgets/Text/SMultiLineEditableText.cpp
@@ -186,6 +186,7 @@ SMultiLineEditableText::SMultiLineEditableText()
 	, bIsDragSelecting(false)
 	, bWasFocusedByLastMouseDown(false)
 	, bHasDragSelectedSinceFocused(false)
+	, bPendingScrollCursorIntoView(false)
 	, CurrentUndoLevel(INDEX_NONE)
 	, bIsChangingText(false)
 	, IsReadOnly(false)
@@ -211,7 +212,21 @@ void SMultiLineEditableText::Construct( const FArguments& InArgs )
 
 	WrapTextAt = InArgs._WrapTextAt;
 	AutoWrapText = InArgs._AutoWrapText;
-	CachedAutoWrapTextWidth = 0.0f;
+	CachedSize = FVector2D(ForceInitToZero);
+
+	ScrollOffset = FVector2D::ZeroVector;
+
+	HScrollBar = InArgs._HScrollBar;
+	if (HScrollBar.IsValid())
+	{
+		HScrollBar->SetOnUserScrolled(FOnUserScrolled::CreateSP(this, &SMultiLineEditableText::OnHScrollBarMoved));
+	}
+
+	VScrollBar = InArgs._VScrollBar;
+	if (VScrollBar.IsValid())
+	{
+		VScrollBar->SetOnUserScrolled(FOnUserScrolled::CreateSP(this, &SMultiLineEditableText::OnVScrollBarMoved));
+	}
 
 	Margin = InArgs._Margin;
 	LineHeightPercentage = InArgs._LineHeightPercentage;
@@ -379,6 +394,16 @@ void SMultiLineEditableText::ForceRefreshTextLayout(const FText& CurrentText)
 	UpdateCursorHighlight();
 }
 
+void SMultiLineEditableText::OnHScrollBarMoved(const float InScrollOffsetFraction)
+{
+	ScrollOffset.X = FMath::Clamp<float>(InScrollOffsetFraction, 0.0, 1.0) * GetDesiredSize().X;
+}
+
+void SMultiLineEditableText::OnVScrollBarMoved(const float InScrollOffsetFraction)
+{
+	ScrollOffset.Y = FMath::Clamp<float>(InScrollOffsetFraction, 0.0, 1.0) * GetDesiredSize().Y;
+}
+
 FReply SMultiLineEditableText::OnKeyboardFocusReceived( const FGeometry& MyGeometry, const FKeyboardFocusEvent& InKeyboardFocusEvent )
 {
 	// Skip the focus received code if it's due to the context menu closing
@@ -890,6 +915,8 @@ FReply SMultiLineEditableText::MoveCursor( FMoveCursor Args )
 
 void SMultiLineEditableText::UpdateCursorHighlight()
 {
+	bPendingScrollCursorIntoView = true;
+
 	RemoveCursorHighlight();
 
 	static const int32 CompositionRangeZOrder = 1; // draw above the text
@@ -947,7 +974,7 @@ void SMultiLineEditableText::UpdateCursorHighlight()
 		{
 			const TArray< FTextLayout::FLineModel >& Lines = TextLayout->GetLineModels();
 
-			for (int LineIndex = SelectionBeginningLineIndex; LineIndex <= SelectionEndLineIndex; LineIndex++)
+			for (int32 LineIndex = SelectionBeginningLineIndex; LineIndex <= SelectionEndLineIndex; LineIndex++)
 			{
 				if ( LineIndex == SelectionBeginningLineIndex )
 				{
@@ -1866,43 +1893,101 @@ void SMultiLineEditableText::Tick( const FGeometry& AllottedGeometry, const doub
 			ForceRefreshTextLayout(TextToSet);
 		}
 	}
-}
 
-int32 SMultiLineEditableText::OnPaint( const FPaintArgs& Args, const FGeometry& AllottedGeometry, const FSlateRect& MyClippingRect, FSlateWindowElementList& OutDrawElements, int32 LayerId, const FWidgetStyle& InWidgetStyle, bool bParentEnabled ) const
-{
-	// Update the auto-wrap size now that we have computed paint geometry; won't take affect until text frame
-	// Note: This is done here rather than in Tick(), because Tick() doesn't get called while resizing windows, but OnPaint() does
-	if ( AutoWrapText.Get() )
+	const float FontMaxCharHeight = FTextEditHelper::GetFontHeight(TextStyle.Font);
+	const float CaretWidth = FTextEditHelper::CalculateCaretWidth(FontMaxCharHeight);
+
+	// Try and make sure that the line containing the cursor is in view
+	if (bPendingScrollCursorIntoView)
 	{
-		CachedAutoWrapTextWidth = AllottedGeometry.Size.X;
-	}
+		bPendingScrollCursorIntoView = false;
 
-	const ETextJustify::Type TextJustification = TextLayout->GetJustification();
+		const TArray<FTextLayout::FLineView>& LineViews = TextLayout->GetLineViews();
+		const int32 LineViewIndex = TextLayout->GetLineViewIndexForTextLocation(LineViews, CursorInfo.GetCursorInteractionLocation(), CursorInfo.GetCursorAlignment() == ECursorAlignment::Right);
+		if (LineViews.IsValidIndex(LineViewIndex))
+		{
+			const FTextLayout::FLineView& LineView = LineViews[LineViewIndex];
+			const FSlateRect LocalLineViewRect(LineView.Offset / TextLayout->GetScale(), (LineView.Offset + LineView.Size) / TextLayout->GetScale());
 
-	if ( TextJustification == ETextJustify::Right )
-	{
-		const FVector2D TextLayoutSize = TextLayout->GetSize();
-		const FVector2D Offset( ( AllottedGeometry.Size.X - TextLayoutSize.X ), 0 );
-		LayerId = TextLayout->OnPaint( Args.WithNewParent(this), TextStyle, AllottedGeometry.MakeChild( Offset, TextLayoutSize ), MyClippingRect, OutDrawElements, LayerId, InWidgetStyle, ShouldBeEnabled( bParentEnabled ) );
+			const FVector2D LocalCursorLocation = TextLayout->GetLocationAt(CursorInfo.GetCursorInteractionLocation(), CursorInfo.GetCursorAlignment() == ECursorAlignment::Right) / TextLayout->GetScale();
+			const FSlateRect LocalCursorRect(LocalCursorLocation, FVector2D(LocalCursorLocation.X + CaretWidth, LocalCursorLocation.Y + FontMaxCharHeight));
+
+			if (LocalCursorRect.Left < 0.0f)
+			{
+				ScrollOffset.X += LocalCursorRect.Left;
+			}
+			else if (LocalCursorRect.Right > AllottedGeometry.Size.X)
+			{
+				ScrollOffset.X += (LocalCursorRect.Right - AllottedGeometry.Size.X);
+			}
+
+			if (LocalLineViewRect.Top < 0.0f)
+			{
+				ScrollOffset.Y += LocalLineViewRect.Top;
+			}
+			else if (LocalLineViewRect.Bottom > AllottedGeometry.Size.Y)
+			{
+				ScrollOffset.Y += (LocalLineViewRect.Bottom - AllottedGeometry.Size.Y);
+			}
+		}
 	}
-	else if ( TextJustification == ETextJustify::Center )
+
 	{
-		FVector2D TextLayoutSize = TextLayout->GetSize();
-		float LayoutWidth = TextLayoutSize.X;
-		if (TextLayoutSize.X == 0)
+		// Need to account for the caret width too
+		const float ContentSize = TextLayout->GetSize().X + CaretWidth;
+		const float VisibleSize = AllottedGeometry.Size.X;
+
+		// If this text box has no size, do not compute a view fraction because it will be wrong and causes pop in when the size is available
+		const float ViewFraction = (VisibleSize > 0.0f && ContentSize > 0.0f) ? VisibleSize / ContentSize : 1;
+		const float ViewOffset = (ContentSize > 0.0f) ? FMath::Clamp<float>(ScrollOffset.X / ContentSize, 0.0f, 1.0f - ViewFraction) : 0.0f;
+	
+		// Update the scrollbar with the clamped version of the offset
+		ScrollOffset.X = ViewOffset * ContentSize;
+	
+		if (HScrollBar.IsValid())
 		{
-			const TSharedRef< FSlateFontMeasure > FontMeasureService = FSlateApplication::Get().GetRenderer()->GetFontMeasureService();
-			TextLayoutSize.X = FontMeasureService->Measure(FString(TEXT("  ")), TextStyle.Font).X;
+			HScrollBar->SetState(ViewOffset, ViewFraction);
+			if (!HScrollBar->IsNeeded())
+			{
+				// We cannot scroll, so ensure that there is no offset
+				ScrollOffset.X = 0.0f;
+			}
 		}
-
-		const FVector2D Offset((AllottedGeometry.Size.X - TextLayoutSize.X) / 2, 0);
-		LayerId = TextLayout->OnPaint( Args.WithNewParent(this), TextStyle, AllottedGeometry.MakeChild( Offset, TextLayoutSize ), MyClippingRect, OutDrawElements, LayerId, InWidgetStyle, ShouldBeEnabled( bParentEnabled ) );
 	}
-	else
+
 	{
-		LayerId = TextLayout->OnPaint( Args.WithNewParent(this), TextStyle, AllottedGeometry, MyClippingRect, OutDrawElements, LayerId, InWidgetStyle, ShouldBeEnabled( bParentEnabled ) );
+		const float ContentSize = TextLayout->GetSize().Y;
+		const float VisibleSize = AllottedGeometry.Size.Y;
+
+		// If this text box has no size, do not compute a view fraction because it will be wrong and causes pop in when the size is available
+		const float ViewFraction = (VisibleSize > 0.0f && ContentSize > 0.0f) ? VisibleSize / ContentSize : 1;
+		const float ViewOffset = (ContentSize > 0.0f) ? FMath::Clamp<float>(ScrollOffset.Y / ContentSize, 0.0f, 1.0f - ViewFraction) : 0.0f;
+	
+		// Update the scrollbar with the clamped version of the offset
+		ScrollOffset.Y = ViewOffset * ContentSize;
+	
+		if (VScrollBar.IsValid())
+		{
+			VScrollBar->SetState(ViewOffset, ViewFraction);
+			if (!VScrollBar->IsNeeded())
+			{
+				// We cannot scroll, so ensure that there is no offset
+				ScrollOffset.Y = 0.0f;
+			}
+		}
 	}
 
+	TextLayout->SetVisibleRegion(AllottedGeometry.Size, ScrollOffset);
+}
+
+int32 SMultiLineEditableText::OnPaint( const FPaintArgs& Args, const FGeometry& AllottedGeometry, const FSlateRect& MyClippingRect, FSlateWindowElementList& OutDrawElements, int32 LayerId, const FWidgetStyle& InWidgetStyle, bool bParentEnabled ) const
+{
+	// Update the auto-wrap size now that we have computed paint geometry; won't take affect until text frame
+	// Note: This is done here rather than in Tick(), because Tick() doesn't get called while resizing windows, but OnPaint() does
+	CachedSize = AllottedGeometry.Size;
+
+	LayerId = TextLayout->OnPaint( Args.WithNewParent(this), TextStyle, AllottedGeometry, MyClippingRect, OutDrawElements, LayerId, InWidgetStyle, ShouldBeEnabled( bParentEnabled ) );
+
 	return LayerId;
 }
 
@@ -1915,15 +2000,16 @@ void SMultiLineEditableText::CacheDesiredSize()
 
 	// Text wrapping can either be used defined (WrapTextAt), automatic (AutoWrapText), or a mixture of both
 	// Take whichever has the smallest value (>1)
-	if(AutoWrapText.Get() && CachedAutoWrapTextWidth >= 1.0f)
+	if(AutoWrapText.Get() && CachedSize.X >= 1.0f)
 	{
-		WrappingWidth = (WrappingWidth >= 1.0f) ? FMath::Min(WrappingWidth, CachedAutoWrapTextWidth) : CachedAutoWrapTextWidth;
+		WrappingWidth = (WrappingWidth >= 1.0f) ? FMath::Min(WrappingWidth, CachedSize.X) : CachedSize.X;
 	}
 
 	TextLayout->SetWrappingWidth( WrappingWidth );
 	TextLayout->SetMargin( OurMargin );
 	TextLayout->SetLineHeightPercentage( LineHeightPercentage.Get() );
 	TextLayout->SetJustification( Justification.Get() );
+	TextLayout->SetVisibleRegion( CachedSize, ScrollOffset );
 	TextLayout->UpdateIfNeeded();
 
 	SWidget::CacheDesiredSize();
@@ -1931,15 +2017,13 @@ void SMultiLineEditableText::CacheDesiredSize()
 
 FVector2D SMultiLineEditableText::ComputeDesiredSize() const
 {
-	const FSlateFontInfo& FontInfo = TextStyle.Font;
-	
-	const float FontMaxCharHeight = FTextEditHelper::GetFontHeight(FontInfo);
+	const float FontMaxCharHeight = FTextEditHelper::GetFontHeight(TextStyle.Font);
 	const float CaretWidth = FTextEditHelper::CalculateCaretWidth(FontMaxCharHeight);
 
 	// The layouts current margin size. We should not report a size smaller then the margins.
 	const FMargin Margin = TextLayout->GetMargin();
 	const FVector2D TextLayoutSize = TextLayout->GetSize();
-	const float WrappingWidth = WrapTextAt.Get();
+	const float WrappingWidth = TextLayout->GetWrappingWidth();
 
 	// If a wrapping width has been provided that should be reported as the desired width.
 	float DesiredWidth = WrappingWidth > 0 ? WrappingWidth : TextLayoutSize.X;
@@ -1959,24 +2043,7 @@ FChildren* SMultiLineEditableText::GetChildren()
 
 void SMultiLineEditableText::OnArrangeChildren( const FGeometry& AllottedGeometry, FArrangedChildren& ArrangedChildren ) const
 {
-	const ETextJustify::Type TextJustification = TextLayout->GetJustification();
-
-	if ( TextJustification == ETextJustify::Right )
-	{
-		const FVector2D TextLayoutSize = TextLayout->GetSize();
-		const FVector2D Offset( ( AllottedGeometry.Size.X - TextLayoutSize.X ), 0 );
-		TextLayout->ArrangeChildren( AllottedGeometry.MakeChild( Offset, TextLayoutSize ), ArrangedChildren );
-	}
-	else if ( TextJustification == ETextJustify::Center )
-	{
-		const FVector2D TextLayoutSize = TextLayout->GetSize();
-		const FVector2D Offset( ( AllottedGeometry.Size.X - TextLayoutSize.X ) / 2, 0 );
-		TextLayout->ArrangeChildren( AllottedGeometry.MakeChild( Offset, TextLayoutSize ), ArrangedChildren );
-	}
-	else
-	{
-		TextLayout->ArrangeChildren( AllottedGeometry, ArrangedChildren );
-	}
+	TextLayout->ArrangeChildren( AllottedGeometry, ArrangedChildren );
 }
 
 bool SMultiLineEditableText::SupportsKeyboardFocus() const
@@ -2025,6 +2092,26 @@ FReply SMultiLineEditableText::OnMouseMove( const FGeometry& MyGeometry, const F
 	return Reply;
 }
 
+FReply SMultiLineEditableText::OnMouseWheel( const FGeometry& MyGeometry, const FPointerEvent& MouseEvent )
+{
+	if (VScrollBar.IsValid() && VScrollBar->IsNeeded())
+	{
+		const float PreviousScrollOffset = ScrollOffset.Y;
+
+		const float ScrollAmount = -MouseEvent.GetWheelDelta() * WheelScrollAmount;
+		ScrollOffset.Y += ScrollAmount;
+
+		const float ContentSize = TextLayout->GetSize().Y;
+		const float ScrollMin = 0.0f;
+		const float ScrollMax = ContentSize - MyGeometry.Size.Y;
+		ScrollOffset.Y = FMath::Clamp(ScrollOffset.Y, ScrollMin, ScrollMax);
+
+		return (PreviousScrollOffset != ScrollOffset.Y) ? FReply::Handled() : FReply::Unhandled();
+	}
+
+	return FReply::Unhandled();
+}
+
 FReply SMultiLineEditableText::OnMouseButtonDoubleClick(const FGeometry& MyGeometry, const FPointerEvent& MouseEvent)
 {
 	FReply Reply = FTextEditHelper::OnMouseButtonDoubleClick(MyGeometry, MouseEvent, SharedThis(this));
@@ -2412,5 +2499,4 @@ void SMultiLineEditableText::FTextInputMethodContext::EndComposition()
 	}
 }
 
-
-#endif //WITH_FANCY_TEXT
\ No newline at end of file
+#endif //WITH_FANCY_TEXT
diff --git a/Engine/Source/Runtime/Slate/Private/Widgets/Text/SRichTextBlock.cpp b/Engine/Source/Runtime/Slate/Private/Widgets/Text/SRichTextBlock.cpp
--- a/Engine/Source/Runtime/Slate/Private/Widgets/Text/SRichTextBlock.cpp
+++ b/Engine/Source/Runtime/Slate/Private/Widgets/Text/SRichTextBlock.cpp
@@ -20,7 +20,7 @@ void SRichTextBlock::Construct( const FArguments& InArgs )
 	LineHeightPercentage = InArgs._LineHeightPercentage;
 	Justification = InArgs._Justification;
 
-	CachedAutoWrapTextWidth = 0.0f;
+	CachedSize = FVector2D(ForceInitToZero);
 
 	{
 		TSharedPtr<IRichTextMarkupParser> Parser = InArgs._Parser;
@@ -72,33 +72,19 @@ void SRichTextBlock::Tick( const FGeometry& AllottedGeometry, const double InCur
 	{
 		SetText(TextToSet);
 	}
+
+	TextLayout->SetVisibleRegion(AllottedGeometry.Size, FVector2D(ForceInitToZero));
 }
 
 int32 SRichTextBlock::OnPaint( const FPaintArgs& Args, const FGeometry& AllottedGeometry, const FSlateRect& MyClippingRect, FSlateWindowElementList& OutDrawElements, int32 LayerId, const FWidgetStyle& InWidgetStyle, bool bParentEnabled ) const
 {
 	SCOPE_CYCLE_COUNTER( STAT_SlateOnPaint_SRichTextBlock );
 
-	// Cache the last known size in case we need it for performing auto-wrap.
-	CachedAutoWrapTextWidth = AllottedGeometry.Size.X;
-
-	const ETextJustify::Type TextJustification = TextLayout->GetJustification();
+	// Update the auto-wrap size now that we have computed paint geometry; won't take affect until text frame
+	// Note: This is done here rather than in Tick(), because Tick() doesn't get called while resizing windows, but OnPaint() does
+	CachedSize = AllottedGeometry.Size;
 
-	if ( TextJustification == ETextJustify::Right )
-	{
-		const FVector2D TextLayoutSize = TextLayout->GetSize();
-		const FVector2D Offset( ( AllottedGeometry.Size.X - TextLayoutSize.X ), 0 );
-		LayerId = TextLayout->OnPaint( Args.WithNewParent(this), TextStyle, AllottedGeometry.MakeChild( Offset, TextLayoutSize ), MyClippingRect, OutDrawElements, LayerId, InWidgetStyle, ShouldBeEnabled( bParentEnabled ) );
-	}
-	else if ( TextJustification == ETextJustify::Center )
-	{
-		const FVector2D TextLayoutSize = TextLayout->GetSize();
-		const FVector2D Offset( ( AllottedGeometry.Size.X - TextLayoutSize.X ) / 2, 0 );
-		LayerId = TextLayout->OnPaint( Args.WithNewParent(this), TextStyle, AllottedGeometry.MakeChild( Offset, TextLayoutSize ), MyClippingRect, OutDrawElements, LayerId, InWidgetStyle, ShouldBeEnabled( bParentEnabled ) );
-	}
-	else
-	{
-		LayerId = TextLayout->OnPaint( Args.WithNewParent(this), TextStyle, AllottedGeometry, MyClippingRect, OutDrawElements, LayerId, InWidgetStyle, ShouldBeEnabled( bParentEnabled ) );
-	}
+	LayerId = TextLayout->OnPaint( Args.WithNewParent(this), TextStyle, AllottedGeometry, MyClippingRect, OutDrawElements, LayerId, InWidgetStyle, ShouldBeEnabled( bParentEnabled ) );
 
 	return LayerId;
 }
@@ -107,17 +93,18 @@ void SRichTextBlock::CacheDesiredSize()
 {
 	float WrappingWidth = WrapTextAt.Get();
 
-	// Text wrapping can either be used defined (WrapTextAt), automatic (AutoWrapText and CachedAutoWrapTextWidth), 
+	// Text wrapping can either be used defined (WrapTextAt), automatic (AutoWrapText and CachedSize), 
 	// or a mixture of both.  Take whichever has the smallest value (>1)
-	if ( AutoWrapText.Get() && CachedAutoWrapTextWidth >= 1.0f )
+	if ( AutoWrapText.Get() && CachedSize.X >= 1.0f )
 	{
-		WrappingWidth = (WrappingWidth >= 1.0f) ? FMath::Min(WrappingWidth, CachedAutoWrapTextWidth) : CachedAutoWrapTextWidth;
+		WrappingWidth = (WrappingWidth >= 1.0f) ? FMath::Min(WrappingWidth, CachedSize.X) : CachedSize.X;
 	}
 
 	TextLayout->SetWrappingWidth( WrappingWidth );
 	TextLayout->SetMargin( Margin.Get() );
 	TextLayout->SetLineHeightPercentage( LineHeightPercentage.Get() );
 	TextLayout->SetJustification( Justification.Get() );
+	TextLayout->SetVisibleRegion( CachedSize, FVector2D(ForceInitToZero) );
 	TextLayout->UpdateIfNeeded();
 
 	SWidget::CacheDesiredSize();
@@ -142,24 +129,7 @@ FChildren* SRichTextBlock::GetChildren()
 
 void SRichTextBlock::OnArrangeChildren( const FGeometry& AllottedGeometry, FArrangedChildren& ArrangedChildren ) const
 {
-	const ETextJustify::Type TextJustification = TextLayout->GetJustification();
-
-	if ( TextJustification == ETextJustify::Right )
-	{
-		const FVector2D TextLayoutSize = TextLayout->GetSize();
-		const FVector2D Offset( ( AllottedGeometry.Size.X - TextLayoutSize.X ), 0 );
-		TextLayout->ArrangeChildren( AllottedGeometry.MakeChild( Offset, TextLayoutSize ), ArrangedChildren );
-	}
-	else if ( TextJustification == ETextJustify::Center )
-	{
-		const FVector2D TextLayoutSize = TextLayout->GetSize();
-		const FVector2D Offset( ( AllottedGeometry.Size.X - TextLayoutSize.X ) / 2, 0 );
-		TextLayout->ArrangeChildren( AllottedGeometry.MakeChild( Offset, TextLayoutSize ), ArrangedChildren );
-	}
-	else
-	{
-		TextLayout->ArrangeChildren( AllottedGeometry, ArrangedChildren );
-	}
+	TextLayout->ArrangeChildren( AllottedGeometry, ArrangedChildren );
 }
 
 void SRichTextBlock::SetText( const TAttribute<FText>& InTextAttr )
diff --git a/Engine/Source/Runtime/Slate/Public/Framework/Text/TextLayout.h b/Engine/Source/Runtime/Slate/Public/Framework/Text/TextLayout.h
--- a/Engine/Source/Runtime/Slate/Public/Framework/Text/TextLayout.h
+++ b/Engine/Source/Runtime/Slate/Public/Framework/Text/TextLayout.h
@@ -182,9 +182,9 @@ public:
 
 		uint8 GetKerning( int32 CurrentIndex, float Scale );
 
-		static int BinarySearchForBeginIndex( const TArray< FTextRange >& Ranges, int32 BeginIndex );
+		static int32 BinarySearchForBeginIndex( const TArray< FTextRange >& Ranges, int32 BeginIndex );
 
-		static int BinarySearchForEndIndex( const TArray< FTextRange >& Ranges, int32 RangeBeginIndex, int32 EndIndex );
+		static int32 BinarySearchForEndIndex( const TArray< FTextRange >& Ranges, int32 RangeBeginIndex, int32 EndIndex );
 
 		TSharedRef< ILayoutBlock > CreateBlock( const FBlockDefinition& BlockDefine, float Scale ) const;
 
@@ -293,6 +293,8 @@ public:
 	FMargin GetMargin() const;
 	void SetMargin( const FMargin& InMargin );
 
+	void SetVisibleRegion( const FVector2D& InViewSize, const FVector2D& InScrollOffset );
+
 	/** Set the iterator to use to detect appropriate soft-wrapping points for lines (or null to go back to using the default) */
 	void SetLineBreakIterator( TSharedPtr<IBreakIterator> InLineBreakIterator );
 
@@ -420,8 +422,6 @@ private:
 
 	void JustifyLayout();
 
-	void JustifyHighlights();
-
 	void CreateLineViewBlocks( int32 LineModelIndex, const int32 StopIndex, int32& OutRunIndex, int32& OutRendererIndex, int32& OutPreviousBlockEnd, TArray< TSharedRef< ILayoutBlock > >& OutSoftLine );
 
 	FBreakCandidate CreateBreakCandidate( int32& OutRunIndex, FLineModel& Line, int32 PreviousBreak, int32 CurrentBreak );
@@ -468,6 +468,12 @@ protected:
 	/** The final size of the text layout on screen. */
 	FVector2D DrawSize;
 
+	/** The size of the text layout that can actually be seen from the parent widget */
+	FVector2D ViewSize;
+
+	/** The scroll offset of the text layout from the parent widget */
+	FVector2D ScrollOffset;
+
 	/** The iterator to use to detect appropriate soft-wrapping points for lines */
 	TSharedPtr<IBreakIterator> LineBreakIterator;
 
diff --git a/Engine/Source/Runtime/Slate/Public/Widgets/Input/SMultiLineEditableTextBox.h b/Engine/Source/Runtime/Slate/Public/Widgets/Input/SMultiLineEditableTextBox.h
--- a/Engine/Source/Runtime/Slate/Public/Widgets/Input/SMultiLineEditableTextBox.h
+++ b/Engine/Source/Runtime/Slate/Public/Widgets/Input/SMultiLineEditableTextBox.h
@@ -82,6 +82,18 @@ public:
 		/** Whether to clear keyboard focus when pressing enter to commit changes */
 		SLATE_ATTRIBUTE( bool, ClearKeyboardFocusOnCommit )
 
+		/** The horizontal scroll bar widget, or null to create one internally */
+		SLATE_ARGUMENT( TSharedPtr< SScrollBar >, HScrollBar )
+
+		/** The vertical scroll bar widget, or null to create one internally */
+		SLATE_ARGUMENT( TSharedPtr< SScrollBar >, VScrollBar )
+
+		/** Padding around the horizontal scrollbar (overrides Style) */
+		SLATE_ATTRIBUTE( FMargin, HScrollBarPadding )
+
+		/** Padding around the vertical scrollbar (overrides Style) */
+		SLATE_ATTRIBUTE( FMargin, VScrollBarPadding )
+
 		/** Called whenever the text is changed interactively by the user */
 		SLATE_EVENT( FOnTextChanged, OnTextChanged )
 
diff --git a/Engine/Source/Runtime/Slate/Public/Widgets/Text/SMultiLineEditableText.h b/Engine/Source/Runtime/Slate/Public/Widgets/Text/SMultiLineEditableText.h
--- a/Engine/Source/Runtime/Slate/Public/Widgets/Text/SMultiLineEditableText.h
+++ b/Engine/Source/Runtime/Slate/Public/Widgets/Text/SMultiLineEditableText.h
@@ -7,6 +7,7 @@
 #include "ITextEditorWidget.h"
 #include "ITextInputMethodSystem.h"
 #include "ITextLayoutMarshaller.h"
+#include "SScrollBar.h"
 
 /** An editable text widget that supports multiple lines and soft word-wrapping. */
 class SLATE_API SMultiLineEditableText : public SWidget, public ITextEditorWidget
@@ -63,6 +64,12 @@ public:
 		/** Sets whether this text box can actually be modified interactively by the user */
 		SLATE_ATTRIBUTE(bool, IsReadOnly)
 
+		/** The horizontal scroll bar widget */
+		SLATE_ARGUMENT(TSharedPtr< SScrollBar >, HScrollBar)
+
+		/** The vertical scroll bar widget */
+		SLATE_ARGUMENT(TSharedPtr< SScrollBar >, VScrollBar)
+
 		/** Called whenever the text is changed interactively by the user */
 		SLATE_EVENT(FOnTextChanged, OnTextChanged)
 
@@ -368,6 +375,7 @@ private:
 	virtual FReply OnMouseButtonDown( const FGeometry& MyGeometry, const FPointerEvent& MouseEvent ) override;
 	virtual FReply OnMouseButtonUp( const FGeometry& MyGeometry, const FPointerEvent& MouseEvent ) override;
 	virtual FReply OnMouseMove( const FGeometry& MyGeometry, const FPointerEvent& MouseEvent ) override;
+	virtual FReply OnMouseWheel( const FGeometry& MyGeometry, const FPointerEvent& MouseEvent ) override;
 	virtual FReply OnMouseButtonDoubleClick(const FGeometry& MyGeometry, const FPointerEvent& MouseEvent) override;
 	virtual FCursorReply OnCursorQuery( const FGeometry& MyGeometry, const FPointerEvent& CursorEvent ) const override;
 	// END SWidget interface
@@ -467,6 +475,9 @@ private:
 	 */
 	void ForceRefreshTextLayout(const FText& CurrentText);
 
+	void OnHScrollBarMoved(const float InScrollOffsetFraction);
+	void OnVScrollBarMoved(const float InScrollOffsetFraction);
+
 private:
 
 	/** The text displayed in this text block */
@@ -490,8 +501,11 @@ private:
 	/** True if we're wrapping text automatically based on the computed horizontal space for this widget */
 	TAttribute< bool > AutoWrapText;
 
-	/** Cached auto-wrap width that this text is using. This is used when determining if the cached string size should be updated */
-	mutable float CachedAutoWrapTextWidth;
+	/** The last known size of the control from the previous OnPaint, used to recalculate wrapping. */
+	mutable FVector2D CachedSize;
+
+	/** The scroll offset (in Slate units) for this text */
+	FVector2D ScrollOffset;
 
 	TAttribute< FMargin > Margin;
 	TAttribute< ETextJustify::Type > Justification; 
@@ -522,6 +536,9 @@ private:
 	/** True if characters were selected by dragging since the last keyboard focus.  Used for text selection. */
 	bool bHasDragSelectedSinceFocused;
 
+	/** True if we're pending a check to ensure the cursor is currently scrolled into view */
+	bool bPendingScrollCursorIntoView;
+
 	/** Undo states */
 	TArray< FUndoState > UndoStates;
 
@@ -549,6 +566,12 @@ private:
 	/** Called when the cursor is moved within the text area */
 	FOnCursorMoved OnCursorMoved;
 
+	/** The horizontal scroll bar widget */
+	TSharedPtr< SScrollBar > HScrollBar;
+
+	/** The vertical scroll bar widget */
+	TSharedPtr< SScrollBar > VScrollBar;
+
 	/** Menu extender for right-click context menu */
 	TSharedPtr<FExtender> MenuExtender;
 
diff --git a/Engine/Source/Runtime/Slate/Public/Widgets/Text/SRichTextBlock.h b/Engine/Source/Runtime/Slate/Public/Widgets/Text/SRichTextBlock.h
--- a/Engine/Source/Runtime/Slate/Public/Widgets/Text/SRichTextBlock.h
+++ b/Engine/Source/Runtime/Slate/Public/Widgets/Text/SRichTextBlock.h
@@ -158,8 +158,8 @@ private:
 	/** True if we're wrapping text automatically based on the computed horizontal space for this widget */
 	TAttribute<bool> AutoWrapText;
 
-	/** The last known width of the control from the previous OnPaint, used to recalculate wrapping. */
-	mutable float CachedAutoWrapTextWidth;
+	/** The last known size of the control from the previous OnPaint, used to recalculate wrapping. */
+	mutable FVector2D CachedSize;
 
 	TAttribute< FMargin > Margin;
 	TAttribute< ETextJustify::Type > Justification; 
diff --git a/Engine/Source/Runtime/SlateCore/Private/Styling/CoreStyle.cpp b/Engine/Source/Runtime/SlateCore/Private/Styling/CoreStyle.cpp
--- a/Engine/Source/Runtime/SlateCore/Private/Styling/CoreStyle.cpp
+++ b/Engine/Source/Runtime/SlateCore/Private/Styling/CoreStyle.cpp
@@ -260,6 +260,17 @@ TSharedRef<ISlateStyle> FCoreStyle::Create( const FName& InStyleSetName )
 		Style->Set( "GenericCommands.Redo", new IMAGE_BRUSH( "Icons/icon_redo_16px", Icon16x16 ) );
 	}
 
+	// SScrollBar defaults...
+	const FScrollBarStyle ScrollBar = FScrollBarStyle()
+		.SetVerticalBackgroundImage( IMAGE_BRUSH( "Common/Scrollbar_Background_Vertical", FVector2D(8,8) ) )
+		.SetHorizontalBackgroundImage( IMAGE_BRUSH( "Common/Scrollbar_Background_Horizontal", FVector2D(8,8) ) )
+		.SetNormalThumbImage( BOX_BRUSH( "Common/Scrollbar_Thumb", FMargin(4.f/16.f) ) )
+		.SetDraggedThumbImage( BOX_BRUSH( "Common/Scrollbar_Thumb", FMargin(4.f/16.f) ) )
+		.SetHoveredThumbImage( BOX_BRUSH( "Common/Scrollbar_Thumb", FMargin(4.f/16.f) ) );
+	{
+		Style->Set( "Scrollbar", ScrollBar );
+	}
+
 	// SButton defaults...
 	const FButtonStyle Button = FButtonStyle()
 		.SetNormal( BOX_BRUSH( "Common/Button", FVector2D(32,32), 8.0f/32.0f ) )
@@ -385,7 +396,8 @@ TSharedRef<ISlateStyle> FCoreStyle::Create( const FName& InStyleSetName )
 		.SetBackgroundImageNormal( BOX_BRUSH( "Common/TextBox", FMargin(4.0f/16.0f) ) )
 		.SetBackgroundImageHovered( BOX_BRUSH( "Common/TextBox_Hovered", FMargin(4.0f/16.0f) ) )
 		.SetBackgroundImageFocused( BOX_BRUSH( "Common/TextBox_Hovered", FMargin(4.0f/16.0f) ) )
-		.SetBackgroundImageReadOnly( BOX_BRUSH( "Common/TextBox_ReadOnly", FMargin(4.0f/16.0f) ) );
+		.SetBackgroundImageReadOnly( BOX_BRUSH( "Common/TextBox_ReadOnly", FMargin(4.0f/16.0f) ) )
+		.SetScrollBarStyle( ScrollBar );
 	{
 		Style->Set( "NormalEditableTextBox", NormalEditableTextBoxStyle );
 		// "NormalFont".
@@ -422,7 +434,8 @@ TSharedRef<ISlateStyle> FCoreStyle::Create( const FName& InStyleSetName )
 			.SetBackgroundImageNormal( BOX_BRUSH( "Common/TextBox", FMargin(4.0f/16.0f) ) )
 			.SetBackgroundImageHovered( BOX_BRUSH( "Common/TextBox_Hovered", FMargin(4.0f/16.0f) ) )
 			.SetBackgroundImageFocused( BOX_BRUSH( "Common/TextBox_Hovered", FMargin(4.0f/16.0f) ) )
-			.SetBackgroundImageReadOnly( BOX_BRUSH( "Common/TextBox_ReadOnly", FMargin(4.0f/16.0f) ) );
+			.SetBackgroundImageReadOnly( BOX_BRUSH( "Common/TextBox_ReadOnly", FMargin(4.0f/16.0f) ) )
+			.SetScrollBarStyle( ScrollBar );
 
 		FInlineEditableTextBlockStyle InlineEditableTextBlockStyle = FInlineEditableTextBlockStyle()
 			.SetTextStyle(InlineEditableTextBlockReadOnly)
@@ -479,17 +492,6 @@ TSharedRef<ISlateStyle> FCoreStyle::Create( const FName& InStyleSetName )
 		Style->Set( "Throbber.CircleChunk", new IMAGE_BRUSH( "Common/Throbber_Piece", FVector2D(8,8) ) );
 	}
 
-	// SScrollBar defaults...
-	{
-		Style->Set( "Scrollbar", FScrollBarStyle()
-			.SetVerticalBackgroundImage( IMAGE_BRUSH( "Common/Scrollbar_Background_Vertical", FVector2D(8,8) ) )
-			.SetHorizontalBackgroundImage( IMAGE_BRUSH( "Common/Scrollbar_Background_Horizontal", FVector2D(8,8) ) )
-			.SetNormalThumbImage( BOX_BRUSH( "Common/Scrollbar_Thumb", FMargin(4.f/16.f) ) )
-			.SetDraggedThumbImage( BOX_BRUSH( "Common/Scrollbar_Thumb", FMargin(4.f/16.f) ) )
-			.SetHoveredThumbImage( BOX_BRUSH( "Common/Scrollbar_Thumb", FMargin(4.f/16.f) ) )
-			);
-	}
-
 	// SExpandableArea defaults...
 	{
 		Style->Set( "ExpandableArea", FExpandableAreaStyle()
@@ -506,7 +508,8 @@ TSharedRef<ISlateStyle> FCoreStyle::Create( const FName& InStyleSetName )
 			.SetBackgroundImageNormal( BOX_BRUSH( "Common/TextBox_Special", FMargin(8.0f/32.0f) ) )
 			.SetBackgroundImageHovered( BOX_BRUSH( "Common/TextBox_Special_Hovered", FMargin(8.0f/32.0f) ) )
 			.SetBackgroundImageFocused( BOX_BRUSH( "Common/TextBox_Special_Hovered", FMargin(8.0f/32.0f) ) )
-			.SetBackgroundImageReadOnly( BOX_BRUSH( "Common/TextBox_ReadOnly", FMargin(4.0f/16.0f) ) );
+			.SetBackgroundImageReadOnly( BOX_BRUSH( "Common/TextBox_ReadOnly", FMargin(4.0f/16.0f) ) )
+			.SetScrollBarStyle( ScrollBar );
 
 		Style->Set( "SearchBox", FSearchBoxStyle()
 			.SetTextBoxStyle( SpecialEditableTextBoxStyle )
diff --git a/Engine/Source/Runtime/SlateCore/Private/Styling/SlateTypes.cpp b/Engine/Source/Runtime/SlateCore/Private/Styling/SlateTypes.cpp
--- a/Engine/Source/Runtime/SlateCore/Private/Styling/SlateTypes.cpp
+++ b/Engine/Source/Runtime/SlateCore/Private/Styling/SlateTypes.cpp
@@ -233,6 +233,8 @@ void FEditableTextBoxStyle::GetResources( TArray< const FSlateBrush* >& OutBrush
 	OutBrushes.Add( &BackgroundImageHovered );
 	OutBrushes.Add( &BackgroundImageFocused );
 	OutBrushes.Add( &BackgroundImageReadOnly );
+
+	ScrollBarStyle.GetResources(OutBrushes);
 }
 
 const FName FEditableTextBoxStyle::TypeName( TEXT("FEditableTextBoxStyle") );
diff --git a/Engine/Source/Runtime/SlateCore/Public/Styling/SlateTypes.h b/Engine/Source/Runtime/SlateCore/Public/Styling/SlateTypes.h
--- a/Engine/Source/Runtime/SlateCore/Public/Styling/SlateTypes.h
+++ b/Engine/Source/Runtime/SlateCore/Public/Styling/SlateTypes.h
@@ -492,6 +492,53 @@
 	FEditableTextStyle& SetCaretImage( const FSlateBrush& InCaretImage ){ CaretImage = InCaretImage; return *this; }
 };
 
+
 /**
+ * Represents the appearance of an SScrollBar
+ */
+USTRUCT()
+struct SLATECORE_API FScrollBarStyle : public FSlateWidgetStyle
+{
+	GENERATED_USTRUCT_BODY()
+
+	FScrollBarStyle();
+
+	virtual ~FScrollBarStyle() {}
+
+	virtual void GetResources( TArray< const FSlateBrush* >& OutBrushes ) const override;
+
+	static const FName TypeName;
+	virtual const FName GetTypeName() const override { return TypeName; };
+
+	static const FScrollBarStyle& GetDefault();
+
+	/** Background image to use when the scrollbar is oriented horizontally */
+	UPROPERTY(EditAnywhere, Category=Appearance)
+	FSlateBrush HorizontalBackgroundImage;
+	FScrollBarStyle& SetHorizontalBackgroundImage( const FSlateBrush& InHorizontalBackgroundImage ){ HorizontalBackgroundImage = InHorizontalBackgroundImage; return *this; }
+
+	/** Background image to use when the scrollbar is oriented vertically */
+	UPROPERTY(EditAnywhere, Category=Appearance)
+	FSlateBrush VerticalBackgroundImage;
+	FScrollBarStyle& SetVerticalBackgroundImage( const FSlateBrush& InVerticalBackgroundImage ){ VerticalBackgroundImage = InVerticalBackgroundImage; return *this; }
+
+	/** Image to use when the scrollbar thumb is in its normal state */
+	UPROPERTY(EditAnywhere, Category=Appearance)
+	FSlateBrush NormalThumbImage;
+	FScrollBarStyle& SetNormalThumbImage( const FSlateBrush& InNormalThumbImage ){ NormalThumbImage = InNormalThumbImage; return *this; }
+
+	/** Image to use when the scrollbar thumb is in its hovered state */
+	UPROPERTY(EditAnywhere, Category=Appearance)
+	FSlateBrush HoveredThumbImage;
+	FScrollBarStyle& SetHoveredThumbImage( const FSlateBrush& InHoveredThumbImage ){ HoveredThumbImage = InHoveredThumbImage; return *this; }
+
+	/** Image to use when the scrollbar thumb is in its dragged state */
+	UPROPERTY(EditAnywhere, Category=Appearance)
+	FSlateBrush DraggedThumbImage;
+	FScrollBarStyle& SetDraggedThumbImage( const FSlateBrush& InDraggedThumbImage ){ DraggedThumbImage = InDraggedThumbImage; return *this; }
+};
+
+
+/**
  * Represents the appearance of an SEditableTextBox
  */
@@ -556,6 +603,21 @@
 	UPROPERTY(EditAnywhere, Category=Appearance)
 	FSlateColor ReadOnlyForegroundColor;
 	FEditableTextBoxStyle& SetReadOnlyForegroundColor(const FSlateColor& InReadOnlyForegroundColor) { ReadOnlyForegroundColor = InReadOnlyForegroundColor; return *this; }
+
+	/** Padding around the horizontal scrollbar */
+	UPROPERTY(EditAnywhere, Category=Appearance)
+	FMargin HScrollBarPadding;
+	FEditableTextBoxStyle& SetHScrollBarPadding( const FMargin& InPadding ){ HScrollBarPadding = InPadding; return *this; }
+
+	/** Padding around the vertical scrollbar */
+	UPROPERTY(EditAnywhere, Category=Appearance)
+	FMargin VScrollBarPadding;
+	FEditableTextBoxStyle& SetVScrollBarPadding( const FMargin& InPadding ){ VScrollBarPadding = InPadding; return *this; }
+
+	/** Style used for the scrollbars */
+	UPROPERTY(EditAnywhere, Category=Appearance)
+	FScrollBarStyle ScrollBarStyle;
+	FEditableTextBoxStyle& SetScrollBarStyle( const FScrollBarStyle& InScrollBarStyle ){ ScrollBarStyle = InScrollBarStyle; return *this; }
 };
 
 
@@ -627,51 +689,5 @@
 
 
 /**
- * Represents the appearance of an SScrollBar
- */
-USTRUCT()
-struct SLATECORE_API FScrollBarStyle : public FSlateWidgetStyle
-{
-	GENERATED_USTRUCT_BODY()
-
-	FScrollBarStyle();
-
-	virtual ~FScrollBarStyle() {}
-
-	virtual void GetResources( TArray< const FSlateBrush* >& OutBrushes ) const override;
-
-	static const FName TypeName;
-	virtual const FName GetTypeName() const override { return TypeName; };
-
-	static const FScrollBarStyle& GetDefault();
-
-	/** Background image to use when the scrollbar is oriented horizontally */
-	UPROPERTY(EditAnywhere, Category=Appearance)
-	FSlateBrush HorizontalBackgroundImage;
-	FScrollBarStyle& SetHorizontalBackgroundImage( const FSlateBrush& InHorizontalBackgroundImage ){ HorizontalBackgroundImage = InHorizontalBackgroundImage; return *this; }
-
-	/** Background image to use when the scrollbar is oriented vertically */
-	UPROPERTY(EditAnywhere, Category=Appearance)
-	FSlateBrush VerticalBackgroundImage;
-	FScrollBarStyle& SetVerticalBackgroundImage( const FSlateBrush& InVerticalBackgroundImage ){ VerticalBackgroundImage = InVerticalBackgroundImage; return *this; }
-
-	/** Image to use when the scrollbar thumb is in its normal state */
-	UPROPERTY(EditAnywhere, Category=Appearance)
-	FSlateBrush NormalThumbImage;
-	FScrollBarStyle& SetNormalThumbImage( const FSlateBrush& InNormalThumbImage ){ NormalThumbImage = InNormalThumbImage; return *this; }
-
-	/** Image to use when the scrollbar thumb is in its hovered state */
-	UPROPERTY(EditAnywhere, Category=Appearance)
-	FSlateBrush HoveredThumbImage;
-	FScrollBarStyle& SetHoveredThumbImage( const FSlateBrush& InHoveredThumbImage ){ HoveredThumbImage = InHoveredThumbImage; return *this; }
-
-	/** Image to use when the scrollbar thumb is in its dragged state */
-	UPROPERTY(EditAnywhere, Category=Appearance)
-	FSlateBrush DraggedThumbImage;
-	FScrollBarStyle& SetDraggedThumbImage( const FSlateBrush& InDraggedThumbImage ){ DraggedThumbImage = InDraggedThumbImage; return *this; }
-};
-
-
-/**
  * Represents the appearance of an SExpandableArea
  */
