Reverted a bad change to SMultiLineEditableText::ComputeDesiredSize

The wrapping width used should only ever be the fixed wrap-width, *not* the auto-wrap width, as the auto-wrap width is computed based on the size we're given, not the size we asked for.

diff --git a/Engine/Source/Runtime/Slate/Private/Widgets/Text/SMultiLineEditableText.cpp b/Engine/Source/Runtime/Slate/Private/Widgets/Text/SMultiLineEditableText.cpp
--- a/Engine/Source/Runtime/Slate/Private/Widgets/Text/SMultiLineEditableText.cpp
+++ b/Engine/Source/Runtime/Slate/Private/Widgets/Text/SMultiLineEditableText.cpp
@@ -2023,7 +2023,7 @@ FVector2D SMultiLineEditableText::ComputeDesiredSize() const
 	// The layouts current margin size. We should not report a size smaller then the margins.
 	const FMargin Margin = TextLayout->GetMargin();
 	const FVector2D TextLayoutSize = TextLayout->GetSize();
-	const float WrappingWidth = TextLayout->GetWrappingWidth();
+	const float WrappingWidth = WrapTextAt.Get();
 
 	// If a wrapping width has been provided that should be reported as the desired width.
 	float DesiredWidth = WrappingWidth > 0 ? WrappingWidth : TextLayoutSize.X;
