Slate - Capturing the mouse when moving tabs now, prevents sticky tab issues at low framerates.

diff --git a/Engine/Source/Runtime/Slate/Private/Widgets/Docking/SDockTab.cpp b/Engine/Source/Runtime/Slate/Private/Widgets/Docking/SDockTab.cpp
--- a/Engine/Source/Runtime/Slate/Private/Widgets/Docking/SDockTab.cpp
+++ b/Engine/Source/Runtime/Slate/Private/Widgets/Docking/SDockTab.cpp
@@ -22,7 +22,7 @@ static float TotalDraggedDistance = 0;
 
 FReply SDockTab::OnMouseButtonDown( const FGeometry& MyGeometry, const FPointerEvent& MouseEvent )
 {
-	if (!this->HasMouseCapture())
+	if ( !HasMouseCapture() )
 	{
 		if ( MouseEvent.GetEffectingButton() == EKeys::LeftMouseButton )
 		{
@@ -37,7 +37,7 @@ FReply SDockTab::OnMouseButtonDown( const FGeometry& MyGeometry, const FPointerE
 				return FReply::Unhandled();
 			}
 
-			return FReply::Handled().DetectDrag( SharedThis(this), EKeys::LeftMouseButton );
+			return FReply::Handled().CaptureMouse(SharedThis(this)).DetectDrag(SharedThis(this), EKeys::LeftMouseButton);
 		}
 		else if ( MouseEvent.GetEffectingButton() == EKeys::MiddleMouseButton )
 		{
@@ -81,7 +81,7 @@ FReply SDockTab::OnDragDetected( const FGeometry& MyGeometry, const FPointerEven
 
 FReply SDockTab::OnMouseButtonUp( const FGeometry& MyGeometry, const FPointerEvent& MouseEvent )
 {
-	if (this->HasMouseCapture())
+	if ( HasMouseCapture() )
 	{
 		if ( MouseEvent.GetEffectingButton() == EKeys::LeftMouseButton )
 		{
