Stylistic change: Changed some usages of FVector2D(ForceInitToZero) to be FVector2D::ZeroVector.

As NickA pointed out, the latter is easier to read and I'd been inconsistent in my usage.

diff --git a/Engine/Source/Runtime/Slate/Private/Widgets/Text/SMultiLineEditableText.cpp b/Engine/Source/Runtime/Slate/Private/Widgets/Text/SMultiLineEditableText.cpp
--- a/Engine/Source/Runtime/Slate/Private/Widgets/Text/SMultiLineEditableText.cpp
+++ b/Engine/Source/Runtime/Slate/Private/Widgets/Text/SMultiLineEditableText.cpp
@@ -212,7 +212,7 @@ void SMultiLineEditableText::Construct( const FArguments& InArgs )
 
 	WrapTextAt = InArgs._WrapTextAt;
 	AutoWrapText = InArgs._AutoWrapText;
-	CachedSize = FVector2D(ForceInitToZero);
+	CachedSize = FVector2D::ZeroVector;
 
 	ScrollOffset = FVector2D::ZeroVector;
 
diff --git a/Engine/Source/Runtime/Slate/Private/Widgets/Text/SRichTextBlock.cpp b/Engine/Source/Runtime/Slate/Private/Widgets/Text/SRichTextBlock.cpp
--- a/Engine/Source/Runtime/Slate/Private/Widgets/Text/SRichTextBlock.cpp
+++ b/Engine/Source/Runtime/Slate/Private/Widgets/Text/SRichTextBlock.cpp
@@ -20,7 +20,7 @@ void SRichTextBlock::Construct( const FArguments& InArgs )
 	LineHeightPercentage = InArgs._LineHeightPercentage;
 	Justification = InArgs._Justification;
 
-	CachedSize = FVector2D(ForceInitToZero);
+	CachedSize = FVector2D::ZeroVector;
 
 	{
 		TSharedPtr<IRichTextMarkupParser> Parser = InArgs._Parser;
@@ -73,7 +73,7 @@ void SRichTextBlock::Tick( const FGeometry& AllottedGeometry, const double InCur
 		SetText(TextToSet);
 	}
 
-	TextLayout->SetVisibleRegion(AllottedGeometry.Size, FVector2D(ForceInitToZero));
+	TextLayout->SetVisibleRegion(AllottedGeometry.Size, FVector2D::ZeroVector);
 }
 
 int32 SRichTextBlock::OnPaint( const FPaintArgs& Args, const FGeometry& AllottedGeometry, const FSlateRect& MyClippingRect, FSlateWindowElementList& OutDrawElements, int32 LayerId, const FWidgetStyle& InWidgetStyle, bool bParentEnabled ) const
@@ -104,7 +104,7 @@ void SRichTextBlock::CacheDesiredSize()
 	TextLayout->SetMargin( Margin.Get() );
 	TextLayout->SetLineHeightPercentage( LineHeightPercentage.Get() );
 	TextLayout->SetJustification( Justification.Get() );
-	TextLayout->SetVisibleRegion( CachedSize, FVector2D(ForceInitToZero) );
+	TextLayout->SetVisibleRegion( CachedSize, FVector2D::ZeroVector );
 	TextLayout->UpdateIfNeeded();
 
 	SWidget::CacheDesiredSize();
