Standardized some tab right click menu options to always show up and be disabled if they arn't availible due to the current state.

diff --git a/Engine/Source/Runtime/Slate/Private/Framework/Docking/SDockingTabStack.cpp b/Engine/Source/Runtime/Slate/Private/Framework/Docking/SDockingTabStack.cpp
--- a/Engine/Source/Runtime/Slate/Private/Framework/Docking/SDockingTabStack.cpp
+++ b/Engine/Source/Runtime/Slate/Private/Framework/Docking/SDockingTabStack.cpp
@@ -531,27 +531,21 @@ TSharedRef<SWidget> SDockingTabStack::MakeContextMenu()
 				FSlateIcon(),
 				FUIAction(
 					FExecuteAction::CreateSP( this, &SDockingTabStack::CloseForegroundTab ),
-					FCanExecuteAction()
+					FCanExecuteAction::CreateSP(this, &SDockingTabStack::CanCloseForegroundTab)
 				)
 			);
 
-			// If the active tab is a document tab, and there is more than one open in this tab well, offer to close the others
-			TSharedPtr<SDockTab> ForegroundTabPtr = TabWell->GetForegroundTab();
-			if (ForegroundTabPtr.IsValid() && (ForegroundTabPtr->GetTabRole() == ETabRole::DocumentTab || ForegroundTabPtr->GetTabRole() == ETabRole::MajorTab) && (TabWell->GetNumTabs() > 1))
-			{
-				const ETabsToClose TabsToClose = CloseDocumentAndMajorTabs;
-
-				MenuBuilder.AddMenuEntry(
-					LOCTEXT("CloseOtherTabs", "Close Other Tabs"),
-					LOCTEXT("CloseOtherTabsTooltil", "Closes all tabs except for the active tab."),
-					FSlateIcon(),
-					FUIAction(
-						FExecuteAction::CreateSP( this, &SDockingTabStack::CloseAllButForegroundTab, TabsToClose ),
-						FCanExecuteAction()
-					)
-				);
-					
-			}
+			const ETabsToClose TabsToClose = CloseDocumentAndMajorTabs;
+
+			MenuBuilder.AddMenuEntry(
+				LOCTEXT("CloseOtherTabs", "Close Other Tabs"),
+				LOCTEXT("CloseOtherTabsTooltil", "Closes all tabs except for the active tab."),
+				FSlateIcon(),
+				FUIAction(
+					FExecuteAction::CreateSP( this, &SDockingTabStack::CloseAllButForegroundTab, TabsToClose ),
+					FCanExecuteAction::CreateSP(this, &SDockingTabStack::CanCloseAllButForegroundTab)
+				)
+			);
 		}
 		MenuBuilder.EndSection();
 	}
@@ -829,12 +823,37 @@ bool SDockingTabStack::CanHideTabWell() const
 	return GetNumTabs() == 1 && (GetTabs()[0]->GetTabRole() != ETabRole::MajorTab && !GetTabs()[0]->IsNomadTabWithMajorTabStyle());
 }
 
+bool SDockingTabStack::CanCloseForegroundTab() const
+{
+	TSharedPtr<SDockTab> ForegroundTabPtr = TabWell->GetForegroundTab();
+	return ForegroundTabPtr.IsValid() && ForegroundTabPtr->CanCloseTab();
+}
+
+bool SDockingTabStack::CanCloseAllButForegroundTab() const
+{
+	// If the active tab is a document tab or major tab and there is at least 1 other closeable tab, offer to close the others
+	TSharedPtr<SDockTab> ForegroundTabPtr = TabWell->GetForegroundTab();
+	if (ForegroundTabPtr.IsValid() && (ForegroundTabPtr->GetTabRole() == ETabRole::DocumentTab || ForegroundTabPtr->GetTabRole() == ETabRole::MajorTab) && (TabWell->GetNumTabs() > 1))
+	{
+		const TArray< TSharedRef<SDockTab> > MyTabs = this->GetTabs().AsArrayCopy();
+		for (int32 TabIndex = 0; TabIndex < MyTabs.Num(); ++TabIndex)
+		{
+			TSharedRef<SDockTab> Tab = MyTabs[TabIndex];
+			if (Tab != ForegroundTabPtr && Tab->CanCloseTab())
+			{
+				return true;
+			}
+		}
+	}
+	return false;
+}
+
 SSplitter::ESizeRule SDockingTabStack::GetSizeRule() const
 {
 	if (this->GetNumTabs() == 1 && this->GetTabs()[0]->ShouldAutosize() )
 	{
 		// If there is a single tab in this stack and it is
-		// sized to content, then the stack's cell shoudl size to Content.
+		// sized to content, then the stack's cell should size to Content.
 		return SSplitter::SizeToContent;
 	}
 	else
diff --git a/Engine/Source/Runtime/Slate/Private/Framework/Docking/SDockingTabStack.h b/Engine/Source/Runtime/Slate/Private/Framework/Docking/SDockingTabStack.h
--- a/Engine/Source/Runtime/Slate/Private/Framework/Docking/SDockingTabStack.h
+++ b/Engine/Source/Runtime/Slate/Private/Framework/Docking/SDockingTabStack.h
@@ -175,6 +175,12 @@ private:
 	/** Only allow hiding the tab well when there is a single tab in it. */
 	bool CanHideTabWell() const;
 
+	/** Only allow closing the tab well when the tab allows it. */
+	bool CanCloseForegroundTab() const;
+
+	/** Only allow closing all other tabs when there are more then one tab open and the tab is of type Document or Major. */
+	bool CanCloseAllButForegroundTab() const;
+
 	/** Keep around our geometry from the last frame so that we can resize the preview windows correctly */
 	FGeometry TabStackGeometry;
 
