On OS X detach the game thread from the 'blessed' Cocoa main thread.
Cocoa isn't an event-polling API as UE4 expects, so previously we were subverting the NSApplication's event handling to pretend that it was. When the engine wasn't running the event loop fast enough, such as when loading where it isn't processed at all, this resulted in unresponsive windows and Spinning-Beachball-Of-Death. That isn't very satisfactory & to some users appears as if the application has crashed. To address these deficiencies without further attempts to subvert Cocoa, the game is now punted onto a separate thread where it can run its own tight-loop, leaving the main thread to handle the Cocoa event run-loop. Events are captured by delegate objects, as Cocoa requires, but dispatched and handled on the game thread which makes Cocoa appear more like other platform APIs to the higher-level UE4 code.
This can all be disabled using the MAC_SEPARATE_GAME_THREAD define in CocoaThread.cpp.

diff --git a/Engine/Source/Developer/DesktopPlatform/Private/Mac/DesktopPlatformMac.cpp b/Engine/Source/Developer/DesktopPlatform/Private/Mac/DesktopPlatformMac.cpp
--- a/Engine/Source/Developer/DesktopPlatform/Private/Mac/DesktopPlatformMac.cpp
+++ b/Engine/Source/Developer/DesktopPlatform/Private/Mac/DesktopPlatformMac.cpp
@@ -4,6 +4,7 @@
 #include "MacApplication.h"
 #include "FeedbackContextMarkup.h"
 #include "MacNativeFeedbackContext.h"
+#include "CocoaThread.h"
 
 #define LOCTEXT_NAMESPACE "DesktopPlatform"
 
@@ -223,99 +224,107 @@ bool FDesktopPlatformMac::SaveFileDialog(const void* ParentWindowHandle, const F
 
 bool FDesktopPlatformMac::OpenDirectoryDialog(const void* ParentWindowHandle, const FString& DialogTitle, const FString& DefaultPath, FString& OutFolderName)
 {
-	SCOPED_AUTORELEASE_POOL;
-	FCocoaScopeContext ContextGuard;
-
 	MacApplication->SetCapture( NULL );
 
-	NSOpenPanel* Panel = [NSOpenPanel openPanel];
-	[Panel setCanChooseFiles: false];
-	[Panel setCanChooseDirectories: true];
-	[Panel setAllowsMultipleSelection: false];
-	[Panel setCanCreateDirectories: true];
+	bool bSuccess = MainThreadReturn(^{
+		SCOPED_AUTORELEASE_POOL;
+		FCocoaScopeContext ContextGuard;
 
-	CFStringRef Title = FPlatformString::TCHARToCFString(*DialogTitle);
-	[Panel setTitle: (NSString*)Title];
-	CFRelease(Title);
+		NSOpenPanel* Panel = [NSOpenPanel openPanel];
+		[Panel setCanChooseFiles: false];
+		[Panel setCanChooseDirectories: true];
+		[Panel setAllowsMultipleSelection: false];
+		[Panel setCanCreateDirectories: true];
 
-	CFStringRef DefaultPathCFString = FPlatformString::TCHARToCFString(*DefaultPath);
-	NSURL* DefaultPathURL = [NSURL fileURLWithPath: (NSString*)DefaultPathCFString];
-	[Panel setDirectoryURL: DefaultPathURL];
-	CFRelease(DefaultPathCFString);
+		CFStringRef Title = FPlatformString::TCHARToCFString(*DialogTitle);
+		[Panel setTitle: (NSString*)Title];
+		CFRelease(Title);
 
-	bool bSuccess = false;
+		CFStringRef DefaultPathCFString = FPlatformString::TCHARToCFString(*DefaultPath);
+		NSURL* DefaultPathURL = [NSURL fileURLWithPath: (NSString*)DefaultPathCFString];
+		[Panel setDirectoryURL: DefaultPathURL];
+		CFRelease(DefaultPathCFString);
 
-	{
-		FScopedSystemModalMode SystemModalScope;
-		NSInteger Result = [Panel runModal];
+		bool bSuccess = false;
 
-		if (Result == NSFileHandlingPanelOKButton)
 		{
-			NSURL *FolderURL = [[Panel URLs] objectAtIndex: 0];
-			TCHAR FolderPath[MAX_PATH];
-			FPlatformString::CFStringToTCHAR((CFStringRef)[FolderURL path], FolderPath);
-			OutFolderName = FolderPath;
-			FPaths::NormalizeFilename(OutFolderName);
+			FScopedSystemModalMode SystemModalScope;
+			NSInteger Result = [Panel runModal];
+
+			if (Result == NSFileHandlingPanelOKButton)
+			{
+				NSURL *FolderURL = [[Panel URLs] objectAtIndex: 0];
+				TCHAR FolderPath[MAX_PATH];
+				FPlatformString::CFStringToTCHAR((CFStringRef)[FolderURL path], FolderPath);
+				OutFolderName = FolderPath;
+				FPaths::NormalizeFilename(OutFolderName);
 
-			bSuccess = true;
+				bSuccess = true;
+			}
 		}
-	}
 
-	[Panel close];
+		[Panel close];
 
+		return bSuccess;
+	});
+	
 	MacApplication->ResetModifierKeys();
-
+	
 	return bSuccess;
 }
 
 bool FDesktopPlatformMac::OpenFontDialog(const void* ParentWindowHandle, FString& OutFontName, float& OutHeight, EFontImportFlags& OutFlags)
 {
-	SCOPED_AUTORELEASE_POOL;
-	FCocoaScopeContext ContextGuard;
-
 	MacApplication->SetCapture( NULL );
+	
+	bool bSuccess = MainThreadReturn(^{
+		SCOPED_AUTORELEASE_POOL;
+		FCocoaScopeContext ContextGuard;
 
-	NSFontPanel* Panel = [NSFontPanel sharedFontPanel];
-	[Panel setFloatingPanel: false];
-	[[Panel standardWindowButton: NSWindowCloseButton] setEnabled: false];
-	FFontDialogAccessoryView* AccessoryView = [[FFontDialogAccessoryView alloc] initWithFrame: NSMakeRect(0.0, 0.0, 190.0, 80.0)];
-	[Panel setAccessoryView: AccessoryView];
+		NSFontPanel* Panel = [NSFontPanel sharedFontPanel];
+		[Panel setFloatingPanel: false];
+		[[Panel standardWindowButton: NSWindowCloseButton] setEnabled: false];
+		FFontDialogAccessoryView* AccessoryView = [[FFontDialogAccessoryView alloc] initWithFrame: NSMakeRect(0.0, 0.0, 190.0, 80.0)];
+		[Panel setAccessoryView: AccessoryView];
 
-	{
-		FScopedSystemModalMode SystemModalScope;
-		[NSApp runModalForWindow: Panel];
-	}
+		{
+			FScopedSystemModalMode SystemModalScope;
+			[NSApp runModalForWindow: Panel];
+		}
 
-	[Panel close];
+		[Panel close];
 
-	bool bSuccess = [AccessoryView result];
+		bool bSuccess = [AccessoryView result];
 
-	[Panel setAccessoryView: NULL];
-	[AccessoryView release];
-	[[Panel standardWindowButton: NSWindowCloseButton] setEnabled: true];
+		[Panel setAccessoryView: NULL];
+		[AccessoryView release];
+		[[Panel standardWindowButton: NSWindowCloseButton] setEnabled: true];
 
-	if( bSuccess )
-	{
-		NSFont* Font = [Panel panelConvertFont: [NSFont userFontOfSize: 0]];
+		if( bSuccess )
+		{
+			NSFont* Font = [Panel panelConvertFont: [NSFont userFontOfSize: 0]];
 
-		TCHAR FontName[MAX_PATH];
-		FPlatformString::CFStringToTCHAR((CFStringRef)[Font fontName], FontName);
+			TCHAR FontName[MAX_PATH];
+			FPlatformString::CFStringToTCHAR((CFStringRef)[Font fontName], FontName);
 
-		OutFontName = FontName;
-		OutHeight = [Font pointSize];
+			OutFontName = FontName;
+			OutHeight = [Font pointSize];
 
-		auto FontFlags = EFontImportFlags::None;
+			auto FontFlags = EFontImportFlags::None;
 
-		if( [Font underlineThickness] >= 1.0 )
-		{
-			FontFlags |= EFontImportFlags::EnableUnderline;
-		}
+			if( [Font underlineThickness] >= 1.0 )
+			{
+				FontFlags |= EFontImportFlags::EnableUnderline;
+			}
 
-		OutFlags = FontFlags;
-	}
+			OutFlags = FontFlags;
+		}
 
+		return bSuccess;
+	});
+	
 	MacApplication->ResetModifierKeys();
-
+	
 	return bSuccess;
 }
 
@@ -397,102 +406,106 @@ bool FDesktopPlatformMac::OpenLauncher(bool Install, FString CommandLineParams )
 
 bool FDesktopPlatformMac::FileDialogShared(bool bSave, const void* ParentWindowHandle, const FString& DialogTitle, const FString& DefaultPath, const FString& DefaultFile, const FString& FileTypes, uint32 Flags, TArray<FString>& OutFilenames)
 {
-	SCOPED_AUTORELEASE_POOL;
-	FCocoaScopeContext ContextGuard;
-
 	MacApplication->SetCapture( NULL );
 
-	NSSavePanel* Panel = bSave ? [NSSavePanel savePanel] : [NSOpenPanel openPanel];
+	bool bSuccess = MainThreadReturn(^{
+		SCOPED_AUTORELEASE_POOL;
+		FCocoaScopeContext ContextGuard;
 
-	if (!bSave)
-	{
-		NSOpenPanel* OpenPanel = (NSOpenPanel*)Panel;
-		[OpenPanel setCanChooseFiles: true];
-		[OpenPanel setCanChooseDirectories: false];
-		[OpenPanel setAllowsMultipleSelection: Flags & EFileDialogFlags::Multiple];
-	}
+		NSSavePanel* Panel = bSave ? [NSSavePanel savePanel] : [NSOpenPanel openPanel];
 
-	[Panel setCanCreateDirectories: bSave];
+		if (!bSave)
+		{
+			NSOpenPanel* OpenPanel = (NSOpenPanel*)Panel;
+			[OpenPanel setCanChooseFiles: true];
+			[OpenPanel setCanChooseDirectories: false];
+			[OpenPanel setAllowsMultipleSelection: Flags & EFileDialogFlags::Multiple];
+		}
 
-	CFStringRef Title = FPlatformString::TCHARToCFString(*DialogTitle);
-	[Panel setTitle: (NSString*)Title];
-	CFRelease(Title);
+		[Panel setCanCreateDirectories: bSave];
 
-	CFStringRef DefaultPathCFString = FPlatformString::TCHARToCFString(*DefaultPath);
-	NSURL* DefaultPathURL = [NSURL fileURLWithPath: (NSString*)DefaultPathCFString];
-	[Panel setDirectoryURL: DefaultPathURL];
-	CFRelease(DefaultPathCFString);
+		CFStringRef Title = FPlatformString::TCHARToCFString(*DialogTitle);
+		[Panel setTitle: (NSString*)Title];
+		CFRelease(Title);
 
-	CFStringRef FileNameCFString = FPlatformString::TCHARToCFString(*DefaultFile);
-	[Panel setNameFieldStringValue: (NSString*)FileNameCFString];
-	CFRelease(FileNameCFString);
+		CFStringRef DefaultPathCFString = FPlatformString::TCHARToCFString(*DefaultPath);
+		NSURL* DefaultPathURL = [NSURL fileURLWithPath: (NSString*)DefaultPathCFString];
+		[Panel setDirectoryURL: DefaultPathURL];
+		CFRelease(DefaultPathCFString);
 
-	FFileDialogAccessoryView* AccessoryView = [[FFileDialogAccessoryView alloc] initWithFrame: NSMakeRect( 0.0, 0.0, 250.0, 85.0 ) dialogPanel: Panel];
-	[Panel setAccessoryView: AccessoryView];
+		CFStringRef FileNameCFString = FPlatformString::TCHARToCFString(*DefaultFile);
+		[Panel setNameFieldStringValue: (NSString*)FileNameCFString];
+		CFRelease(FileNameCFString);
 
-	TArray<FString> FileTypesArray;
-	int32 NumFileTypes = FileTypes.ParseIntoArray(&FileTypesArray, TEXT("|"), true);
+		FFileDialogAccessoryView* AccessoryView = [[FFileDialogAccessoryView alloc] initWithFrame: NSMakeRect( 0.0, 0.0, 250.0, 85.0 ) dialogPanel: Panel];
+		[Panel setAccessoryView: AccessoryView];
 
-	NSMutableArray* AllowedFileTypes = [NSMutableArray arrayWithCapacity: NumFileTypes];
+		TArray<FString> FileTypesArray;
+		int32 NumFileTypes = FileTypes.ParseIntoArray(&FileTypesArray, TEXT("|"), true);
 
-	if( NumFileTypes > 0 )
-	{
-		for( int32 Index = 0; Index < NumFileTypes; ++Index )
+		NSMutableArray* AllowedFileTypes = [NSMutableArray arrayWithCapacity: NumFileTypes];
+
+		if( NumFileTypes > 0 )
 		{
-			CFStringRef Type = FPlatformString::TCHARToCFString(*FileTypesArray[Index]);
-			[AllowedFileTypes addObject: (NSString*)Type];
-			CFRelease(Type);
+			for( int32 Index = 0; Index < NumFileTypes; ++Index )
+			{
+				CFStringRef Type = FPlatformString::TCHARToCFString(*FileTypesArray[Index]);
+				[AllowedFileTypes addObject: (NSString*)Type];
+				CFRelease(Type);
+			}
 		}
-	}
 
-	[AccessoryView AddAllowedFileTypes:AllowedFileTypes];
+		[AccessoryView AddAllowedFileTypes:AllowedFileTypes];
 
-	bool bSuccess = false;
-	NSWindow* FocusWindow = [[NSApplication sharedApplication] keyWindow];
+		bool bSuccess = false;
+		NSWindow* FocusWindow = [[NSApplication sharedApplication] keyWindow];
 
-	{
-		FScopedSystemModalMode SystemModalScope;
+		{
+			FScopedSystemModalMode SystemModalScope;
 
-		NSInteger Result = [Panel runModal];
-		[AccessoryView release];
+			NSInteger Result = [Panel runModal];
+			[AccessoryView release];
 
-		if (Result == NSFileHandlingPanelOKButton)
-		{
-			if (bSave)
-			{
-				TCHAR FilePath[MAX_PATH];
-				FPlatformString::CFStringToTCHAR((CFStringRef)[[Panel URL] path], FilePath);
-				new(OutFilenames) FString(FilePath);
-			}
-			else
+			if (Result == NSFileHandlingPanelOKButton)
 			{
-				NSOpenPanel* OpenPanel = (NSOpenPanel*)Panel;
-				for (NSURL *FileURL in [OpenPanel URLs])
+				if (bSave)
 				{
 					TCHAR FilePath[MAX_PATH];
-					FPlatformString::CFStringToTCHAR((CFStringRef)[FileURL path], FilePath);
+					FPlatformString::CFStringToTCHAR((CFStringRef)[[Panel URL] path], FilePath);
 					new(OutFilenames) FString(FilePath);
 				}
-			}
+				else
+				{
+					NSOpenPanel* OpenPanel = (NSOpenPanel*)Panel;
+					for (NSURL *FileURL in [OpenPanel URLs])
+					{
+						TCHAR FilePath[MAX_PATH];
+						FPlatformString::CFStringToTCHAR((CFStringRef)[FileURL path], FilePath);
+						new(OutFilenames) FString(FilePath);
+					}
+				}
 
-			// Make sure all filenames gathered have their paths normalized
-			for (auto FilenameIt = OutFilenames.CreateIterator(); FilenameIt; ++FilenameIt)
-			{
-				FPaths::NormalizeFilename(*FilenameIt);
-			}
+				// Make sure all filenames gathered have their paths normalized
+				for (auto FilenameIt = OutFilenames.CreateIterator(); FilenameIt; ++FilenameIt)
+				{
+					FPaths::NormalizeFilename(*FilenameIt);
+				}
 
-			bSuccess = true;
+				bSuccess = true;
+			}
 		}
-	}
 
-	[Panel close];
+		[Panel close];
+		
+		if(FocusWindow)
+		{
+			[FocusWindow makeKeyWindow];
+		}
 
-	MacApplication->ResetModifierKeys();
+		return bSuccess;
+	});
 	
-	if(FocusWindow)
-	{
-		[FocusWindow makeKeyWindow];
-	}
+	MacApplication->ResetModifierKeys();
 
 	return bSuccess;
 }
diff --git a/Engine/Source/Developer/DesktopPlatform/Private/Mac/MacNativeFeedbackContext.cpp b/Engine/Source/Developer/DesktopPlatform/Private/Mac/MacNativeFeedbackContext.cpp
--- a/Engine/Source/Developer/DesktopPlatform/Private/Mac/MacNativeFeedbackContext.cpp
+++ b/Engine/Source/Developer/DesktopPlatform/Private/Mac/MacNativeFeedbackContext.cpp
@@ -3,6 +3,7 @@
 #include "DesktopPlatformPrivatePCH.h"
 #include "MacNativeFeedbackContext.h"
 #include "MacApplication.h"
+#include "CocoaThread.h"
 
 @implementation FMacNativeFeedbackContextWindowController
 
@@ -254,17 +255,26 @@ FMacNativeFeedbackContext::FMacNativeFeedbackContext()
 	: FFeedbackContext()
 	, WindowController(nil)
 	, Context( NULL )
+	, OutstandingTasks(0)
 	, SlowTaskCount( 0 )
 	, bShowingConsoleForSlowTask( false )
 {
-	WindowController = [[FMacNativeFeedbackContextWindowController alloc] init];
-	
+	MainThreadCall(^{
+		WindowController = [[FMacNativeFeedbackContextWindowController alloc] init];
+	}, true);
 	SetDefaultTextColor();
 }
 
 FMacNativeFeedbackContext::~FMacNativeFeedbackContext()
 {
-	[WindowController release];
+	do
+	{
+		FPlatformMisc::PumpMessages( true );
+	} while(OutstandingTasks);
+	
+	MainThreadCall(^{
+		[WindowController release];
+	}, true);
 }
 
 void FMacNativeFeedbackContext::Serialize( const TCHAR* Data, ELogVerbosity::Type Verbosity, const class FName& Category )
@@ -306,24 +316,31 @@ void FMacNativeFeedbackContext::Serialize( const TCHAR* Data, ELogVerbosity::Typ
 						*S -= '0';
 					}
 					
-					NSColor* Colors[2];
-					NSString* AttributeKeys[2];
+					NSMutableArray* Colors = [[NSMutableArray alloc] init];
+					NSMutableArray* AttributeKeys = [[NSMutableArray alloc] init];
 					
 					// Get FOREGROUND_INTENSITY and calculate final color
 					CGFloat Intensity = String[3] ? 1.0 : 0.5;
-					Colors[0] = [NSColor colorWithSRGBRed:(String[0] ? 1.0 * Intensity : 0.0) green:(String[1] ? 1.0 * Intensity : 0.0) blue:(String[2] ? 1.0 * Intensity : 0.0) alpha:1.0];
+					[Colors addObject:[NSColor colorWithSRGBRed:(String[0] ? 1.0 * Intensity : 0.0) green:(String[1] ? 1.0 * Intensity : 0.0) blue:(String[2] ? 1.0 * Intensity : 0.0) alpha:1.0]];
 					
 					// Get BACKGROUND_INTENSITY and calculate final color
 					Intensity = String[7] ? 1.0 : 0.5;
-					Colors[1] = [NSColor colorWithSRGBRed:(String[4] ? 1.0 * Intensity : 0.0) green:(String[5] ? 1.0 * Intensity : 0.0) blue:(String[6] ? 1.0 * Intensity : 0.0) alpha:1.0];
+					[Colors addObject:[NSColor colorWithSRGBRed:(String[4] ? 1.0 * Intensity : 0.0) green:(String[5] ? 1.0 * Intensity : 0.0) blue:(String[6] ? 1.0 * Intensity : 0.0) alpha:1.0]];
 					
-					AttributeKeys[0] = NSForegroundColorAttributeName;
-					AttributeKeys[1] = NSBackgroundColorAttributeName;
+					[AttributeKeys addObject:NSForegroundColorAttributeName];
+					[AttributeKeys addObject:NSBackgroundColorAttributeName];
 					
-					if( TextViewTextColor )
-					[TextViewTextColor release];
-					
-					TextViewTextColor = [[NSDictionary alloc] initWithObjects:(id *)Colors forKeys:(id *)AttributeKeys count:2];
+					OutstandingTasks++;
+					MainThreadCall(^{
+						if( TextViewTextColor )
+							[TextViewTextColor release];
+						
+						TextViewTextColor = [[NSDictionary alloc] initWithObjects:Colors forKeys:AttributeKeys];
+						
+						[Colors release];
+						[AttributeKeys release];
+						OutstandingTasks--;
+					}, false);
 				}
 			}
 			else
@@ -334,11 +351,16 @@ void FMacNativeFeedbackContext::Serialize( const TCHAR* Data, ELogVerbosity::Typ
 				FCString::Sprintf(OutputString,TEXT("%s%s"),*FOutputDevice::FormatLogLine(Verbosity, Category, Data, GPrintLogTimes),LINE_TERMINATOR);
 				
 				CFStringRef CocoaText = FPlatformString::TCHARToCFString(OutputString);
-				NSAttributedString *AttributedString = [[NSAttributedString alloc] initWithString:(NSString*)CocoaText attributes:TextViewTextColor];
-				[[WindowController->TextView textStorage] appendAttributedString:AttributedString];
-				[WindowController->TextView scrollRangeToVisible:NSMakeRange([[WindowController->TextView string] length], 0)];
-				CFRelease(CocoaText);
-				[AttributedString release];
+				
+				OutstandingTasks++;
+				MainThreadCall(^{
+					NSAttributedString *AttributedString = [[NSAttributedString alloc] initWithString:(NSString*)CocoaText attributes:TextViewTextColor];
+					[[WindowController->TextView textStorage] appendAttributedString:AttributedString];
+					[WindowController->TextView scrollRangeToVisible:NSMakeRange([[WindowController->TextView string] length], 0)];
+					[AttributedString release];
+					CFRelease(CocoaText);
+					OutstandingTasks--;
+				}, false);
 				
 				if(!MacApplication)
 				{
@@ -382,18 +404,20 @@ void FMacNativeFeedbackContext::BeginSlowTask( const FText& Task, bool bShowProg
 
 	if(SlowTaskCount > 0 && bShowProgressDialog && WindowController != NULL && !bShowingConsoleForSlowTask)
 	{
-		[WindowController setTitleText:Task.ToString().GetNSString()];
-		[WindowController setStatusText:@"Progress:"];
-		[WindowController setShowCancelButton:bInShowCancelButton];
-		[WindowController setShowProgress:true];
-		[WindowController setProgress:0 total:100];
-		[WindowController setShowProgress:false];
-		
-		SetDefaultTextColor();
-		
-		[WindowController showWindow];
-		
-		bShowingConsoleForSlowTask = true;
+		MainThreadCall(^{
+			[WindowController setTitleText:Task.ToString().GetNSString()];
+			[WindowController setStatusText:@"Progress:"];
+			[WindowController setShowCancelButton:bInShowCancelButton];
+			[WindowController setShowProgress:true];
+			[WindowController setProgress:0 total:100];
+			[WindowController setShowProgress:false];
+			
+			SetDefaultTextColor();
+			
+			[WindowController showWindow];
+			
+			bShowingConsoleForSlowTask = true;
+		}, true);
 	}
 }
 
@@ -404,8 +428,13 @@ void FMacNativeFeedbackContext::EndSlowTask()
 
 	if(SlowTaskCount == 0 && bShowingConsoleForSlowTask)
 	{
-		if(WindowController != NULL) [WindowController hideWindow];
-		bShowingConsoleForSlowTask = false;
+		MainThreadCall(^{
+			if(WindowController != NULL)
+			{
+				[WindowController hideWindow];
+			}
+			bShowingConsoleForSlowTask = false;
+		}, true);
 	}
 }
 
@@ -413,7 +442,9 @@ bool FMacNativeFeedbackContext::StatusUpdate( int32 Numerator, int32 Denominator
 {
 	if(WindowController != NULL && bShowingConsoleForSlowTask)
 	{
-		[WindowController setStatusText:NewStatus.ToString().GetNSString()];
+		MainThreadCall(^{
+			[WindowController setStatusText:NewStatus.ToString().GetNSString()];
+		}, true);
 		UpdateProgress(Numerator, Denominator);
 	}
 	return true;
@@ -426,11 +457,13 @@ bool FMacNativeFeedbackContext::StatusForceUpdate( int32 Numerator, int32 Denomi
 
 void FMacNativeFeedbackContext::UpdateProgress(int32 Numerator, int32 Denominator)
 {
-	if(WindowController != NULL && bShowingConsoleForSlowTask)
-	{
-		[WindowController setShowProgress:true];
-		[WindowController setProgress:Numerator total:Denominator];
-	}
+	MainThreadCall(^{
+		if(WindowController != NULL && bShowingConsoleForSlowTask)
+		{
+			[WindowController setShowProgress:true];
+			[WindowController setProgress:Numerator total:Denominator];
+		}
+	}, true);
 }
 
 FContextSupplier* FMacNativeFeedbackContext::GetContext() const
@@ -446,18 +479,26 @@ void FMacNativeFeedbackContext::SetContext( FContextSupplier* InSupplier )
 void FMacNativeFeedbackContext::SetDefaultTextColor()
 {
 	SCOPED_AUTORELEASE_POOL;
+	FScopeLock ScopeLock( &CriticalSection );
 	
-	NSColor* Colors[2];
-	NSString* AttributeKeys[2];
-	
-	Colors[0] = [NSColor grayColor];
-	Colors[1] = [NSColor blackColor];
+	NSMutableArray* Colors = [[NSMutableArray alloc] init];
+	NSMutableArray* AttributeKeys = [[NSMutableArray alloc] init];
 	
-	AttributeKeys[0] = NSForegroundColorAttributeName;
-	AttributeKeys[1] = NSBackgroundColorAttributeName;
+	[Colors addObject:[NSColor grayColor]];
+	[Colors addObject:[NSColor blackColor]];
 	
-	if( TextViewTextColor )
-	[TextViewTextColor release];
+	[AttributeKeys addObject:NSForegroundColorAttributeName];
+	[AttributeKeys addObject:NSBackgroundColorAttributeName];
 	
-	TextViewTextColor = [[NSDictionary alloc] initWithObjects:(id *)Colors forKeys:(id *)AttributeKeys count:2];
+	OutstandingTasks++;
+	MainThreadCall(^{
+		if( TextViewTextColor )
+			[TextViewTextColor release];
+		
+		TextViewTextColor = [[NSDictionary alloc] initWithObjects:Colors forKeys:AttributeKeys];
+		
+		[Colors release];
+		[AttributeKeys release];
+		OutstandingTasks--;
+	}, false);
 }
diff --git a/Engine/Source/Developer/DesktopPlatform/Private/Mac/MacNativeFeedbackContext.h b/Engine/Source/Developer/DesktopPlatform/Private/Mac/MacNativeFeedbackContext.h
--- a/Engine/Source/Developer/DesktopPlatform/Private/Mac/MacNativeFeedbackContext.h
+++ b/Engine/Source/Developer/DesktopPlatform/Private/Mac/MacNativeFeedbackContext.h
@@ -64,6 +64,7 @@ private:
 	FMacNativeFeedbackContextWindowController* WindowController;
 	NSDictionary* TextViewTextColor;
 	FContextSupplier* Context;
+	uint64 OutstandingTasks;
 	int32 SlowTaskCount;
 	bool bShowingConsoleForSlowTask;
 };
diff --git a/Engine/Source/Programs/SlateViewer/Private/Mac/SlateViewerMainMac.cpp b/Engine/Source/Programs/SlateViewer/Private/Mac/SlateViewerMainMac.cpp
--- a/Engine/Source/Programs/SlateViewer/Private/Mac/SlateViewerMainMac.cpp
+++ b/Engine/Source/Programs/SlateViewer/Private/Mac/SlateViewerMainMac.cpp
@@ -2,6 +2,7 @@
 
 #include "SlateViewerApp.h"
 #include "ExceptionHandling.h"
+#include "CocoaThread.h"
 
 static FString GSavedCommandLine;
 
@@ -19,20 +20,35 @@ static FString GSavedCommandLine;
     [NSApp terminate:self];
 }
 
-- (void)applicationDidFinishLaunching:(NSNotification *)Notification
+- (void) runGameThread:(id)Arg
 {
-	//install the custom quit event handler
-    NSAppleEventManager* appleEventManager = [NSAppleEventManager sharedAppleEventManager];
-    [appleEventManager setEventHandler:self andSelector:@selector(handleQuitEvent:withReplyEvent:) forEventClass:kCoreEventClass andEventID:kAEQuitApplication];
-	
 	FPlatformMisc::SetGracefulTerminationHandler();
 	FPlatformMisc::SetCrashHandler(nullptr);
 	
 	RunSlateViewer(*GSavedCommandLine);
-
+	
 	[NSApp terminate: self];
 }
 
+- (NSApplicationTerminateReply)applicationShouldTerminate:(NSApplication *)Sender;
+{
+	if(!GIsRequestingExit)
+	{
+		GIsRequestingExit = true;
+		return NSTerminateCancel;
+	}
+	return NSTerminateNow;
+}
+
+- (void)applicationDidFinishLaunching:(NSNotification *)Notification
+{
+	//install the custom quit event handler
+    NSAppleEventManager* appleEventManager = [NSAppleEventManager sharedAppleEventManager];
+    [appleEventManager setEventHandler:self andSelector:@selector(handleQuitEvent:withReplyEvent:) forEventClass:kCoreEventClass andEventID:kAEQuitApplication];
+	
+	RunGameThread(self, @selector(runGameThread:));
+}
+
 @end
 
 int main(int argc, char *argv[])
diff --git a/Engine/Source/Programs/SymbolDebugger/Private/Mac/SymbolDebuggerMainMac.cpp b/Engine/Source/Programs/SymbolDebugger/Private/Mac/SymbolDebuggerMainMac.cpp
--- a/Engine/Source/Programs/SymbolDebugger/Private/Mac/SymbolDebuggerMainMac.cpp
+++ b/Engine/Source/Programs/SymbolDebugger/Private/Mac/SymbolDebuggerMainMac.cpp
@@ -1,6 +1,7 @@
 // Copyright 1998-2014 Epic Games, Inc. All Rights Reserved.
 
 #include "SymbolDebuggerApp.h"
+#include "CocoaThread.h"
 
 static FString GSavedCommandLine;
 
@@ -12,13 +13,41 @@ static FString GSavedCommandLine;
 
 @implementation UE4AppDelegate
 
-- (void)applicationDidFinishLaunching:(NSNotification *)Notification
+//handler for the quit apple event used by the Dock menu
+- (void)handleQuitEvent:(NSAppleEventDescriptor*)Event withReplyEvent:(NSAppleEventDescriptor*)ReplyEvent
 {
-	RunSymbolDebugger(*GSavedCommandLine);
+    [NSApp terminate:self];
+}
 
+- (void) runGameThread:(id)Arg
+{
+	FPlatformMisc::SetGracefulTerminationHandler();
+	FPlatformMisc::SetCrashHandler(nullptr);
+	
+	RunSymbolDebugger(*GSavedCommandLine);
+	
 	[NSApp terminate: self];
 }
 
+- (NSApplicationTerminateReply)applicationShouldTerminate:(NSApplication *)Sender;
+{
+	if(!GIsRequestingExit)
+	{
+		GIsRequestingExit = true;
+		return NSTerminateCancel;
+	}
+	return NSTerminateNow;
+}
+
+- (void)applicationDidFinishLaunching:(NSNotification *)Notification
+{
+	//install the custom quit event handler
+    NSAppleEventManager* appleEventManager = [NSAppleEventManager sharedAppleEventManager];
+    [appleEventManager setEventHandler:self andSelector:@selector(handleQuitEvent:withReplyEvent:) forEventClass:kCoreEventClass andEventID:kAEQuitApplication];
+	
+	RunGameThread(self, @selector(runGameThread:));
+}
+
 @end
 
 int main(int argc, char *argv[])
diff --git a/Engine/Source/Programs/UnrealFrontend/Private/Mac/MacUnrealFrontendMain.cpp b/Engine/Source/Programs/UnrealFrontend/Private/Mac/MacUnrealFrontendMain.cpp
--- a/Engine/Source/Programs/UnrealFrontend/Private/Mac/MacUnrealFrontendMain.cpp
+++ b/Engine/Source/Programs/UnrealFrontend/Private/Mac/MacUnrealFrontendMain.cpp
@@ -3,6 +3,7 @@
 #include "UnrealFrontendMain.h"
 #include "ExceptionHandling.h"
 #include "RequiredProgramMainCPPInclude.h"
+#include "CocoaThread.h"
 
 
 static FString GSavedCommandLine;
@@ -23,14 +24,15 @@ static FString GSavedCommandLine;
     [self requestQuit:self];
 }
 
-- (void)applicationDidFinishLaunching:(NSNotification *)Notification
+- (IBAction)requestQuit:(id)Sender
+{
+	GIsRequestingExit = true;
+}
+
+- (void) runGameThread:(id)Arg
 {
-	//install the custom quit event handler
-    NSAppleEventManager* appleEventManager = [NSAppleEventManager sharedAppleEventManager];
-    [appleEventManager setEventHandler:self andSelector:@selector(handleQuitEvent:withReplyEvent:) forEventClass:kCoreEventClass andEventID:kAEQuitApplication];
-	
 	FPlatformMisc::SetGracefulTerminationHandler();
-	FPlatformMisc::SetCrashHandler(NULL);
+	FPlatformMisc::SetCrashHandler(nullptr);
 	
 #if !UE_BUILD_SHIPPING
 	if (FParse::Param(*GSavedCommandLine,TEXT("crashreports")))
@@ -42,26 +44,40 @@ static FString GSavedCommandLine;
 #if UE_BUILD_DEBUG
 	if (!GAlwaysReportCrash)
 #else
-	if (FPlatformMisc::IsDebuggerPresent() && !GAlwaysReportCrash)
+		if (FPlatformMisc::IsDebuggerPresent() && !GAlwaysReportCrash)
 #endif
-	{
-		UnrealFrontendMain(*GSavedCommandLine);
-	}
-	else
-	{
-		GIsGuarded = 1;
-		UnrealFrontendMain(*GSavedCommandLine);
-		GIsGuarded = 0;
-	}
-
+		{
+			UnrealFrontendMain(*GSavedCommandLine);
+		}
+		else
+		{
+			GIsGuarded = 1;
+			UnrealFrontendMain(*GSavedCommandLine);
+			GIsGuarded = 0;
+		}
+	
 	FEngineLoop::AppExit();
-
+	
 	[NSApp terminate: self];
 }
 
-- (IBAction)requestQuit:(id)Sender
+- (NSApplicationTerminateReply)applicationShouldTerminate:(NSApplication *)Sender;
 {
-	GIsRequestingExit = true;
+	if(!GIsRequestingExit)
+	{
+		[self requestQuit:self];
+		return NSTerminateCancel;
+	}
+	return NSTerminateNow;
+}
+
+- (void)applicationDidFinishLaunching:(NSNotification *)Notification
+{
+	//install the custom quit event handler
+    NSAppleEventManager* appleEventManager = [NSAppleEventManager sharedAppleEventManager];
+    [appleEventManager setEventHandler:self andSelector:@selector(handleQuitEvent:withReplyEvent:) forEventClass:kCoreEventClass andEventID:kAEQuitApplication];
+	
+	RunGameThread(self, @selector(runGameThread:));
 }
 
 @end
diff --git a/Engine/Source/Runtime/Core/Private/Mac/CocoaTextView.cpp b/Engine/Source/Runtime/Core/Private/Mac/CocoaTextView.cpp
--- a/Engine/Source/Runtime/Core/Private/Mac/CocoaTextView.cpp
+++ b/Engine/Source/Runtime/Core/Private/Mac/CocoaTextView.cpp
@@ -3,6 +3,7 @@
 #include "Core.h"
 #include "CocoaTextView.h"
 #include "CocoaWindow.h"
+#include "CocoaThread.h"
 
 @implementation FCocoaTextView
 
@@ -163,7 +164,8 @@
 {
 	if (IMMContext.IsValid() && [self hasMarkedText])
 	{
-		uint32 SelectionLocation, SelectionLength;
+		__block uint32 SelectionLocation;
+		__block uint32 SelectionLength;
 		if (replacementRange.location == NSNotFound)
 		{
 			if (markedRange.location != NSNotFound)
@@ -173,8 +175,10 @@
 			}
 			else
 			{
-				ITextInputMethodContext::ECaretPosition CaretPosition;
-				IMMContext->GetSelectionRange(SelectionLocation, SelectionLength, CaretPosition);
+				GameThreadCall(^{
+					ITextInputMethodContext::ECaretPosition CaretPosition;
+					IMMContext->GetSelectionRange(SelectionLocation, SelectionLength, CaretPosition);
+				}, true);
 			}
 		}
 		else
@@ -193,9 +197,11 @@
 			TheString = (NSString*)aString;
 		}
 		
-		FString TheFString(TheString);
-		IMMContext->SetTextInRange(SelectionLocation, SelectionLength, TheFString);
-		IMMContext->SetSelectionRange(SelectionLocation+TheFString.Len(), 0, ITextInputMethodContext::ECaretPosition::Ending);
+		GameThreadCall(^{
+			FString TheFString(TheString);
+			IMMContext->SetTextInRange(SelectionLocation, SelectionLength, TheFString);
+			IMMContext->SetSelectionRange(SelectionLocation+TheFString.Len(), 0, ITextInputMethodContext::ECaretPosition::Ending);
+		}, true);
 		
 		[self unmarkText];
 		[[self inputContext] invalidateCharacterCoordinates]; // recentering
@@ -219,7 +225,8 @@
 {
 	if (IMMContext.IsValid())
 	{
-		uint32 SelectionLocation, SelectionLength;
+		__block uint32 SelectionLocation;
+		__block uint32 SelectionLength;
 		if (replacementRange.location == NSNotFound)
 		{
 			if (markedRange.location != NSNotFound)
@@ -229,8 +236,10 @@
 			}
 			else
 			{
-				ITextInputMethodContext::ECaretPosition CaretPosition;
-				IMMContext->GetSelectionRange(SelectionLocation, SelectionLength, CaretPosition);
+				GameThreadCall(^{
+					ITextInputMethodContext::ECaretPosition CaretPosition;
+					IMMContext->GetSelectionRange(SelectionLocation, SelectionLength, CaretPosition);
+				}, true);
 			}
 		}
 		else
@@ -241,14 +250,18 @@
 		
 		if ([aString length] == 0)
 		{
-			IMMContext->SetTextInRange(SelectionLocation, SelectionLength, FString());
+			GameThreadCall(^{
+				IMMContext->SetTextInRange(SelectionLocation, SelectionLength, FString());
+			}, true);
 			[self unmarkText];
 		}
 		else
 		{
 			if(markedRange.location == NSNotFound)
 			{
-				IMMContext->BeginComposition();
+				GameThreadCall(^{
+					IMMContext->BeginComposition();
+				}, true);
 			}
 			markedRange = NSMakeRange(SelectionLocation, [aString length]);
 			
@@ -284,10 +297,12 @@
 				TheString = (NSString*)aString;
 			}
 			
-			FString TheFString(TheString);
-			IMMContext->SetTextInRange(SelectionLocation, SelectionLength, TheFString);
-			IMMContext->UpdateCompositionRange(CompositionRange.location, CompositionRange.length);
-			IMMContext->SetSelectionRange(markedRange.location + selectedRange.location, 0, ITextInputMethodContext::ECaretPosition::Ending);
+			GameThreadCall(^{
+				FString TheFString(TheString);
+				IMMContext->SetTextInRange(SelectionLocation, SelectionLength, TheFString);
+				IMMContext->UpdateCompositionRange(CompositionRange.location, CompositionRange.length);
+				IMMContext->SetSelectionRange(markedRange.location + selectedRange.location, 0, ITextInputMethodContext::ECaretPosition::Ending);
+			}, true);
 		}
 		[[self inputContext] invalidateCharacterCoordinates]; // recentering
 	}
@@ -306,8 +321,10 @@
 		markedRange = {NSNotFound, 0};
 		if (IMMContext.IsValid())
 		{
-			IMMContext->UpdateCompositionRange(0, 0);
-			IMMContext->EndComposition();
+			GameThreadCall(^{
+				IMMContext->UpdateCompositionRange(0, 0);
+				IMMContext->EndComposition();
+			}, true);
 		}
 	}
 }
@@ -318,10 +335,12 @@
 {
 	if (IMMContext.IsValid())
 	{
-		uint32 SelectionLocation, SelectionLength;
-		ITextInputMethodContext::ECaretPosition CaretPosition;
-		IMMContext->GetSelectionRange(SelectionLocation, SelectionLength, CaretPosition);
-		
+		__block uint32 SelectionLocation;
+		__block uint32 SelectionLength;
+		__block ITextInputMethodContext::ECaretPosition CaretPosition;
+		GameThreadCall(^{
+			IMMContext->GetSelectionRange(SelectionLocation, SelectionLength, CaretPosition);
+		}, true);
 		return {SelectionLocation, SelectionLength};
 	}
 	else
@@ -358,8 +377,10 @@
 	NSAttributedString* AttributedString = nil;
 	if (IMMContext.IsValid())
 	{
-		FString String;
-		IMMContext->GetTextInRange(aRange.location, aRange.length, String);
+		__block FString String;
+		GameThreadCall(^{
+			IMMContext->GetTextInRange(aRange.location, aRange.length, String);
+		}, true);
 		CFStringRef CFString = FPlatformString::TCHARToCFString(*String);
 		if(CFString)
 		{
@@ -390,8 +411,11 @@
 {
 	if (IMMContext.IsValid())
 	{
-		FVector2D Position, Size;
-		IMMContext->GetTextBounds(aRange.location, aRange.length, Position, Size);
+		__block FVector2D Position;
+		__block FVector2D Size;
+		GameThreadCall(^{
+			IMMContext->GetTextBounds(aRange.location, aRange.length, Position, Size);
+		}, true);
 		
 		if(actualRange)
 		{
@@ -419,7 +443,7 @@
 	if (IMMContext.IsValid())
 	{
 		FVector2D Pos(aPoint.x, aPoint.y);
-		int32 Index = IMMContext->GetCharacterIndexFromPoint(Pos);
+		int32 Index = GameThreadReturn(^{ return IMMContext->GetCharacterIndexFromPoint(Pos); });
 		NSUInteger Idx = Index == INDEX_NONE ? NSNotFound : (NSUInteger)Index;
 		return Idx;
 	}
diff --git a/Engine/Source/Runtime/Core/Private/Mac/CocoaThread.cpp b/Engine/Source/Runtime/Core/Private/Mac/CocoaThread.cpp
new file mode 100644
--- /dev/null
+++ b/Engine/Source/Runtime/Core/Private/Mac/CocoaThread.cpp
@@ -0,0 +1,162 @@
+// Copyright 1998-2014 Epic Games, Inc. All Rights Reserved.
+
+#include "Core.h"
+#include "CocoaThread.h"
+
+#define MAC_SEPARATE_GAME_THREAD 1 // Separate the main & game threads so that we better handle the interaction between the Cocoa's event delegates and UE4's event polling.
+
+static FCocoaGameThread* GCocoaGameThread = nil;
+static NSRunLoop* GCocoaGameRunLoop = nil;
+
+@implementation NSThread (FCocoaThread)
+
++ (NSThread*) gameThread
+{
+	if(!GCocoaGameThread)
+	{
+		return [NSThread mainThread];
+	}
+	return GCocoaGameThread;
+}
+
++ (bool) isGameThread
+{
+	bool const bIsGameThread = [[NSThread currentThread] isGameThread];
+	return bIsGameThread;
+}
+
+- (bool) isGameThread
+{
+	bool const bIsGameThread = (self == GCocoaGameThread);
+	return bIsGameThread;
+}
+
+- (void) performBlock:(dispatch_block_t)Block onRunLoop:(NSRunLoop*)RunLoop forMode:(NSString*)RunMode wait:(bool)bWait inMode:(NSString*)WaitMode
+{
+	check(RunLoop);
+	if(RunLoop != [NSRunLoop currentRunLoop])
+	{
+		if(bWait)
+		{
+			__block bool bDone = false;
+			[RunLoop performBlock:^{ Block(); bDone = true; } forMode:RunMode wake:true];
+			while(!bDone)
+			{
+				CFRunLoopRunInMode((CFStringRef)WaitMode, 0.0, true);
+			}
+		}
+		else
+		{
+			[RunLoop performBlock:Block forMode:RunMode wake:true];
+		}
+	}
+	else
+	{
+		Block();
+	}
+}
+
+@end
+
+@implementation NSRunLoop (FCocoaRunLoop)
+
++ (NSRunLoop*)gameRunLoop
+{
+	if(!GCocoaGameRunLoop)
+	{
+		return [NSRunLoop mainRunLoop];
+	}
+	return GCocoaGameRunLoop;
+}
+
+- (void)performBlock:(dispatch_block_t)Block forMode:(NSString *)Mode wake:(bool)bWakeRunLoop
+{
+	CFRunLoopRef RunLoop = [self getCFRunLoop];
+	check(RunLoop);
+	CFRunLoopPerformBlock(RunLoop, (CFStringRef)Mode, Block);
+	if(bWakeRunLoop)
+	{
+		CFRunLoopWakeUp(RunLoop);
+	}
+}
+
+@end
+
+@implementation FCocoaGameThread : NSThread
+
+- (id)init
+{
+	id Self = [super init];
+	if(Self)
+	{
+		GCocoaGameThread = Self;
+	}
+	return Self;
+}
+
+- (id)initWithTarget:(id)Target selector:(SEL)Selector object:(id)Argument
+{
+	id Self = [super initWithTarget:Target selector:Selector object:Argument];
+	if(Self)
+	{
+		GCocoaGameThread = Self;
+	}
+	return Self;
+}
+
+- (void)dealloc
+{
+	GCocoaGameThread = nil;
+	GCocoaGameRunLoop = nil;
+	[super dealloc];
+}
+
+- (void)main
+{
+	GCocoaGameRunLoop = [NSRunLoop currentRunLoop];
+	[super main];
+}
+
+@end
+
+void PerformBlockOnRunLoop(dispatch_block_t Block, NSRunLoop* const RunLoop, bool const bWait)
+{
+	[[NSThread currentThread] performBlock:Block onRunLoop:RunLoop forMode:NSDefaultRunLoopMode wait:bWait inMode:NSDefaultRunLoopMode];
+}
+
+void MainThreadCall(dispatch_block_t Block, bool const bWait)
+{
+	[[NSThread currentThread] performBlock:Block onRunLoop:[NSRunLoop mainRunLoop] forMode:NSDefaultRunLoopMode wait:bWait inMode:NSDefaultRunLoopMode];
+}
+
+void GameThreadCall(dispatch_block_t Block, bool const bWait)
+{
+	[[NSThread currentThread] performBlock:Block onRunLoop:[NSRunLoop gameRunLoop] forMode:NSDefaultRunLoopMode wait:bWait inMode:NSDefaultRunLoopMode];
+}
+
+void RunGameThread(id Target, SEL Selector)
+{
+#if MAC_SEPARATE_GAME_THREAD
+	// Create a separate game thread and set it to the stack size to be the same as the main thread default of 8MB ( http://developer.apple.com/library/mac/#qa/qa1419/_index.html )
+	FCocoaGameThread* GameThread = [[FCocoaGameThread alloc] initWithTarget:Target selector:Selector object:nil];
+	[GameThread setStackSize:8*1024*1024];
+	[GameThread start];
+#else
+	[Target performSelector:Selector withObject:nil];
+#endif
+}
+
+void ProcessGameThreadEvents(void)
+{
+#if MAC_SEPARATE_GAME_THREAD
+	// Make one pass through the loop, processing all events
+	CFRunLoopRef RunLoop = CFRunLoopGetCurrent();
+	CFRunLoopRunInMode(kCFRunLoopDefaultMode, 0, false);
+#else
+	SCOPED_AUTORELEASE_POOL;
+	while( NSEvent *Event = [NSApp nextEventMatchingMask: NSAnyEventMask untilDate: nil inMode: NSDefaultRunLoopMode dequeue: true] )
+	{
+		[NSApp sendEvent: Event];
+	}
+#endif
+}
diff --git a/Engine/Source/Runtime/Core/Private/Mac/CocoaWindow.cpp b/Engine/Source/Runtime/Core/Private/Mac/CocoaWindow.cpp
--- a/Engine/Source/Runtime/Core/Private/Mac/CocoaWindow.cpp
+++ b/Engine/Source/Runtime/Core/Private/Mac/CocoaWindow.cpp
@@ -2,9 +2,18 @@
 
 #include "Core.h"
 #include "CocoaWindow.h"
-#include "CocoaTextView.h"
 #include "MacApplication.h"
-#include "MacWindow.h"
+#include "CocoaTextView.h"
+#include "MacEvent.h"
+#include "CocoaThread.h"
+
+TArray< FCocoaWindow* > GRunningModalWindows;
+
+NSString* NSWindowRedrawContents = @"NSWindowRedrawContents";
+NSString* NSDraggingExited = @"NSDraggingExited";
+NSString* NSDraggingUpdated = @"NSDraggingUpdated";
+NSString* NSPrepareForDragOperation = @"NSPrepareForDragOperation";
+NSString* NSPerformDragOperation = @"NSPerformDragOperation";
 
 /**
  * Custom window class used for input handling
@@ -157,7 +166,8 @@
 {
 	if(bNeedsRedraw && bForwardEvents && ([self isVisible] && [super alphaValue] > 0.0f))
 	{
-		MacApplication->OnWindowRedrawContents( self );
+		NSNotification* Notification = [NSNotification notificationWithName:NSWindowRedrawContents object:self];
+		FMacEvent::SendToGameRunLoop(Notification, self, EMacEventSendMethod::Sync);
 	}
 	bNeedsRedraw = false;
 }
@@ -273,7 +283,8 @@
 		DeferFrame = FrameRect;
 		if(self.bForwardEvents)
 		{
-			MacApplication->OnWindowDidResize( self );
+			NSNotification* Notification = [NSNotification notificationWithName:NSWindowDidResizeNotification object:self];
+			FMacEvent::SendToGameRunLoop(Notification, self, EMacEventSendMethod::Async);
 		}
 	}
 }
@@ -283,7 +294,9 @@
 	NSSize Size = [self frame].size;
 	if(!bRenderInitialised || ([self isVisible] && [super alphaValue] > 0.0f && (Size.width > 1 || Size.height > 1)))
 	{
-		[super setFrameOrigin:Point];
+		MainThreadCall(^{
+			[super setFrameOrigin:Point];
+		}, true);
 		bDeferSetOrigin = false;
 	}
 	else
@@ -297,7 +310,7 @@
 {
 	if(self.bForwardEvents)
 	{
-		MacApplication->ProcessEvent( Event );
+		FMacEvent::SendToGameRunLoop(Event, EMacEventSendMethod::Async);
 	}
 }
 
@@ -305,7 +318,7 @@
 {
 	if(self.bForwardEvents)
 	{
-		MacApplication->ProcessEvent( Event );
+		FMacEvent::SendToGameRunLoop(Event, EMacEventSendMethod::Async);
 	}
 }
 
@@ -324,7 +337,7 @@
 	WindowMode = self.TargetWindowMode;
 	if(self.bForwardEvents)
 	{
-		MacApplication->OnWindowDidResize( self );
+		FMacEvent::SendToGameRunLoop(notification, self, EMacEventSendMethod::Async);
 	}
 }
 
@@ -334,7 +347,7 @@
 	self.TargetWindowMode = EWindowMode::Windowed;
 	if(self.bForwardEvents)
 	{
-		MacApplication->OnWindowDidResize( self );
+		FMacEvent::SendToGameRunLoop(notification, self, EMacEventSendMethod::Async);
 	}
 }
 
@@ -354,7 +367,7 @@
 	
 	if(self.bForwardEvents)
 	{
-		MacApplication->OnWindowDidBecomeKey( self );
+		FMacEvent::SendToGameRunLoop(Notification, self, EMacEventSendMethod::Async);
 	}
 }
 
@@ -365,7 +378,7 @@
 	
 	if(self.bForwardEvents)
 	{
-		MacApplication->OnWindowDidResignKey( self );
+		FMacEvent::SendToGameRunLoop(Notification, self, EMacEventSendMethod::Async);
 	}
 }
 
@@ -373,7 +386,7 @@
 {
 	if(self.bForwardEvents)
 	{
-		MacApplication->OnWindowWillMove( self );
+		FMacEvent::SendToGameRunLoop(Notification, self, EMacEventSendMethod::Async);
 	}
 }
 
@@ -386,7 +399,7 @@
 	
 	if(self.bForwardEvents)
 	{
-		MacApplication->OnWindowDidMove( self );
+		FMacEvent::SendToGameRunLoop(Notification, self, EMacEventSendMethod::Async);
 	}
 }
 
@@ -458,7 +471,7 @@
 	bZoomed = [self isZoomed];
 	if(self.bForwardEvents)
 	{
-		MacApplication->OnWindowDidResize( self );
+		FMacEvent::SendToGameRunLoop(Notification, self, EMacEventSendMethod::Async);
 	}
 	bNeedsRedraw = true;
 }
@@ -467,7 +480,7 @@
 {
 	if(self.bForwardEvents && MacApplication)
 	{
-		MacApplication->OnWindowDidClose( self );
+		FMacEvent::SendToGameRunLoop(notification, self, EMacEventSendMethod::Async);
 	}
 }
 
@@ -475,7 +488,7 @@
 {
 	if(self.bForwardEvents)
 	{
-		MacApplication->ProcessEvent( Event );
+		FMacEvent::SendToGameRunLoop(Event, EMacEventSendMethod::Async);
 	}
 }
 
@@ -490,7 +503,7 @@
 	
 	if(self.bForwardEvents)
 	{
-		MacApplication->ProcessEvent( Event );
+		FMacEvent::SendToGameRunLoop(Event, EMacEventSendMethod::Async);
 	}
 }
 
@@ -498,7 +511,7 @@
 {
 	if(self.bForwardEvents)
 	{
-		MacApplication->ProcessEvent( Event );
+		FMacEvent::SendToGameRunLoop(Event, EMacEventSendMethod::Async);
 	}
 }
 
@@ -506,7 +519,7 @@
 {
 	if(self.bForwardEvents)
 	{
-		MacApplication->ProcessEvent( Event );
+		FMacEvent::SendToGameRunLoop(Event, EMacEventSendMethod::Async);
 	}
 }
 
@@ -514,7 +527,7 @@
 {
 	if(self.bForwardEvents)
 	{
-		MacApplication->ProcessEvent( Event );
+		FMacEvent::SendToGameRunLoop(Event, EMacEventSendMethod::Async);
 	}
 }
 
@@ -522,7 +535,7 @@
 {
 	if(self.bForwardEvents)
 	{
-		MacApplication->ProcessEvent( Event );
+		FMacEvent::SendToGameRunLoop(Event, EMacEventSendMethod::Async);
 	}
 }
 
@@ -535,7 +548,8 @@
 {
 	if(self.bForwardEvents)
 	{
-		MacApplication->OnDragOut( self );
+		NSNotification* Notification = [NSNotification notificationWithName:NSDraggingExited object:Sender];
+		FMacEvent::SendToGameRunLoop(Notification, self, EMacEventSendMethod::Async);
 	}
 }
 
@@ -543,7 +557,8 @@
 {
 	if(self.bForwardEvents)
 	{
-		MacApplication->OnDragOver( self );
+		NSNotification* Notification = [NSNotification notificationWithName:NSDraggingUpdated object:Sender];
+		FMacEvent::SendToGameRunLoop(Notification, self, EMacEventSendMethod::Async);
 	}
 	return NSDragOperationGeneric;
 }
@@ -552,7 +567,8 @@
 {
 	if(self.bForwardEvents)
 	{
-		MacApplication->OnDragEnter( self, [Sender draggingPasteboard] );
+		NSNotification* Notification = [NSNotification notificationWithName:NSPrepareForDragOperation object:Sender];
+		FMacEvent::SendToGameRunLoop(Notification, self, EMacEventSendMethod::Async);
 	}
 	return YES;
 }
@@ -561,7 +577,8 @@
 {
 	if(self.bForwardEvents)
 	{
-		MacApplication->OnDragDrop( self );
+		NSNotification* Notification = [NSNotification notificationWithName:NSPerformDragOperation object:Sender];
+		FMacEvent::SendToGameRunLoop(Notification, self, EMacEventSendMethod::Async);
 	}
 	return YES;
 }
@@ -571,7 +588,7 @@
 	BOOL Movable = [super isMovable];
 	if(Movable && bRenderInitialised && MacApplication)
 	{
-		Movable &= MacApplication->IsWindowMovable(self, NULL);
+		Movable &= (BOOL)(GameThreadReturn(^{ return MacApplication->IsWindowMovable(self, NULL); }));
 	}
 	return Movable;
 }
@@ -610,32 +627,32 @@
 
 - (void)mouseDown:(NSEvent*)Event
 {
-	MacApplication->ProcessEvent( Event );
+	FMacEvent::SendToGameRunLoop(Event, EMacEventSendMethod::Async);
 }
 
 - (void)rightMouseDown:(NSEvent*)Event
 {
-	MacApplication->ProcessEvent( Event );
+	FMacEvent::SendToGameRunLoop(Event, EMacEventSendMethod::Async);
 }
 
 - (void)otherMouseDown:(NSEvent*)Event
 {
-	MacApplication->ProcessEvent( Event );
+	FMacEvent::SendToGameRunLoop(Event, EMacEventSendMethod::Async);
 }
 
 - (void)mouseUp:(NSEvent*)Event
 {
-	MacApplication->ProcessEvent( Event );
+	FMacEvent::SendToGameRunLoop(Event, EMacEventSendMethod::Async);
 }
 
 - (void)rightMouseUp:(NSEvent*)Event
 {
-	MacApplication->ProcessEvent( Event );
+	FMacEvent::SendToGameRunLoop(Event, EMacEventSendMethod::Async);
 }
 
 - (void)otherMouseUp:(NSEvent*)Event
 {
-	MacApplication->ProcessEvent( Event );
+	FMacEvent::SendToGameRunLoop(Event, EMacEventSendMethod::Async);
 }
 
 @end
diff --git a/Engine/Source/Runtime/Core/Private/Mac/MacApplication.cpp b/Engine/Source/Runtime/Core/Private/Mac/MacApplication.cpp
--- a/Engine/Source/Runtime/Core/Private/Mac/MacApplication.cpp
+++ b/Engine/Source/Runtime/Core/Private/Mac/MacApplication.cpp
@@ -5,6 +5,7 @@
 #include "MacApplication.h"
 #include "MacWindow.h"
 #include "MacCursor.h"
+#include "MacEvent.h"
 #include "GenericApplicationMessageHandler.h"
 #include "HIDInputInterface.h"
 #include "AnalyticsEventAttribute.h"
@@ -12,6 +13,30 @@
 
 #include "ModuleManager.h"
 
+NSString* NSWindowDraggingFinished = @"NSWindowDraggingFinished";
+
+void HandleNSEvent(NSEvent* Event)
+{
+	const bool bIsMouseClickOrKeyEvent = [Event type] == NSLeftMouseDown || [Event type] == NSLeftMouseUp
+	|| [Event type] == NSRightMouseDown || [Event type] == NSRightMouseUp
+	|| [Event type] == NSOtherMouseDown || [Event type] == NSOtherMouseUp
+	|| [Event type] == NSKeyDown || ([Event type] == NSKeyUp && !([Event modifierFlags] & NSCommandKeyMask));
+	
+	if( MacApplication )
+	{
+		if( !bIsMouseClickOrKeyEvent || [Event window] == NULL )
+		{
+			FMacEvent::SendToGameRunLoop(Event, EMacEventSendMethod::Async);
+		}
+		
+		if( [Event type] == NSLeftMouseUp )
+		{
+			NSNotification* Notification = [NSNotification notificationWithName:NSWindowDraggingFinished object:[Event window]];
+			FMacEvent::SendToGameRunLoop(Notification, [Event window], EMacEventSendMethod::Async);
+		}
+	}
+}
+
 FMacApplication* MacApplication = NULL;
 
 #if WITH_EDITOR
@@ -67,10 +92,19 @@ FMacApplication::FMacApplication()
 	, bIsMouseCursorLocked( false )
 	, ModifierKeysFlags( 0 )
 	, CurrentModifierFlags( 0 )
+	, EventMonitor( NULL )
 {
 	CGDisplayRegisterReconfigurationCallback(FMacApplication::OnDisplayReconfiguration, this);
 
-	[NSEvent addGlobalMonitorForEventsMatchingMask:NSMouseMovedMask handler:^(NSEvent* Event){ ProcessEvent(Event); }];
+	[NSEvent addGlobalMonitorForEventsMatchingMask:NSMouseMovedMask handler:^(NSEvent* Event){
+		FMacEvent::SendToGameRunLoop(Event, EMacEventSendMethod::Async);
+	}];
+	
+	EventMonitor = (void*)[NSEvent addLocalMonitorForEventsMatchingMask:NSAnyEventMask handler:^(NSEvent* IncomingEvent)
+	{
+		HandleNSEvent(IncomingEvent);
+		return IncomingEvent;
+	}];
 
 #if WITH_EDITOR
 	NSMutableArray* MultiTouchDevices = (__bridge NSMutableArray*)MTDeviceCreateList();
@@ -128,6 +162,12 @@ FMacApplication::FMacApplication()
 
 FMacApplication::~FMacApplication()
 {
+	if(EventMonitor)
+	{
+		[NSEvent removeMonitor:(id)EventMonitor];
+		EventMonitor = nullptr;
+	}
+	
 	if(AppActivationObserver)
 	{
 		[[NSNotificationCenter defaultCenter] removeObserver:AppActivationObserver];
@@ -172,8 +212,9 @@ void FMacApplication::SetMessageHandler( const TSharedRef< FGenericApplicationMe
 	HIDInput->SetMessageHandler( InMessageHandler );
 }
 
-static TSharedPtr< FMacWindow > FindWindowByNSWindow( const TArray< TSharedRef< FMacWindow > >& WindowsToSearch, FCocoaWindow* const WindowHandle )
+static TSharedPtr< FMacWindow > FindWindowByNSWindow( const TArray< TSharedRef< FMacWindow > >& WindowsToSearch, FCriticalSection* const Mutex, FCocoaWindow* const WindowHandle )
 {
+	FScopeLock Lock(Mutex);
 	for (int32 WindowIndex=0; WindowIndex < WindowsToSearch.Num(); ++WindowIndex)
 	{
 		TSharedRef< FMacWindow > Window = WindowsToSearch[ WindowIndex ];
@@ -186,6 +227,106 @@ static TSharedPtr< FMacWindow > FindWindowByNSWindow( const TArray< TSharedRef<
 	return TSharedPtr< FMacWindow >( NULL );
 }
 
+FCocoaWindow* FMacApplication::FindMacEventWindow( NSEvent* CocoaEvent )
+{
+	return MacApplication->FindEventWindow( CocoaEvent );
+}
+
+TSharedPtr< FMacWindow > FMacApplication::FindMacWindowByNSWindow( FCocoaWindow* const WindowHandle )
+{
+	return FindWindowByNSWindow(MacApplication->Windows, &MacApplication->WindowsMutex, WindowHandle);
+}
+
+void FMacApplication::ProcessEvent(FMacEvent const* const Event)
+{
+	// Must have an event
+	check(Event);
+	
+	// Determine the type of the event to forward to the right handling routine
+	NSEvent* AppKitEvent = Event->GetEvent();
+	NSNotification* Notification = Event->GetNotification();
+
+	if(AppKitEvent)
+	{
+		// Process a standard NSEvent as we always did
+		MacApplication->ProcessNSEvent(AppKitEvent, Event->GetWindow(), Event->GetMousePosition());
+	}
+	else if(Notification)
+	{
+		// Notifications need to mapped to the right handler function, that's no longer handled in the
+		// call location.
+		NSString* const NotificationName = [Notification name];
+		TSharedPtr<FMacWindow> Window = Event->GetWindow();
+		if(Window.IsValid())
+		{
+			FCocoaWindow* CocoaWindow = Window->GetWindowHandle();
+			if(NotificationName == NSWindowDidResizeNotification)
+			{
+				MacApplication->OnWindowDidResize( CocoaWindow );
+			}
+			else if(NotificationName == NSWindowDidEnterFullScreenNotification)
+			{
+				MacApplication->OnWindowDidResize( CocoaWindow );
+			}
+			else if(NotificationName == NSWindowDidExitFullScreenNotification)
+			{
+				MacApplication->OnWindowDidResize( CocoaWindow );
+			}
+			else if(NotificationName == NSWindowDidBecomeKeyNotification)
+			{
+				MacApplication->OnWindowDidBecomeKey( CocoaWindow );
+			}
+			else if(NotificationName == NSWindowDidResignKeyNotification)
+			{
+				MacApplication->OnWindowDidResignKey( CocoaWindow );
+			}
+			else if(NotificationName == NSWindowWillMoveNotification)
+			{
+				MacApplication->OnWindowWillMove( CocoaWindow );
+			}
+			else if(NotificationName == NSWindowDidMoveNotification)
+			{
+				MacApplication->OnWindowDidMove( CocoaWindow );
+			}
+			else if(NotificationName == NSWindowWillCloseNotification)
+			{
+				MacApplication->OnWindowDidClose( CocoaWindow );
+			}
+			else if(NotificationName == NSWindowRedrawContents)
+			{
+				MacApplication->OnWindowRedrawContents( CocoaWindow );
+			}
+			else if(NotificationName == NSDraggingExited)
+			{
+				MacApplication->OnDragOut( CocoaWindow );
+			}
+			else if(NotificationName == NSDraggingUpdated)
+			{
+				MacApplication->OnDragOver( CocoaWindow );
+			}
+			else if(NotificationName == NSPrepareForDragOperation)
+			{
+				MacApplication->OnDragEnter( CocoaWindow, [(id < NSDraggingInfo >)[Notification object] draggingPasteboard] );
+			}
+			else if(NotificationName == NSPerformDragOperation)
+			{
+				MacApplication->OnDragDrop( CocoaWindow );
+			}
+			else if(NotificationName == NSWindowDraggingFinished)
+			{
+				MacApplication->OnWindowDraggingFinished();
+			}
+			else
+			{
+				check(false);
+			}
+		}
+	}
+	
+	// We are responsible for deallocating the event
+	delete Event;
+}
+
 void FMacApplication::PumpMessages( const float TimeDelta )
 {
 	FPlatformMisc::PumpMessages( true );
@@ -208,7 +349,7 @@ bool FMacApplication::IsWindowMovable(FCocoaWindow* Win, bool* OutMovableByBackg
 	}
 
 	FCocoaWindow* NativeWindow = (FCocoaWindow*)Win;
-	TSharedPtr< FMacWindow > CurrentEventWindow = FindWindowByNSWindow( Windows, NativeWindow );
+	TSharedPtr< FMacWindow > CurrentEventWindow = FindWindowByNSWindow( Windows, &WindowsMutex, NativeWindow );
 	if( CurrentEventWindow.IsValid() )
 	{
 		const FVector2D CursorPos = static_cast<FMacCursor*>( Cursor.Get() )->GetPosition();
@@ -272,14 +413,13 @@ void FMacApplication::HandleModifierChange(TSharedPtr< FMacWindow > CurrentEvent
 	}
 }
 
-void FMacApplication::ProcessEvent( NSEvent* Event )
+void FMacApplication::ProcessNSEvent(NSEvent* const Event, TSharedPtr< FMacWindow > CurrentEventWindow, FVector2D const MousePosition)
 {
 	SCOPED_AUTORELEASE_POOL;
 
 	const NSEventType EventType = [Event type];
-
-	FCocoaWindow* NativeWindow = FindEventWindow( Event );
-	TSharedPtr< FMacWindow > CurrentEventWindow = FindWindowByNSWindow( Windows, NativeWindow );
+	
+	FCocoaWindow* NativeWindow = (FCocoaWindow*)[Event window];
 	if( NativeWindow && !CurrentEventWindow.IsValid() && LastEventWindow.IsValid() )
 	{
 		CurrentEventWindow = LastEventWindow;
@@ -289,6 +429,7 @@ void FMacApplication::ProcessEvent( NSEvent* Event )
 	if( CurrentEventWindow.IsValid() )
 	{
 		LastEventWindow = CurrentEventWindow;
+		NativeWindow = CurrentEventWindow->GetWindowHandle();
 	}
 
 	if( !NativeWindow )
@@ -409,14 +550,6 @@ void FMacApplication::ProcessEvent( NSEvent* Event )
 		case NSRightMouseDown:
 		case NSOtherMouseDown:
 		{
-			if( CurrentEventWindow.IsValid() && CurrentEventWindow->IsRegularWindow() )
-			{
-				bool IsMouseOverTitleBar = false;
-				bool IsMovable = IsWindowMovable(NativeWindow, &IsMouseOverTitleBar);
-				[NativeWindow setMovable: IsMovable];
-				[NativeWindow setMovableByWindowBackground: IsMouseOverTitleBar];
-			}
-			
 			EMouseButtons::Type Button = [Event type] == NSLeftMouseDown ? EMouseButtons::Left : EMouseButtons::Right;
 			if ([Event type] == NSOtherMouseDown)
 			{
@@ -622,6 +755,20 @@ void FMacApplication::ProcessEvent( NSEvent* Event )
 	}
 }
 
+void FMacApplication::ProcessEvent( NSEvent* Event )
+{
+	SCOPED_AUTORELEASE_POOL;
+
+	const NSEventType EventType = [Event type];
+
+	FCocoaWindow* NativeWindow = FindEventWindow( Event );
+	TSharedPtr< FMacWindow > CurrentEventWindow = FindWindowByNSWindow( Windows, &WindowsMutex, NativeWindow );
+	
+	FMacCursor* MacCursor = (FMacCursor*)Cursor.Get();
+	
+	ProcessNSEvent(Event, CurrentEventWindow, MacCursor->GetPosition());
+}
+
 FCocoaWindow* FMacApplication::FindEventWindow( NSEvent* Event )
 {
 	SCOPED_AUTORELEASE_POOL;
@@ -727,7 +874,7 @@ TSharedPtr<FGenericWindow> FMacApplication::LocateWindowUnderCursor( const FVect
 
 		if (NSPointInRect(Position, VisibleFrame))
 		{
-			TSharedPtr<FMacWindow> MacWindow = FindWindowByNSWindow(Windows, (FCocoaWindow*)NativeWindow);
+			TSharedPtr<FMacWindow> MacWindow = FindWindowByNSWindow(Windows, &WindowsMutex, (FCocoaWindow*)NativeWindow);
 			return MacWindow;
 		}
 	}
@@ -925,7 +1072,7 @@ void FMacApplication::OnDragEnter( FCocoaWindow* Window, void *InPasteboard )
 {
 	SCOPED_AUTORELEASE_POOL;
 
-	TSharedPtr< FMacWindow > EventWindow = FindWindowByNSWindow( Windows, Window );
+	TSharedPtr< FMacWindow > EventWindow = FindWindowByNSWindow( Windows, &WindowsMutex, Window );
 	if( !EventWindow.IsValid() )
 	{
 		return;
@@ -966,7 +1113,7 @@ void FMacApplication::OnDragEnter( FCocoaWindow* Window, void *InPasteboard )
 
 void FMacApplication::OnDragOver( FCocoaWindow* Window )
 {
-	TSharedPtr< FMacWindow > EventWindow = FindWindowByNSWindow( Windows, Window );
+	TSharedPtr< FMacWindow > EventWindow = FindWindowByNSWindow( Windows, &WindowsMutex, Window );
 	if( EventWindow.IsValid() )
 	{
 		MessageHandler->OnDragOver( EventWindow.ToSharedRef() );
@@ -975,7 +1122,7 @@ void FMacApplication::OnDragOver( FCocoaWindow* Window )
 
 void FMacApplication::OnDragOut( FCocoaWindow* Window )
 {
-	TSharedPtr< FMacWindow > EventWindow = FindWindowByNSWindow( Windows, Window );
+	TSharedPtr< FMacWindow > EventWindow = FindWindowByNSWindow( Windows, &WindowsMutex, Window );
 	if( EventWindow.IsValid() )
 	{
 		MessageHandler->OnDragLeave( EventWindow.ToSharedRef() );
@@ -984,7 +1131,7 @@ void FMacApplication::OnDragOut( FCocoaWindow* Window )
 
 void FMacApplication::OnDragDrop( FCocoaWindow* Window )
 {
-	TSharedPtr< FMacWindow > EventWindow = FindWindowByNSWindow( Windows, Window );
+	TSharedPtr< FMacWindow > EventWindow = FindWindowByNSWindow( Windows, &WindowsMutex, Window );
 	if( EventWindow.IsValid() )
 	{
 		MessageHandler->OnDragDrop( EventWindow.ToSharedRef() );
@@ -993,7 +1140,7 @@ void FMacApplication::OnDragDrop( FCocoaWindow* Window )
 
 void FMacApplication::OnWindowDidBecomeKey( FCocoaWindow* Window )
 {
-	TSharedPtr< FMacWindow > EventWindow = FindWindowByNSWindow( Windows, Window );
+	TSharedPtr< FMacWindow > EventWindow = FindWindowByNSWindow( Windows, &WindowsMutex, Window );
 	if( EventWindow.IsValid() )
 	{
 		MessageHandler->OnWindowActivationChanged( EventWindow.ToSharedRef(), EWindowActivation::Activate );
@@ -1003,7 +1150,7 @@ void FMacApplication::OnWindowDidBecomeKey( FCocoaWindow* Window )
 
 void FMacApplication::OnWindowDidResignKey( FCocoaWindow* Window )
 {
-	TSharedPtr< FMacWindow > EventWindow = FindWindowByNSWindow( Windows, Window );
+	TSharedPtr< FMacWindow > EventWindow = FindWindowByNSWindow( Windows, &WindowsMutex, Window );
 	if( EventWindow.IsValid() )
 	{
 		MessageHandler->OnWindowActivationChanged( EventWindow.ToSharedRef(), EWindowActivation::Deactivate );
@@ -1027,7 +1174,7 @@ void FMacApplication::OnWindowDidMove( FCocoaWindow* Window )
 	const int32 X = (int32)WindowFrame.origin.x;
 	const int32 Y = FPlatformMisc::ConvertSlateYPositionToCocoa( (int32)WindowFrame.origin.y ) - OpenGLFrame.size.height + 1;
 
-	TSharedPtr< FMacWindow > EventWindow = FindWindowByNSWindow( Windows, Window );
+	TSharedPtr< FMacWindow > EventWindow = FindWindowByNSWindow( Windows, &WindowsMutex, Window );
 	if( EventWindow.IsValid() )
 	{
 		MessageHandler->OnMovedWindow( EventWindow.ToSharedRef(), X, Y );
@@ -1042,7 +1189,7 @@ void FMacApplication::OnWindowDidResize( FCocoaWindow* Window )
 
 	OnWindowDidMove( Window );
 	
-	TSharedPtr< FMacWindow > EventWindow = FindWindowByNSWindow( Windows, Window );
+	TSharedPtr< FMacWindow > EventWindow = FindWindowByNSWindow( Windows, &WindowsMutex, Window );
 	if( EventWindow.IsValid() )
 	{
 		// default is no override
@@ -1062,7 +1209,7 @@ void FMacApplication::OnWindowDidResize( FCocoaWindow* Window )
 
 void FMacApplication::OnWindowRedrawContents( FCocoaWindow* Window )
 {
-	TSharedPtr< FMacWindow > EventWindow = FindWindowByNSWindow( Windows, Window );
+	TSharedPtr< FMacWindow > EventWindow = FindWindowByNSWindow( Windows, &WindowsMutex, Window );
 	if( EventWindow.IsValid() )
 	{
 		SCOPED_AUTORELEASE_POOL;
@@ -1072,7 +1219,7 @@ void FMacApplication::OnWindowRedrawContents( FCocoaWindow* Window )
 
 void FMacApplication::OnWindowDidClose( FCocoaWindow* Window )
 {
-	TSharedPtr< FMacWindow > EventWindow = FindWindowByNSWindow( Windows, Window );
+	TSharedPtr< FMacWindow > EventWindow = FindWindowByNSWindow( Windows, &WindowsMutex, Window );
 	if( EventWindow.IsValid() )
 	{
 		SCOPED_AUTORELEASE_POOL;
@@ -1085,7 +1232,7 @@ void FMacApplication::OnWindowDidClose( FCocoaWindow* Window )
 
 bool FMacApplication::OnWindowDestroyed( FCocoaWindow* Window )
 {
-	TSharedPtr< FMacWindow > EventWindow = FindWindowByNSWindow( Windows, Window );
+	TSharedPtr< FMacWindow > EventWindow = FindWindowByNSWindow( Windows, &WindowsMutex, Window );
 	if( EventWindow.IsValid() )
 	{
 		MessageHandler->OnWindowActivationChanged( EventWindow.ToSharedRef(), EWindowActivation::Deactivate );
diff --git a/Engine/Source/Runtime/Core/Private/Mac/MacEvent.cpp b/Engine/Source/Runtime/Core/Private/Mac/MacEvent.cpp
new file mode 100644
--- /dev/null
+++ b/Engine/Source/Runtime/Core/Private/Mac/MacEvent.cpp
@@ -0,0 +1,118 @@
+// Copyright 1998-2014 Epic Games, Inc. All Rights Reserved.
+
+#include "Core.h"
+#include "MacEvent.h"
+#include "MacApplication.h"
+#include "CocoaThread.h"
+#include "CocoaWindow.h"
+
+FMacEvent::FMacEvent(NSEvent* const Event)
+: EventData(Event)
+{
+	check(Event);
+	[Event retain];
+	
+	NSWindow* Window = FMacApplication::FindMacEventWindow(Event);
+	if(Window && [Window isKindOfClass:[FCocoaWindow class]])
+	{
+		MacWindow = FMacApplication::FindMacWindowByNSWindow((FCocoaWindow*)Window);
+	}
+	
+	CGEventRef CGEvent = [Event CGEvent];
+	CGPoint CGMouseLocation = CGEventGetLocation(CGEvent);
+	MousePosition.X = CGMouseLocation.x;
+	MousePosition.Y = CGMouseLocation.y;
+}
+
+FMacEvent::FMacEvent(NSNotification* const Notification, NSWindow* const Window)
+: EventData(Notification)
+{
+	check(Notification);
+	[Notification retain];
+	
+	if(Window && [Window isKindOfClass:[FCocoaWindow class]])
+	{
+		MacWindow = FMacApplication::FindMacWindowByNSWindow((FCocoaWindow*)Window);
+	}
+	
+	CGEventRef Event = CGEventCreate(NULL);
+	CGPoint CursorPos = CGEventGetLocation(Event);
+	MousePosition.X = CursorPos.x;
+	MousePosition.Y = CursorPos.y;
+	CFRelease(Event);
+}
+
+void FMacEvent::SendToGameRunLoop(FMacEvent const* const Event, EMacEventSendMethod SendMethod)
+{
+	NSThread* GameThread = [NSThread gameThread];
+	check(GameThread);
+	
+	NSRunLoop* GameRunLoop = [NSRunLoop gameRunLoop];
+	check(GameRunLoop);
+	
+	dispatch_block_t Block = ^{ FMacApplication::ProcessEvent(Event); };
+	const bool bWait = (SendMethod == EMacEventSendMethod::Sync);
+	
+	[GameThread performBlock: Block onRunLoop: GameRunLoop forMode: NSDefaultRunLoopMode wait: bWait inMode: NSDefaultRunLoopMode];
+}
+	
+void FMacEvent::SendToGameRunLoop(NSEvent* const Event, EMacEventSendMethod SendMethod)
+{
+	if(MacApplication)
+	{
+		FMacEvent* MacEvent = new FMacEvent(Event);
+		FMacEvent::SendToGameRunLoop(MacEvent, SendMethod);
+	}
+}
+
+void FMacEvent::SendToGameRunLoop(NSNotification* const Notification, NSWindow* const Window, EMacEventSendMethod SendMethod)
+{
+	if(MacApplication)
+	{
+		FMacEvent* MacEvent = new FMacEvent(Notification, Window);
+		FMacEvent::SendToGameRunLoop(MacEvent, SendMethod);
+	}
+}
+
+FMacEvent::~FMacEvent(void)
+{
+	if(EventData)
+	{
+		[EventData release];
+		EventData = nil;
+	}
+}
+
+NSEvent* FMacEvent::GetEvent(void) const
+{
+	if(EventData && [EventData isKindOfClass:[NSEvent class]])
+	{
+		return (NSEvent*)EventData;
+	}
+	else
+	{
+		return nil;
+	}
+}
+
+NSNotification* FMacEvent::GetNotification(void) const
+{
+	if(EventData && [EventData isKindOfClass:[NSNotification class]])
+	{
+		return (NSNotification*)EventData;
+	}
+	else
+	{
+		return nil;
+	}
+}
+
+TSharedPtr<FMacWindow> FMacEvent::GetWindow(void) const
+{
+	return MacWindow;
+}
+
+FVector2D FMacEvent::GetMousePosition(void) const
+{
+	return MousePosition;
+}
diff --git a/Engine/Source/Runtime/Core/Private/Mac/MacPlatformMisc.cpp b/Engine/Source/Runtime/Core/Private/Mac/MacPlatformMisc.cpp
--- a/Engine/Source/Runtime/Core/Private/Mac/MacPlatformMisc.cpp
+++ b/Engine/Source/Runtime/Core/Private/Mac/MacPlatformMisc.cpp
@@ -10,6 +10,7 @@
 #include "VarargsHelper.h"
 #include "MacApplication.h"
 #include "MacCursor.h"
+#include "CocoaThread.h"
 #include "Runtime/Launch/Resources/Version.h"
 #include "EngineVersion.h"
 
@@ -466,30 +467,7 @@ void FMacPlatformMisc::PumpMessages( bool bFromMainLoop )
 {
 	if( bFromMainLoop )
 	{
-		SCOPED_AUTORELEASE_POOL;
-
-		while( NSEvent *Event = [NSApp nextEventMatchingMask: NSAnyEventMask untilDate: nil inMode: NSDefaultRunLoopMode dequeue: true] )
-		{
-			const bool bIsMouseClickOrKeyEvent = [Event type] == NSLeftMouseDown || [Event type] == NSLeftMouseUp
-										 || [Event type] == NSRightMouseDown || [Event type] == NSRightMouseUp
-										 || [Event type] == NSOtherMouseDown || [Event type] == NSOtherMouseUp
-										 || [Event type] == NSKeyDown || ([Event type] == NSKeyUp && !([Event modifierFlags] & NSCommandKeyMask));
-
-			if( MacApplication )
-			{
-				if( !bIsMouseClickOrKeyEvent || [Event window] == NULL )
-				{
-					MacApplication->ProcessEvent( Event );
-				}
-
-				if( [Event type] == NSLeftMouseUp )
-				{
-					MacApplication->OnWindowDraggingFinished();
-				}
-			}
-
-			[NSApp sendEvent: Event];
-		}
+		ProcessGameThreadEvents();
 	}
 }
 
diff --git a/Engine/Source/Runtime/Core/Private/Mac/MacPlatformOutputDevices.cpp b/Engine/Source/Runtime/Core/Private/Mac/MacPlatformOutputDevices.cpp
--- a/Engine/Source/Runtime/Core/Private/Mac/MacPlatformOutputDevices.cpp
+++ b/Engine/Source/Runtime/Core/Private/Mac/MacPlatformOutputDevices.cpp
@@ -10,6 +10,7 @@
 #include "MacApplication.h"
 #include "MacPlatformOutputDevicesPrivate.h"
 #include "MacPlatformFeedbackContextPrivate.h"
+#include "CocoaThread.h"
 
 //////////////////////////////////
 // FMacPlatformOutputDevices
@@ -152,6 +153,7 @@ FOutputDeviceConsoleMac::FOutputDeviceConsoleMac()
 	, TextView(NULL)
 	, ScrollView(NULL)
 	, TextViewTextColor(NULL)
+	, OutstandingTasks(0)
 {
 }
 
@@ -330,7 +332,8 @@ void FOutputDeviceConsoleMac::Serialize( const TCHAR* Data, ELogVerbosity::Type
 					[AttributeKeys addObject:NSForegroundColorAttributeName];
 					[AttributeKeys addObject:NSBackgroundColorAttributeName];
 					
-					dispatch_block_t Block = ^{
+					OutstandingTasks++;
+					MainThreadCall(^{
 						if( TextViewTextColor )
 							[TextViewTextColor release];
 						
@@ -339,17 +342,7 @@ void FOutputDeviceConsoleMac::Serialize( const TCHAR* Data, ELogVerbosity::Type
 						[Colors release];
 						[AttributeKeys release];
 						OutstandingTasks--;
-					};
-					
-					OutstandingTasks++;
-					if([NSThread isMainThread])
-					{
-						Block();
-					}
-					else
-					{
-						dispatch_async(dispatch_get_main_queue(), Block);
-					}
+					}, false);
 				}
 			}
 			else
@@ -361,24 +354,15 @@ void FOutputDeviceConsoleMac::Serialize( const TCHAR* Data, ELogVerbosity::Type
 
 				CFStringRef CocoaText = FPlatformString::TCHARToCFString(OutputString);
 				
-				dispatch_block_t Block = ^{
+				OutstandingTasks++;
+				MainThreadCall(^{
 					NSAttributedString *AttributedString = [[NSAttributedString alloc] initWithString:(NSString*)CocoaText attributes:TextViewTextColor];
 					[[TextView textStorage] appendAttributedString:AttributedString];
 					[TextView scrollRangeToVisible:NSMakeRange([[TextView string] length], 0)];
 					[AttributedString release];
 					CFRelease(CocoaText);
 					OutstandingTasks--;
-				};
-				
-				OutstandingTasks++;
-				if([NSThread isMainThread])
-				{
-					Block();
-				}
-				else
-				{
-					dispatch_async(dispatch_get_main_queue(), Block);
-				}
+				}, false);
 				
 				if(!MacApplication)
 				{
@@ -415,7 +399,8 @@ void FOutputDeviceConsoleMac::SetDefaultTextColor()
 	[AttributeKeys addObject:NSForegroundColorAttributeName];
 	[AttributeKeys addObject:NSBackgroundColorAttributeName];
 
-	dispatch_block_t Block = ^{
+	OutstandingTasks++;
+	MainThreadCall(^{
 		if( TextViewTextColor )
 			[TextViewTextColor release];
 		
@@ -424,15 +409,5 @@ void FOutputDeviceConsoleMac::SetDefaultTextColor()
 		[Colors release];
 		[AttributeKeys release];
 		OutstandingTasks--;
-	};
-	
-	OutstandingTasks++;
-	if([NSThread isMainThread])
-	{
-		Block();
-	}
-	else
-	{
-		dispatch_async(dispatch_get_main_queue(), Block);
-	}
+	}, false);
 }
diff --git a/Engine/Source/Runtime/Core/Private/Mac/MacTextInputMethodSystem.cpp b/Engine/Source/Runtime/Core/Private/Mac/MacTextInputMethodSystem.cpp
--- a/Engine/Source/Runtime/Core/Private/Mac/MacTextInputMethodSystem.cpp
+++ b/Engine/Source/Runtime/Core/Private/Mac/MacTextInputMethodSystem.cpp
@@ -5,6 +5,7 @@
 #include "MacWindow.h"
 #include "MacTextInputMethodSystem.h"
 #include "CocoaTextView.h"
+#include "CocoaThread.h"
 
 DEFINE_LOG_CATEGORY_STATIC(LogMacTextInputMethodSystem, Log, All);
 
@@ -49,11 +50,13 @@ namespace
 		if(ContextWindow.IsValid())
 		{
 			FCocoaWindow* CocoaWindow = (FCocoaWindow*)ContextWindow->GetOSWindowHandle();
-			if(CocoaWindow && [CocoaWindow openGLView])
-			{
-				FCocoaTextView* TextView = (FCocoaTextView*)[CocoaWindow openGLView];
-				[[TextView inputContext] invalidateCharacterCoordinates];
-			}
+			MainThreadCall(^{
+				if(CocoaWindow && [CocoaWindow openGLView])
+				{
+					FCocoaTextView* TextView = (FCocoaTextView*)[CocoaWindow openGLView];
+					[[TextView inputContext] invalidateCharacterCoordinates];
+				}
+			}, true);
 		}
 	}
 	
@@ -62,11 +65,13 @@ namespace
 		if(ContextWindow.IsValid())
 		{
 			FCocoaWindow* CocoaWindow = (FCocoaWindow*)ContextWindow->GetOSWindowHandle();
-			if(CocoaWindow && [CocoaWindow openGLView])
-			{
-				FCocoaTextView* TextView = (FCocoaTextView*)[CocoaWindow openGLView];
-				[[TextView inputContext] invalidateCharacterCoordinates];
-			}
+			MainThreadCall(^{
+				if(CocoaWindow && [CocoaWindow openGLView])
+				{
+					FCocoaTextView* TextView = (FCocoaTextView*)[CocoaWindow openGLView];
+					[[TextView inputContext] invalidateCharacterCoordinates];
+				}
+			}, true);
 		}
 	}
 	
@@ -75,11 +80,13 @@ namespace
 		if(ContextWindow.IsValid())
 		{
 			FCocoaWindow* CocoaWindow = (FCocoaWindow*)ContextWindow->GetOSWindowHandle();
-			if(CocoaWindow && [CocoaWindow openGLView])
-			{
-				FCocoaTextView* TextView = (FCocoaTextView*)[CocoaWindow openGLView];
-				[[TextView inputContext] invalidateCharacterCoordinates];
-			}
+			MainThreadCall(^{
+				if(CocoaWindow && [CocoaWindow openGLView])
+				{
+					FCocoaTextView* TextView = (FCocoaTextView*)[CocoaWindow openGLView];
+					[[TextView inputContext] invalidateCharacterCoordinates];
+				}
+			}, true);
 		}
 	}
 	
@@ -88,12 +95,14 @@ namespace
 		if(ContextWindow.IsValid())
 		{
 			FCocoaWindow* CocoaWindow = (FCocoaWindow*)ContextWindow->GetOSWindowHandle();
-			if(CocoaWindow && [CocoaWindow openGLView])
-			{
-				FCocoaTextView* TextView = (FCocoaTextView*)[CocoaWindow openGLView];
-				[[TextView inputContext] discardMarkedText];
-				[TextView unmarkText];
-			}
+			MainThreadCall(^{
+				if(CocoaWindow && [CocoaWindow openGLView])
+				{
+					FCocoaTextView* TextView = (FCocoaTextView*)[CocoaWindow openGLView];
+					[[TextView inputContext] discardMarkedText];
+					[TextView unmarkText];
+				}
+			}, true);
 		}
 	}
 }
@@ -142,22 +151,30 @@ void FMacTextInputMethodSystem::ActivateContext(const TSharedRef<ITextInputMetho
 	TSharedPtr<ITextInputMethodChangeNotifier> Notifier(NotifierRef.Pin());
 	FTextInputMethodChangeNotifier* MacNotifier = (FTextInputMethodChangeNotifier*)Notifier.Get();
 	TSharedPtr<FGenericWindow> GenericWindow = Context->GetWindow();
+	bool bActivatedContext = false;
 	if(GenericWindow.IsValid())
 	{
 		MacNotifier->SetContextWindow(GenericWindow);
 		FCocoaWindow* CocoaWindow = (FCocoaWindow*)GenericWindow->GetOSWindowHandle();
-		if(CocoaWindow && [CocoaWindow openGLView])
-		{
-			NSView* GLView = [CocoaWindow openGLView];
-			if([GLView isKindOfClass:[FCocoaTextView class]])
+		bActivatedContext = MainThreadReturn(^{
+			bool bSuccess = false;
+			if(CocoaWindow && [CocoaWindow openGLView])
 			{
-				FCocoaTextView* TextView = (FCocoaTextView*)GLView;
-				[TextView activateInputMethod:Context];
-				return;
+				NSView* GLView = [CocoaWindow openGLView];
+				if([GLView isKindOfClass:[FCocoaTextView class]])
+				{
+					FCocoaTextView* TextView = (FCocoaTextView*)GLView;
+					[TextView activateInputMethod:Context];
+					bSuccess = true;
+				}
 			}
-		}
+			return bSuccess;
+		});
+	}
+	if(!bActivatedContext)
+	{
+		UE_LOG(LogMacTextInputMethodSystem, Error, TEXT("Activating a context failed when its window couldn't be found."));
 	}
-	UE_LOG(LogMacTextInputMethodSystem, Error, TEXT("Activating a context failed when its window couldn't be found."));
 }
 
 void FMacTextInputMethodSystem::DeactivateContext(const TSharedRef<ITextInputMethodContext>& Context)
@@ -172,20 +189,28 @@ void FMacTextInputMethodSystem::DeactivateContext(const TSharedRef<ITextInputMet
 	TSharedPtr<ITextInputMethodChangeNotifier> Notifier(NotifierRef.Pin());
 	FTextInputMethodChangeNotifier* MacNotifier = (FTextInputMethodChangeNotifier*)Notifier.Get();
 	TSharedPtr<FGenericWindow> GenericWindow = MacNotifier->GetContextWindow();
+	bool bDeactivatedContext = false;
 	if(GenericWindow.IsValid())
 	{
 		FCocoaWindow* CocoaWindow = (FCocoaWindow*)GenericWindow->GetOSWindowHandle();
-		if(CocoaWindow && [CocoaWindow openGLView])
-		{
-			NSView* GLView = [CocoaWindow openGLView];
-			if([GLView isKindOfClass:[FCocoaTextView class]])
+		bDeactivatedContext = MainThreadReturn(^{
+			bool bSuccess = false;
+			if(CocoaWindow && [CocoaWindow openGLView])
 			{
-				FCocoaTextView* TextView = (FCocoaTextView*)GLView;
-				[TextView deactivateInputMethod];
-				return;
+				NSView* GLView = [CocoaWindow openGLView];
+				if([GLView isKindOfClass:[FCocoaTextView class]])
+				{
+					FCocoaTextView* TextView = (FCocoaTextView*)GLView;
+					[TextView deactivateInputMethod];
+					bSuccess = true;
+				}
 			}
-		}
+			return bSuccess;
+		});
+	}
+	if(!bDeactivatedContext)
+	{
+		UE_LOG(LogMacTextInputMethodSystem, Error, TEXT("Deactivating a context failed when its window couldn't be found."));
 	}
-	UE_LOG(LogMacTextInputMethodSystem, Error, TEXT("Deactivating a context failed when its window couldn't be found."));
 }
 // ITextInputMethodSystem Interface End
diff --git a/Engine/Source/Runtime/Core/Private/Mac/MacWindow.cpp b/Engine/Source/Runtime/Core/Private/Mac/MacWindow.cpp
--- a/Engine/Source/Runtime/Core/Private/Mac/MacWindow.cpp
+++ b/Engine/Source/Runtime/Core/Private/Mac/MacWindow.cpp
@@ -4,6 +4,7 @@
 #include "MacWindow.h"
 #include "MacApplication.h"
 #include "CocoaTextView.h"
+#include "CocoaThread.h"
 
 TArray< FCocoaWindow* > FMacWindow::RunningModalWindows;
 
@@ -84,7 +85,7 @@ void FMacWindow::Initialize( FMacApplication* const Application, const TSharedRe
 		WindowStyle &= ~(NSTexturedBackgroundWindowMask);
 	}
 
-	WindowHandle = [[FCocoaWindow alloc] initWithContentRect: ViewRect styleMask: WindowStyle backing: NSBackingStoreBuffered defer: NO];
+	WindowHandle = MainThreadReturn(^{ return [[FCocoaWindow alloc] initWithContentRect: ViewRect styleMask: WindowStyle backing: NSBackingStoreBuffered defer: NO]; });
 
 	if( WindowHandle == NULL )
 	{
@@ -94,77 +95,79 @@ void FMacWindow::Initialize( FMacApplication* const Application, const TSharedRe
 		return;
 	}
 
-	// Disable automatic release on close - explicit retain/release makes the Destroy() logic much easier to follow
-	[WindowHandle setReleasedWhenClosed: NO];
-	
-	[WindowHandle setWindowMode: WindowMode];
-	[WindowHandle setAcceptsInput: Definition->AcceptsInput];
-	[WindowHandle setDisplayReconfiguring: false];
-	[WindowHandle setAcceptsMouseMovedEvents: YES];
-	[WindowHandle setDelegate: WindowHandle];
-
-	// @todo: We really need a window type in tab manager to know what window level to use and whether or not a window should hide on deactivate
-	if(Definition->IsRegularWindow && (!InParent.IsValid() || Definition->IsModalWindow || (!Definition->SupportsMaximize && !Definition->SupportsMinimize)))
-	{
-		[WindowHandle setLevel: NSNormalWindowLevel];
-	}
-	else
-	{
-		[WindowHandle setLevel: NSFloatingWindowLevel];
-		[WindowHandle setHidesOnDeactivate: YES];
-	}
-	
-	// Use of rounded corners will always render with system values for rounding
-	[WindowHandle setRoundedCorners: (Definition->CornerRadius != 0)];
-	
-	if( !Definition->HasOSWindowBorder )
-	{
-		[WindowHandle setBackgroundColor: [NSColor darkGrayColor]];
-		[WindowHandle setHasShadow: YES];
-	}
+	MainThreadCall(^{
+		// Disable automatic release on close - explicit retain/release makes the Destroy() logic much easier to follow
+		[WindowHandle setReleasedWhenClosed: NO];
+		
+		[WindowHandle setWindowMode: WindowMode];
+		[WindowHandle setAcceptsInput: Definition->AcceptsInput];
+		[WindowHandle setDisplayReconfiguring: false];
+		[WindowHandle setAcceptsMouseMovedEvents: YES];
+		[WindowHandle setDelegate: WindowHandle];
+
+		// @todo: We really need a window type in tab manager to know what window level to use and whether or not a window should hide on deactivate
+		if(Definition->IsRegularWindow && (!InParent.IsValid() || Definition->IsModalWindow || (!Definition->SupportsMaximize && !Definition->SupportsMinimize)))
+		{
+			[WindowHandle setLevel: NSNormalWindowLevel];
+		}
+		else
+		{
+			[WindowHandle setLevel: NSFloatingWindowLevel];
+			[WindowHandle setHidesOnDeactivate: YES];
+		}
+		
+		// Use of rounded corners will always render with system values for rounding
+		[WindowHandle setRoundedCorners: (Definition->CornerRadius != 0)];
+		
+		if( !Definition->HasOSWindowBorder )
+		{
+			[WindowHandle setBackgroundColor: [NSColor darkGrayColor]];
+			[WindowHandle setHasShadow: YES];
+		}
 
-	[WindowHandle setOpaque: NO];
+		[WindowHandle setOpaque: NO];
 
-	ReshapeWindow( X, Y, SizeX, SizeY );
+		ReshapeWindow( X, Y, SizeX, SizeY );
 
-	if( Definition->IsRegularWindow )
-	{
-		CFStringRef CFName = FPlatformString::TCHARToCFString( *Definition->Title );
-		[WindowHandle setTitle: ( NSString *)CFName];
-		CFRelease( CFName );
+		if( Definition->IsRegularWindow )
+		{
+			CFStringRef CFName = FPlatformString::TCHARToCFString( *Definition->Title );
+			[WindowHandle setTitle: ( NSString *)CFName];
+			CFRelease( CFName );
 
-		[NSApp addWindowsItem: WindowHandle title: [WindowHandle title] filename: NO];
+			[NSApp addWindowsItem: WindowHandle title: [WindowHandle title] filename: NO];
 
-		// Tell Cocoa that we are opting into drag and drop.
-		// Only makes sense for regular windows (windows that last a while.)
-		[WindowHandle registerForDraggedTypes: [NSArray arrayWithObject: NSFilenamesPboardType]];
+			// Tell Cocoa that we are opting into drag and drop.
+			// Only makes sense for regular windows (windows that last a while.)
+			[WindowHandle registerForDraggedTypes: [NSArray arrayWithObject: NSFilenamesPboardType]];
 
-		if( Definition->HasOSWindowBorder )
+			if( Definition->HasOSWindowBorder )
+			{
+				[WindowHandle setCollectionBehavior: NSWindowCollectionBehaviorFullScreenPrimary|NSWindowCollectionBehaviorDefault|NSWindowCollectionBehaviorManaged|NSWindowCollectionBehaviorParticipatesInCycle];
+			}
+			else
+			{
+				[WindowHandle setCollectionBehavior: NSWindowCollectionBehaviorFullScreenAuxiliary|NSWindowCollectionBehaviorDefault|NSWindowCollectionBehaviorManaged|NSWindowCollectionBehaviorParticipatesInCycle];
+			}
+		}
+		else if(Definition->AppearsInTaskbar)
 		{
-			[WindowHandle setCollectionBehavior: NSWindowCollectionBehaviorFullScreenPrimary|NSWindowCollectionBehaviorDefault|NSWindowCollectionBehaviorManaged|NSWindowCollectionBehaviorParticipatesInCycle];
+			[WindowHandle setCollectionBehavior:NSWindowCollectionBehaviorFullScreenAuxiliary|NSWindowCollectionBehaviorDefault|NSWindowCollectionBehaviorManaged|NSWindowCollectionBehaviorParticipatesInCycle];
 		}
 		else
 		{
-			[WindowHandle setCollectionBehavior: NSWindowCollectionBehaviorFullScreenAuxiliary|NSWindowCollectionBehaviorDefault|NSWindowCollectionBehaviorManaged|NSWindowCollectionBehaviorParticipatesInCycle];
+			[WindowHandle setCollectionBehavior:NSWindowCollectionBehaviorCanJoinAllSpaces|NSWindowCollectionBehaviorTransient|NSWindowCollectionBehaviorIgnoresCycle];
 		}
-	}
-	else if(Definition->AppearsInTaskbar)
-	{
-		[WindowHandle setCollectionBehavior:NSWindowCollectionBehaviorFullScreenAuxiliary|NSWindowCollectionBehaviorDefault|NSWindowCollectionBehaviorManaged|NSWindowCollectionBehaviorParticipatesInCycle];
-	}
-	else
-	{
-		[WindowHandle setCollectionBehavior:NSWindowCollectionBehaviorCanJoinAllSpaces|NSWindowCollectionBehaviorTransient|NSWindowCollectionBehaviorIgnoresCycle];
-	}
 
-	if( Definition->SupportsTransparency )
-	{
-		SetOpacity( Definition->Opacity );
-	}
-	else
-	{
-		SetOpacity( 1.0f );
-	}
+		if( Definition->SupportsTransparency )
+		{
+			SetOpacity( Definition->Opacity );
+		}
+		else
+		{
+			SetOpacity( 1.0f );
+		}
+	}, true);
 }
 
 FMacWindow::FMacWindow()
@@ -190,26 +193,28 @@ void FMacWindow::ReshapeWindow( int32 X, int32 Y, int32 Width, int32 Height )
 
 	if(WindowMode == EWindowMode::Windowed || WindowMode == EWindowMode::WindowedFullscreen)
 	{
-		BOOL DisplayIfNeeded = (WindowMode == EWindowMode::Windowed);
-		
-		const int32 InvertedY = FPlatformMisc::ConvertSlateYPositionToCocoa(Y) - Height + 1;
-		if (Definition->HasOSWindowBorder)
-		{
-			[WindowHandle setFrame: [WindowHandle frameRectForContentRect: NSMakeRect(X, InvertedY, FMath::Max(Width, 1), FMath::Max(Height, 1))] display:DisplayIfNeeded];
-		}
-		else
-		{
-			[WindowHandle setFrame: NSMakeRect(X, InvertedY, FMath::Max(Width, 1), FMath::Max(Height, 1)) display:DisplayIfNeeded];
-		}
-		
-		// Force resize back to screen size in fullscreen - not ideally pretty but means we don't
-		// have to subvert the OS X or UE fullscreen handling events elsewhere.
-		if(WindowMode != EWindowMode::Windowed)
-		{
-			[WindowHandle setFrame: [WindowHandle screen].frame display:YES];
-		}
-		
-		WindowHandle->bZoomed = [WindowHandle isZoomed];
+		MainThreadCall(^{
+			BOOL DisplayIfNeeded = (WindowMode == EWindowMode::Windowed);
+			
+			const int32 InvertedY = FPlatformMisc::ConvertSlateYPositionToCocoa(Y) - Height + 1;
+			if (Definition->HasOSWindowBorder)
+			{
+				[WindowHandle setFrame: [WindowHandle frameRectForContentRect: NSMakeRect(X, InvertedY, FMath::Max(Width, 1), FMath::Max(Height, 1))] display:DisplayIfNeeded];
+			}
+			else
+			{
+				[WindowHandle setFrame: NSMakeRect(X, InvertedY, FMath::Max(Width, 1), FMath::Max(Height, 1)) display:DisplayIfNeeded];
+			}
+			
+			// Force resize back to screen size in fullscreen - not ideally pretty but means we don't
+			// have to subvert the OS X or UE fullscreen handling events elsewhere.
+			if(WindowMode != EWindowMode::Windowed)
+			{
+				[WindowHandle setFrame: [WindowHandle screen].frame display:YES];
+			}
+			
+			WindowHandle->bZoomed = [WindowHandle isZoomed];
+		}, true);
 	}
 	
 	MessageHandler->FinishedReshapingWindow( SharedThis( this ) );
@@ -229,18 +234,21 @@ bool FMacWindow::GetFullScreenInfo( int32& X, int32& Y, int32& Width, int32& Hei
 
 void FMacWindow::MoveWindowTo( int32 X, int32 Y )
 {
-	SCOPED_AUTORELEASE_POOL;
-
-	const int32 InvertedY = FPlatformMisc::ConvertSlateYPositionToCocoa(Y) - [WindowHandle frame].size.height + 1;
-	[WindowHandle setFrameOrigin: NSMakePoint(X, InvertedY)];
+	MainThreadCall(^{
+		SCOPED_AUTORELEASE_POOL;
+		const int32 InvertedY = FPlatformMisc::ConvertSlateYPositionToCocoa(Y) - [WindowHandle frame].size.height + 1;
+		[WindowHandle setFrameOrigin: NSMakePoint(X, InvertedY)];
+	}, true);
 }
 
 void FMacWindow::BringToFront( bool bForce )
 {
 	if (bIsVisible)
 	{
-		SCOPED_AUTORELEASE_POOL;
-		[WindowHandle orderFrontAndMakeMain:IsRegularWindow() andKey:IsRegularWindow()];
+		MainThreadCall(^{
+			SCOPED_AUTORELEASE_POOL;
+			[WindowHandle orderFrontAndMakeMain:IsRegularWindow() andKey:IsRegularWindow()];
+		}, true);
 	}
 }
 
@@ -271,7 +279,9 @@ void FMacWindow::Destroy()
 			}
 
 			// Close the window, but don't destruct it...
-			[Window performClose:nil];
+			MainThreadCall(^{
+				[Window performClose:nil];
+			}, true);
 			
 			// Since event handling of the change in focus may try to communicate with this window,
 			// dispatch a call to release it which should be handled afterward.
@@ -285,33 +295,39 @@ void FMacWindow::Destroy()
 
 void FMacWindow::Minimize()
 {
-	SCOPED_AUTORELEASE_POOL;
-	[WindowHandle miniaturize:nil];
+	MainThreadCall(^{
+		SCOPED_AUTORELEASE_POOL;
+		[WindowHandle miniaturize:nil];
+	}, true);
 }
 
 void FMacWindow::Maximize()
 {
-	SCOPED_AUTORELEASE_POOL;
-	if( ![WindowHandle isZoomed] )
-	{
-		WindowHandle->bZoomed = true;
-		[WindowHandle zoom:nil];
-	}
+	MainThreadCall(^{
+		SCOPED_AUTORELEASE_POOL;
+		if( ![WindowHandle isZoomed] )
+		{
+			WindowHandle->bZoomed = true;
+			[WindowHandle zoom:nil];
+		}
+	}, true);
 }
 
 void FMacWindow::Restore()
 {
-	SCOPED_AUTORELEASE_POOL;
-	if( [WindowHandle isZoomed] )
-	{
-		WindowHandle->bZoomed = !WindowHandle->bZoomed;
-		[WindowHandle zoom:nil];
-	}
-	else
-	{
-		WindowHandle->bZoomed = false;
-		[WindowHandle deminiaturize:nil];
-	}
+	MainThreadCall(^{
+		SCOPED_AUTORELEASE_POOL;
+		if( [WindowHandle isZoomed] )
+		{
+			WindowHandle->bZoomed = !WindowHandle->bZoomed;
+			[WindowHandle zoom:nil];
+		}
+		else
+		{
+			WindowHandle->bZoomed = false;
+			[WindowHandle deminiaturize:nil];
+		}
+	}, true);
 }
 
 void FMacWindow::Show()
@@ -323,8 +339,12 @@ void FMacWindow::Show()
 		{
 			AddModalWindow(WindowHandle);
 		}
+		
 		bool bMakeMainAndKey = ([WindowHandle canBecomeKeyWindow] && Definition->ActivateWhenFirstShown);
-		[WindowHandle orderFrontAndMakeMain:bMakeMainAndKey andKey:bMakeMainAndKey];
+		
+		MainThreadCall(^{
+			[WindowHandle orderFrontAndMakeMain:bMakeMainAndKey andKey:bMakeMainAndKey];
+		}, true);
 		
 		bIsVisible = [WindowHandle isVisible];
 		static bool bCannotRecurse = false;
@@ -355,7 +375,9 @@ void FMacWindow::Hide()
 		{
 			RemoveModalWindow(WindowHandle);
 		}
-		[WindowHandle orderOut:nil];
+		MainThreadCall(^{
+			[WindowHandle orderOut:nil];
+		}, true);
 	}
 }
 
@@ -378,8 +400,6 @@ void FMacWindow::SetWindowMode( EWindowMode::Type NewWindowMode )
 			Behaviour |= NSWindowCollectionBehaviorFullScreenPrimary;
 		}
 		
-		[WindowHandle setCollectionBehavior: Behaviour];
-		
 		if(!bIsFullscreen)
 		{
 			PreFullscreenWindowRect = [WindowHandle openGLFrame];
@@ -387,7 +407,12 @@ void FMacWindow::SetWindowMode( EWindowMode::Type NewWindowMode )
 		}
 		
 		WindowHandle.TargetWindowMode = NewWindowMode;
-		[WindowHandle toggleFullScreen:nil];
+		
+		MainThreadCall(^{
+			[WindowHandle setCollectionBehavior: Behaviour];
+			[WindowHandle toggleFullScreen:nil];
+		}, true);
+		
 		// Ensure that the window has transitioned BEFORE leaving this function
 		// this prevents problems with failure to correctly update mouse locks
 		// and OpenGL contexts due to bad event ordering.
@@ -439,20 +464,26 @@ bool FMacWindow::GetRestoredDimensions(int32& X, int32& Y, int32& Width, int32&
 
 void FMacWindow::SetWindowFocus()
 {
-	SCOPED_AUTORELEASE_POOL;
-	[WindowHandle orderFrontAndMakeMain:true andKey:true];
+	MainThreadCall(^{
+		SCOPED_AUTORELEASE_POOL;
+		[WindowHandle orderFrontAndMakeMain:true andKey:true];
+	}, true);
 }
 
 void FMacWindow::SetOpacity( const float InOpacity )
 {
-	SCOPED_AUTORELEASE_POOL;
-	[WindowHandle setAlphaValue:InOpacity];
+	MainThreadCall(^{
+		SCOPED_AUTORELEASE_POOL;
+		[WindowHandle setAlphaValue:InOpacity];
+	}, true);
 }
 
 void FMacWindow::Enable( bool bEnable )
 {
-	SCOPED_AUTORELEASE_POOL;
-	[WindowHandle setIgnoresMouseEvents: !bEnable];
+	MainThreadCall(^{
+		SCOPED_AUTORELEASE_POOL;
+		[WindowHandle setIgnoresMouseEvents: !bEnable];
+	}, true);
 }
 
 bool FMacWindow::IsPointInWindow( int32 X, int32 Y ) const
@@ -509,12 +540,14 @@ void FMacWindow::SetText(const TCHAR* const Text)
 {
 	SCOPED_AUTORELEASE_POOL;
 	CFStringRef CFName = FPlatformString::TCHARToCFString( Text );
-	[WindowHandle setTitle: (NSString*)CFName];
-	if(IsRegularWindow())
-	{
-		[NSApp changeWindowsItem: WindowHandle title: (NSString*)CFName filename: NO];
-	}
-	CFRelease( CFName );
+	MainThreadCall(^{
+		[WindowHandle setTitle: (NSString*)CFName];
+		if(IsRegularWindow())
+		{
+			[NSApp changeWindowsItem: WindowHandle title: (NSString*)CFName filename: NO];
+		}
+		CFRelease( CFName );
+	}, true);
 }
 
 bool FMacWindow::IsRegularWindow() const
@@ -531,17 +564,19 @@ void FMacWindow::OnDisplayReconfiguration(CGDirectDisplayID Display, CGDisplayCh
 {
 	if(WindowHandle)
 	{
-		if(Flags & kCGDisplayBeginConfigurationFlag)
-		{
-			[WindowHandle setMovable: YES];
-			[WindowHandle setMovableByWindowBackground: NO];
-			
-			[WindowHandle setDisplayReconfiguring: true];
-		}
-		else if(Flags & kCGDisplayDesktopShapeChangedFlag)
-		{
-			[WindowHandle setDisplayReconfiguring: false];
-		}
+		MainThreadCall(^{
+			if(Flags & kCGDisplayBeginConfigurationFlag)
+			{
+				[WindowHandle setMovable: YES];
+				[WindowHandle setMovableByWindowBackground: NO];
+				
+				[WindowHandle setDisplayReconfiguring: true];
+			}
+			else if(Flags & kCGDisplayDesktopShapeChangedFlag)
+			{
+				[WindowHandle setDisplayReconfiguring: false];
+			}
+		}, true);
 	}
 }
 
@@ -549,8 +584,10 @@ bool FMacWindow::OnIMKKeyDown(NSEvent* Event)
 {
 	if(WindowHandle && [WindowHandle openGLView])
 	{
-		FCocoaTextView* View = (FCocoaTextView*)[WindowHandle openGLView];
-		return View && [View imkKeyDown:Event];
+		return MainThreadReturn(^{
+			FCocoaTextView* View = (FCocoaTextView*)[WindowHandle openGLView];
+			return View && [View imkKeyDown:Event];
+		});
 	}
 	else
 	{
diff --git a/Engine/Source/Runtime/Core/Public/Mac/CocoaThread.h b/Engine/Source/Runtime/Core/Public/Mac/CocoaThread.h
new file mode 100644
--- /dev/null
+++ b/Engine/Source/Runtime/Core/Public/Mac/CocoaThread.h
@@ -0,0 +1,63 @@
+// Copyright 1998-2014 Epic Games, Inc. All Rights Reserved.
+#pragma once
+
+#include "Core.h"
+
+#ifdef __OBJC__
+
+@interface NSThread (FCocoaThread)
++ (NSThread*) gameThread; // Returns the main game thread, or nil if has yet to be constructed.
++ (bool) isGameThread; // True if the current thread is the main game thread, else false.
+- (bool) isGameThread; // True if this thread object is the main game thread, else false.
+- (void) performBlock:(dispatch_block_t)Block onRunLoop:(NSRunLoop*)RunLoop forMode:(NSString*)RunMode wait:(bool)bWait inMode:(NSString*)WaitMode; // Perform Block on RunLoop in RunMode, which must not be the current thread's runloop and wait in WaitMode if bWait is true or return immediately after queueing to the RunLoop.
+@end
+
+@interface NSRunLoop (FCocoaRunLoop)
++ (NSRunLoop*)gameRunLoop; // The run loop for the main game thread, or nil if the main game thread has not been started.
+- (void)performBlock:(dispatch_block_t)Block forMode:(NSString *)Mode wake:(bool)bWakeRunLoop; // Add Block to the run loop for processing in the given Mode, only waking the run loop if bWakeRunLoop is set otherwise processing occurs the next time the run-loop is processed.
+@end
+
+@interface FCocoaGameThread : NSThread
+- (id)init; // Override that sets the variable backing +[NSThread gameThread], do not override in a subclass.
+- (id)initWithTarget:(id)Target selector:(SEL)Selector object:(id)Argument; // Override that sets the variable backing +[NSThread gameThread], do not override in a subclass.
+- (void)dealloc; // Override that clears the variable backing +[NSThread gameThread], do not override in a subclass.
+- (void)main; // Override that sets the variable backing +[NSRunLoop gameRunLoop], do not override in a subclass.
+@end
+
+template<typename ReturnType>
+ReturnType PerformBlockOnRunLoop(ReturnType (^Block)(void), NSRunLoop* const RunLoop)
+{
+	__block ReturnType ReturnValue;
+	[[NSThread currentThread] performBlock:^{ ReturnValue = Block(); } onRunLoop:RunLoop forMode:NSDefaultRunLoopMode wait:true inMode:NSDefaultRunLoopMode];
+	return ReturnValue;
+}
+
+void PerformBlockOnRunLoop(dispatch_block_t Block, NSRunLoop* const RunLoop, bool const bWait);
+
+template<typename ReturnType>
+ReturnType MainThreadReturn(ReturnType (^Block)(void))
+{
+	return PerformBlockOnRunLoop(Block, [NSRunLoop mainRunLoop]);
+}
+
+void MainThreadCall(dispatch_block_t Block, bool const bWait);
+
+template<typename ReturnType>
+ReturnType GameThreadReturn(ReturnType (^Block)(void))
+{
+	return PerformBlockOnRunLoop(Block, [NSRunLoop gameRunLoop]);
+}
+
+void GameThreadCall(dispatch_block_t Block, bool const bWait);
+
+void RunGameThread(id Target, SEL Selector);
+
+void ProcessGameThreadEvents(void);
+
+#else
+
+class NSThread;
+class NSRunLoop;
+class FCocoaGameThread;
+
+#endif
diff --git a/Engine/Source/Runtime/Core/Public/Mac/CocoaWindow.h b/Engine/Source/Runtime/Core/Public/Mac/CocoaWindow.h
--- a/Engine/Source/Runtime/Core/Public/Mac/CocoaWindow.h
+++ b/Engine/Source/Runtime/Core/Public/Mac/CocoaWindow.h
@@ -82,6 +82,12 @@
 
 @end
 
+extern NSString* NSWindowRedrawContents;
+extern NSString* NSDraggingExited;
+extern NSString* NSDraggingUpdated;
+extern NSString* NSPrepareForDragOperation;
+extern NSString* NSPerformDragOperation;
+
 #else // __OBJC__
 
 class FCocoaWindow;
diff --git a/Engine/Source/Runtime/Core/Public/Mac/MacApplication.h b/Engine/Source/Runtime/Core/Public/Mac/MacApplication.h
--- a/Engine/Source/Runtime/Core/Public/Mac/MacApplication.h
+++ b/Engine/Source/Runtime/Core/Public/Mac/MacApplication.h
@@ -6,6 +6,7 @@
 #include "MacWindow.h"
 #include "MacTextInputMethodSystem.h"
 
+class FMacEvent;
 
 /**
  * Mac-specific application implementation.
@@ -61,6 +62,8 @@ public:
 		return TextInputMethodSystem.Get();
 	}
 
+	void ProcessNSEvent(NSEvent* const Event, TSharedPtr< FMacWindow > CurrentEventWindow, FVector2D const MousePosition);
+
 	void ProcessEvent(NSEvent* Event);
 
 	void OnWindowDraggingFinished();
@@ -103,6 +106,9 @@ public:
 
 	void OnMouseCursorLock( bool bLockEnabled );
 
+	static FCocoaWindow* FindMacEventWindow( NSEvent* CocoaEvent );
+	static TSharedPtr< FMacWindow > FindMacWindowByNSWindow( FCocoaWindow* const WindowHandle );
+	static void ProcessEvent(FMacEvent const* const Event);
 
 private:
 	static void OnDisplayReconfiguration(CGDirectDisplayID Display, CGDisplayChangeSummaryFlags Flags, void* UserInfo);
@@ -138,6 +144,7 @@ private:
 
 	EMouseButtons::Type LastPressedMouseButton;
 
+	FCriticalSection WindowsMutex;
 	TArray< TSharedRef< FMacWindow > > Windows;
 
 	TSharedRef< class HIDInputInterface > HIDInput;
@@ -174,6 +181,8 @@ private:
 	int32 GestureUsage[EGestureEvent::Count];
 #endif
 
+	void* EventMonitor;
+
 	friend class FMacWindow;
 };
 
diff --git a/Engine/Source/Runtime/Core/Public/Mac/MacEvent.h b/Engine/Source/Runtime/Core/Public/Mac/MacEvent.h
new file mode 100644
--- /dev/null
+++ b/Engine/Source/Runtime/Core/Public/Mac/MacEvent.h
@@ -0,0 +1,60 @@
+// Copyright 1998-2014 Epic Games, Inc. All Rights Reserved.
+#pragma once
+
+#include "Core.h"
+#include "MacWindow.h"
+
+#ifndef __OBJC__
+
+class NSEvent;
+class NSNotification;
+class NSObject;
+
+#endif
+
+enum EMacEventSendMethod
+{
+	Async,
+	Sync
+};
+
+class FMacEvent
+{
+	// Constructor for an NSEvent based FMacEvent
+	FMacEvent(NSEvent* const Event);
+	
+	// Constructor for an NSNotification based FMacEvent
+	FMacEvent(NSNotification* const Notification, NSWindow* const Window);
+	
+	// Send the event to the game run loop for processing using the specified SendMethod.
+	static void SendToGameRunLoop(FMacEvent const* const Event, EMacEventSendMethod SendMethod);
+	
+public:
+	// Send an NSEvent to the Game run loop as an FMacEvent using the specified SendMethod
+	static void SendToGameRunLoop(NSEvent* const Event, EMacEventSendMethod SendMethod);
+	
+	// Send an NSNotification to the Game run loop as an FMacEvent using the specified SendMethod
+	static void SendToGameRunLoop(NSNotification* const Notification, NSWindow* const Window, EMacEventSendMethod SendMethod);
+	
+	// Destructor
+	~FMacEvent(void);
+	
+	// Get the NSEvent for this FMacEvent, nil if not an NSEvent.
+	NSEvent* GetEvent(void) const;
+	
+	// Get the NSNotification for this FMacEvent, nil if not an NSNotification.
+	NSNotification* GetNotification(void) const;
+	
+	// Get the FMacWindow for this FMacEvent, may be nullptr if not a window event, or for a non-FMacWindow.
+	TSharedPtr<FMacWindow> GetWindow(void) const;
+	
+	// Get the mouse position for this FMacEvent at the time of posting on the AppKit Main Thread.
+	FVector2D GetMousePosition(void) const;
+	
+private:
+	TSharedPtr<FMacWindow> MacWindow; // The retained UE4 window for this event, or nullptr
+	NSObject* EventData; // The retained NSEvent or NSNotification for this event or nil.
+	FVector2D MousePosition; // The mouse position for the event in UE4 coordinates.
+};
+
+
diff --git a/Engine/Source/Runtime/Launch/Private/Mac/LaunchMac.cpp b/Engine/Source/Runtime/Launch/Private/Mac/LaunchMac.cpp
--- a/Engine/Source/Runtime/Launch/Private/Mac/LaunchMac.cpp
+++ b/Engine/Source/Runtime/Launch/Private/Mac/LaunchMac.cpp
@@ -9,6 +9,7 @@
 #endif
 
 #include <signal.h>
+#include "CocoaThread.h"
 
 
 static FString GSavedCommandLine;
@@ -92,24 +93,26 @@ void EngineCrashHandler(const FGenericCrashContext & GenericContext)
 
 - (IBAction)requestQuit:(id)Sender
 {
-	if (GEngine)
-	{
-		if (GIsEditor)
+	GameThreadCall(^{
+		if (GEngine)
 		{
-			if (IsRunningCommandlet())
+			if (GIsEditor)
 			{
-				GIsRequestingExit = true;
+				if (IsRunningCommandlet())
+				{
+					GIsRequestingExit = true;
+				}
+				else
+				{
+					GEngine->DeferredCommands.Add(TEXT("CLOSE_SLATE_MAINFRAME"));
+				}
 			}
 			else
 			{
-				GEngine->DeferredCommands.Add(TEXT("CLOSE_SLATE_MAINFRAME"));
+				GEngine->DeferredCommands.Add(TEXT("EXIT"));
 			}
 		}
-		else
-		{
-			GEngine->DeferredCommands.Add(TEXT("EXIT"));
-		}
-	}
+	}, false);
 }
 
 - (IBAction)showAboutWindow:(id)Sender
@@ -140,6 +143,53 @@ void EngineCrashHandler(const FGenericCrashContext & GenericContext)
     [self requestQuit:self];
 }
 
+- (NSApplicationTerminateReply)applicationShouldTerminate:(NSApplication *)Sender;
+{
+	if(!GIsRequestingExit)
+	{
+		[self requestQuit:self];
+		return NSTerminateCancel;
+	}
+	return NSTerminateNow;
+}
+
+- (void) runGameThread:(id)Arg
+{
+	bool bIsBuildMachine = false;
+#if !UE_BUILD_SHIPPING
+	if (FParse::Param(*GSavedCommandLine, TEXT("BUILDMACHINE")))
+	{
+		bIsBuildMachine = true;
+	}
+#endif
+	
+#if UE_BUILD_DEBUG
+	if( true && !GAlwaysReportCrash )
+#else
+	if( FPlatformMisc::IsDebuggerPresent() && !GAlwaysReportCrash )
+#endif
+	{
+		// Don't use exception handling when a debugger is attached to exactly trap the crash. This does NOT check
+		// whether we are the first instance or not!
+		GuardedMain( *GSavedCommandLine );
+	}
+	else
+	{
+		if (!bIsBuildMachine)
+		{
+			FPlatformMisc::SetCrashHandler(EngineCrashHandler);
+		}
+		GIsGuarded = 1;
+		// Run the guarded code.
+		GuardedMain( *GSavedCommandLine );
+		GIsGuarded = 0;
+	}
+
+	FEngineLoop::AppExit();
+
+	[NSApp terminate: nil];
+}
+
 - (void)applicationDidFinishLaunching:(NSNotification *)Notification
 {
 	//install the custom quit event handler
@@ -178,40 +228,8 @@ void EngineCrashHandler(const FGenericCrashContext & GenericContext)
 		GSavedCommandLine += Argument;
 	}
 #endif
-
-	bool bIsBuildMachine = false;
-#if !UE_BUILD_SHIPPING
-	if (FParse::Param(*GSavedCommandLine, TEXT("BUILDMACHINE")))
-	{
-		bIsBuildMachine = true;
-	}
-#endif
-
-#if UE_BUILD_DEBUG
-	if( true && !GAlwaysReportCrash )
-#else
-	if( FPlatformMisc::IsDebuggerPresent() && !GAlwaysReportCrash )
-#endif
-	{
-		// Don't use exception handling when a debugger is attached to exactly trap the crash. This does NOT check
-		// whether we are the first instance or not!
-		GuardedMain( *GSavedCommandLine );
-	}
-	else
-	{
-		if (!bIsBuildMachine)
-		{
-			FPlatformMisc::SetCrashHandler(EngineCrashHandler);
-		}
-		GIsGuarded = 1;
-		// Run the guarded code.
-		GuardedMain( *GSavedCommandLine );
-		GIsGuarded = 0;
-	}
-
-	FEngineLoop::AppExit();
 	
-	[NSApp terminate: nil];
+	RunGameThread(self, @selector(runGameThread:));
 }
 
 @end
diff --git a/Engine/Source/Runtime/Slate/Private/Framework/MultiBox/Mac/MacMenu.cpp b/Engine/Source/Runtime/Slate/Private/Framework/MultiBox/Mac/MacMenu.cpp
--- a/Engine/Source/Runtime/Slate/Private/Framework/MultiBox/Mac/MacMenu.cpp
+++ b/Engine/Source/Runtime/Slate/Private/Framework/MultiBox/Mac/MacMenu.cpp
@@ -2,6 +2,7 @@
 
 #include "SlatePrivatePCH.h"
 #include "MacMenu.h"
+#include "CocoaThread.h"
 
 @interface FMacMenuItem : NSMenuItem
 {
@@ -96,47 +97,50 @@
 		}
 		else
 		{
-			TSharedRef< SWidget > Widget = FSlateMacMenu::GetMenuBuilderWidget(MacMenu);
+			return GameThreadReturn(^{
+				TSharedRef< SWidget > Widget = FSlateMacMenu::GetMenuBuilderWidget(MacMenu);
 
-			TSharedRef<const FMultiBox> MultiBox = StaticCastSharedRef<SMultiBoxWidget>(Widget)->GetMultiBox();
-			const TArray<TSharedRef<const FMultiBlock>>& MenuBlocks = MultiBox->GetBlocks();
+				TSharedRef<const FMultiBox> MultiBox = StaticCastSharedRef<SMultiBoxWidget>(Widget)->GetMultiBox();
+				const TArray<TSharedRef<const FMultiBlock>>& MenuBlocks = MultiBox->GetBlocks();
 
-			for (int32 Index = 0; Index < MenuBlocks.Num(); Index++)
-			{
-				EMultiBlockType::Type Type = MenuBlocks[Index]->GetType();
-
-				if (Type == EMultiBlockType::MenuEntry)
+				for (int32 Index = 0; Index < MenuBlocks.Num(); Index++)
 				{
-					TSharedRef<const FMenuEntryBlock> Block = StaticCastSharedRef<const FMenuEntryBlock>(MenuBlocks[Index]);
-					if (Block->GetAction().IsValid())
+					EMultiBlockType::Type Type = MenuBlocks[Index]->GetType();
+
+					if (Type == EMultiBlockType::MenuEntry)
 					{
-						const TSharedRef<const FInputGesture>& Gesture = Block->GetAction()->GetActiveGesture();
-						if (!Gesture->GetKeyText().IsEmpty() && Gesture->GetKeyText().ToString()[0] == Character)
+						TSharedRef<const FMenuEntryBlock> Block = StaticCastSharedRef<const FMenuEntryBlock>(MenuBlocks[Index]);
+						if (Block->GetAction().IsValid())
 						{
-							uint32 EventModifiers = [event modifierFlags] & (NSCommandKeyMask | NSShiftKeyMask | NSAlternateKeyMask | NSControlKeyMask);
-							uint32 GestureModfiers = 0;
-							if (Gesture->bCtrl)
-							{
-								GestureModfiers |= NSCommandKeyMask;
-							}
-							if (Gesture->bShift)
+							const TSharedRef<const FInputGesture>& Gesture = Block->GetAction()->GetActiveGesture();
+							if (!Gesture->GetKeyText().IsEmpty() && Gesture->GetKeyText().ToString()[0] == Character)
 							{
-								GestureModfiers |= NSShiftKeyMask;
+								uint32 EventModifiers = [event modifierFlags] & (NSCommandKeyMask | NSShiftKeyMask | NSAlternateKeyMask | NSControlKeyMask);
+								uint32 GestureModfiers = 0;
+								if (Gesture->bCtrl)
+								{
+									GestureModfiers |= NSCommandKeyMask;
+								}
+								if (Gesture->bShift)
+								{
+									GestureModfiers |= NSShiftKeyMask;
+								}
+								if (Gesture->bAlt)
+								{
+									GestureModfiers |= NSAlternateKeyMask;
+								}
+								if (Gesture->bCmd)
+								{
+									GestureModfiers |= NSControlKeyMask;
+								}
+
+								return EventModifiers == GestureModfiers;
 							}
-							if (Gesture->bAlt)
-							{
-								GestureModfiers |= NSAlternateKeyMask;
-							}
-							if (Gesture->bCmd)
-							{
-								GestureModfiers |= NSControlKeyMask;
-							}
-
-							return EventModifiers == GestureModfiers;
 						}
 					}
 				}
-			}
+				return false;
+			});
 		}
 
 		return NO;
@@ -152,35 +156,37 @@
 
 void FSlateMacMenu::UpdateWithMultiBox(const TSharedRef< FMultiBox >& MultiBox)
 {
-	int32 NumItems = [[NSApp mainMenu] numberOfItems];
-	FText HelpTitle = NSLOCTEXT("MainMenu", "HelpMenu", "Help");
-	NSMenuItem* HelpMenu = [[[NSApp mainMenu] itemWithTitle:HelpTitle.ToString().GetNSString()] retain];
-	for (int32 Index = NumItems - 1; Index > 0; Index--)
-	{
-		[[NSApp mainMenu] removeItemAtIndex:Index];
-	}
-	
-	const TArray<TSharedRef<const FMultiBlock> >& MenuBlocks = MultiBox->GetBlocks();
-	
-	for (int32 Index = 0; Index < MenuBlocks.Num(); Index++)
-	{
-		TSharedRef<const FMenuEntryBlock> Block = StaticCastSharedRef<const FMenuEntryBlock>(MenuBlocks[Index]);
-		FMacMenu* Menu = [[FMacMenu alloc] initWithMenuEntryBlock:Block];
+	MainThreadCall(^{
+		int32 NumItems = [[NSApp mainMenu] numberOfItems];
+		FText HelpTitle = NSLOCTEXT("MainMenu", "HelpMenu", "Help");
+		NSMenuItem* HelpMenu = [[[NSApp mainMenu] itemWithTitle:HelpTitle.ToString().GetNSString()] retain];
+		for (int32 Index = NumItems - 1; Index > 0; Index--)
+		{
+			[[NSApp mainMenu] removeItemAtIndex:Index];
+		}
 		
-		NSMenuItem* MenuItem = [[NSMenuItem new] autorelease];
-		[MenuItem setTitle:[Menu title]];
-		[[NSApp mainMenu] addItem:MenuItem];
-		[MenuItem setSubmenu:Menu];
-	}
-	
-	if([[NSApp mainMenu] itemWithTitle:HelpTitle.ToString().GetNSString()] == nil && HelpMenu)
-	{
-		[[NSApp mainMenu] addItem:HelpMenu];
-	}
-	if(HelpMenu)
-	{
-		[HelpMenu release];
-	}
+		const TArray<TSharedRef<const FMultiBlock> >& MenuBlocks = MultiBox->GetBlocks();
+		
+		for (int32 Index = 0; Index < MenuBlocks.Num(); Index++)
+		{
+			TSharedRef<const FMenuEntryBlock> Block = StaticCastSharedRef<const FMenuEntryBlock>(MenuBlocks[Index]);
+			FMacMenu* Menu = [[FMacMenu alloc] initWithMenuEntryBlock:Block];
+			
+			NSMenuItem* MenuItem = [[NSMenuItem new] autorelease];
+			[MenuItem setTitle:[Menu title]];
+			[[NSApp mainMenu] addItem:MenuItem];
+			[MenuItem setSubmenu:Menu];
+		}
+		
+		if([[NSApp mainMenu] itemWithTitle:HelpTitle.ToString().GetNSString()] == nil && HelpMenu)
+		{
+			[[NSApp mainMenu] addItem:HelpMenu];
+		}
+		if(HelpMenu)
+		{
+			[HelpMenu release];
+		}
+	}, true);
 }
 
 static const TSharedRef<SWidget> FindTextBlockWidget(TSharedRef<SWidget> Content)
@@ -206,225 +212,246 @@ static const TSharedRef<SWidget> FindTextBlockWidget(TSharedRef<SWidget> Content
 
 CFStringRef FSlateMacMenu::GetMenuItemTitle(const TSharedRef<const FMenuEntryBlock>& Block)
 {
-	TAttribute<FText> Label;
-	if (!Block->LabelOverride.IsBound() && Block->LabelOverride.Get().IsEmpty() && Block->GetAction().IsValid())
-	{
-		Label = Block->GetAction()->GetLabel();
-	}
-	else if (!Block->LabelOverride.Get().IsEmpty())
-	{
-		Label = Block->LabelOverride;
-	}
-	else if (Block->EntryWidget.IsValid())
-	{
-		const TSharedRef<SWidget>& TextBlockWidget = FindTextBlockWidget(Block->EntryWidget.ToSharedRef());
-		if (TextBlockWidget != SNullWidget::NullWidget)
+	return GameThreadReturn(^{
+		TAttribute<FText> Label;
+		if (!Block->LabelOverride.IsBound() && Block->LabelOverride.Get().IsEmpty() && Block->GetAction().IsValid())
 		{
-			Label = FText::FromString(StaticCastSharedRef<STextBlock>(TextBlockWidget)->GetText());
+			Label = Block->GetAction()->GetLabel();
+		}
+		else if (!Block->LabelOverride.Get().IsEmpty())
+		{
+			Label = Block->LabelOverride;
+		}
+		else if (Block->EntryWidget.IsValid())
+		{
+			const TSharedRef<SWidget>& TextBlockWidget = FindTextBlockWidget(Block->EntryWidget.ToSharedRef());
+			if (TextBlockWidget != SNullWidget::NullWidget)
+			{
+				Label = FText::FromString(StaticCastSharedRef<STextBlock>(TextBlockWidget)->GetText());
+			}
 		}
-	}
 
-	CFStringRef Title = FPlatformString::TCHARToCFString(*Label.Get().ToString());
-	return Title;
+		CFStringRef Title = FPlatformString::TCHARToCFString(*Label.Get().ToString());
+		return Title;
+	});
 }
 
 NSImage* FSlateMacMenu::GetMenuItemIcon(const TSharedRef<const FMenuEntryBlock>& Block)
 {
-	NSImage* MenuImage = nil;
-	FSlateIcon Icon;
-	if(Block->IconOverride.IsSet())
-	{
-		Icon = Block->IconOverride;
-	}
-	else if(Block->GetAction().IsValid() && Block->GetAction()->GetIcon().IsSet())
-	{
-		Icon = Block->GetAction()->GetIcon();
-	}
-	if(Icon.IsSet())
-	{
-		if(Icon.GetIcon())
+	return GameThreadReturn(^{
+		NSImage* MenuImage = nil;
+		FSlateIcon Icon;
+		if(Block->IconOverride.IsSet())
+		{
+			Icon = Block->IconOverride;
+		}
+		else if(Block->GetAction().IsValid() && Block->GetAction()->GetIcon().IsSet())
 		{
-			FSlateBrush const* IconBrush = Icon.GetIcon();
-			FName ResourceName = IconBrush->GetResourceName();
-			MenuImage = [[NSImage alloc] initWithContentsOfFile: ResourceName.ToString().GetNSString()];
-			if(MenuImage)
+			Icon = Block->GetAction()->GetIcon();
+		}
+		if(Icon.IsSet())
+		{
+			if(Icon.GetIcon())
 			{
-				[MenuImage setSize:NSMakeSize(16.0f, 16.0f)];
+				FSlateBrush const* IconBrush = Icon.GetIcon();
+				FName ResourceName = IconBrush->GetResourceName();
+				MenuImage = [[NSImage alloc] initWithContentsOfFile: ResourceName.ToString().GetNSString()];
+				if(MenuImage)
+				{
+					[MenuImage setSize:NSMakeSize(16.0f, 16.0f)];
+				}
 			}
 		}
-	}
-	return MenuImage;
+		return MenuImage;
+	});
 }
 
 void FSlateMacMenu::ExecuteMenuItemAction(const TSharedRef< const class FMenuEntryBlock >& Block)
 {
-	TSharedPtr< const FUICommandList > ActionList = Block->GetActionList();
-	if (ActionList.IsValid() && Block->GetAction().IsValid())
-	{
-		ActionList->ExecuteAction(Block->GetAction().ToSharedRef());
-	}
-	else
-	{
-		// There is no action list or action associated with this block via a UI command.  Execute any direct action we have
-		Block->GetDirectActions().Execute();
-	}
-}
-
-CFStringRef FSlateMacMenu::GetMenuItemKeyEquivalent(const TSharedRef<const class FMenuEntryBlock>& Block, uint32* OutModifiers)
-{
-	if (Block->GetAction().IsValid())
-	{
-		const TSharedRef<const FInputGesture>& Gesture = Block->GetAction()->GetActiveGesture();
-
-		*OutModifiers = 0;
-		if (Gesture->bCtrl)
+	GameThreadCall(^{
+		TSharedPtr< const FUICommandList > ActionList = Block->GetActionList();
+		if (ActionList.IsValid() && Block->GetAction().IsValid())
 		{
-			*OutModifiers |= NSCommandKeyMask;
+			ActionList->ExecuteAction(Block->GetAction().ToSharedRef());
 		}
-		if (Gesture->bShift)
-		{
-			*OutModifiers |= NSShiftKeyMask;
-		}
-		if (Gesture->bAlt)
+		else
 		{
-			*OutModifiers |= NSAlternateKeyMask;
+			// There is no action list or action associated with this block via a UI command.  Execute any direct action we have
+			Block->GetDirectActions().Execute();
 		}
-		if (Gesture->bCmd)
+	}, false);
+}
+
+CFStringRef FSlateMacMenu::GetMenuItemKeyEquivalent(const TSharedRef<const class FMenuEntryBlock>& Block, uint32* OutModifiers)
+{
+	return GameThreadReturn(^{
+		if (Block->GetAction().IsValid())
 		{
-			*OutModifiers |= NSControlKeyMask;
-		}
+			const TSharedRef<const FInputGesture>& Gesture = Block->GetAction()->GetActiveGesture();
 
-		FString KeyString = Gesture->GetKeyText().ToString().ToLower();
-		return FPlatformString::TCHARToCFString(*KeyString);
-	}
-	return FPlatformString::TCHARToCFString(TEXT(""));
+			*OutModifiers = 0;
+			if (Gesture->bCtrl)
+			{
+				*OutModifiers |= NSCommandKeyMask;
+			}
+			if (Gesture->bShift)
+			{
+				*OutModifiers |= NSShiftKeyMask;
+			}
+			if (Gesture->bAlt)
+			{
+				*OutModifiers |= NSAlternateKeyMask;
+			}
+			if (Gesture->bCmd)
+			{
+				*OutModifiers |= NSControlKeyMask;
+			}
+
+			FString KeyString = Gesture->GetKeyText().ToString().ToLower();
+			return FPlatformString::TCHARToCFString(*KeyString);
+		}
+		return FPlatformString::TCHARToCFString(TEXT(""));
+	});
 }
 
 void FSlateMacMenu::UpdateMenu(FMacMenu* Menu)
 {
-	FText WindowLabel = NSLOCTEXT("MainMenu", "WindowMenu", "Window");
-	
-	bool bIsWindowMenu = (WindowLabel.ToString().Compare(FString([Menu title])) == 0);
-	
-	[Menu removeAllItems];
-	
-	if(bIsWindowMenu)
-	{
-		NSMenuItem* MinimizeItem = [[[NSMenuItem alloc] initWithTitle:@"Minimize" action:@selector(miniaturize:) keyEquivalent:@"m"] autorelease];
-		NSMenuItem* ZoomItem = [[[NSMenuItem alloc] initWithTitle:@"Zoom" action:@selector(performZoom:) keyEquivalent:@""] autorelease];
-		NSMenuItem* CloseItem = [[[NSMenuItem alloc] initWithTitle:@"Close" action:@selector(performClose:) keyEquivalent:@"w"] autorelease];
-		NSMenuItem* BringAllToFrontItem = [[[NSMenuItem alloc] initWithTitle:@"Bring All to Front" action:@selector(arrangeInFront:) keyEquivalent:@""] autorelease];
+	GameThreadCall(^{
+		FText WindowLabel = NSLOCTEXT("MainMenu", "WindowMenu", "Window");
 		
-		[Menu addItem:MinimizeItem];
-		[Menu addItem:ZoomItem];
-		[Menu addItem:CloseItem];
-		[Menu addItem:[NSMenuItem separatorItem]];
-		[Menu addItem:BringAllToFrontItem];
-		[Menu addItem:[NSMenuItem separatorItem]];
-	}
-
-	TSharedRef< SWidget > Widget = GetMenuBuilderWidget(Menu);
+		bool bIsWindowMenu = (WindowLabel.ToString().Compare(FString([Menu title])) == 0);
+		
+		[Menu removeAllItems];
+		
+		if(bIsWindowMenu)
+		{
+			NSMenuItem* MinimizeItem = [[[NSMenuItem alloc] initWithTitle:@"Minimize" action:@selector(miniaturize:) keyEquivalent:@"m"] autorelease];
+			NSMenuItem* ZoomItem = [[[NSMenuItem alloc] initWithTitle:@"Zoom" action:@selector(performZoom:) keyEquivalent:@""] autorelease];
+			NSMenuItem* CloseItem = [[[NSMenuItem alloc] initWithTitle:@"Close" action:@selector(performClose:) keyEquivalent:@"w"] autorelease];
+			NSMenuItem* BringAllToFrontItem = [[[NSMenuItem alloc] initWithTitle:@"Bring All to Front" action:@selector(arrangeInFront:) keyEquivalent:@""] autorelease];
+			
+			[Menu addItem:MinimizeItem];
+			[Menu addItem:ZoomItem];
+			[Menu addItem:CloseItem];
+			[Menu addItem:[NSMenuItem separatorItem]];
+			[Menu addItem:BringAllToFrontItem];
+			[Menu addItem:[NSMenuItem separatorItem]];
+		}
 
-	TSharedRef<const FMultiBox> MultiBox = StaticCastSharedRef<SMultiBoxWidget>(Widget)->GetMultiBox();
-	const TArray<TSharedRef<const FMultiBlock>>& MenuBlocks = MultiBox->GetBlocks();
+		TSharedRef< SWidget > Widget = GetMenuBuilderWidget(Menu);
 
-	for (int32 Index = 0; Index < MenuBlocks.Num(); Index++)
-	{
-		EMultiBlockType::Type Type = MenuBlocks[Index]->GetType();
+		TSharedRef<const FMultiBox> MultiBox = StaticCastSharedRef<SMultiBoxWidget>(Widget)->GetMultiBox();
+		const TArray<TSharedRef<const FMultiBlock>>& MenuBlocks = MultiBox->GetBlocks();
 
-		if (Type == EMultiBlockType::MenuEntry)
+		for (int32 Index = 0; Index < MenuBlocks.Num(); Index++)
 		{
-			TSharedRef<const FMenuEntryBlock> Block = StaticCastSharedRef<const FMenuEntryBlock>(MenuBlocks[Index]);
-			FMacMenuItem* MenuItem = [[[FMacMenuItem alloc] initWithMenuEntryBlock:Block] autorelease];
+			EMultiBlockType::Type Type = MenuBlocks[Index]->GetType();
 
-			if (Block->bIsSubMenu)
+			if (Type == EMultiBlockType::MenuEntry)
 			{
-				FMacMenu* SubMenu = [[FMacMenu alloc] initWithMenuEntryBlock:Block];
-				[MenuItem setSubmenu:SubMenu];
-			}
-			else
-			{
-				[MenuItem setState:FSlateMacMenu::GetMenuItemState(Block)];
-			}
-			
-			if(!bIsWindowMenu)
-			{
-				bool bIsWindowSubMenu = false;
-				NSMenu* SuperMenu = [Menu supermenu];
-				while(!bIsWindowSubMenu && SuperMenu)
+				TSharedRef<const FMenuEntryBlock> Block = StaticCastSharedRef<const FMenuEntryBlock>(MenuBlocks[Index]);
+				FMacMenuItem* MenuItem = [[[FMacMenuItem alloc] initWithMenuEntryBlock:Block] autorelease];
+
+				if (Block->bIsSubMenu)
 				{
-					bIsWindowSubMenu = (WindowLabel.ToString().Compare(FString([SuperMenu title])) == 0);
-					SuperMenu = [SuperMenu supermenu];
+					FMacMenu* SubMenu = [[FMacMenu alloc] initWithMenuEntryBlock:Block];
+					[MenuItem setSubmenu:SubMenu];
+				}
+				else
+				{
+					[MenuItem setState:FSlateMacMenu::GetMenuItemState(Block)];
 				}
 				
-				if(!bIsWindowSubMenu)
+				if(!bIsWindowMenu)
 				{
-					[MenuItem setImage:nil];
+					bool bIsWindowSubMenu = false;
+					NSMenu* SuperMenu = [Menu supermenu];
+					while(!bIsWindowSubMenu && SuperMenu)
+					{
+						bIsWindowSubMenu = (WindowLabel.ToString().Compare(FString([SuperMenu title])) == 0);
+						SuperMenu = [SuperMenu supermenu];
+					}
+					
+					if(!bIsWindowSubMenu)
+					{
+						[MenuItem setImage:nil];
+					}
 				}
-			}
 
-			[Menu addItem:MenuItem];
+				[Menu addItem:MenuItem];
+			}
+			else if (Type == EMultiBlockType::MenuSeparator)
+			{
+				[Menu addItem:[NSMenuItem separatorItem]];
+			}
 		}
-		else if (Type == EMultiBlockType::MenuSeparator)
+		
+		if(bIsWindowMenu)
 		{
-			[Menu addItem:[NSMenuItem separatorItem]];
+			[NSApp setWindowsMenu:Menu];
 		}
-	}
-	
-	if(bIsWindowMenu)
-	{
-		[NSApp setWindowsMenu:Menu];
-	}
+	}, true);
 }
 
 bool FSlateMacMenu::IsMenuItemEnabled(const TSharedRef<const class FMenuEntryBlock>& Block)
 {
-	TSharedPtr<const FUICommandList> ActionList = Block->GetActionList();
-	TSharedPtr<const FUICommandInfo> Action = Block->GetAction();
-	const FUIAction& DirectActions = Block->GetDirectActions();
+	bool bEnabled = GameThreadReturn(^{
+		TSharedPtr<const FUICommandList> ActionList = Block->GetActionList();
+		TSharedPtr<const FUICommandInfo> Action = Block->GetAction();
+		const FUIAction& DirectActions = Block->GetDirectActions();
 
-	bool bEnabled = true;
-	if (ActionList.IsValid() && Action.IsValid())
-	{
-		bEnabled = ActionList->CanExecuteAction(Action.ToSharedRef());
-	}
-	else
-	{
-		// There is no action list or action associated with this block via a UI command.  Execute any direct action we have
-		bEnabled = DirectActions.CanExecute();
-	}
+		bool bEnabled = true;
+		if (ActionList.IsValid() && Action.IsValid())
+		{
+			bEnabled = ActionList->CanExecuteAction(Action.ToSharedRef());
+		}
+		else
+		{
+			// There is no action list or action associated with this block via a UI command.  Execute any direct action we have
+			bEnabled = DirectActions.CanExecute();
+		}
+		
+		return bEnabled;
+	});
 
 	return bEnabled;
 }
 
 int32 FSlateMacMenu::GetMenuItemState(const TSharedRef<const class FMenuEntryBlock>& Block)
 {
-	TSharedPtr<const FUICommandList> ActionList = Block->GetActionList();
-	TSharedPtr<const FUICommandInfo> Action = Block->GetAction();
-	const FUIAction& DirectActions = Block->GetDirectActions();
+	bool bIsChecked = GameThreadReturn(^{
+		bool bIsChecked = false;
+		TSharedPtr<const FUICommandList> ActionList = Block->GetActionList();
+		TSharedPtr<const FUICommandInfo> Action = Block->GetAction();
+		const FUIAction& DirectActions = Block->GetDirectActions();
 
-	bool bIsChecked = true;
-	if (ActionList.IsValid() && Action.IsValid())
-	{
-		bIsChecked = ActionList->IsChecked(Action.ToSharedRef());
-	}
-	else
-	{
-		// There is no action list or action associated with this block via a UI command.  Execute any direct action we have
-		bIsChecked = DirectActions.IsChecked();
-	}
+		if (ActionList.IsValid() && Action.IsValid())
+		{
+			bIsChecked = ActionList->IsChecked(Action.ToSharedRef());
+		}
+		else
+		{
+			// There is no action list or action associated with this block via a UI command.  Execute any direct action we have
+			bIsChecked = DirectActions.IsChecked();
+		}
+		
+		return bIsChecked;
+	});
 
 	return bIsChecked ? NSOnState : NSOffState;
 }
 
 TSharedRef<SWidget> FSlateMacMenu::GetMenuBuilderWidget(FMacMenu* Menu)
 {
-	// @todo: ideally this should be cached and updated only when the menu contents change. See menuHasKeyEquivalent.
-	const bool bShouldCloseWindowAfterMenuSelection = true;
-	FMenuBuilder MenuBuilder(bShouldCloseWindowAfterMenuSelection, [Menu GetMenuEntryBlock]->GetActionList(), [Menu GetMenuEntryBlock]->Extender);
-	{
-		// Have the menu fill its contents
-		[Menu GetMenuEntryBlock]->EntryBuilder.ExecuteIfBound(MenuBuilder);
-	}
-	return MenuBuilder.MakeWidget();
+	return GameThreadReturn(^{
+		// @todo: ideally this should be cached and updated only when the menu contents change. See menuHasKeyEquivalent.
+		const bool bShouldCloseWindowAfterMenuSelection = true;
+		FMenuBuilder MenuBuilder(bShouldCloseWindowAfterMenuSelection, [Menu GetMenuEntryBlock]->GetActionList(), [Menu GetMenuEntryBlock]->Extender);
+		{
+			// Have the menu fill its contents
+			[Menu GetMenuEntryBlock]->EntryBuilder.ExecuteIfBound(MenuBuilder);
+		}
+		TSharedPtr<SWidget> Widget = MenuBuilder.MakeWidget();
+		return Widget;
+	}).ToSharedRef();
 }
