Add minimum desired width properties to a few text based widgets.

diff --git a/Engine/Source/Runtime/Slate/Private/Widgets/Text/STextBlock.cpp b/Engine/Source/Runtime/Slate/Private/Widgets/Text/STextBlock.cpp
--- a/Engine/Source/Runtime/Slate/Private/Widgets/Text/STextBlock.cpp
+++ b/Engine/Source/Runtime/Slate/Private/Widgets/Text/STextBlock.cpp
@@ -18,6 +18,7 @@ void STextBlock::Construct( const FArguments& InArgs )
 	HighlightText = InArgs._HighlightText;
 	WrapTextAt = InArgs._WrapTextAt;
 	AutoWrapText = InArgs._AutoWrapText;
+	MinDesiredWidth = InArgs._MinDesiredWidth;
 
 	CachedWrapTextWidth = 0.0f;
 	CachedAutoWrapTextWidth = 0.0f;
@@ -272,6 +273,8 @@ void STextBlock::CacheDesiredSize()
 			TextMeasurement.X = FMath::Min(TextMeasurement.X, WrappingWidth);
 		}
 
+		TextMeasurement.X = FMath::Max(TextMeasurement.X, MinDesiredWidth.Get());
+
 		const FVector2D CurrentShadowOffset = GetShadowOffset();
 		const FVector2D AbsoluteShadowOffset(FMath::Abs(CurrentShadowOffset.X), FMath::Abs(CurrentShadowOffset.Y));
 		this->Advanced_SetDesiredSize(TextMeasurement + AbsoluteShadowOffset);
diff --git a/Engine/Source/Runtime/Slate/Public/Widgets/Input/SNumericEntryBox.h b/Engine/Source/Runtime/Slate/Public/Widgets/Input/SNumericEntryBox.h
--- a/Engine/Source/Runtime/Slate/Public/Widgets/Input/SNumericEntryBox.h
+++ b/Engine/Source/Runtime/Slate/Public/Widgets/Input/SNumericEntryBox.h
@@ -38,6 +38,7 @@ public:
 		, _MinSliderValue(0)				
 		, _MaxSliderValue(100)
 		, _SliderExponent(1.f)
+		, _MinDesiredValueWidth(0)
 	{}		
 
 		/** Style to use for the editable text box within this widget */
@@ -73,6 +74,8 @@ public:
 		SLATE_ATTRIBUTE( TOptional< NumericType >, MaxSliderValue )
 		/** Use exponential scale for the slider */
 		SLATE_ATTRIBUTE( float, SliderExponent )
+		/** The minimum desired width for the value portion of the control. */
+		SLATE_ATTRIBUTE( float, MinDesiredValueWidth )
 		/** Called whenever the text is changed interactively by the user */
 		SLATE_EVENT( FOnValueChanged, OnValueChanged )
 		/** Called whenever the text is committed.  This happens when the user presses enter or the text box loses focus. */
@@ -98,6 +101,7 @@ public:
 		OnValueCommitted = InArgs._OnValueCommitted;
 		ValueAttribute = InArgs._Value;
 		UndeterminedString = InArgs._UndeterminedString;
+		MinDesiredValueWidth = InArgs._MinDesiredValueWidth;
 		
 		BorderImageNormal = &InArgs._EditableTextBoxStyle->BackgroundImageNormal;
 		BorderImageHovered = &InArgs._EditableTextBoxStyle->BackgroundImageHovered;
@@ -121,7 +125,8 @@ public:
 				.MinValue(InArgs._MinValue)
 				.SliderExponent(InArgs._SliderExponent)
 				.OnBeginSliderMovement(InArgs._OnBeginSliderMovement)
-				.OnEndSliderMovement(InArgs._OnEndSliderMovement);
+				.OnEndSliderMovement(InArgs._OnEndSliderMovement)
+				.MinDesiredWidth(InArgs._MinDesiredValueWidth);
 		}
 
 		// Always create an editable text box.  In the case of an undetermined value being passed in, we cant use the spinbox.
@@ -134,7 +139,8 @@ public:
 			.OnTextChanged( this, &SNumericEntryBox<NumericType>::OnTextChanged  )
 			.OnTextCommitted( this, &SNumericEntryBox<NumericType>::OnTextCommitted )
 			.SelectAllTextOnCommit( true )
-			.ContextMenuExtender( InArgs._ContextMenuExtender );
+			.ContextMenuExtender( InArgs._ContextMenuExtender )
+			.MinDesiredWidth(InArgs._MinDesiredValueWidth);
 
 		TSharedRef<SHorizontalBox> HorizontalBox = SNew( SHorizontalBox );
 		if( InArgs._Label.Widget != SNullWidget::NullWidget )
@@ -452,6 +458,8 @@ private:
 	const FSlateBrush* BorderImageHovered;
 	/** Styling: border image to draw when focused */
 	const FSlateBrush* BorderImageFocused;
+	/** Prevents the value portion of the control from being smaller than desired in certain cases. */
+	TAttribute<float> MinDesiredValueWidth;
 };
 
 
diff --git a/Engine/Source/Runtime/Slate/Public/Widgets/Input/SSpinBox.h b/Engine/Source/Runtime/Slate/Public/Widgets/Input/SSpinBox.h
--- a/Engine/Source/Runtime/Slate/Public/Widgets/Input/SSpinBox.h
+++ b/Engine/Source/Runtime/Slate/Public/Widgets/Input/SSpinBox.h
@@ -33,6 +33,7 @@ public:
 		, _OnValueCommitted()
 		, _ClearKeyboardFocusOnCommit( false )
 		, _SelectAllTextOnCommit( true )
+		, _MinDesiredWidth(0.0f)
 		{}
 	
 		/** The style used to draw this spinbox */
@@ -68,6 +69,8 @@ public:
 		SLATE_ATTRIBUTE( bool, ClearKeyboardFocusOnCommit )
 		/** Whether to select all text when pressing enter to commit changes */
 		SLATE_ATTRIBUTE( bool, SelectAllTextOnCommit )
+		/** Minimum width that a spin box should be */
+		SLATE_ATTRIBUTE( float, MinDesiredWidth )
 
 	SLATE_END_ARGS()
 
@@ -90,6 +93,8 @@ public:
 		OnValueCommitted = InArgs._OnValueCommitted;
 		OnBeginSliderMovement = InArgs._OnBeginSliderMovement;
 		OnEndSliderMovement = InArgs._OnEndSliderMovement;
+		MinDesiredWidth = InArgs._MinDesiredWidth;
+		ArrowImageWidth = InArgs._Style->ArrowsImage.ImageSize.X;
 	
 		MinValue = InArgs._MinValue;
 		MaxValue = InArgs._MaxValue;
@@ -130,6 +135,7 @@ public:
 				SAssignNew(TextBlock, STextBlock)
 				.Font(InArgs._Font)
 				.Text( this, &SSpinBox<NumericType>::GetValueAsString )
+				.MinDesiredWidth( this, &SSpinBox<NumericType>::GetTextMinDesiredWidth )
 			]
 			+SHorizontalBox::Slot()
 			.FillWidth(1.0f)
@@ -146,6 +152,7 @@ public:
 				.OnTextCommitted( this, &SSpinBox<NumericType>::TextField_OnTextCommitted )
 				.ClearKeyboardFocusOnCommit( InArgs._ClearKeyboardFocusOnCommit )
 				.SelectAllTextOnCommit( InArgs._SelectAllTextOnCommit )
+				.MinDesiredWidth( this, &SSpinBox<NumericType>::GetTextMinDesiredWidth )
 			]			
 			+SHorizontalBox::Slot()
 			.AutoWidth()
@@ -655,6 +662,11 @@ private:
 	NumericType GetMinSliderValue() const { return MinSliderValue.Get().Get(TNumericLimits<NumericType>::Lowest()); }
 	NumericType GetMaxSliderValue() const { return MaxSliderValue.Get().Get(TNumericLimits<NumericType>::Max()); }
 
+	float GetTextMinDesiredWidth() const
+	{
+		return FMath::Max(0.0f, MinDesiredWidth.Get() - ArrowImageWidth);
+	}
+
 	TAttribute<float> SliderExponent;
 
 	bool bDragging;
@@ -673,4 +685,10 @@ private:
 	/** This is the cached value the user believes it to be (usually different due to truncation to an int). Used for identifying 
 		external forces on the spinbox and syncing the internal value to them. Synced when a value is committed to the spinbox. */
 	NumericType CachedExternalValue;
+
+	/** Prevents the spinbox from being smaller than desired in certain cases (e.g. when it is empty) */
+	TAttribute<float> MinDesiredWidth;
+
+	/** The width of the arrow image, used for calculating the min desired width of the internal text controls. */
+	float ArrowImageWidth;
 };
diff --git a/Engine/Source/Runtime/Slate/Public/Widgets/Text/STextBlock.h b/Engine/Source/Runtime/Slate/Public/Widgets/Text/STextBlock.h
--- a/Engine/Source/Runtime/Slate/Public/Widgets/Text/STextBlock.h
+++ b/Engine/Source/Runtime/Slate/Public/Widgets/Text/STextBlock.h
@@ -35,6 +35,7 @@ public:
 		, _HighlightText()
 		, _WrapTextAt(0.0f)
 		, _AutoWrapText(false)
+		, _MinDesiredWidth(0.0f)
 		{}
 
 		/** The text displayed in this text block */
@@ -72,6 +73,9 @@ public:
 			desired size will not be clamped.  This works best in cases where the text block's size is not affecting other widget's layout. */
 		SLATE_ATTRIBUTE( bool, AutoWrapText )
 
+		/** Minimum width that a text block should be */
+		SLATE_ATTRIBUTE( float, MinDesiredWidth )
+
 		/** Called when this text is double clicked */
 		SLATE_EVENT( FOnClicked, OnDoubleClicked )
 
@@ -241,4 +245,7 @@ private:
 
 	/** If a request to cache the string size occurred and we should recache the text */
 	mutable bool bRequestCache;
+
+	/** Prevents the text block from being smaller than desired in certain cases (e.g. when it is empty) */
+	TAttribute<float> MinDesiredWidth;
 };
