UMG - Adding MinDesiredSize support to the UMG textblock widget.

diff --git a/Engine/Source/Runtime/Slate/Private/Widgets/Text/STextBlock.cpp b/Engine/Source/Runtime/Slate/Private/Widgets/Text/STextBlock.cpp
--- a/Engine/Source/Runtime/Slate/Private/Widgets/Text/STextBlock.cpp
+++ b/Engine/Source/Runtime/Slate/Private/Widgets/Text/STextBlock.cpp
@@ -331,3 +331,8 @@ void STextBlock::SetShadowColorAndOpacity(const TAttribute<FLinearColor>& InShad
 {
 	ShadowColorAndOpacity = InShadowColorAndOpacity;
 }
+
+void STextBlock::SetMinDesiredWidth(const TAttribute<float>& InMinDesiredWidth)
+{
+	MinDesiredWidth = InMinDesiredWidth;
+}
diff --git a/Engine/Source/Runtime/Slate/Public/Widgets/Text/STextBlock.h b/Engine/Source/Runtime/Slate/Public/Widgets/Text/STextBlock.h
--- a/Engine/Source/Runtime/Slate/Public/Widgets/Text/STextBlock.h
+++ b/Engine/Source/Runtime/Slate/Public/Widgets/Text/STextBlock.h
@@ -161,6 +161,9 @@ public:
 	/** See ShadowColorAndOpacity attribute */
 	void SetShadowColorAndOpacity(const TAttribute<FLinearColor>& InShadowColorAndOpacity);
 
+	/** See MinDesiredWidth attribute */
+	void SetMinDesiredWidth(const TAttribute<float>& InMinDesiredWidth);
+
 	// SWidget interface
 	virtual int32 OnPaint( const FPaintArgs& Args, const FGeometry& AllottedGeometry, const FSlateRect& MyClippingRect, FSlateWindowElementList& OutDrawElements, int32 LayerId, const FWidgetStyle& InWidgetStyle, bool bParentEnabled ) const override;
 	virtual FReply OnMouseButtonDoubleClick( const FGeometry& InMyGeometry, const FPointerEvent& InMouseEvent ) override;
diff --git a/Engine/Source/Runtime/UMG/Private/Components/TextBlock.cpp b/Engine/Source/Runtime/UMG/Private/Components/TextBlock.cpp
--- a/Engine/Source/Runtime/UMG/Private/Components/TextBlock.cpp
+++ b/Engine/Source/Runtime/UMG/Private/Components/TextBlock.cpp
@@ -99,6 +99,7 @@ void UTextBlock::SynchronizeProperties()
 	MyTextBlock->SetShadowColorAndOpacity(ShadowColorAndOpacityBinding);
 	MyTextBlock->SetAutoWrapText(AutoWrapText);
 	MyTextBlock->SetWrapTextAt(WrapTextAt != 0 ? WrapTextAt : TAttribute<float>());
+	MyTextBlock->SetMinDesiredWidth(MinDesiredWidth);
 
 #if WITH_EDITOR
 	MyEditorTextBlock->SetText(TextBinding);
diff --git a/Engine/Source/Runtime/UMG/Public/Components/TextBlock.h b/Engine/Source/Runtime/UMG/Public/Components/TextBlock.h
--- a/Engine/Source/Runtime/UMG/Public/Components/TextBlock.h
+++ b/Engine/Source/Runtime/UMG/Public/Components/TextBlock.h
@@ -81,6 +81,10 @@ public:
 	UPROPERTY(EditDefaultsOnly, Category=Appearance, AdvancedDisplay)
 	bool AutoWrapText;
 
+	/** The minimum desired size for the text */
+	UPROPERTY(EditDefaultsOnly, Category=Appearance, AdvancedDisplay)
+	float MinDesiredWidth;
+
 	///** Called when this text is double clicked */
 	//SLATE_EVENT(FOnClicked, OnDoubleClicked)
 
