Slate - Exposing a way on the SDPIScaler to set the scale and content post construct.

diff --git a/Engine/Source/Runtime/Slate/Private/Widgets/Layout/SDPIScaler.cpp b/Engine/Source/Runtime/Slate/Private/Widgets/Layout/SDPIScaler.cpp
--- a/Engine/Source/Runtime/Slate/Private/Widgets/Layout/SDPIScaler.cpp
+++ b/Engine/Source/Runtime/Slate/Private/Widgets/Layout/SDPIScaler.cpp
@@ -9,11 +9,12 @@ SDPIScaler::SDPIScaler()
 
 void SDPIScaler::Construct( const FArguments& InArgs )
 {
-	this->ChildSlot
+	ChildSlot
 	[
 		InArgs._Content.Widget
 	];
-	this->DPIScale = InArgs._DPIScale;
+	
+	DPIScale = InArgs._DPIScale;
 }
 
 void SDPIScaler::OnArrangeChildren( const FGeometry& AllottedGeometry, FArrangedChildren& ArrangedChildren ) const
@@ -42,3 +43,16 @@ FChildren* SDPIScaler::GetChildren()
 {
 	return &ChildSlot;
 }
+
+void SDPIScaler::SetContent(TSharedRef<SWidget> InContent)
+{
+	ChildSlot
+	[
+		InContent
+	];
+}
+
+void SDPIScaler::SetDPIScale(TAttribute<float> InDPIScale)
+{
+	DPIScale = InDPIScale;
+}
diff --git a/Engine/Source/Runtime/Slate/Public/Widgets/Layout/SDPIScaler.h b/Engine/Source/Runtime/Slate/Public/Widgets/Layout/SDPIScaler.h
--- a/Engine/Source/Runtime/Slate/Public/Widgets/Layout/SDPIScaler.h
+++ b/Engine/Source/Runtime/Slate/Public/Widgets/Layout/SDPIScaler.h
@@ -36,6 +36,12 @@ public:
 
 	virtual FChildren* GetChildren() override;
 
+	/** See the Content attribute */
+	void SetContent(TSharedRef<SWidget> InContent);
+
+	/** See the DPIScale attribute */
+	void SetDPIScale(TAttribute<float> InDPIScale);
+
 protected:
 
 	/**
