Slate - Removing the mouse capture on the tab, it has no effect.

diff --git a/Engine/Source/Runtime/Slate/Private/Widgets/Docking/SDockTab.cpp b/Engine/Source/Runtime/Slate/Private/Widgets/Docking/SDockTab.cpp
--- a/Engine/Source/Runtime/Slate/Private/Widgets/Docking/SDockTab.cpp
+++ b/Engine/Source/Runtime/Slate/Private/Widgets/Docking/SDockTab.cpp
@@ -37,7 +37,7 @@ FReply SDockTab::OnMouseButtonDown( const FGeometry& MyGeometry, const FPointerE
 				return FReply::Unhandled();
 			}
 
-			return FReply::Handled().CaptureMouse(SharedThis(this)).DetectDrag(SharedThis(this), EKeys::LeftMouseButton);
+			return FReply::Handled().DetectDrag(SharedThis(this), EKeys::LeftMouseButton);
 		}
 		else if ( MouseEvent.GetEffectingButton() == EKeys::MiddleMouseButton )
 		{
