Clicking on a numeric entry box selects all text again

diff --git a/Engine/Source/Runtime/Slate/Private/Framework/Text/TextEditHelper.cpp b/Engine/Source/Runtime/Slate/Private/Framework/Text/TextEditHelper.cpp
--- a/Engine/Source/Runtime/Slate/Private/Framework/Text/TextEditHelper.cpp
+++ b/Engine/Source/Runtime/Slate/Private/Framework/Text/TextEditHelper.cpp
@@ -339,7 +339,7 @@ FReply FTextEditHelper::OnMouseMove( const FGeometry& InMyGeometry, const FPoint
 {
 	FReply Reply = FReply::Unhandled();
 
-	if( TextEditor->IsDragSelecting() && TextEditor->GetWidget()->HasMouseCapture() )
+	if( TextEditor->IsDragSelecting() && TextEditor->GetWidget()->HasMouseCapture() && InMouseEvent.GetCursorDelta() != FVector2D::ZeroVector )
 	{
 		TextEditor->MoveCursor( FMoveCursor::ViaScreenPointer( InMyGeometry.AbsoluteToLocal( InMouseEvent.GetScreenSpacePosition( ) ), InMyGeometry.Scale, ECursorAction::SelectText ) );
 		TextEditor->SetHasDragSelectedSinceFocused( true );
