Slate list scrolling now respects DPI when right-click-drag scrolling or touch scrolling.

diff --git a/Engine/Source/Runtime/Slate/Private/Widgets/Views/STableViewBase.cpp b/Engine/Source/Runtime/Slate/Private/Widgets/Views/STableViewBase.cpp
--- a/Engine/Source/Runtime/Slate/Private/Widgets/Views/STableViewBase.cpp
+++ b/Engine/Source/Runtime/Slate/Private/Widgets/Views/STableViewBase.cpp
@@ -320,7 +320,7 @@ FReply STableViewBase::OnMouseMove( const FGeometry& MyGeometry, const FPointerE
 {	
 	if( MouseEvent.IsMouseButtonDown( EKeys::RightMouseButton ) )
 	{
-		const float ScrollByAmount = MouseEvent.GetCursorDelta().Y;
+		const float ScrollByAmount = MouseEvent.GetCursorDelta().Y / MyGeometry.Scale;
 		// If scrolling with the right mouse button, we need to remember how much we scrolled.
 		// If we did not scroll at all, we will bring up the context menu when the mouse is released.
 		AmountScrolledWhileRightMouseDown += FMath::Abs( ScrollByAmount );
@@ -346,7 +346,7 @@ FReply STableViewBase::OnMouseMove( const FGeometry& MyGeometry, const FPointerE
 			// Check if the mouse has moved.
 			if( AmountScrolled != 0 )
 			{
-				SoftwareCursorPosition.Y += MouseEvent.GetCursorDelta().Y;
+				SoftwareCursorPosition.Y += ScrollByAmount;
 			}
 
 			return Reply;
@@ -423,7 +423,7 @@ FReply STableViewBase::OnTouchMoved( const FGeometry& MyGeometry, const FPointer
 {
 	if (bStartedTouchInteraction)
 	{
-		const float ScrollByAmount = InTouchEvent.GetCursorDelta().Y;
+		const float ScrollByAmount = InTouchEvent.GetCursorDelta().Y / MyGeometry.Scale;
 		AmountScrolledWhileRightMouseDown += FMath::Abs( ScrollByAmount );
 
 		this->InertialScrollManager.AddScrollSample( ScrollByAmount, FPlatformTime::Seconds());
