Refactored the SToolTip widget to be a SCompoundWidget instead of a SBorder. Also fixed a bug with setting the content programatically after the fact.

diff --git a/Engine/Source/Editor/GraphEditor/Private/KismetNodes/SGraphNodeK2Base.cpp b/Engine/Source/Editor/GraphEditor/Private/KismetNodes/SGraphNodeK2Base.cpp
--- a/Engine/Source/Editor/GraphEditor/Private/KismetNodes/SGraphNodeK2Base.cpp
+++ b/Engine/Source/Editor/GraphEditor/Private/KismetNodes/SGraphNodeK2Base.cpp
@@ -194,7 +194,7 @@ TSharedPtr<SToolTip> SGraphNodeK2Base::GetComplexTooltip()
 		+SVerticalBox::Slot()
 			.AutoHeight()
 		[
-			DefaultToolTip->GetContent()
+			DefaultToolTip->GetContentWidget()
 		]
 	];
 
diff --git a/Engine/Source/Editor/Kismet/Private/SBlueprintPalette.cpp b/Engine/Source/Editor/Kismet/Private/SBlueprintPalette.cpp
--- a/Engine/Source/Editor/Kismet/Private/SBlueprintPalette.cpp
+++ b/Engine/Source/Editor/Kismet/Private/SBlueprintPalette.cpp
@@ -1475,7 +1475,7 @@ TSharedPtr<SToolTip> SBlueprintPaletteItem::ConstructToolTipWidget() const
 					SNew(SVerticalBox)
 					+SVerticalBox::Slot()
 					[
-						TooltipWidget->GetContent()
+						TooltipWidget->GetContentWidget()
 					]
 					+SVerticalBox::Slot()
 						.AutoHeight()
diff --git a/Engine/Source/Runtime/Slate/Private/Widgets/SToolTip.cpp b/Engine/Source/Runtime/Slate/Private/Widgets/SToolTip.cpp
--- a/Engine/Source/Runtime/Slate/Private/Widgets/SToolTip.cpp
+++ b/Engine/Source/Runtime/Slate/Private/Widgets/SToolTip.cpp
@@ -15,46 +15,49 @@ float SToolTip::GetToolTipWrapWidth()
 }
 
 
-/**
- * Construct this widget
- *
- * @param	InArgs	The declaration data for this widget
- */
 void SToolTip::Construct( const FArguments& InArgs )
 {
 	TextContent = InArgs._Text;
 	bIsInteractive = InArgs._IsInteractive;
-	if( InArgs._Content.Widget != SNullWidget::NullWidget )
+	Font = InArgs._Font;
+	ColorAndOpacity = InArgs._ColorAndOpacity;
+	TextMargin = InArgs._TextMargin;
+	BorderImage = InArgs._BorderImage;
+	
+	SetContentWidget(InArgs._Content.Widget);
+}
+
+
+void SToolTip::SetContentWidget(const TSharedRef<SWidget>& InContentWidget)
+{
+	if (InContentWidget != SNullWidget::NullWidget)
 	{
-		// Widget content argument takes precedence
-		// overrides the text content.
-		WidgetContent = InArgs._Content.Widget;
+		// Widget content argument takes precedence over the text content.
+		WidgetContent = InContentWidget;
 	}
 
-	TSharedRef<SWidget> ToolTipContent = ( !WidgetContent.IsValid() )
+	ToolTipContent = (!WidgetContent.IsValid())
 		? StaticCastSharedRef<SWidget>(
-			SNew( STextBlock )
-			.Text( TextContent )
-			.Font( InArgs._Font )
-			.ColorAndOpacity( InArgs._ColorAndOpacity )
-			.WrapTextAt_Static( &SToolTip::GetToolTipWrapWidth )
+		SNew(STextBlock)
+		.Text(TextContent)
+		.Font(Font)
+		.ColorAndOpacity(ColorAndOpacity)
+		.WrapTextAt_Static(&SToolTip::GetToolTipWrapWidth)
 		)
-		: InArgs._Content.Widget;
+		: InContentWidget;
 
-	SBorder::Construct(
-		SBorder::FArguments()
-		.BorderImage(InArgs._BorderImage)
-		.Padding(InArgs._TextMargin)
+	ChildSlot
+	[
+		SNew(SBorder)
+		.BorderImage(BorderImage)
+		.Padding(TextMargin)
 		[
-			ToolTipContent
-		]	
-	);
-
-	
+			ToolTipContent.ToSharedRef()
+		]
+	];
 }
 
 
-/** @return True if the tool tip has no content to display at this particular moment (remember, attribute callbacks can change this at any time!) */
 bool SToolTip::IsEmpty() const
 {
 	return !WidgetContent.IsValid() && TextContent.Get().IsEmpty();
diff --git a/Engine/Source/Runtime/Slate/Public/Widgets/SToolTip.h b/Engine/Source/Runtime/Slate/Public/Widgets/SToolTip.h
--- a/Engine/Source/Runtime/Slate/Public/Widgets/SToolTip.h
+++ b/Engine/Source/Runtime/Slate/Public/Widgets/SToolTip.h
@@ -7,7 +7,7 @@
  * Slate tool tip widget
  */
 class SLATE_API SToolTip
-	: public SBorder
+	: public SCompoundWidget
 	, public IToolTip
 {
 public:
@@ -61,11 +61,13 @@ public:
 		return AsShared();
 	}
 
-	virtual const TSharedRef<SWidget>& GetContentWidget( ) override
+	virtual TSharedRef<SWidget> GetContentWidget( ) override
 	{
-		return GetContent();
+		return ToolTipContent.ToSharedRef();
 	}
 
+	virtual void SetContentWidget(const TSharedRef<SWidget>& InContentWidget) override;
+
 	virtual bool IsEmpty( ) const override;
 	virtual bool IsInteractive( ) const override;
 
@@ -81,6 +83,21 @@ private:
 	// Content widget.
 	TWeakPtr<SWidget> WidgetContent;
 
+	// Wrapped content within the widget;
+	TSharedPtr<SWidget> ToolTipContent;
+
+	// Font used for the text displayed (where applicable)
+	TAttribute<FSlateFontInfo> Font;
+	
+	// Color and opacity used for the text displayed (where applicable)
+	TAttribute<FSlateColor> ColorAndOpacity;
+
+	// Margin between the tool tip border and the text content
+	TAttribute<FMargin> TextMargin;
+
+	// The background/border image to display
+	TAttribute<const FSlateBrush*> BorderImage;
+	
 	// Whether the tooltip should be considered interactive.
 	TAttribute<bool> bIsInteractive;
 };
diff --git a/Engine/Source/Runtime/SlateCore/Public/Widgets/IToolTip.h b/Engine/Source/Runtime/SlateCore/Public/Widgets/IToolTip.h
--- a/Engine/Source/Runtime/SlateCore/Public/Widgets/IToolTip.h
+++ b/Engine/Source/Runtime/SlateCore/Public/Widgets/IToolTip.h
@@ -22,7 +22,14 @@ public:
 	 *
 	 * @return The content widget.
 	 */
-	virtual const TSharedRef<SWidget>& GetContentWidget( ) = 0;
+	virtual TSharedRef<SWidget> GetContentWidget( ) = 0;
+
+	/**
+	 * Sets the tool tip's content widget.
+	 *
+	 * @param InContentWidget The new content widget to set.
+	 */
+	virtual void SetContentWidget( const TSharedRef<SWidget>& InContentWidget ) = 0;
 
 	/**
 	 * Checks whether this tool tip has no content to display right now.
