Fix asynchronous dispatch of menu actions on Mac OS X.

diff --git a/Engine/Source/Runtime/Slate/Private/Framework/MultiBox/Mac/MacMenu.cpp b/Engine/Source/Runtime/Slate/Private/Framework/MultiBox/Mac/MacMenu.cpp
--- a/Engine/Source/Runtime/Slate/Private/Framework/MultiBox/Mac/MacMenu.cpp
+++ b/Engine/Source/Runtime/Slate/Private/Framework/MultiBox/Mac/MacMenu.cpp
@@ -268,17 +268,19 @@ NSImage* FSlateMacMenu::GetMenuItemIcon(const TSharedRef<const FMenuEntryBlock>&
 
 void FSlateMacMenu::ExecuteMenuItemAction(const TSharedRef< const class FMenuEntryBlock >& Block)
 {
+    TSharedPtr< const class FMenuEntryBlock>* MenuBlock = new TSharedPtr< const class FMenuEntryBlock>(Block);
 	GameThreadCall(^{
-		TSharedPtr< const FUICommandList > ActionList = Block->GetActionList();
-		if (ActionList.IsValid() && Block->GetAction().IsValid())
+		TSharedPtr< const FUICommandList > ActionList = (*MenuBlock)->GetActionList();
+		if (ActionList.IsValid() && (*MenuBlock)->GetAction().IsValid())
 		{
-			ActionList->ExecuteAction(Block->GetAction().ToSharedRef());
+			ActionList->ExecuteAction((*MenuBlock)->GetAction().ToSharedRef());
 		}
 		else
 		{
 			// There is no action list or action associated with this block via a UI command.  Execute any direct action we have
-			Block->GetDirectActions().Execute();
+			(*MenuBlock)->GetDirectActions().Execute();
 		}
+        delete MenuBlock;
 	}, false);
 }
 
