Fixed EditableText's HintText not respecting DPI scale.

diff --git a/Engine/Source/Runtime/Slate/Private/Widgets/Input/SEditableText.cpp b/Engine/Source/Runtime/Slate/Private/Widgets/Input/SEditableText.cpp
--- a/Engine/Source/Runtime/Slate/Private/Widgets/Input/SEditableText.cpp
+++ b/Engine/Source/Runtime/Slate/Private/Widgets/Input/SEditableText.cpp
@@ -1289,7 +1289,7 @@ int32 SEditableText::OnPaint( const FPaintArgs& Args, const FGeometry& AllottedG
 		FSlateDrawElement::MakeText(
 			OutDrawElements,
 			LayerId + TextLayer,
-			AllottedGeometry.ToPaintGeometry( FVector2D( 0, DrawPositionY ), AllottedGeometry.Size, ScaleInverse ),
+			AllottedGeometry.ToPaintGeometry( FVector2D( 0, DrawPositionY ), AllottedGeometry.Size ),
 			ThisHintText,          // Text
 			FontInfo,              // Font information (font name, size)
 			MyClippingRect,        // Clipping rect
