Fixes to list view overscroll

Inertial scroll was always being cleared when scrolling in the overscroll region. This meant that scrolling out of the region did not correctly sample.
    Solution: ClearInertialScrollVelocity on mouse/touch up (if overscrolling).
Side effect: Inertial scroll still operates even when overscrolling(this was previously also the case, but they were cleared on tick).
    Solution: Add CanUseInertialScroll to control when inertial scroll can occur

This addresses TTP#339241 - iOS style scrolling has odd velocity and drag properties when dragging/velocity past the end of the list

diff --git a/Engine/Source/Runtime/Slate/Private/Widgets/Views/STableViewBase.cpp b/Engine/Source/Runtime/Slate/Private/Widgets/Views/STableViewBase.cpp
--- a/Engine/Source/Runtime/Slate/Private/Widgets/Views/STableViewBase.cpp
+++ b/Engine/Source/Runtime/Slate/Private/Widgets/Views/STableViewBase.cpp
@@ -156,19 +156,19 @@ void STableViewBase::Tick( const FGeometry& AllottedGeometry, const double InCur
 			this->InertialScrollManager.UpdateScrollVelocity(InDeltaTime);
 			const float ScrollVelocity = this->InertialScrollManager.GetScrollVelocity();
 
-			if ( ScrollVelocity != 0.f )
+			if ( ScrollVelocity != 0.f && CanUseInertialScroll(ScrollVelocity) )
 			{
 				this->ScrollBy(AllottedGeometry, ScrollVelocity * InDeltaTime, EAllowOverscroll::Yes);
 			}
 
-			Overscroll.UpdateOverscroll( InDeltaTime );
-
-			if (Overscroll.GetOverscroll() != 0.0f)
+			// If we are currently in overscroll, the list will need refreshing.
+			// Do this before UpdateOverscroll, as that could cause GetOverscroll() to be 0
+			if ( Overscroll.GetOverscroll() != 0.0f )
 			{
-				InertialScrollManager.ClearScrollVelocity();	
-				
 				this->RequestListRefresh();
 			}
+
+			Overscroll.UpdateOverscroll( InDeltaTime );
 		}
 
 		FGeometry PanelGeometry = FindChildGeometry( AllottedGeometry, ItemsPanel.ToSharedRef() );
@@ -291,6 +291,12 @@ FReply STableViewBase::OnMouseButtonUp( const FGeometry& MyGeometry, const FPoin
 {
 	if ( MouseEvent.GetEffectingButton() == EKeys::RightMouseButton )
 	{
+		if ( IsRightClickScrolling() && Overscroll.GetOverscroll() != 0.f )
+		{
+			// Clear the scroll velocity if are overscrolling
+			InertialScrollManager.ClearScrollVelocity();
+		}
+
 		OnRightMouseButtonUp( MouseEvent.GetScreenSpacePosition() );
 
 		FReply Reply = FReply::Handled().ReleaseMouseCapture();
@@ -329,7 +335,11 @@ FReply STableViewBase::OnMouseMove( const FGeometry& MyGeometry, const FPointerE
 		// the mouse and dragging the view?
 		if( IsRightClickScrolling() )
 		{
-			this->InertialScrollManager.AddScrollSample(-ScrollByAmount, FPlatformTime::Seconds());
+			if ( CanUseInertialScroll( -ScrollByAmount ) )
+			{
+				this->InertialScrollManager.AddScrollSample(-ScrollByAmount, FPlatformTime::Seconds());
+			}
+
 			const float AmountScrolled = this->ScrollBy( MyGeometry, -ScrollByAmount, EAllowOverscroll::Yes );
 
 			FReply Reply = FReply::Handled();
@@ -426,7 +436,11 @@ FReply STableViewBase::OnTouchMoved( const FGeometry& MyGeometry, const FPointer
 		const float ScrollByAmount = InTouchEvent.GetCursorDelta().Y / MyGeometry.Scale;
 		AmountScrolledWhileRightMouseDown += FMath::Abs( ScrollByAmount );
 
-		this->InertialScrollManager.AddScrollSample( ScrollByAmount, FPlatformTime::Seconds());
+		if ( CanUseInertialScroll( ScrollByAmount ) )
+		{
+			this->InertialScrollManager.AddScrollSample( ScrollByAmount, FPlatformTime::Seconds());
+		}
+
 		const float AmountScrolled = this->ScrollBy( MyGeometry, ScrollByAmount, EAllowOverscroll::Yes );
 
 		if (AmountScrolledWhileRightMouseDown > SlatePanTriggerDistance)
@@ -446,6 +460,12 @@ FReply STableViewBase::OnTouchMoved( const FGeometry& MyGeometry, const FPointer
 
 FReply STableViewBase::OnTouchEnded( const FGeometry& MyGeometry, const FPointerEvent& InTouchEvent )
 {
+	if ( IsRightClickScrolling() && Overscroll.GetOverscroll() != 0.f )
+	{
+		// Clear the scroll velocity if are overscrolling
+		InertialScrollManager.ClearScrollVelocity();
+	}
+
 	AmountScrolledWhileRightMouseDown = 0;
 	bStartedTouchInteraction = false;
 
@@ -669,6 +689,14 @@ TSharedRef<class SWidget> STableViewBase::GetScrollWidget()
 	return SharedThis(this);
 }
 
+bool STableViewBase::CanUseInertialScroll( float ScrollAmount ) const
+{
+	const auto CurrentOverscroll = Overscroll.GetOverscroll();
+
+	// We allow sampling for the inertial scroll if we are not in the overscroll region,
+	// Or if we are scrolling outwards of the overscroll region
+	return CurrentOverscroll == 0.f || FMath::Sign(CurrentOverscroll) != FMath::Sign(ScrollAmount);
+}
 
 STableViewBase::FOverscroll::FOverscroll()
 : OverscrollAmount( 0.0f )
@@ -681,6 +709,13 @@ float STableViewBase::FOverscroll::ScrollBy( float Delta )
 	const float EasedDelta = Delta / (FMath::Abs(OverscrollAmount/ListConstants::OvershootMax)+1.0f);
 	OverscrollAmount = FMath::Clamp(OverscrollAmount + EasedDelta, -ListConstants::OvershootMax, ListConstants::OvershootMax);
 
+	// Don't allow an interaction to change from positive <-> negative overscroll
+	const bool bCrossedOverscrollBoundary = FMath::Sign(ValueBeforeDeltaApplied) != FMath::Sign(OverscrollAmount);
+	if ( bCrossedOverscrollBoundary && ValueBeforeDeltaApplied != 0.f )
+	{
+		OverscrollAmount = 0.f;
+	}
+
 	return ValueBeforeDeltaApplied - OverscrollAmount;
 }
 
diff --git a/Engine/Source/Runtime/Slate/Public/Widgets/Views/STableViewBase.h b/Engine/Source/Runtime/Slate/Public/Widgets/Views/STableViewBase.h
--- a/Engine/Source/Runtime/Slate/Public/Widgets/Views/STableViewBase.h
+++ b/Engine/Source/Runtime/Slate/Public/Widgets/Views/STableViewBase.h
@@ -249,6 +249,13 @@ protected:
 	/**	Whether the software cursor should be drawn in the viewport */
 	bool bShowSoftwareCursor;
 
+private:
+
+	/** Check whether the current state of the table warrants inertial scroll by the specified amount */
+	bool CanUseInertialScroll( float ScrollAmount ) const;
+
+protected:
+
 	struct SLATE_API FOverscroll
 	{
 	public:
