Added support for deferring painting of logical children.
Popups no longer use the window overlay for showing SMenuAnchor::UseCurrentWindow popups; instead deferred painting is used.
Fixed some issues with computing popup placement when DPI scaling is enabled.

diff --git a/Engine/Source/Runtime/Slate/Private/Framework/Application/SlateApplication.cpp b/Engine/Source/Runtime/Slate/Private/Framework/Application/SlateApplication.cpp
--- a/Engine/Source/Runtime/Slate/Private/Framework/Application/SlateApplication.cpp
+++ b/Engine/Source/Runtime/Slate/Private/Framework/Application/SlateApplication.cpp
@@ -601,7 +601,7 @@ void FSlateApplication::DrawWindowAndChildren( const TSharedRef<SWindow>& Window
 		FGeometry WindowGeometry = WindowToDraw->GetWindowGeometryInWindow();
 		int32 MaxLayerId = 0;
 		{
-			MaxLayerId = WindowToDraw->Paint(
+			MaxLayerId = WindowToDraw->PaintWindow(
 				FPaintArgs(WindowToDraw, *HittestGrid, WindowToDraw->GetPositionInScreen()),
 				WindowGeometry, WindowToDraw->GetClippingRectangleInWindow(),
 				WindowElementList,
diff --git a/Engine/Source/Runtime/Slate/Private/Widgets/Input/SComboButton.cpp b/Engine/Source/Runtime/Slate/Private/Widgets/Input/SComboButton.cpp
--- a/Engine/Source/Runtime/Slate/Private/Widgets/Input/SComboButton.cpp
+++ b/Engine/Source/Runtime/Slate/Private/Widgets/Input/SComboButton.cpp
@@ -15,54 +15,54 @@ void SComboButton::Construct( const FArguments& InArgs )
 	
 	OnGetMenuContent = InArgs._OnGetMenuContent;
 	OnComboBoxOpened = InArgs._OnComboBoxOpened;
-	Method = InArgs._Method;
 	ContentWidgetPtr = InArgs._MenuContent.Widget;
-	ContentScale = InArgs._ContentScale;
-	Placement = InArgs._MenuPlacement;
 	bIsFocusable = InArgs._IsFocusable;
 
 	TSharedPtr<SHorizontalBox> HBox;
 
-	this->ChildSlot
-	.HAlign( InArgs._HAlign )
-	.VAlign( InArgs._VAlign )
-	[
-		SNew(SButton)
-		.ContentPadding(FMargin(1,0))
-		.ButtonStyle(OurButtonStyle)	
-		.ClickMethod( EButtonClickMethod::MouseDown )
-		.OnClicked( this, &SComboButton::OnButtonClicked )
-		.ContentPadding( InArgs._ContentPadding )
-		.ForegroundColor( InArgs._ForegroundColor )
-		.ButtonColorAndOpacity( InArgs._ButtonColorAndOpacity )
-		.IsFocusable( InArgs._IsFocusable )
+	SMenuAnchor::Construct( SMenuAnchor::FArguments()
+		.Placement( InArgs._MenuPlacement )
+		.Method( InArgs._Method )
 		[
-			// Button and down arrow on the right
-			// +-------------------+---+
-			// | Button Content    | v |
-			// +-------------------+---+
-			SAssignNew(HBox, SHorizontalBox)
-			+ SHorizontalBox::Slot()
-			.  Expose( ButtonContentSlot )
-			.  FillWidth(1)
-			.  VAlign(VAlign_Center)
+			SNew( SButton )
+			.ContentPadding( FMargin( 1, 0 ) )
+			.ButtonStyle( OurButtonStyle )
+			.ClickMethod( EButtonClickMethod::MouseDown )
+			.OnClicked( this, &SComboButton::OnButtonClicked )
+			.ContentPadding( InArgs._ContentPadding )
+			.ForegroundColor( InArgs._ForegroundColor )
+			.ButtonColorAndOpacity( InArgs._ButtonColorAndOpacity )
+			.IsFocusable( InArgs._IsFocusable )
 			[
-				InArgs._ButtonContent.Widget
-			]
-			+SHorizontalBox::Slot()
-			.  AutoWidth()
-			.  HAlign(HAlign_Center)
-			.  VAlign(VAlign_Center)
-			.  Padding( InArgs._HasDownArrow ? 2 : 0)
-			[
-				SNew(SImage)
-				. Visibility( InArgs._HasDownArrow ? EVisibility::Visible : EVisibility::Collapsed )
-				. Image(&InArgs._ComboButtonStyle->DownArrowImage)
-				// Inherit tinting from parent
-				. ColorAndOpacity( FSlateColor::UseForeground() )
+				// Button and down arrow on the right
+				// +-------------------+---+
+				// | Button Content    | v |
+				// +-------------------+---+
+				SAssignNew( HBox, SHorizontalBox )
+				+ SHorizontalBox::Slot()
+				.Expose( ButtonContentSlot )
+				.FillWidth( 1 )
+				.HAlign( InArgs._HAlign )
+				.VAlign( InArgs._VAlign )
+				[
+					InArgs._ButtonContent.Widget
+				]
+				+ SHorizontalBox::Slot()
+				.AutoWidth()
+				.HAlign( HAlign_Center )
+				.VAlign( VAlign_Center )
+				.Padding( InArgs._HasDownArrow ? 2 : 0 )
+				[
+					SNew( SImage )
+					.Visibility( InArgs._HasDownArrow ? EVisibility::Visible : EVisibility::Collapsed )
+					.Image( &InArgs._ComboButtonStyle->DownArrowImage )
+					// Inherit tinting from parent
+					. ColorAndOpacity( FSlateColor::UseForeground() )
+				]
 			]
 		]
-	];
+	);
+
 	
 	// The menu that pops up when we press the button.
 	// We keep this content around, and then put it into a new window when we need to pop
diff --git a/Engine/Source/Runtime/Slate/Private/Widgets/Input/SMenuAnchor.cpp b/Engine/Source/Runtime/Slate/Private/Widgets/Input/SMenuAnchor.cpp
--- a/Engine/Source/Runtime/Slate/Private/Widgets/Input/SMenuAnchor.cpp
+++ b/Engine/Source/Runtime/Slate/Private/Widgets/Input/SMenuAnchor.cpp
@@ -10,9 +10,11 @@
  */
 void SMenuAnchor::Construct( const FArguments& InArgs )
 {
-	PopupLayerSlot = nullptr;
+	Children.Add( new FSimpleSlot() );
+	Children.Add( new FSimpleSlot() );
+	
 
-	this->ChildSlot
+	Children[0]
 	.Padding(0)
 	[
 		InArgs._Content.Widget
@@ -24,21 +26,114 @@ void SMenuAnchor::Construct( const FArguments& InArgs )
 	Method = InArgs._Method;
 }
 
-void SMenuAnchor::ComputeMenuPlacement( const FGeometry& AllottedGeometry, const FVector2D& PopupDesiredSize, FVector2D& NewWindowSize, FVector2D& NewPosition )
+
+void SMenuAnchor::Tick( const FGeometry& AllottedGeometry, const double InCurrentTime, const float InDeltaTime )
+{
+	SPanel::Tick(AllottedGeometry, InCurrentTime, InDeltaTime);
+
+	TSharedPtr<SWindow> PopupWindow = PopupWindowPtr.Pin();
+	if ( PopupWindow.IsValid() && Method == CreateNewWindow )
+	{
+		// Figure out where our attached pop-up window should be placed.				
+		const FVector2D PopupContentDesiredSize = PopupWindow->GetContent()->GetDesiredSize();
+		FGeometry PopupGeometry = ComputeMenuPlacement( AllottedGeometry, PopupContentDesiredSize, Placement.Get( ) );
+		const FVector2D NewPosition = PopupGeometry.AbsolutePosition;
+		const FVector2D NewSize = PopupGeometry.GetDrawSize( );
+
+		const FSlateRect NewShape = FSlateRect( NewPosition.X, NewPosition.Y, NewPosition.X + NewSize.X, NewPosition.Y + NewSize.Y );
+
+		// We made a window for showing the popup.
+		// Update the window's position!
+		if( PopupWindow->IsMorphing() )
+		{
+			if( NewShape != PopupWindow->GetMorphTargetShape() )
+			{
+				// Update the target shape
+				PopupWindow->UpdateMorphTargetShape( NewShape );
+				// Set size immediately if not morphing size
+				if(!PopupWindow->IsMorphingSize())
+				{
+					PopupWindow->ReshapeWindow( PopupWindow->GetPositionInScreen(), NewSize );
+				}
+			}
+		}
+		else
+		{
+			const FVector2D WindowPosition = PopupWindow->GetPositionInScreen();
+			const FVector2D WindowSize = PopupWindow->GetSizeInScreen();
+			if ( NewPosition != WindowPosition || NewSize != WindowSize )
+			{
+				PopupWindow->ReshapeWindow( NewShape );
+			}
+		}
+	}
+
+	/** The tick is ending, so the window was not dismissed this tick. */
+	bDismissedThisTick = false;
+}
+
+void SMenuAnchor::OnArrangeChildren( const FGeometry& AllottedGeometry, FArrangedChildren& ArrangedChildren ) const
+{
+	ArrangeSingleChild( AllottedGeometry, ArrangedChildren, Children[0], FVector2D::UnitVector);
+	if ( IsOpen() )
+	{
+		// @todo umg na AllottedGeometry here is in Window-Space when computed via OnPaint.
+		//           The incorrect geometry causes the popup to fail workspace-edge tests.
+		const FGeometry PopupGeometry = ComputeMenuPlacement( AllottedGeometry, Children[1].GetWidget()->GetDesiredSize(), Placement.Get() );
+		ArrangedChildren.AddWidget( FArrangedWidget( Children[1].GetWidget(), PopupGeometry ) );
+	}
+}
+
+FVector2D SMenuAnchor::ComputeDesiredSize() const
+{
+	return Children[0].GetWidget()->GetDesiredSize();
+}
+
+FChildren* SMenuAnchor::GetChildren()
+{
+	return &Children;
+}
+
+int32 SMenuAnchor::OnPaint( const FPaintArgs& Args, const FGeometry& AllottedGeometry, const FSlateRect& MyClippingRect, FSlateWindowElementList& OutDrawElements, int32 LayerId, const FWidgetStyle& InWidgetStyle, bool bParentEnabled ) const
+{
+	FArrangedChildren ArrangedChildren( EVisibility::Visible );
+	ArrangeChildren( AllottedGeometry, ArrangedChildren );
+	
+	// There may be zero elements in this array if our child collapsed/hidden
+	if ( ArrangedChildren.Num() > 0 )
+	{
+		FArrangedWidget& TheChild = ArrangedChildren( 0 );
+
+		const FSlateRect ChildClippingRect = AllottedGeometry.GetClippingRect().IntersectionWith( MyClippingRect );
+
+		LayerId = TheChild.Widget->Paint( Args.WithNewParent( this ), TheChild.Geometry, ChildClippingRect, OutDrawElements, LayerId + 1, InWidgetStyle, ShouldBeEnabled( bParentEnabled ) );
+
+		if ( IsOpen( ) && ArrangedChildren.Num( ) > 1 && !ChildClippingRect.IsEmpty() )
+		{
+			OutDrawElements.QueueDeferredPainting(
+				FSlateWindowElementList::FDeferredPaint( ArrangedChildren[1].Widget, Args, ArrangedChildren[1].Geometry, MyClippingRect, InWidgetStyle, bParentEnabled ) );
+		}
+	}
+
+	return LayerId;
+}
+
+FGeometry SMenuAnchor::ComputeMenuPlacement( const FGeometry& AllottedGeometry, const FVector2D& PopupDesiredSize, EMenuPlacement PlacementMode )
 {
-	const EMenuPlacement PlacementMode = Placement.Get();
+	FVector2D PopupSize = PopupDesiredSize * AllottedGeometry.Scale;
 
 	if ( PlacementMode == MenuPlacement_ComboBox )
 	{
-		NewWindowSize = FVector2D( FMath::Max(AllottedGeometry.Size.X, PopupDesiredSize.X), PopupDesiredSize.Y );
-		NewWindowSize *= AllottedGeometry.Scale;
+		PopupSize = FVector2D( FMath::Max( AllottedGeometry.Size.X, PopupDesiredSize.X ), PopupDesiredSize.Y ) * AllottedGeometry.Scale;
 	}
 
 	const float AbsoluteX = AllottedGeometry.AbsolutePosition.X;
 	const float AbsoluteY = AllottedGeometry.AbsolutePosition.Y;
 	FSlateRect Anchor(AbsoluteX, AbsoluteY, AbsoluteX + AllottedGeometry.Size.X*AllottedGeometry.Scale, AbsoluteY + AllottedGeometry.Size.Y*AllottedGeometry.Scale);
 	EOrientation Orientation = (PlacementMode == MenuPlacement_MenuRight) ? Orient_Horizontal : Orient_Vertical;
-	NewPosition = FSlateApplication::Get().CalculatePopupWindowPosition( Anchor, NewWindowSize, Orientation );
+	const FVector2D NewPositionDesktopSpace = FSlateApplication::Get().CalculatePopupWindowPosition( Anchor, PopupSize, Orientation );
+
+	return FGeometry( FVector2D::ZeroVector, NewPositionDesktopSpace, PopupSize / AllottedGeometry.Scale, AllottedGeometry.Scale );
 }
 
 void SMenuAnchor::RequestClosePopupWindow( const TSharedRef<SWindow>& PopupWindow )
@@ -65,74 +160,14 @@ void SMenuAnchor::OnClickedOutsidePopup()
 	}
 }
 
-void SMenuAnchor::Tick( const FGeometry& AllottedGeometry, const double InCurrentTime, const float InDeltaTime )
-{
-	SCompoundWidget::Tick(AllottedGeometry, InCurrentTime, InDeltaTime);
-
-	TSharedPtr<SWindow> PopupWindow = PopupWindowPtr.Pin();
-	if ( PopupWindow.IsValid() )
-	{
-		// Figure out where our attached pop-up window should be placed.				
-		if (Method == CreateNewWindow)
-		{
-			FVector2D NewPosition = FVector2D::ZeroVector;
-			FVector2D NewSize = PopupWindow->GetSizeInScreen();
-			ComputeMenuPlacement( AllottedGeometry, PopupWindow->GetContent()->GetDesiredSize(), NewSize, NewPosition );
-
-			// We made a window for showing the popup.
-			// Update the window's position!
-			if( PopupWindow->IsMorphing() )
-			{
-				FSlateRect NewShape = FSlateRect(NewPosition.X, NewPosition.Y, NewPosition.X + NewSize.X, NewPosition.Y + NewSize.Y);
-				if( NewShape != PopupWindow->GetMorphTargetShape() )
-				{
-					// Update the target shape
-					PopupWindow->UpdateMorphTargetShape( NewShape );
-					// Set size immediately if not morphing size
-					if(!PopupWindow->IsMorphingSize())
-					{
-						PopupWindow->ReshapeWindow( PopupWindow->GetPositionInScreen(), NewSize );
-					}
-				}
-			}
-			else
-			{
-				if( NewPosition != PopupWindow->GetPositionInScreen() || NewSize != PopupWindow->GetSizeInScreen() )
-				{
-					PopupWindow->ReshapeWindow( NewPosition, NewSize );
-				}
-			}
-		}
-		else
-		{
-			// Using the PopupLayer to show this popup instead of a window.
-			// Update the Popup's position so it appears anchored to the menu anchor.
-			if (ensure(PopupLayerSlot != nullptr))
-			{
-				FVector2D NewPosition = FVector2D::ZeroVector;
-				FVector2D NewSize = PopupLayerSlot->GetWidget()->GetDesiredSize();
-
-				ComputeMenuPlacement( AllottedGeometry, PopupLayerSlot->GetWidget()->GetDesiredSize(), NewSize, NewPosition );
-			
-				(*PopupLayerSlot)
-					.DesktopPosition( NewPosition )
-					.WidthOverride( NewSize.X )
-					.Scale(AllottedGeometry.Scale);
-			}
-		}
-	}
-
-	/** The tick is ending, so the window was not dismissed this tick. */
-	bDismissedThisTick = false;
-}
 
 void SMenuAnchor::SetContent(TSharedRef<SWidget> InContent)
 {
-	ChildSlot
-		.Padding(0)
-		[
-			InContent
-		];
+	Children[0]
+	.Padding(0)
+	[
+		InContent
+	];
 }
 
 void SMenuAnchor::SetMenuContent(TSharedRef<SWidget> InMenuContent)
@@ -179,7 +214,6 @@ void SMenuAnchor::SetIsOpen( bool InIsOpen, const bool bFocusMenu )
 			// other pop-up menus are currently auto-sized
 			const FVector2D DesiredContentSize = MenuContentRef->GetDesiredSize();  // @todo: This is ignoring any window border size!
 			const EMenuPlacement PlacementMode = Placement.Get();
-			bool bShouldAutoSize = (PlacementMode != MenuPlacement_ComboBox);
 
 			const FVector2D NewPosition = MyGeometry.AbsolutePosition;
 			FVector2D NewWindowSize = DesiredContentSize;
@@ -203,7 +237,7 @@ void SMenuAnchor::SetIsOpen( bool InIsOpen, const bool bFocusMenu )
 			if (Method == CreateNewWindow)
 			{
 				// Open the pop-up
-				TSharedRef<SWindow> PopupWindow = FSlateApplication::Get().PushMenu( AsShared(), MenuContentRef, NewPosition, TransitionEffect, bFocusMenu, bShouldAutoSize, NewWindowSize, SummonLocationSize );
+				TSharedRef<SWindow> PopupWindow = FSlateApplication::Get().PushMenu( AsShared(), MenuContentRef, NewPosition, TransitionEffect, bFocusMenu, false, NewWindowSize, SummonLocationSize );
 				PopupWindow->SetRequestDestroyWindowOverride( FRequestDestroyWindowOverride::CreateSP(this, &SMenuAnchor::RequestClosePopupWindow) );
 				PopupWindowPtr = PopupWindow;
 			}
@@ -212,23 +246,13 @@ void SMenuAnchor::SetIsOpen( bool InIsOpen, const bool bFocusMenu )
 				// We are re-using the current window instead of creating a new one.
 				// The popup will be presented via an overlay service.
 				ensure(Method==UseCurrentWindow);
-				TSharedRef<SWindow> PopupWindow = MyWidgetPath.GetWindow();
-				PopupWindowPtr = PopupWindow;
-				
-				TSharedPtr<SWeakWidget> PopupWrapperPtr;
-				PopupLayerSlot = &PopupWindow->AddPopupLayerSlot()
-				.DesktopPosition(NewPosition)
-				.WidthOverride(NewWindowSize.X)
+				Children[1]
 				[
-					SAssignNew( PopupWrapperPtr, SWeakWidget )
-					.PossiblyNullContent
-					(
-						MenuContentRef
-					)
+					MenuContentRef
 				];
 
 				// We want to support dismissing the popup widget when the user clicks outside it.
-				FSlateApplication::Get().GetPopupSupport().RegisterClickNotification( PopupWrapperPtr.ToSharedRef(), FOnClickedOutside::CreateSP(this, &SMenuAnchor::OnClickedOutsidePopup) );
+				FSlateApplication::Get( ).GetPopupSupport( ).RegisterClickNotification( MenuContentRef, FOnClickedOutside::CreateSP( this, &SMenuAnchor::OnClickedOutsidePopup ) );
 			}
 
 			if(bFocusMenu)
@@ -254,12 +278,7 @@ void SMenuAnchor::SetIsOpen( bool InIsOpen, const bool bFocusMenu )
 			{
 				// Close the popup overlay
 				ensure(Method==UseCurrentWindow);
-				check(PopupWindowPtr.IsValid());
-				check(PopupLayerSlot != nullptr);
-				PopupWindowPtr.Pin()->RemovePopupLayerSlot(PopupLayerSlot->GetWidget());
-				// We no longer have a popup open, so reset all the tracking state associated.
-				PopupLayerSlot = nullptr;
-				PopupWindowPtr.Reset();
+				Children[1].DetachWidget();
 			}
 
 			if (OnMenuOpenChanged.IsBound())
@@ -272,7 +291,7 @@ void SMenuAnchor::SetIsOpen( bool InIsOpen, const bool bFocusMenu )
 
 bool SMenuAnchor::IsOpen() const
 {
-	return PopupWindowPtr.Pin().IsValid();
+	return PopupWindowPtr.Pin().IsValid() || Children[1].GetWidget() != SNullWidget::NullWidget;
 }
 
 bool SMenuAnchor::ShouldOpenDueToClick() const
@@ -319,14 +338,8 @@ SMenuAnchor::~SMenuAnchor()
 			// Request that the popup be closed.
 			PopupWindow->RequestDestroyWindow();
 		}
-		// Close the popup overlay
-		else if (Method == UseCurrentWindow && PopupLayerSlot != nullptr)
-		{
-			PopupWindow->RemovePopupLayerSlot(PopupLayerSlot->GetWidget());
-		}
 		
 		// We no longer have a popup open, so reset all the tracking state associated.
 		PopupWindowPtr.Reset();
-		PopupLayerSlot = nullptr;
 	}
 }
diff --git a/Engine/Source/Runtime/Slate/Public/Widgets/Input/SComboButton.h b/Engine/Source/Runtime/Slate/Public/Widgets/Input/SComboButton.h
--- a/Engine/Source/Runtime/Slate/Public/Widgets/Input/SComboButton.h
+++ b/Engine/Source/Runtime/Slate/Public/Widgets/Input/SComboButton.h
@@ -20,11 +20,10 @@ public:
 		, _HasDownArrow(true)
 		, _ForegroundColor(FCoreStyle::Get().GetSlateColor("InvertedForeground"))
 		, _ButtonColorAndOpacity(FLinearColor::White)
-		, _ContentScale(FVector2D(1,1))
 		, _ContentPadding(FMargin(5))
 		, _MenuPlacement(MenuPlacement_ComboBox)
 		, _HAlign(HAlign_Fill)
-		, _VAlign(VAlign_Fill)
+		, _VAlign(VAlign_Center)
 		, _Method( SMenuAnchor::CreateNewWindow )
 		{}
 
@@ -47,7 +46,6 @@ public:
 
 		SLATE_ATTRIBUTE( FSlateColor, ForegroundColor )
 		SLATE_ATTRIBUTE( FSlateColor, ButtonColorAndOpacity )
-		SLATE_ATTRIBUTE( FVector2D, ContentScale )
 		SLATE_ATTRIBUTE( FMargin, ContentPadding )
 		SLATE_ATTRIBUTE( EMenuPlacement, MenuPlacement )
 		SLATE_ARGUMENT( EHorizontalAlignment, HAlign )
diff --git a/Engine/Source/Runtime/Slate/Public/Widgets/Input/SMenuAnchor.h b/Engine/Source/Runtime/Slate/Public/Widgets/Input/SMenuAnchor.h
--- a/Engine/Source/Runtime/Slate/Public/Widgets/Input/SMenuAnchor.h
+++ b/Engine/Source/Runtime/Slate/Public/Widgets/Input/SMenuAnchor.h
@@ -9,7 +9,7 @@ DECLARE_DELEGATE_OneParam(FOnIsOpenChanged, bool)
  * A PopupAnchor summons a Popup relative to its content.
  * Summoning a popup relative to the cursor is accomplished via the application.
  */
-class SLATE_API SMenuAnchor : public SCompoundWidget
+class SLATE_API SMenuAnchor : public SPanel
 {
 public:
 	enum EMethod
@@ -54,10 +54,6 @@ public:
 	
 	SMenuAnchor();
 
-	// SWidget interface
-	virtual void Tick( const FGeometry& AllottedGeometry, const double InCurrentTime, const float InDeltaTime ) override;
-	// End of SWidget interface
-
 	/** See Content Slot attribute */
 	void SetContent(TSharedRef<SWidget> InContent);
 
@@ -86,16 +82,24 @@ public:
 	bool HasOpenSubMenus() const;
 
 protected:
+	// SWidget interface
+	virtual void Tick( const FGeometry& AllottedGeometry, const double InCurrentTime, const float InDeltaTime ) override;
+	virtual void OnArrangeChildren( const FGeometry& AllottedGeometry, FArrangedChildren& ArrangedChildren ) const override;
+	virtual FVector2D ComputeDesiredSize( ) const override;
+	virtual FChildren* GetChildren() override;
+	int32 OnPaint( const FPaintArgs& Args, const FGeometry& AllottedGeometry, const FSlateRect& MyClippingRect, FSlateWindowElementList& OutDrawElements, int32 LayerId, const FWidgetStyle& InWidgetStyle, bool bParentEnabled ) const;
+	// End of SWidget interface
 
 	/**
-	 * Computes the desired position for a menu that's attached to this anchor
+	 * Computes the desired position for a popup that is attached to this anchor.
 	 *
 	 * @param   AllottedGeometry    Our widget's geometry
 	 * @param   PopupDesiredSize    How big the popup needs to be for all of its contents to be visible
-	 * @param   WindowSize          (Out) New size of the window
-	 * @param   NewPosition         (Out) The computed window position
+	 * @param   PlacementMode       Should the popup be below the anchor? Right of it? Some other placement?
+	 *
+	 * @return Geometry computed for the popup widget.
 	 */
-	void ComputeMenuPlacement( const FGeometry& AllottedGeometry, const FVector2D& PopupDesiredSize, FVector2D& WindowSize, FVector2D& NewPosition );
+	static FGeometry ComputeMenuPlacement( const FGeometry& AllottedGeometry, const FVector2D& PopupDesiredSize, EMenuPlacement PlacementMode );
 
 	/** Invoked when the popup window is being closed by the application */
 	void RequestClosePopupWindow( const TSharedRef<SWindow>& PopupWindow );
@@ -107,10 +111,7 @@ protected:
 	/** A pointer to the window presenting this popup. Pointer is null when the popup is not visible */
 	TWeakPtr<SWindow> PopupWindowPtr;
 
-	/** When not making a new window for the popup, we rely on the SPopupLayer to achieve the effect. */
-	struct FPopupLayerSlot* PopupLayerSlot;
-	
-	/** Static menu content */
+	/** Static menu content to use when the delegate used when OnGetMenuContent is not defined. */
 	TSharedPtr<SWidget> MenuContent;
 	
 	/** Callback invoked when the popup is being summoned */
@@ -128,5 +129,5 @@ protected:
 	/** Should we summon a new window for this popup or  */
 	EMethod Method;
 
-
+	TPanelChildren<FSimpleSlot> Children;
 };
diff --git a/Engine/Source/Runtime/SlateCore/Private/Rendering/DrawElements.cpp b/Engine/Source/Runtime/SlateCore/Private/Rendering/DrawElements.cpp
--- a/Engine/Source/Runtime/SlateCore/Private/Rendering/DrawElements.cpp
+++ b/Engine/Source/Runtime/SlateCore/Private/Rendering/DrawElements.cpp
@@ -15,6 +15,7 @@ void FSlateDrawElement::MakeDebugQuad( FSlateWindowElementList& ElementList, uin
 	DrawElt.Scale = PaintGeometry.DrawScale;
 }
 
+
 void FSlateDrawElement::MakeBox( 
 	FSlateWindowElementList& ElementList,
 	uint32 InLayer, 
@@ -35,6 +36,7 @@ void FSlateDrawElement::MakeBox(
 	DrawElt.Scale = PaintGeometry.DrawScale;
 }
 
+
 void FSlateDrawElement::MakeRotatedBox( 
 	FSlateWindowElementList& ElementList,
 	uint32 InLayer, 
@@ -60,6 +62,7 @@ void FSlateDrawElement::MakeRotatedBox(
 	DrawElt.DataPayload.SetRotatedBoxPayloadProperties( InBrush, Angle, RotationPoint, InTint );
 }
 
+
 void FSlateDrawElement::MakeText( FSlateWindowElementList& ElementList, uint32 InLayer, const FPaintGeometry& PaintGeometry, const FString& InText, const int32 StartIndex, const int32 EndIndex, const FSlateFontInfo& InFontInfo, const FSlateRect& InClippingRect,ESlateDrawEffect::Type InDrawEffects, const FLinearColor& InTint )
 {
 	FSlateDrawElement& DrawElt = ElementList.AddUninitialized();
@@ -73,6 +76,7 @@ void FSlateDrawElement::MakeText( FSlateWindowElementList& ElementList, uint32 I
 	DrawElt.Scale = PaintGeometry.DrawScale;
 }
 
+
 void FSlateDrawElement::MakeText( FSlateWindowElementList& ElementList, uint32 InLayer, const FPaintGeometry& PaintGeometry, const FString& InText, const FSlateFontInfo& InFontInfo, const FSlateRect& InClippingRect,ESlateDrawEffect::Type InDrawEffects, const FLinearColor& InTint )
 {
 	FSlateDrawElement& DrawElt = ElementList.AddUninitialized();
@@ -86,6 +90,7 @@ void FSlateDrawElement::MakeText( FSlateWindowElementList& ElementList, uint32 I
 	DrawElt.Scale = PaintGeometry.DrawScale;
 }
 
+
 void FSlateDrawElement::MakeText( FSlateWindowElementList& ElementList, uint32 InLayer, const FPaintGeometry& PaintGeometry, const FText& InText, const FSlateFontInfo& InFontInfo, const FSlateRect& InClippingRect,ESlateDrawEffect::Type InDrawEffects, const FLinearColor& InTint )
 {
 	FSlateDrawElement& DrawElt = ElementList.AddUninitialized();
@@ -99,6 +104,7 @@ void FSlateDrawElement::MakeText( FSlateWindowElementList& ElementList, uint32 I
 	DrawElt.Scale = PaintGeometry.DrawScale;
 }
 
+
 void FSlateDrawElement::MakeGradient( FSlateWindowElementList& ElementList, uint32 InLayer, const FPaintGeometry& InPaintGeometry, TArray<FSlateGradientStop> InGradientStops, EOrientation InGradientType, const FSlateRect& InClippingRect, ESlateDrawEffect::Type InDrawEffects, bool bGammaCorrect )
 {
 	FSlateDrawElement& DrawElt = ElementList.AddUninitialized();
@@ -112,6 +118,7 @@ void FSlateDrawElement::MakeGradient( FSlateWindowElementList& ElementList, uint
 	DrawElt.Scale = InPaintGeometry.DrawScale;
 }
 
+
 void FSlateDrawElement::MakeSpline( FSlateWindowElementList& ElementList, uint32 InLayer, const FPaintGeometry& PaintGeometry, const FVector2D& InStart, const FVector2D& InStartDir, const FVector2D& InEnd, const FVector2D& InEndDir, const FSlateRect InClippingRect, float InThickness, ESlateDrawEffect::Type InDrawEffects, const FLinearColor& InTint )
 {
 	FSlateDrawElement& DrawElt = ElementList.AddUninitialized();
@@ -124,11 +131,13 @@ void FSlateDrawElement::MakeSpline( FSlateWindowElementList& ElementList, uint32
 	DrawElt.Scale = PaintGeometry.DrawScale;
 }
 
+
 void FSlateDrawElement::MakeDrawSpaceSpline( FSlateWindowElementList& ElementList, uint32 InLayer, const FVector2D& InStart, const FVector2D& InStartDir, const FVector2D& InEnd, const FVector2D& InEndDir, const FSlateRect InClippingRect, float InThickness, ESlateDrawEffect::Type InDrawEffects, const FColor& InTint )
 {
 	MakeSpline( ElementList, InLayer, FPaintGeometry::Identity(), InStart, InStartDir, InEnd, InEndDir, InClippingRect, InThickness, InDrawEffects, InTint );
 }
 
+
 void FSlateDrawElement::MakeLines( FSlateWindowElementList& ElementList, uint32 InLayer, const FPaintGeometry& PaintGeometry, const TArray<FVector2D>& Points, const FSlateRect InClippingRect, ESlateDrawEffect::Type InDrawEffects, const FLinearColor& InTint, bool bAntialias )
 {
 	FSlateDrawElement& DrawElt = ElementList.AddUninitialized();
@@ -141,6 +150,7 @@ void FSlateDrawElement::MakeLines( FSlateWindowElementList& ElementList, uint32
 	DrawElt.Scale = PaintGeometry.DrawScale;
 }
 
+
 void FSlateDrawElement::MakeViewport( FSlateWindowElementList& ElementList, uint32 InLayer, const FPaintGeometry& PaintGeometry, TSharedPtr<const ISlateViewport> Viewport, const FSlateRect& InClippingRect, bool bGammaCorrect, bool bAllowBlending, ESlateDrawEffect::Type InDrawEffects, const FLinearColor& InTint )
 {
 	FSlateDrawElement& DrawElt = ElementList.AddUninitialized();
@@ -153,6 +163,8 @@ void FSlateDrawElement::MakeViewport( FSlateWindowElementList& ElementList, uint
 	DrawElt.DrawEffects = InDrawEffects;
 	DrawElt.Scale = PaintGeometry.DrawScale;
 }
+
+
 void FSlateDrawElement::MakeCustom( FSlateWindowElementList& ElementList, uint32 InLayer, TSharedPtr<ICustomSlateElement, ESPMode::ThreadSafe> CustomDrawer )
 {
 	FSlateDrawElement& DrawElt = ElementList.AddUninitialized();
@@ -160,4 +172,46 @@ void FSlateDrawElement::MakeCustom( FSlateWindowElementList& ElementList, uint32
 	DrawElt.Layer = InLayer;
 	DrawElt.DataPayload.SetCustomDrawerPayloadProperties( CustomDrawer );
 	DrawElt.ClippingRect = FSlateRect(1,1,1,1);
-}
\ No newline at end of file
+}
+
+
+FSlateWindowElementList::FDeferredPaint::FDeferredPaint( const TSharedRef<SWidget>& InWidgetToPaint, const FPaintArgs& InArgs, const FGeometry InAllottedGeometry, const FSlateRect InMyClippingRect, const FWidgetStyle& InWidgetStyle, bool InParentEnabled )
+: WidgetToPaintPtr( InWidgetToPaint )
+, Args( InArgs )
+, AllottedGeometry( InAllottedGeometry )
+, MyClippingRect( InMyClippingRect )
+, WidgetStyle( InWidgetStyle )
+, bParentEnabled( InParentEnabled )
+{
+}
+
+
+int32 FSlateWindowElementList::FDeferredPaint::ExecutePaint( int32 LayerId, FSlateWindowElementList& OutDrawElements ) const
+{
+	TSharedPtr<SWidget> WidgetToPaint = WidgetToPaintPtr.Pin();
+	if ( WidgetToPaint.IsValid() )
+	{
+		return WidgetToPaint->Paint( Args, AllottedGeometry, OutDrawElements.GetWindow()->GetClippingRectangleInWindow(), OutDrawElements, LayerId, WidgetStyle, bParentEnabled );
+	}
+
+	return LayerId;
+}
+
+
+void FSlateWindowElementList::QueueDeferredPainting( const FDeferredPaint& InDeferredPaint )
+{
+	DeferredPaintList.Add( MakeShareable( new FDeferredPaint( InDeferredPaint ) ) );
+}
+
+
+int32 FSlateWindowElementList::PaintDeferred( int32 LayerId )
+{
+	for ( const TSharedRef<FDeferredPaint>& PaintArgs : DeferredPaintList )
+	{
+		LayerId = PaintArgs->ExecutePaint( LayerId, *this );
+	}
+
+	return LayerId;
+}
+
+
diff --git a/Engine/Source/Runtime/SlateCore/Private/Widgets/SWindow.cpp b/Engine/Source/Runtime/SlateCore/Private/Widgets/SWindow.cpp
--- a/Engine/Source/Runtime/SlateCore/Private/Widgets/SWindow.cpp
+++ b/Engine/Source/Runtime/SlateCore/Private/Widgets/SWindow.cpp
@@ -1660,7 +1660,7 @@ SWindow::SWindow()
 	, bHasMinimizeButton( false )
 	, bHasMaximizeButton( false )
 	, bHasSizingFrame( false )
-    , bIsModalWindow( false )
+	, bIsModalWindow( false )
 	, InitialDesiredScreenPosition( FVector2D::ZeroVector )
 	, InitialDesiredSize( FVector2D::ZeroVector )
 	, ScreenPosition( FVector2D::ZeroVector )
@@ -1678,6 +1678,14 @@ SWindow::SWindow()
 }
 
 
+int32 SWindow::PaintWindow( const FPaintArgs& Args, const FGeometry& AllottedGeometry, const FSlateRect& MyClippingRect, FSlateWindowElementList& OutDrawElements, int32 LayerId, const FWidgetStyle& InWidgetStyle, bool bParentEnabled ) const
+{
+	LayerId = Paint(Args, AllottedGeometry, MyClippingRect, OutDrawElements, LayerId, InWidgetStyle, bParentEnabled);
+	LayerId = OutDrawElements.PaintDeferred( LayerId );
+	return LayerId;
+}
+
+
 FOptionalSize SWindow::GetTitleBarSize() const
 {
 	return TitleBarSize;
diff --git a/Engine/Source/Runtime/SlateCore/Public/Layout/ArrangedChildren.h b/Engine/Source/Runtime/SlateCore/Public/Layout/ArrangedChildren.h
--- a/Engine/Source/Runtime/SlateCore/Public/Layout/ArrangedChildren.h
+++ b/Engine/Source/Runtime/SlateCore/Public/Layout/ArrangedChildren.h
@@ -97,11 +97,23 @@ class SLATECORE_API FArrangedChildren
 		return Array.Num();
 	}
 
+	const FArrangedWidget& operator[]( int32 Index ) const
+	{
+		return Array[Index];
+	}
+
+	FArrangedWidget& operator[]( int32 Index )
+	{
+		return Array[Index];
+	}
+
+	// @todo umg na deprecate this
 	const FArrangedWidget& operator()(int32 Index) const
 	{
 		return Array[Index];
 	}
 
+	// @todo umg na deprecate this
 	FArrangedWidget& operator()(int32 Index)
 	{
 		return Array[Index];
diff --git a/Engine/Source/Runtime/SlateCore/Public/Rendering/DrawElements.h b/Engine/Source/Runtime/SlateCore/Public/Rendering/DrawElements.h
--- a/Engine/Source/Runtime/SlateCore/Public/Rendering/DrawElements.h
+++ b/Engine/Source/Runtime/SlateCore/Public/Rendering/DrawElements.h
@@ -570,6 +570,7 @@ public:
 class FSlateWindowElementList
 {
 	friend class FSlateElementBatcher;
+
 public:
 	/** 
 	 * Construct a new list of elements with which to paint a window.
@@ -614,6 +615,30 @@ public:
 		const int32 InsertIdx = DrawElements.AddUninitialized();
 		return *(new(DrawElements.GetTypedData() + InsertIdx) FSlateDrawElement());
 	}
+
+	/**
+	* Some widgets may want to paint their children after after another, loosely-related widget finished painting.
+	* Or they may want to paint "after everyone".
+	*/
+	struct FDeferredPaint
+	{
+	public:
+		SLATECORE_API FDeferredPaint( const TSharedRef<SWidget>& InWidgetToPaint, const FPaintArgs& InArgs, const FGeometry InAllottedGeometry, const FSlateRect InMyClippingRect, const FWidgetStyle& InWidgetStyle, bool InParentEnabled );
+
+		int32 ExecutePaint( int32 LayerId, FSlateWindowElementList& OutDrawElements ) const;
+
+	private:
+		const TWeakPtr<SWidget> WidgetToPaintPtr;
+		const FPaintArgs Args;
+		const FGeometry AllottedGeometry;
+		const FSlateRect MyClippingRect;
+		const FWidgetStyle WidgetStyle;
+		const bool bParentEnabled;
+	};
+
+	SLATECORE_API void QueueDeferredPainting( const FDeferredPaint& InDeferredPaint );
+
+	int32 PaintDeferred( int32 LayerId );
 	
 	/**
 	 * Remove all the elements from this draw list.
@@ -656,4 +681,10 @@ protected:
 
 	/** All the indices for every batch. */
 	TArray<SlateIndex> BatchedIndices;
+
+	/**
+	 * Some widgets want their logical children to appear at a different "layer" in the physical hierarchy.
+	 * We accomplish this by deferring their painting.
+	 */
+	TArray< TSharedRef<FDeferredPaint> > DeferredPaintList;
 };
diff --git a/Engine/Source/Runtime/SlateCore/Public/Widgets/SWindow.h b/Engine/Source/Runtime/SlateCore/Public/Widgets/SWindow.h
--- a/Engine/Source/Runtime/SlateCore/Public/Widgets/SWindow.h
+++ b/Engine/Source/Runtime/SlateCore/Public/Widgets/SWindow.h
@@ -207,6 +207,9 @@ public:
 		}
 	}
 
+	/** Paint the window and all of its contents. Not the same as Paint(). */
+	int32 PaintWindow( const FPaintArgs& Args, const FGeometry& AllottedGeometry, const FSlateRect& MyClippingRect, FSlateWindowElementList& OutDrawElements, int32 LayerId, const FWidgetStyle& InWidgetStyle, bool bParentEnabled ) const;
+
 	/**
 	 * Returns the size of the title bar as a Slate size parameter.  Does not take into account application scale!
 	 *
