Slate: Removed legacy operator() from FArrangedWidgets in favor of operator[]. Fixed uses cases.

diff --git a/Engine/Source/Developer/SlateReflector/Private/Models/WidgetReflectorNode.h b/Engine/Source/Developer/SlateReflector/Private/Models/WidgetReflectorNode.h
--- a/Engine/Source/Developer/SlateReflector/Private/Models/WidgetReflectorNode.h
+++ b/Engine/Source/Developer/SlateReflector/Private/Models/WidgetReflectorNode.h
@@ -52,7 +52,7 @@ public:
 		for (int32 WidgetIndex=0; WidgetIndex < ArrangedChildren.Num(); ++WidgetIndex)
 		{
 			// Note that we include both visible and invisible children!
-			NewNode->ChildNodes.Add( FReflectorNode::NewTreeFrom( ArrangedChildren(WidgetIndex) )  );
+			NewNode->ChildNodes.Add( FReflectorNode::NewTreeFrom( ArrangedChildren[WidgetIndex] )  );
 		}
 
 		return NewNode;
@@ -70,11 +70,11 @@ public:
 	{
 		if (NodeIndexToFind < WidgetPathToFind.Widgets.Num())
 		{
-			const FArrangedWidget& WidgetToFind = WidgetPathToFind.Widgets(NodeIndexToFind);
+			const FArrangedWidget& WidgetToFind = WidgetPathToFind.Widgets[NodeIndexToFind];
 
 			for (int32 NodeIndex=0; NodeIndex < CandidateNodes.Num(); ++NodeIndex)
 			{
-				if (CandidateNodes[NodeIndex]->Widget.Pin() == WidgetPathToFind.Widgets(NodeIndexToFind).Widget)
+				if (CandidateNodes[NodeIndex]->Widget.Pin() == WidgetPathToFind.Widgets[NodeIndexToFind].Widget)
 				{
 					SearchResult.Add(CandidateNodes[NodeIndex]);
 					FindWidgetPath(CandidateNodes[NodeIndex]->ChildNodes, WidgetPathToFind, SearchResult, NodeIndexToFind + 1);
diff --git a/Engine/Source/Developer/SlateReflector/Private/Widgets/SWidgetReflector.cpp b/Engine/Source/Developer/SlateReflector/Private/Widgets/SWidgetReflector.cpp
--- a/Engine/Source/Developer/SlateReflector/Private/Widgets/SWidgetReflector.cpp
+++ b/Engine/Source/Developer/SlateReflector/Private/Widgets/SWidgetReflector.cpp
@@ -154,7 +154,7 @@ void SWidgetReflector::SetWidgetsToVisualize( const FWidgetPath& InWidgetsToVisu
 
 	if (InWidgetsToVisualize.IsValid())
 	{
-		ReflectorTreeRoot.Add(FReflectorNode::NewTreeFrom(InWidgetsToVisualize.Widgets(0)));
+		ReflectorTreeRoot.Add(FReflectorNode::NewTreeFrom(InWidgetsToVisualize.Widgets[0]));
 		PickedPath.Empty();
 
 		FReflectorNode::FindWidgetPath( ReflectorTreeRoot, InWidgetsToVisualize, PickedPath );
@@ -228,7 +228,7 @@ int32 SWidgetReflector::VisualizePickAsRectangles( const FWidgetPath& InWidgetsT
 
 	for (int32 WidgetIndex = 0; WidgetIndex < InWidgetsToVisualize.Widgets.Num(); ++WidgetIndex)
 	{
-		const FArrangedWidget& WidgetGeometry = InWidgetsToVisualize.Widgets(WidgetIndex);
+		const FArrangedWidget& WidgetGeometry = InWidgetsToVisualize.Widgets[WidgetIndex];
 		const float ColorFactor = static_cast<float>(WidgetIndex)/InWidgetsToVisualize.Widgets.Num();
 		const FLinearColor Tint(1.0f - ColorFactor, ColorFactor, 0.0f, 1.0f);
 
diff --git a/Engine/Source/Editor/BehaviorTreeEditor/Private/BehaviorTreeConnectionDrawingPolicy.cpp b/Engine/Source/Editor/BehaviorTreeEditor/Private/BehaviorTreeConnectionDrawingPolicy.cpp
--- a/Engine/Source/Editor/BehaviorTreeEditor/Private/BehaviorTreeConnectionDrawingPolicy.cpp
+++ b/Engine/Source/Editor/BehaviorTreeEditor/Private/BehaviorTreeConnectionDrawingPolicy.cpp
@@ -72,7 +72,7 @@ void FBehaviorTreeConnectionDrawingPolicy::Draw(TMap<TSharedRef<SWidget>, FArran
 	NodeWidgetMap.Empty();
 	for (int32 NodeIndex = 0; NodeIndex < ArrangedNodes.Num(); ++NodeIndex)
 	{
-		FArrangedWidget& CurWidget = ArrangedNodes(NodeIndex);
+		FArrangedWidget& CurWidget = ArrangedNodes[NodeIndex];
 		TSharedRef<SGraphNode> ChildNode = StaticCastSharedRef<SGraphNode>(CurWidget.Widget);
 		NodeWidgetMap.Add(ChildNode->GetNodeObj(), NodeIndex);
 	}
diff --git a/Engine/Source/Editor/BehaviorTreeEditor/Private/SGraphNode_BehaviorTree.cpp b/Engine/Source/Editor/BehaviorTreeEditor/Private/SGraphNode_BehaviorTree.cpp
--- a/Engine/Source/Editor/BehaviorTreeEditor/Private/SGraphNode_BehaviorTree.cpp
+++ b/Engine/Source/Editor/BehaviorTreeEditor/Private/SGraphNode_BehaviorTree.cpp
@@ -684,7 +684,7 @@ TSharedPtr<SGraphNode> SGraphNode_BehaviorTree::GetSubNodeUnderCursor(const FGeo
 		int32 HoveredIndex = SWidget::FindChildUnderMouse( ArrangedChildren, MouseEvent );
 		if ( HoveredIndex != INDEX_NONE )
 		{
-			ResultNode = StaticCastSharedRef<SGraphNode>(ArrangedChildren(HoveredIndex).Widget);
+			ResultNode = StaticCastSharedRef<SGraphNode>(ArrangedChildren[HoveredIndex].Widget);
 		}
 	}
 
diff --git a/Engine/Source/Editor/EnvironmentQueryEditor/Private/EnvironmentQueryConnectionDrawingPolicy.cpp b/Engine/Source/Editor/EnvironmentQueryEditor/Private/EnvironmentQueryConnectionDrawingPolicy.cpp
--- a/Engine/Source/Editor/EnvironmentQueryEditor/Private/EnvironmentQueryConnectionDrawingPolicy.cpp
+++ b/Engine/Source/Editor/EnvironmentQueryEditor/Private/EnvironmentQueryConnectionDrawingPolicy.cpp
@@ -32,7 +32,7 @@ void FEnvironmentQueryConnectionDrawingPolicy::Draw(TMap<TSharedRef<SWidget>, FA
 	NodeWidgetMap.Empty();
 	for (int32 NodeIndex = 0; NodeIndex < ArrangedNodes.Num(); ++NodeIndex)
 	{
-		FArrangedWidget& CurWidget = ArrangedNodes(NodeIndex);
+		FArrangedWidget& CurWidget = ArrangedNodes[NodeIndex];
 		TSharedRef<SGraphNode> ChildNode = StaticCastSharedRef<SGraphNode>(CurWidget.Widget);
 		NodeWidgetMap.Add(ChildNode->GetNodeObj(), NodeIndex);
 	}
diff --git a/Engine/Source/Editor/EnvironmentQueryEditor/Private/SGraphNode_EnvironmentQuery.cpp b/Engine/Source/Editor/EnvironmentQueryEditor/Private/SGraphNode_EnvironmentQuery.cpp
--- a/Engine/Source/Editor/EnvironmentQueryEditor/Private/SGraphNode_EnvironmentQuery.cpp
+++ b/Engine/Source/Editor/EnvironmentQueryEditor/Private/SGraphNode_EnvironmentQuery.cpp
@@ -473,7 +473,7 @@ TSharedPtr<SGraphNode> SGraphNode_EnvironmentQuery::GetSubNodeUnderCursor(const
 		int32 HoveredIndex = SWidget::FindChildUnderMouse( ArrangedChildren, MouseEvent );
 		if ( HoveredIndex != INDEX_NONE )
 		{
-			ResultNode = StaticCastSharedRef<SGraphNode>(ArrangedChildren(HoveredIndex).Widget);
+			ResultNode = StaticCastSharedRef<SGraphNode>(ArrangedChildren[HoveredIndex].Widget);
 		}
 	}
 
diff --git a/Engine/Source/Editor/GraphEditor/Private/ConnectionDrawingPolicy.cpp b/Engine/Source/Editor/GraphEditor/Private/ConnectionDrawingPolicy.cpp
--- a/Engine/Source/Editor/GraphEditor/Private/ConnectionDrawingPolicy.cpp
+++ b/Engine/Source/Editor/GraphEditor/Private/ConnectionDrawingPolicy.cpp
@@ -1011,7 +1011,7 @@ void FStateMachineConnectionDrawingPolicy::DetermineLinkGeometry(
 
 		UAnimStateNodeBase* State = CastChecked<UAnimStateNodeBase>(InputPin->GetOwningNode());
 		int32 StateIndex = NodeWidgetMap.FindChecked(State);
-		EndWidgetGeometry = &(ArrangedNodes(StateIndex));
+		EndWidgetGeometry = &(ArrangedNodes[StateIndex]);
 	}
 	else if (UAnimStateTransitionNode* TransNode = Cast<UAnimStateTransitionNode>(InputPin->GetOwningNode()))
 	{
@@ -1023,8 +1023,8 @@ void FStateMachineConnectionDrawingPolicy::DetermineLinkGeometry(
 			int32* NextNodeIndex = NodeWidgetMap.Find(NextState);
 			if ((PrevNodeIndex != NULL) && (NextNodeIndex != NULL))
 			{
-				StartWidgetGeometry = &(ArrangedNodes(*PrevNodeIndex));
-				EndWidgetGeometry = &(ArrangedNodes(*NextNodeIndex));
+				StartWidgetGeometry = &(ArrangedNodes[*PrevNodeIndex]);
+				EndWidgetGeometry = &(ArrangedNodes[*NextNodeIndex]);
 			}
 		}
 	}
@@ -1046,7 +1046,7 @@ void FStateMachineConnectionDrawingPolicy::Draw(TMap<TSharedRef<SWidget>, FArran
 	NodeWidgetMap.Empty();
 	for (int32 NodeIndex = 0; NodeIndex < ArrangedNodes.Num(); ++NodeIndex)
 	{
-		FArrangedWidget& CurWidget = ArrangedNodes(NodeIndex);
+		FArrangedWidget& CurWidget = ArrangedNodes[NodeIndex];
 		TSharedRef<SGraphNode> ChildNode = StaticCastSharedRef<SGraphNode>(CurWidget.Widget);
 		NodeWidgetMap.Add(ChildNode->GetNodeObj(), NodeIndex);
 	}
diff --git a/Engine/Source/Editor/GraphEditor/Private/SGraphNode.cpp b/Engine/Source/Editor/GraphEditor/Private/SGraphNode.cpp
--- a/Engine/Source/Editor/GraphEditor/Private/SGraphNode.cpp
+++ b/Engine/Source/Editor/GraphEditor/Private/SGraphNode.cpp
@@ -1006,7 +1006,7 @@ TSharedPtr<SGraphPin> SGraphNode::GetHoveredPin( const FGeometry& MyGeometry, co
 		int32 HoveredPinIndex = SWidget::FindChildUnderMouse( ArrangedPins, MouseEvent );
 		if ( HoveredPinIndex != INDEX_NONE )
 		{
-			return StaticCastSharedRef<SGraphPin>(ArrangedPins(HoveredPinIndex).Widget);
+			return StaticCastSharedRef<SGraphPin>(ArrangedPins[HoveredPinIndex].Widget);
 		}
 	}
 	
diff --git a/Engine/Source/Editor/GraphEditor/Private/SGraphPanel.cpp b/Engine/Source/Editor/GraphEditor/Private/SGraphPanel.cpp
--- a/Engine/Source/Editor/GraphEditor/Private/SGraphPanel.cpp
+++ b/Engine/Source/Editor/GraphEditor/Private/SGraphPanel.cpp
@@ -161,7 +161,7 @@ int32 SGraphPanel::OnPaint( const FPaintArgs& Args, const FGeometry& AllottedGeo
 		TArray<FGraphDiffControl::FNodeMatch> NodeMatches;
 		for (int32 ChildIndex = 0; ChildIndex < ArrangedChildren.Num(); ++ChildIndex)
 		{
-			FArrangedWidget& CurWidget = ArrangedChildren(ChildIndex);
+			FArrangedWidget& CurWidget = ArrangedChildren[ChildIndex];
 			TSharedRef<SGraphNode> ChildNode = StaticCastSharedRef<SGraphNode>(CurWidget.Widget);
 			
 			// Examine node to see what layers we should be drawing in
@@ -514,7 +514,7 @@ void SGraphPanel::OnArrangeChildren( const FGeometry& AllottedGeometry, FArrange
 	FArrangedChildren MyArrangedChildren(ArrangedChildren.GetFilter());
 	for (int32 ChildIndex = 0; ChildIndex < ArrangedChildren.Num(); ++ChildIndex)
 	{
-		FArrangedWidget& CurWidget = ArrangedChildren(ChildIndex);
+		FArrangedWidget& CurWidget = ArrangedChildren[ChildIndex];
 		TSharedRef<SGraphNode> ChildNode = StaticCastSharedRef<SGraphNode>(CurWidget.Widget);
 
 		TArray<FOverlayWidgetInfo> OverlayWidgets = ChildNode->GetOverlayWidgets(false, CurWidget.Geometry.Size);
@@ -645,7 +645,7 @@ TSharedPtr<SWidget> SGraphPanel::OnSummonContextMenu(const FGeometry& MyGeometry
 			const int32 HoveredNodeIndex = SWidget::FindChildUnderMouse( ArrangedNodes, MouseEvent );
 			if (HoveredNodeIndex != INDEX_NONE)
 			{
-				const FArrangedWidget& HoveredNode = ArrangedNodes(HoveredNodeIndex);
+				const FArrangedWidget& HoveredNode = ArrangedNodes[HoveredNodeIndex];
 				const TSharedRef<SGraphNode>& GraphNode = StaticCastSharedRef<SGraphNode>(HoveredNode.Widget);
 				NodeUnderCursor = GraphNode->GetNodeObj();
 
diff --git a/Engine/Source/Editor/GraphEditor/Private/SNodePanel.cpp b/Engine/Source/Editor/GraphEditor/Private/SNodePanel.cpp
--- a/Engine/Source/Editor/GraphEditor/Private/SNodePanel.cpp
+++ b/Engine/Source/Editor/GraphEditor/Private/SNodePanel.cpp
@@ -578,7 +578,7 @@ FReply SNodePanel::OnMouseButtonDown( const FGeometry& MyGeometry, const FPointe
 			// PRESSING ON A NODE!
 
 			// This changes selection and starts dragging it.
-			const FArrangedWidget& NodeGeometry = ArrangedChildren(NodeUnderMouseIndex);
+			const FArrangedWidget& NodeGeometry = ArrangedChildren[NodeUnderMouseIndex];
 			const FVector2D MousePositionInNode = NodeGeometry.Geometry.AbsoluteToLocal(MouseEvent.GetScreenSpacePosition());
 			TSharedRef<SNode> NodeWidgetUnderMouse = StaticCastSharedRef<SNode>( NodeGeometry.Widget );
 
diff --git a/Engine/Source/Editor/IntroTutorials/Private/STutorialOverlay.cpp b/Engine/Source/Editor/IntroTutorials/Private/STutorialOverlay.cpp
--- a/Engine/Source/Editor/IntroTutorials/Private/STutorialOverlay.cpp
+++ b/Engine/Source/Editor/IntroTutorials/Private/STutorialOverlay.cpp
@@ -132,7 +132,7 @@ int32 STutorialOverlay::TraverseWidgets(TSharedRef<SWidget> InWidget, const FGeo
 	InWidget->ArrangeChildren(InGeometry, ArrangedChildren);
 	for(int32 ChildIndex = 0; ChildIndex < ArrangedChildren.Num(); ChildIndex++)
 	{
-		const FArrangedWidget& ArrangedWidget = ArrangedChildren(ChildIndex);
+		const FArrangedWidget& ArrangedWidget = ArrangedChildren[ChildIndex];
 		LayerId = TraverseWidgets(ArrangedWidget.Widget, ArrangedWidget.Geometry, MyClippingRect, OutDrawElements, LayerId);
 	}
 
diff --git a/Engine/Source/Editor/IntroTutorials/Private/TutorialStructCustomization.cpp b/Engine/Source/Editor/IntroTutorials/Private/TutorialStructCustomization.cpp
--- a/Engine/Source/Editor/IntroTutorials/Private/TutorialStructCustomization.cpp
+++ b/Engine/Source/Editor/IntroTutorials/Private/TutorialStructCustomization.cpp
@@ -206,7 +206,7 @@ private:
 
 		for(int32 PathIndex = Path.Widgets.Num() - 1; PathIndex >= 0; PathIndex--)
 		{
-			TSharedRef<SWidget> PathWidget = Path.Widgets(PathIndex).Widget;
+			TSharedRef<SWidget> PathWidget = Path.Widgets[PathIndex].Widget;
 			if(PathWidget->GetTag() != NAME_None)
 			{
 				PickedWidgetName = PathWidget->GetTag();
diff --git a/Engine/Source/Editor/Sequencer/Private/SSequencerSectionAreaView.cpp b/Engine/Source/Editor/Sequencer/Private/SSequencerSectionAreaView.cpp
--- a/Engine/Source/Editor/Sequencer/Private/SSequencerSectionAreaView.cpp
+++ b/Engine/Source/Editor/Sequencer/Private/SSequencerSectionAreaView.cpp
@@ -124,7 +124,7 @@ int32 SSequencerSectionAreaView::OnPaint( const FPaintArgs& Args, const FGeometr
 
 	for (int32 ChildIndex = 0; ChildIndex < ArrangedChildren.Num(); ++ChildIndex)
 	{
-		FArrangedWidget& CurWidget = ArrangedChildren(ChildIndex);
+		FArrangedWidget& CurWidget = ArrangedChildren[ChildIndex];
 		FSlateRect ChildClipRect = MyClippingRect.IntersectionWith( CurWidget.Geometry.GetClippingRect() );
 		const int32 CurWidgetsMaxLayerId = CurWidget.Widget->Paint( Args.WithNewParent(this), CurWidget.Geometry, ChildClipRect, OutDrawElements, LayerId+1, InWidgetStyle, ShouldBeEnabled( bParentEnabled ) );
 	}
diff --git a/Engine/Source/Editor/UMGEditor/Private/Designer/SDesignerView.cpp b/Engine/Source/Editor/UMGEditor/Private/Designer/SDesignerView.cpp
--- a/Engine/Source/Editor/UMGEditor/Private/Designer/SDesignerView.cpp
+++ b/Engine/Source/Editor/UMGEditor/Private/Designer/SDesignerView.cpp
@@ -66,7 +66,7 @@ static bool LocateWidgetsUnderCursor_Helper(FArrangedWidget& Candidate, FVector2
 			// A widget's children are implicitly Z-ordered from first to last
 			for ( int32 ChildIndex = ArrangedChildren.Num() - 1; !bHitChildWidget && ChildIndex >= 0; --ChildIndex )
 			{
-				FArrangedWidget& SomeChild = ArrangedChildren(ChildIndex);
+				FArrangedWidget& SomeChild = ArrangedChildren[ChildIndex];
 				bHitChildWidget = ( SomeChild.Widget->IsEnabled() || bIgnoreEnabledStatus ) && LocateWidgetsUnderCursor_Helper(SomeChild, InAbsoluteCursorLocation, OutWidgetsUnderCursor, bIgnoreEnabledStatus);
 			}
 		}
diff --git a/Engine/Source/Editor/WorldBrowser/Private/Tiles/SWorldComposition.cpp b/Engine/Source/Editor/WorldBrowser/Private/Tiles/SWorldComposition.cpp
--- a/Engine/Source/Editor/WorldBrowser/Private/Tiles/SWorldComposition.cpp
+++ b/Engine/Source/Editor/WorldBrowser/Private/Tiles/SWorldComposition.cpp
@@ -300,7 +300,7 @@ public:
 
 		for (int32 ChildIndex = 0; ChildIndex < ArrangedChildren.Num(); ++ChildIndex)
 		{
-			FArrangedWidget& CurWidget = ArrangedChildren(ChildIndex);
+			FArrangedWidget& CurWidget = ArrangedChildren[ChildIndex];
 			TSharedRef<SWorldTileItem> ChildNode = StaticCastSharedRef<SWorldTileItem>(CurWidget.Widget);
 		
 			ChildNode->bAffectedByMarquee = SelectionToVisualize->Contains(ChildNode->GetObjectBeingDisplayed());
@@ -606,7 +606,7 @@ protected:
 		if (NodeUnderMouseIndex != INDEX_NONE)
 		{
 			// PRESSING ON A NODE!
-			const FArrangedWidget& NodeGeometry = ArrangedChildren(NodeUnderMouseIndex);
+			const FArrangedWidget& NodeGeometry = ArrangedChildren[NodeUnderMouseIndex];
 			const FVector2D MousePositionInNode = NodeGeometry.Geometry.AbsoluteToLocal(MouseEvent.GetScreenSpacePosition());
 			TSharedRef<SNode> NodeWidgetUnderMouse = StaticCastSharedRef<SNode>( NodeGeometry.Widget );
 
diff --git a/Engine/Source/Runtime/Slate/Private/Framework/Application/SlateApplication.cpp b/Engine/Source/Runtime/Slate/Private/Framework/Application/SlateApplication.cpp
--- a/Engine/Source/Runtime/Slate/Private/Framework/Application/SlateApplication.cpp
+++ b/Engine/Source/Runtime/Slate/Private/Framework/Application/SlateApplication.cpp
@@ -113,7 +113,7 @@ void FSlateApplication::MouseCaptorHelper::SetMouseCaptor( const FWidgetPath& Ev
 		if ( !MouseCaptorWeakPath.IsValid() )
 		{
 			// If the target widget wasn't found on the event path then start the search from the root
-			FWidgetPath NewMouseCaptorPath = EventPath.GetPathDownTo( EventPath.Widgets(0).Widget );
+			FWidgetPath NewMouseCaptorPath = EventPath.GetPathDownTo( EventPath.Widgets[0].Widget );
 			NewMouseCaptorPath.ExtendPathTo( FWidgetMatcher( WidgetRef ) );
 			MouseCaptorWeakPath = NewMouseCaptorPath;
 		}
@@ -272,7 +272,7 @@ static bool LocateWidgetsUnderCursor_Helper( FArrangedWidget& Candidate, FVector
 			// A widget's children are implicitly Z-ordered from first to last
 			for (int32 ChildIndex = ArrangedChildren.Num()-1; !bHitChildWidget && ChildIndex >= 0; --ChildIndex)
 			{
-				FArrangedWidget& SomeChild = ArrangedChildren(ChildIndex);
+				FArrangedWidget& SomeChild = ArrangedChildren[ChildIndex];
 				bHitChildWidget = (SomeChild.Widget->IsEnabled() || bIgnoreEnabledStatus) && LocateWidgetsUnderCursor_Helper( SomeChild, InAbsoluteCursorLocation, OutWidgetsUnderCursor, bIgnoreEnabledStatus );
 			}
 		}
@@ -1506,7 +1506,7 @@ bool FSlateApplication::SetKeyboardFocus( const FWidgetPath& InFocusPath, const
 
 	for (int32 WidgetIndex = InFocusPath.Widgets.Num()-1; !bFocusTransferComplete  && WidgetIndex >= 0; --WidgetIndex)
 	{
-		const FArrangedWidget& WidgetToFocus = InFocusPath.Widgets( WidgetIndex );
+		const FArrangedWidget& WidgetToFocus = InFocusPath.Widgets[ WidgetIndex ];
 		
 		// Does this widget support keyboard focus?  If so, then we'll go ahead and set it!
 		if( WidgetToFocus.Widget->SupportsKeyboardFocus() )
@@ -1549,7 +1549,7 @@ bool FSlateApplication::SetKeyboardFocus( const FWidgetPath& InFocusPath, const
 					{
 						for (int32 ChildIndex=0; ChildIndex < NewFocusPath.Widgets.Num(); ++ChildIndex)
 						{
-							NotifyUsAboutFocusChange.AddUnique( NewFocusPath.Widgets(ChildIndex).Widget );
+							NotifyUsAboutFocusChange.AddUnique( NewFocusPath.Widgets[ChildIndex].Widget );
 						}
 
 						FScopedSwitchWorldHack SwitchWorld( NewFocusPath );
@@ -1801,14 +1801,14 @@ void FSlateApplication::ProcessReply( const FWidgetPath& CurrentEventPath, const
 		// Pretend that the mouse left all the previously hovered widgets, and a drag entered them.
 		for (int32 WidgetIndex=0; WidgetIndex < WidgetsUnderMouse->Widgets.Num(); ++WidgetIndex)
 		{
-			const FArrangedWidget& SomeWidget = WidgetsUnderMouse->Widgets(WidgetIndex);
+			const FArrangedWidget& SomeWidget = WidgetsUnderMouse->Widgets[WidgetIndex];
 			SomeWidget.Widget->OnMouseLeave( *InMouseEvent );
 		}
 
 		FDragDropEvent DragDropEvent( *InMouseEvent, ReplyDragDropContent );
 		for (int32 WidgetIndex=0; WidgetIndex < WidgetsUnderMouse->Widgets.Num(); ++WidgetIndex)
 		{
-			const FArrangedWidget& SomeWidget = WidgetsUnderMouse->Widgets(WidgetIndex);
+			const FArrangedWidget& SomeWidget = WidgetsUnderMouse->Widgets[WidgetIndex];
 			SomeWidget.Widget->OnDragEnter( SomeWidget.Geometry, DragDropEvent );
 		}
 	}
@@ -1847,7 +1847,7 @@ void FSlateApplication::ProcessReply( const FWidgetPath& CurrentEventPath, const
 		{
 			// The requested mouse captor was not in the event path.
 			// We will attempt to find it in this window; if we don't find it, then give up.
-			NewJoystickCaptorPath = CurrentEventPath.GetPathDownTo( CurrentEventPath.Widgets(0).Widget );
+			NewJoystickCaptorPath = CurrentEventPath.GetPathDownTo( CurrentEventPath.Widgets[0].Widget );
 			NewJoystickCaptorPath.ExtendPathTo( FWidgetMatcher(RequestedJoystickCaptor.ToSharedRef()) );
 		}
 
@@ -1917,7 +1917,7 @@ void FSlateApplication::LockCursor( const TSharedPtr<SWidget>& Widget )
 			GeneratePathToWidgetChecked( Widget.ToSharedRef(), WidgetPath );
 
 			// The last widget in the path should be the widget we are locking the cursor to
-			FArrangedWidget& WidgetGeom = WidgetPath.Widgets( WidgetPath.Widgets.Num() - 1 );
+			FArrangedWidget& WidgetGeom = WidgetPath.Widgets[ WidgetPath.Widgets.Num() - 1 ];
 
 			TSharedRef<SWindow> Window = WidgetPath.GetWindow();
 			// Do not attempt to lock the cursor to the window if its not in the foreground.  It would cause annoying side effects
@@ -2008,7 +2008,7 @@ void FSlateApplication::QueryCursor()
 				CursorResult = FCursorReply::Unhandled();
 				for( int32 WidgetIndex = WidgetsToQueryForCursor.Widgets.Num() - 1; !CursorResult.IsEventHandled() && WidgetIndex >= 0; --WidgetIndex )
 				{
-					const FArrangedWidget& WidgetToQuery = WidgetsToQueryForCursor.Widgets( WidgetIndex );
+					const FArrangedWidget& WidgetToQuery = WidgetsToQueryForCursor.Widgets[ WidgetIndex ];
 					CursorResult = WidgetToQuery.Widget->OnCursorQuery( WidgetToQuery.Geometry, CursorEvent );
 				}
 
@@ -2132,7 +2132,7 @@ void FSlateApplication::UpdateToolTip( bool AllowSpawningOfNewToolTips )
 	TSharedPtr<SWidget> WidgetProvidingNewToolTip;
 	for ( int32 WidgetIndex=WidgetsToQueryForToolTip.Widgets.Num()-1; WidgetIndex >= 0; --WidgetIndex )
 	{
-		FArrangedWidget* CurWidgetGeometry = &WidgetsToQueryForToolTip.Widgets(WidgetIndex);
+		FArrangedWidget* CurWidgetGeometry = &WidgetsToQueryForToolTip.Widgets[WidgetIndex];
 		const TSharedRef<SWidget>& CurWidget = CurWidgetGeometry->Widget;
 
 		if( !NewToolTip.IsValid() )
@@ -2181,7 +2181,7 @@ void FSlateApplication::UpdateToolTip( bool AllowSpawningOfNewToolTips )
 		// Some widgets might want to provide an alternative Tooltip Handler.
 		for ( int32 WidgetIndex=WidgetsToQueryForToolTip.Widgets.Num()-1; !bOnVisualizeTooltipHandled && WidgetIndex >= 0; --WidgetIndex )
 		{
-			const FArrangedWidget& CurWidgetGeometry = WidgetsToQueryForToolTip.Widgets(WidgetIndex);
+			const FArrangedWidget& CurWidgetGeometry = WidgetsToQueryForToolTip.Widgets[WidgetIndex];
 			bOnVisualizeTooltipHandled = CurWidgetGeometry.Widget->OnVisualizeTooltip( NewToolTip.IsValid() ? NewToolTip->AsWidget() : TSharedPtr<SWidget>() );
 			if (bOnVisualizeTooltipHandled)
 			{
@@ -2943,7 +2943,7 @@ bool FSlateApplication::ProcessKeyCharEvent( FCharacterEvent& InCharacterEvent )
 		InCharacterEvent.SetEventPath( EventPath );
 		for( int32 WidgetIndex = EventPath.Widgets.Num()-1; !Reply.IsEventHandled() && WidgetIndex >= 0; --WidgetIndex )
 		{
-			FArrangedWidget& SomeWidgetGettingEvent = EventPath.Widgets( WidgetIndex );
+			FArrangedWidget& SomeWidgetGettingEvent = EventPath.Widgets[ WidgetIndex ];
 			if (SomeWidgetGettingEvent.Widget->IsEnabled())
 			{
 				// Widget newly under cursor, so send a MouseEnter.
@@ -3018,7 +3018,7 @@ bool FSlateApplication::ProcessKeyDownEvent( FKeyboardEvent& InKeyboardEvent )
 			// Tunnel the keyboard event
 			for( int32 WidgetIndex = 0; !Reply.IsEventHandled() && WidgetIndex < EventPath.Widgets.Num(); ++WidgetIndex )
 			{
-				FArrangedWidget& SomeWidgetGettingEvent = EventPath.Widgets( WidgetIndex );
+				FArrangedWidget& SomeWidgetGettingEvent = EventPath.Widgets[ WidgetIndex ];
 				if (SomeWidgetGettingEvent.Widget->IsEnabled())
 				{
 					Reply = SomeWidgetGettingEvent.Widget->OnPreviewKeyDown( SomeWidgetGettingEvent.Geometry, InKeyboardEvent ).SetHandler(SomeWidgetGettingEvent.Widget);
@@ -3031,7 +3031,7 @@ bool FSlateApplication::ProcessKeyDownEvent( FKeyboardEvent& InKeyboardEvent )
 			// Send out key down events.
 			for( int32 WidgetIndex = EventPath.Widgets.Num()-1; !Reply.IsEventHandled() && WidgetIndex >= 0; --WidgetIndex )
 			{
-				FArrangedWidget& SomeWidgetGettingEvent = EventPath.Widgets( WidgetIndex );
+				FArrangedWidget& SomeWidgetGettingEvent = EventPath.Widgets[ WidgetIndex ];
 				if (SomeWidgetGettingEvent.Widget->IsEnabled())
 				{
 					Reply = SomeWidgetGettingEvent.Widget->OnKeyDown( SomeWidgetGettingEvent.Geometry, InKeyboardEvent ).SetHandler(SomeWidgetGettingEvent.Widget);
@@ -3101,7 +3101,7 @@ bool FSlateApplication::ProcessKeyUpEvent( FKeyboardEvent& InKeyboardEvent )
 		// Send out mouse enter events.
 		for( int32 WidgetIndex = EventPath.Widgets.Num()-1; !Reply.IsEventHandled() && WidgetIndex >= 0; --WidgetIndex )
 		{
-			FArrangedWidget& SomeWidgetGettingEvent = EventPath.Widgets( WidgetIndex );
+			FArrangedWidget& SomeWidgetGettingEvent = EventPath.Widgets[ WidgetIndex ];
 			if (SomeWidgetGettingEvent.Widget->IsEnabled())
 			{
 				// Widget newly under cursor, so send a MouseEnter.
@@ -3241,7 +3241,7 @@ bool FSlateApplication::ProcessMouseButtonDownEvent( const TSharedPtr< FGenericW
 
 			for( int32 WidgetIndex = 0; !Reply.IsEventHandled() && WidgetIndex < WidgetsUnderCursor.Widgets.Num() ; ++WidgetIndex )
 			{
-				FArrangedWidget& CurWidget = WidgetsUnderCursor.Widgets(WidgetIndex);
+				FArrangedWidget& CurWidget = WidgetsUnderCursor.Widgets[WidgetIndex];
 
 				Reply = CurWidget.Widget->OnPreviewMouseButtonDown(CurWidget.Geometry, MouseEvent).SetHandler( CurWidget.Widget );
 				ProcessReply( WidgetsUnderCursor, Reply, &WidgetsUnderCursor, &MouseEvent);
@@ -3249,7 +3249,7 @@ bool FSlateApplication::ProcessMouseButtonDownEvent( const TSharedPtr< FGenericW
 
 			for( int32 WidgetIndex = WidgetsUnderCursor.Widgets.Num()-1; !Reply.IsEventHandled() && WidgetIndex >= 0 ; --WidgetIndex )
 			{
-				FArrangedWidget& CurWidget = WidgetsUnderCursor.Widgets(WidgetIndex);				
+				FArrangedWidget& CurWidget = WidgetsUnderCursor.Widgets[WidgetIndex];				
 
 				if (!Reply.IsEventHandled())
 				{
@@ -3279,7 +3279,7 @@ bool FSlateApplication::ProcessMouseButtonDownEvent( const TSharedPtr< FGenericW
 				bool bFocusCandidateFound = false;
 				for( int32 WidgetIndex = WidgetsUnderCursor.Widgets.Num()-1; !bFocusCandidateFound && WidgetIndex >= 0 ; --WidgetIndex )
 				{
-					FArrangedWidget& CurWidget = WidgetsUnderCursor.Widgets(WidgetIndex);
+					FArrangedWidget& CurWidget = WidgetsUnderCursor.Widgets[WidgetIndex];
 					if (CurWidget.Widget->SupportsKeyboardFocus())
 					{
 						bFocusCandidateFound = true;
@@ -3353,7 +3353,7 @@ bool FSlateApplication::ProcessMouseButtonDoubleClickEvent( const TSharedPtr< FG
 	FReply Reply = FReply::Unhandled();
 	for( int32 WidgetIndex = WidgetsUnderCursor.Widgets.Num()-1; !Reply.IsEventHandled() && WidgetIndex >= 0; --WidgetIndex )
 	{
-		FArrangedWidget& CurWidget = WidgetsUnderCursor.Widgets(WidgetIndex);
+		FArrangedWidget& CurWidget = WidgetsUnderCursor.Widgets[WidgetIndex];
 		Reply = CurWidget.Widget->OnMouseButtonDoubleClick( CurWidget.Geometry, InMouseEvent ).SetHandler( CurWidget.Widget );
 		ProcessReply(WidgetsUnderCursor, Reply, &WidgetsUnderCursor, &InMouseEvent);
 
@@ -3445,7 +3445,7 @@ bool FSlateApplication::ProcessMouseButtonUpEvent( FPointerEvent& MouseEvent )
 		FReply Reply = FReply::Unhandled();
 		for( int32 WidgetIndex = WidgetsUnderCursor.Widgets.Num()-1; !Reply.IsEventHandled() && WidgetIndex >= 0; --WidgetIndex )
 		{
-			FArrangedWidget& CurWidget = WidgetsUnderCursor.Widgets(WidgetIndex);
+			FArrangedWidget& CurWidget = WidgetsUnderCursor.Widgets[WidgetIndex];
 			if (MouseEvent.IsTouchEvent())
 			{
 				Reply = CurWidget.Widget->OnTouchEnded( CurWidget.Geometry, MouseEvent ).SetHandler( CurWidget.Widget );
@@ -3528,7 +3528,7 @@ bool FSlateApplication::ProcessMouseWheelOrGestureEvent( FPointerEvent& InWheelE
 	FReply Reply = FReply::Unhandled();
 	for( int32 WidgetIndex = EventPath.Widgets.Num()-1; !Reply.IsEventHandled() && WidgetIndex >= 0; --WidgetIndex )
 	{
-		FArrangedWidget& CurWidget = EventPath.Widgets(WidgetIndex);
+		FArrangedWidget& CurWidget = EventPath.Widgets[WidgetIndex];
 		// Gesture event gets first shot, if slate doesn't respond to it, we'll try the wheel event.
 		if( InGestureEvent!=NULL )
 		{
@@ -3751,7 +3751,7 @@ bool FSlateApplication::ProcessMouseMoveEvent( FPointerEvent& MouseEvent )
 		// Send out mouse enter events.
 		for( int32 WidgetIndex = WidgetsUnderCursor.Widgets.Num()-1; WidgetIndex >= 0; --WidgetIndex )
 		{
-			FArrangedWidget& SomeWidgetUnderCursor = WidgetsUnderCursor.Widgets( WidgetIndex );
+			FArrangedWidget& SomeWidgetUnderCursor = WidgetsUnderCursor.Widgets[ WidgetIndex ];
 			if ( !LastWidgetsUnderCursor.ContainsWidget(SomeWidgetUnderCursor.Widget) )
 			{
 				// Widget newly under cursor, so send a MouseEnter.
@@ -3773,7 +3773,7 @@ bool FSlateApplication::ProcessMouseMoveEvent( FPointerEvent& MouseEvent )
 		// Bubble the MouseMove event.
 		for( int32 WidgetIndex = WidgetsUnderCursor.Widgets.Num()-1; !Reply.IsEventHandled() && WidgetIndex >= 0; --WidgetIndex )
 		{
-			FArrangedWidget& CurWidget = WidgetsUnderCursor.Widgets(WidgetIndex);
+			FArrangedWidget& CurWidget = WidgetsUnderCursor.Widgets[WidgetIndex];
 
 			if (MouseEvent.IsTouchEvent())
 			{
diff --git a/Engine/Source/Runtime/Slate/Private/Framework/Docking/SDockingTabWell.cpp b/Engine/Source/Runtime/Slate/Private/Framework/Docking/SDockingTabWell.cpp
--- a/Engine/Source/Runtime/Slate/Private/Framework/Docking/SDockingTabWell.cpp
+++ b/Engine/Source/Runtime/Slate/Private/Framework/Docking/SDockingTabWell.cpp
@@ -132,7 +132,7 @@ int32 SDockingTabWell::OnPaint( const FPaintArgs& Args, const FGeometry& Allotte
 	// Draw all inactive tabs first
 	for (int32 ChildIndex = 0; ChildIndex < ArrangedChildren.Num(); ++ChildIndex)
 	{
-		FArrangedWidget& CurWidget = ArrangedChildren(ChildIndex);
+		FArrangedWidget& CurWidget = ArrangedChildren[ChildIndex];
 		if (CurWidget.Widget == ForegroundTab)
 		{
 			ForegroundTabGeometry = &CurWidget;
diff --git a/Engine/Source/Runtime/Slate/Private/Framework/MultiBox/SClippingHorizontalBox.cpp b/Engine/Source/Runtime/Slate/Private/Framework/MultiBox/SClippingHorizontalBox.cpp
--- a/Engine/Source/Runtime/Slate/Private/Framework/MultiBox/SClippingHorizontalBox.cpp
+++ b/Engine/Source/Runtime/Slate/Private/Framework/MultiBox/SClippingHorizontalBox.cpp
@@ -24,7 +24,7 @@ void SClippingHorizontalBox::OnArrangeChildren( const FGeometry& AllottedGeometr
 	int32 IndexClippedAt = NumChildren;
 	for (int32 ChildIdx = NumChildren - 2; ChildIdx >= 0; --ChildIdx)
 	{
-		const FArrangedWidget& CurWidget = ArrangedChildren(ChildIdx);
+		const FArrangedWidget& CurWidget = ArrangedChildren[ChildIdx];
 		if (FMath::TruncToInt(CurWidget.Geometry.AbsolutePosition.X + CurWidget.Geometry.Size.X) > FMath::TruncToInt(AllottedGeometry.AbsolutePosition.X + AllottedGeometry.Size.X))
 		{
 			ArrangedChildren.Remove(ChildIdx);
@@ -40,13 +40,13 @@ void SClippingHorizontalBox::OnArrangeChildren( const FGeometry& AllottedGeometr
 	else
 	{
 		// Right align the wrap button
-		FArrangedWidget& ArrangedButton = ArrangedChildren(ArrangedChildren.Num() - 1);
+		FArrangedWidget& ArrangedButton = ArrangedChildren[ArrangedChildren.Num() - 1];
 		ArrangedButton.Geometry.AbsolutePosition = AllottedGeometry.AbsolutePosition + AllottedGeometry.Size - ArrangedButton.Geometry.Size;
 
 		// Further remove any children that the wrap button overlaps with
 		for (int32 ChildIdx = IndexClippedAt - 1; ChildIdx >= 0; --ChildIdx)
 		{
-			const FArrangedWidget& CurWidget = ArrangedChildren(ChildIdx);
+			const FArrangedWidget& CurWidget = ArrangedChildren[ChildIdx];
 			if (FMath::TruncToInt(CurWidget.Geometry.AbsolutePosition.X + CurWidget.Geometry.Size.X) > FMath::TruncToInt(ArrangedButton.Geometry.AbsolutePosition.X))
 			{
 				ArrangedChildren.Remove(ChildIdx);
@@ -69,9 +69,9 @@ int32 SClippingHorizontalBox::OnPaint( const FPaintArgs& Args, const FGeometry&
 	if ((ClippedArrangedChildren.Num() != 0) && (ArrangedChildren.Num() != 0))
 	{
 		int32 IndexClippedAt = ClippedArrangedChildren.Num() - 1;
-		const FArrangedWidget& LastCippedChild = ClippedArrangedChildren(IndexClippedAt);
-		const FArrangedWidget& FirstChild = ArrangedChildren(0);
-		const FArrangedWidget& LastChild = ArrangedChildren(ArrangedChildren.Num() - 1);
+		const FArrangedWidget& LastCippedChild = ClippedArrangedChildren[IndexClippedAt];
+		const FArrangedWidget& FirstChild = ArrangedChildren[0];
+		const FArrangedWidget& LastChild = ArrangedChildren[ArrangedChildren.Num() - 1];
 		FGeometry ResizedGeometry = AllottedGeometry;
 	
 		// If only the last child/block, which is the wrap button, is being clipped
@@ -87,7 +87,7 @@ int32 SClippingHorizontalBox::OnPaint( const FPaintArgs& Args, const FGeometry&
 		else
 		{
 			// Children/blocks are being clipped, calculate the size of the custom border
-			const FArrangedWidget& NextChild = (IndexClippedAt + 1 < ClippedArrangedChildren.Num())? ClippedArrangedChildren(IndexClippedAt + 1): LastCippedChild;
+			const FArrangedWidget& NextChild = (IndexClippedAt + 1 < ClippedArrangedChildren.Num())? ClippedArrangedChildren[IndexClippedAt + 1]: LastCippedChild;
 			ResizedGeometry.Size.X = NextChild.Geometry.AbsolutePosition.X + NextChild.Geometry.Size.X - FirstChild.Geometry.AbsolutePosition.X;
 		}
 	
diff --git a/Engine/Source/Runtime/Slate/Private/Widgets/Input/SMenuAnchor.cpp b/Engine/Source/Runtime/Slate/Private/Widgets/Input/SMenuAnchor.cpp
--- a/Engine/Source/Runtime/Slate/Private/Widgets/Input/SMenuAnchor.cpp
+++ b/Engine/Source/Runtime/Slate/Private/Widgets/Input/SMenuAnchor.cpp
@@ -102,7 +102,7 @@ int32 SMenuAnchor::OnPaint( const FPaintArgs& Args, const FGeometry& AllottedGeo
 	// There may be zero elements in this array if our child collapsed/hidden
 	if ( ArrangedChildren.Num() > 0 )
 	{
-		FArrangedWidget& TheChild = ArrangedChildren( 0 );
+		FArrangedWidget& TheChild = ArrangedChildren[0];
 
 		const FSlateRect ChildClippingRect = AllottedGeometry.GetClippingRect().IntersectionWith( MyClippingRect );
 
diff --git a/Engine/Source/Runtime/Slate/Private/Widgets/Layout/SBox.cpp b/Engine/Source/Runtime/Slate/Private/Widgets/Layout/SBox.cpp
--- a/Engine/Source/Runtime/Slate/Private/Widgets/Layout/SBox.cpp
+++ b/Engine/Source/Runtime/Slate/Private/Widgets/Layout/SBox.cpp
@@ -167,7 +167,7 @@ int32 SBox::OnPaint( const FPaintArgs& Args, const FGeometry& AllottedGeometry,
 	if( ArrangedChildren.Num() > 0 )
 	{
 		check( ArrangedChildren.Num() == 1 );
-		FArrangedWidget& TheChild = ArrangedChildren(0);
+		FArrangedWidget& TheChild = ArrangedChildren[0];
 
 		const FSlateRect ChildClippingRect = AllottedGeometry.GetClippingRect().InsetBy( ChildSlot.SlotPadding.Get() * AllottedGeometry.Scale ).IntersectionWith(MyClippingRect);
 
diff --git a/Engine/Source/Runtime/Slate/Private/Widgets/Layout/SFxWidget.cpp b/Engine/Source/Runtime/Slate/Private/Widgets/Layout/SFxWidget.cpp
--- a/Engine/Source/Runtime/Slate/Private/Widgets/Layout/SFxWidget.cpp
+++ b/Engine/Source/Runtime/Slate/Private/Widgets/Layout/SFxWidget.cpp
@@ -82,7 +82,7 @@ int32 SFxWidget::OnPaint( const FPaintArgs& Args, const FGeometry& AllottedGeome
 	{
 		// We can only have one direct descendant.
 		check( ArrangedChildren.Num() == 1 );
-		const FArrangedWidget& TheChild = ArrangedChildren(0);
+		const FArrangedWidget& TheChild = ArrangedChildren[0];
 
 		// SFxWidgets are able to ignore parent clipping.
 		const FSlateRect ChildClippingRect = (bIgnoreClipping.Get())
diff --git a/Engine/Source/Runtime/Slate/Private/Widgets/Layout/SGridPanel.cpp b/Engine/Source/Runtime/Slate/Private/Widgets/Layout/SGridPanel.cpp
--- a/Engine/Source/Runtime/Slate/Private/Widgets/Layout/SGridPanel.cpp
+++ b/Engine/Source/Runtime/Slate/Private/Widgets/Layout/SGridPanel.cpp
@@ -66,7 +66,7 @@ int32 SGridPanel::OnPaint( const FPaintArgs& Args, const FGeometry& AllottedGeom
 	int32 LastGridLayer = 0;
 	for (int32 ChildIndex = 0; ChildIndex < Slots.Num(); ++ChildIndex)
 	{
-		FArrangedWidget& CurWidget = ArrangedChildren(ChildIndex);
+		FArrangedWidget& CurWidget = ArrangedChildren[ChildIndex];
 		if (CurWidget.Widget->GetVisibility().IsVisible())
 		{
 			const FSlot& CurSlot = Slots[ChildIndex];
diff --git a/Engine/Source/Runtime/Slate/Private/Widgets/Layout/SSplitter.cpp b/Engine/Source/Runtime/Slate/Private/Widgets/Layout/SSplitter.cpp
--- a/Engine/Source/Runtime/Slate/Private/Widgets/Layout/SSplitter.cpp
+++ b/Engine/Source/Runtime/Slate/Private/Widgets/Layout/SSplitter.cpp
@@ -164,7 +164,7 @@ int32 SSplitter::OnPaint( const FPaintArgs& Args, const FGeometry& AllottedGeome
 
 	for( int32 ChildIndex = 0; ChildIndex < ArrangedChildren.Num(); ++ChildIndex )
 	{
-		const FGeometry& GeometryAfterSplitter = ArrangedChildren( FMath::Clamp(ChildIndex + 1, 0, ArrangedChildren.Num()-1) ).Geometry;
+		const FGeometry& GeometryAfterSplitter = ArrangedChildren[ FMath::Clamp(ChildIndex + 1, 0, ArrangedChildren.Num()-1) ].Geometry;
 
 		const float HalfHitDetectionSplitterHandleSize = ( HitDetectionSplitterHandleSize / 2 );
 		const float HalfPhysicalSplitterHandleSize = ( PhysicalSplitterHandleSize / 2 );
@@ -494,7 +494,7 @@ void SSplitter::HandleResizing( const float PhysicalSplitterHandleSize, const ES
 	//  - Prev vs. Next refers to the widgets in the order they are laid out (left->right, top->bottom).
 	//  - New vs. Old refers to the Old values for width/height vs. the post-resize values.
 
-	const float HandlePos = ChildGeometries(DraggedHandle+1).Geometry.Position.Component(AxisIndex) - PhysicalSplitterHandleSize / 2;
+	const float HandlePos = ChildGeometries[DraggedHandle+1].Geometry.Position.Component(AxisIndex) - PhysicalSplitterHandleSize / 2;
 	float Delta = LocalMousePos.Component(AxisIndex) - HandlePos;
 
 	const int32 SlotBeforeDragHandle = FindResizeableSlotBeforeHandle( DraggedHandle, Children );
@@ -529,7 +529,7 @@ void SSplitter::HandleResizing( const float PhysicalSplitterHandleSize, const ES
 			FSlotInfo SlotInfo;
 
 			SlotInfo.Slot = &Children[ SlotsAfterDragHandleIndicies[ SlotIndex ] ];
-			SlotInfo.Geometry = &ChildGeometries( SlotsAfterDragHandleIndicies[ SlotIndex ] ).Geometry;
+			SlotInfo.Geometry = &ChildGeometries[ SlotsAfterDragHandleIndicies[ SlotIndex ] ].Geometry;
 			SlotInfo.NewSize = SlotInfo.Geometry->Size.Component( AxisIndex );
 
 			SlotsAfterDragHandle.Add( SlotInfo );
@@ -537,7 +537,7 @@ void SSplitter::HandleResizing( const float PhysicalSplitterHandleSize, const ES
 
 		// Get references the prev and next children and their layout settings so that we can modify them.
 		FSlot& PrevChild = Children[SlotBeforeDragHandle];
-		FGeometry& PrevChildGeom = ChildGeometries(SlotBeforeDragHandle).Geometry;
+		FGeometry& PrevChildGeom = ChildGeometries[SlotBeforeDragHandle].Geometry;
 
 		// Compute the new sizes of the children
 		float NewPrevChildLength = ClampChild( PrevChildGeom.Size.Component(AxisIndex) + Delta );
@@ -632,8 +632,8 @@ int32 SSplitter::GetHandleBeingResizedFromMousePosition( float PhysicalSplitterH
 	// Search for the two widgets between which the cursor currently resides.
 	for ( int32 ChildIndex = 1; ChildIndex < ChildGeometries.Num(); ++ChildIndex )
 	{
-		FGeometry& PrevChild = ChildGeometries(ChildIndex - 1).Geometry;
-		FGeometry& NextChild = ChildGeometries(ChildIndex).Geometry;
+		FGeometry& PrevChild = ChildGeometries[ChildIndex - 1].Geometry;
+		FGeometry& NextChild = ChildGeometries[ChildIndex].Geometry;
 		float PrevBound = PrevChild.Position.Component(AxisIndex) + PrevChild.Size.Component(AxisIndex) - HalfHitDetectionSplitterHandleSize + HalfPhysicalSplitterHandleSize;
 		float NextBound = NextChild.Position.Component(AxisIndex) + HalfHitDetectionSplitterHandleSize - HalfPhysicalSplitterHandleSize;
 
@@ -798,13 +798,13 @@ FCursorReply SSplitter2x2::OnCursorQuery( const FGeometry& MyGeometry, const FPo
 void SSplitter2x2::ResizeChildren( FArrangedChildren& ArrangedChildren, FVector2D LocalMousePos )
 {
 	// Compute the handle position.  The last child is used because it is always the furthest away from the origin
-	const FVector2D HandlePos = ArrangedChildren(3).Geometry.Position - (FVector2D( SplitterHandleSize , SplitterHandleSize ) * .5f);
+	const FVector2D HandlePos = ArrangedChildren[3].Geometry.Position - (FVector2D( SplitterHandleSize , SplitterHandleSize ) * .5f);
 	FVector2D Delta = LocalMousePos - HandlePos;
 
-	FGeometry& TopLeftGeom = ArrangedChildren(0).Geometry;
-	FGeometry& BotLeftGeom = ArrangedChildren(1).Geometry;
-	FGeometry& TopRightGeom = ArrangedChildren(2).Geometry;
-	FGeometry& BotRightGeom = ArrangedChildren(3).Geometry;
+	FGeometry& TopLeftGeom = ArrangedChildren[0].Geometry;
+	FGeometry& BotLeftGeom = ArrangedChildren[1].Geometry;
+	FGeometry& TopRightGeom = ArrangedChildren[2].Geometry;
+	FGeometry& BotRightGeom = ArrangedChildren[3].Geometry;
 
 	FSlot& TopLeft = Children[0];
 	FSlot& BotLeft = Children[1];
@@ -890,8 +890,8 @@ int32 SSplitter2x2::CalculateResizingAxis( const FGeometry& MyGeometry, const FV
 	// Search for the two widgets between which the cursor currently resides.
 	for ( int32 ChildIndex = 1; ChildIndex < ArrangedChildren.Num(); ++ChildIndex )
 	{
-		FGeometry& PrevChild = ArrangedChildren(ChildIndex - 1).Geometry;
-		FGeometry& NextChild = ArrangedChildren(ChildIndex).Geometry;
+		FGeometry& PrevChild = ArrangedChildren[ChildIndex - 1].Geometry;
+		FGeometry& NextChild = ArrangedChildren[ChildIndex].Geometry;
 		FVector2D PrevBound = PrevChild.Position + PrevChild.Size;
 		FVector2D NextBound = NextChild.Position;
 
diff --git a/Engine/Source/Runtime/Slate/Private/Widgets/SCanvas.cpp b/Engine/Source/Runtime/Slate/Private/Widgets/SCanvas.cpp
--- a/Engine/Source/Runtime/Slate/Private/Widgets/SCanvas.cpp
+++ b/Engine/Source/Runtime/Slate/Private/Widgets/SCanvas.cpp
@@ -107,7 +107,7 @@ int32 SCanvas::OnPaint( const FPaintArgs& Args, const FGeometry& AllottedGeometr
 
 	for (int32 ChildIndex = 0; ChildIndex < ArrangedChildren.Num(); ++ChildIndex)
 	{
-		FArrangedWidget& CurWidget = ArrangedChildren(ChildIndex);
+		FArrangedWidget& CurWidget = ArrangedChildren[ChildIndex];
 		FSlateRect ChildClipRect = MyClippingRect.IntersectionWith(CurWidget.Geometry.GetClippingRect());
 		const int32 CurWidgetsMaxLayerId = CurWidget.Widget->Paint(Args.WithNewParent(this), CurWidget.Geometry, ChildClipRect, OutDrawElements, MaxLayerId + 1, InWidgetStyle, ShouldBeEnabled(bParentEnabled));
 
diff --git a/Engine/Source/Runtime/Slate/Private/Widgets/SWeakWidget.cpp b/Engine/Source/Runtime/Slate/Private/Widgets/SWeakWidget.cpp
--- a/Engine/Source/Runtime/Slate/Private/Widgets/SWeakWidget.cpp
+++ b/Engine/Source/Runtime/Slate/Private/Widgets/SWeakWidget.cpp
@@ -59,7 +59,7 @@ int32 SWeakWidget::OnPaint( const FPaintArgs& Args, const FGeometry& AllottedGeo
 	if( ArrangedChildren.Num() > 0 )
 	{
 		check( ArrangedChildren.Num() == 1 );
-		FArrangedWidget& TheChild = ArrangedChildren(0);
+		FArrangedWidget& TheChild = ArrangedChildren[0];
 
 		return TheChild.Widget->Paint( 
 			Args.WithNewParent(this), 
diff --git a/Engine/Source/Runtime/SlateCore/Private/Layout/WidgetPath.cpp b/Engine/Source/Runtime/SlateCore/Private/Layout/WidgetPath.cpp
--- a/Engine/Source/Runtime/SlateCore/Private/Layout/WidgetPath.cpp
+++ b/Engine/Source/Runtime/SlateCore/Private/Layout/WidgetPath.cpp
@@ -7,9 +7,9 @@ FArrangedWidget FWidgetPath::FindArrangedWidget( TSharedRef<const SWidget> Widge
 {
 	for(int32 WidgetIndex = 0; WidgetIndex < Widgets.Num(); ++WidgetIndex)
 	{
-		if ( Widgets(WidgetIndex).Widget == WidgetToFind )
+		if ( Widgets[WidgetIndex].Widget == WidgetToFind )
 		{
-			return Widgets(WidgetIndex);
+			return Widgets[WidgetIndex];
 		}
 	}
 
@@ -64,29 +64,29 @@ bool FWidgetPath::MoveFocus(int32 PathLevel, EFocusMoveDirection::Type MoveDirec
 	
 		// Arrange the children so we can iterate through them regardless of widget type.
 		FArrangedChildren ArrangedChildren(EVisibility::Visible);
-		Widgets(PathLevel).Widget->ArrangeChildren( Widgets(PathLevel).Geometry, ArrangedChildren );
+		Widgets[PathLevel].Widget->ArrangeChildren( Widgets[PathLevel].Geometry, ArrangedChildren );
 
 		// Find the currently focused child among the children.
-		int32 FocusedChildIndex = ArrangedChildren.FindItemIndex( Widgets(PathLevel+1) );
+		int32 FocusedChildIndex = ArrangedChildren.FindItemIndex( Widgets[PathLevel+1] );
 		FocusedChildIndex = (FocusedChildIndex) % ArrangedChildren.Num() + MoveDirectionAsInt;
 
 		// Now actually search for the widget.
 		for( ; FocusedChildIndex < ArrangedChildren.Num() && FocusedChildIndex >= 0; FocusedChildIndex += MoveDirectionAsInt )
 		{
 			// Neither disabled widgets nor their children can be focused.
-			if ( ArrangedChildren(FocusedChildIndex).Widget->IsEnabled() )
+			if ( ArrangedChildren[FocusedChildIndex].Widget->IsEnabled() )
 			{
 				// Look for a focusable descendant.
-				FArrangedChildren PathToFocusableChild = GeneratePathToWidget( FFocusableWidgetMatcher(), ArrangedChildren(FocusedChildIndex), MoveDirection );
+				FArrangedChildren PathToFocusableChild = GeneratePathToWidget( FFocusableWidgetMatcher(), ArrangedChildren[FocusedChildIndex], MoveDirection );
 				// Either we found a focusable descendant, or an immediate child that is focusable.
-				const bool bFoundNextFocusable = ( PathToFocusableChild.Num() > 0 ) || ArrangedChildren(FocusedChildIndex).Widget->SupportsKeyboardFocus();
+				const bool bFoundNextFocusable = ( PathToFocusableChild.Num() > 0 ) || ArrangedChildren[FocusedChildIndex].Widget->SupportsKeyboardFocus();
 				if ( bFoundNextFocusable )
 				{
 					// We found the next focusable widget, so make this path point at the new widget by:
 					// First, truncating the FocusPath up to the current level (i.e. PathLevel).
 					Widgets.Remove( PathLevel+1, Widgets.Num()-PathLevel-1 );
 					// Second, add the immediate child that is focus or whose descendant is focused.
-					Widgets.AddWidget( ArrangedChildren(FocusedChildIndex) );
+					Widgets.AddWidget( ArrangedChildren[FocusedChildIndex] );
 					// Add path to focused descendants if any.
 					Widgets.Append( PathToFocusableChild );
 					// We successfully moved focus!
@@ -106,7 +106,7 @@ FWeakWidgetPath::FWeakWidgetPath( const FWidgetPath& InWidgetPath )
 {
 	for( int32 WidgetIndex = 0; WidgetIndex < InWidgetPath.Widgets.Num(); ++WidgetIndex )
 	{
-		Widgets.Add( TWeakPtr<SWidget>( InWidgetPath.Widgets(WidgetIndex).Widget ) );
+		Widgets.Add( TWeakPtr<SWidget>( InWidgetPath.Widgets[WidgetIndex].Widget ) );
 	}
 }
 
@@ -158,13 +158,13 @@ FWeakWidgetPath::EPathResolutionResult::Result FWeakWidgetPath::ToWidgetPath( FW
 				// Find the next widget in the path among the arranged children.
 				for( int32 SearchIndex = 0; !bFoundChild && SearchIndex < ArrangedChildren.Num(); ++SearchIndex )
 				{						
-					if ( ArrangedChildren(SearchIndex).Widget == WidgetPtrs[WidgetIndex+1] )
+					if ( ArrangedChildren[SearchIndex].Widget == WidgetPtrs[WidgetIndex+1] )
 					{
 						bFoundChild = true;
 						// Remember the widget and the associated geometry.
-						PathWithGeometries.AddWidget( ArrangedChildren(SearchIndex) );
+						PathWithGeometries.AddWidget( ArrangedChildren[SearchIndex] );
 						// The next child in the vertical slice will be arranged with respect to its parent's geometry.
-						ParentGeometry = ArrangedChildren(SearchIndex).Geometry;
+						ParentGeometry = ArrangedChildren[SearchIndex].Geometry;
 					}
 				}
 			}
diff --git a/Engine/Source/Runtime/SlateCore/Private/Widgets/SCompoundWidget.cpp b/Engine/Source/Runtime/SlateCore/Private/Widgets/SCompoundWidget.cpp
--- a/Engine/Source/Runtime/SlateCore/Private/Widgets/SCompoundWidget.cpp
+++ b/Engine/Source/Runtime/SlateCore/Private/Widgets/SCompoundWidget.cpp
@@ -18,7 +18,7 @@ int32 SCompoundWidget::OnPaint( const FPaintArgs& Args, const FGeometry& Allotte
 	if( ArrangedChildren.Num() > 0 )
 	{
 		check( ArrangedChildren.Num() == 1 );
-		FArrangedWidget& TheChild = ArrangedChildren(0);
+		FArrangedWidget& TheChild = ArrangedChildren[0];
 
 		const FSlateRect ChildClippingRect = AllottedGeometry.GetClippingRect().InsetBy( ChildSlot.SlotPadding.Get() * AllottedGeometry.Scale ).IntersectionWith(MyClippingRect);
 
diff --git a/Engine/Source/Runtime/SlateCore/Private/Widgets/SOverlay.cpp b/Engine/Source/Runtime/SlateCore/Private/Widgets/SOverlay.cpp
--- a/Engine/Source/Runtime/SlateCore/Private/Widgets/SOverlay.cpp
+++ b/Engine/Source/Runtime/SlateCore/Private/Widgets/SOverlay.cpp
@@ -93,7 +93,7 @@ int32 SOverlay::OnPaint( const FPaintArgs& Args, const FGeometry& AllottedGeomet
 
 	for (int32 ChildIndex = 0; ChildIndex < ArrangedChildren.Num(); ++ChildIndex)
 	{
-		FArrangedWidget& CurWidget = ArrangedChildren(ChildIndex);
+		FArrangedWidget& CurWidget = ArrangedChildren[ChildIndex];
 		FSlateRect ChildClipRect = MyClippingRect.IntersectionWith( CurWidget.Geometry.GetClippingRect() );
 		const int32 CurWidgetsMaxLayerId = CurWidget.Widget->Paint( Args.WithNewParent(this), CurWidget.Geometry, ChildClipRect, OutDrawElements, MaxLayerId + 1, InWidgetStyle, ShouldBeEnabled( bParentEnabled ) );
 
diff --git a/Engine/Source/Runtime/SlateCore/Private/Widgets/SPanel.cpp b/Engine/Source/Runtime/SlateCore/Private/Widgets/SPanel.cpp
--- a/Engine/Source/Runtime/SlateCore/Private/Widgets/SPanel.cpp
+++ b/Engine/Source/Runtime/SlateCore/Private/Widgets/SPanel.cpp
@@ -35,7 +35,7 @@ int32 SPanel::PaintArrangedChildren( const FPaintArgs& Args, const FArrangedChil
 
 	for (int32 ChildIndex = 0; ChildIndex < ArrangedChildren.Num(); ++ChildIndex)
 	{
-		const FArrangedWidget& CurWidget = ArrangedChildren(ChildIndex);
+		const FArrangedWidget& CurWidget = ArrangedChildren[ChildIndex];
 		FSlateRect ChildClipRect = MyClippingRect.IntersectionWith(CurWidget.Geometry.GetClippingRect());
 		
 		if (ChildClipRect.GetSize().Size() > 0)
diff --git a/Engine/Source/Runtime/SlateCore/Private/Widgets/SWidget.cpp b/Engine/Source/Runtime/SlateCore/Private/Widgets/SWidget.cpp
--- a/Engine/Source/Runtime/SlateCore/Private/Widgets/SWidget.cpp
+++ b/Engine/Source/Runtime/SlateCore/Private/Widgets/SWidget.cpp
@@ -263,7 +263,7 @@ void SWidget::TickWidgetsRecursively( const FGeometry& AllottedGeometry, const d
 	// Recur!
 	for(int32 ChildIndex=0; ChildIndex < ArrangedChildren.Num(); ++ChildIndex)
 	{
-		FArrangedWidget& SomeChild = ArrangedChildren(ChildIndex);
+		FArrangedWidget& SomeChild = ArrangedChildren[ChildIndex];
 		SomeChild.Widget->TickWidgetsRecursively( SomeChild.Geometry, InCurrentTime, InDeltaTime );
 	}
 }
@@ -351,7 +351,7 @@ void SWidget::FindChildGeometries_Helper( const FGeometry& MyGeometry, const TSe
 	// See if we found any of the widgets on this level.
 	for(int32 ChildIndex=0; ChildIndex < NumChildren; ++ChildIndex )
 	{
-		const FArrangedWidget& CurChild = ArrangedChildren( ChildIndex );
+		const FArrangedWidget& CurChild = ArrangedChildren[ ChildIndex ];
 		
 		if ( WidgetsToFind.Contains(CurChild.Widget) )
 		{
@@ -366,7 +366,7 @@ void SWidget::FindChildGeometries_Helper( const FGeometry& MyGeometry, const TSe
 		// Look for widgets among the children.
 		for( int32 ChildIndex=0; ChildIndex < NumChildren; ++ChildIndex )
 		{
-			const FArrangedWidget& CurChild = ArrangedChildren( ChildIndex );
+			const FArrangedWidget& CurChild = ArrangedChildren[ ChildIndex ];
 			CurChild.Widget->FindChildGeometries_Helper( CurChild.Geometry, WidgetsToFind, OutResult );
 		}	
 	}	
@@ -394,7 +394,7 @@ int32 SWidget::FindChildUnderMouse( const FArrangedChildren& Children, const FPo
 	const int32 NumChildren = Children.Num();
 	for( int32 ChildIndex=NumChildren-1; ChildIndex >= 0; --ChildIndex )
 	{
-		const FArrangedWidget& Candidate = Children(ChildIndex);
+		const FArrangedWidget& Candidate = Children[ChildIndex];
 		const bool bCandidateUnderCursor = 
 			// Candidate is physically under the cursor
 			Candidate.Geometry.IsUnderLocation( AbsoluteCursorLocation ) && 
diff --git a/Engine/Source/Runtime/SlateCore/Public/Layout/ArrangedChildren.h b/Engine/Source/Runtime/SlateCore/Public/Layout/ArrangedChildren.h
--- a/Engine/Source/Runtime/SlateCore/Public/Layout/ArrangedChildren.h
+++ b/Engine/Source/Runtime/SlateCore/Public/Layout/ArrangedChildren.h
@@ -107,17 +107,17 @@ class SLATECORE_API FArrangedChildren
 		return Array[Index];
 	}
 
-	// @todo umg na deprecate this
-	const FArrangedWidget& operator()(int32 Index) const
-	{
-		return Array[Index];
-	}
-
-	// @todo umg na deprecate this
-	FArrangedWidget& operator()(int32 Index)
-	{
-		return Array[Index];
-	}
+	//// @todo umg na deprecate this
+	//const FArrangedWidget& operator()(int32 Index) const
+	//{
+	//	return Array[Index];
+	//}
+
+	//// @todo umg na deprecate this
+	//FArrangedWidget& operator()(int32 Index)
+	//{
+	//	return Array[Index];
+	//}
 
 	const FArrangedWidget& Last() const
 	{
diff --git a/Engine/Source/Runtime/SlateCore/Public/Layout/WidgetPath.h b/Engine/Source/Runtime/SlateCore/Public/Layout/WidgetPath.h
--- a/Engine/Source/Runtime/SlateCore/Public/Layout/WidgetPath.h
+++ b/Engine/Source/Runtime/SlateCore/Public/Layout/WidgetPath.h
@@ -62,8 +62,8 @@ public:
 		bool bCopiedMarker = false;
 		for( int32 WidgetIndex = 0; !bCopiedMarker && WidgetIndex < Widgets.Num(); ++WidgetIndex )
 		{
-			ClippedPath.AddWidget( Widgets(WidgetIndex) );
-			bCopiedMarker = (Widgets(WidgetIndex).Widget == MarkerWidget);
+			ClippedPath.AddWidget( Widgets[WidgetIndex] );
+			bCopiedMarker = (Widgets[WidgetIndex].Widget == MarkerWidget);
 		}
 		
 		if ( bCopiedMarker )
@@ -83,7 +83,7 @@ public:
 	{
 		for(int32 WidgetIndex = 0; WidgetIndex < Widgets.Num(); ++WidgetIndex)
 		{
-			if ( Widgets(WidgetIndex).Widget == WidgetToFind )
+			if ( Widgets[WidgetIndex].Widget == WidgetToFind )
 			{
 				return true;
 			}
@@ -103,7 +103,7 @@ public:
 	{
 		check(IsValid());
 
-		TSharedRef<SWindow> FirstWidgetWindow = StaticCastSharedRef<SWindow>(Widgets(0).Widget);
+		TSharedRef<SWindow> FirstWidgetWindow = StaticCastSharedRef<SWindow>(Widgets[0].Widget);
 		return FirstWidgetWindow;
 	}
 
@@ -116,7 +116,7 @@ public:
 	{
 		check(IsValid());
 
-		TSharedRef<SWindow> FirstWidgetWindow = StaticCastSharedRef<SWindow>(Widgets(0).Widget);
+		TSharedRef<SWindow> FirstWidgetWindow = StaticCastSharedRef<SWindow>(Widgets[0].Widget);
 		return FirstWidgetWindow;
 	}
 	
@@ -128,7 +128,7 @@ public:
 		FString StringBuffer;
 		for( int32 WidgetIndex = Widgets.Num()-1; WidgetIndex >= 0; --WidgetIndex )
 		{
-			StringBuffer += Widgets(WidgetIndex).ToString();
+			StringBuffer += Widgets[WidgetIndex].ToString();
 			StringBuffer += TEXT("\n");
 		}
 		return StringBuffer;
@@ -152,7 +152,7 @@ public:
 
 		for( int32 WidgetIndex=0; WidgetIndex < Extension.Num(); ++WidgetIndex )
 		{
-			this->Widgets.AddWidget( Extension(WidgetIndex) );
+			this->Widgets.AddWidget( Extension[WidgetIndex] );
 		}
 
 		return Extension.Num() > 0;
diff --git a/Engine/Source/Runtime/SlateCore/Public/Layout/WidgetPath.inl b/Engine/Source/Runtime/SlateCore/Public/Layout/WidgetPath.inl
--- a/Engine/Source/Runtime/SlateCore/Public/Layout/WidgetPath.inl
+++ b/Engine/Source/Runtime/SlateCore/Public/Layout/WidgetPath.inl
@@ -22,7 +22,7 @@ bool FWidgetPath::SearchForWidgetRecursively( const MatchRuleType& MatchRule, co
 
 	for( int32 ChildIndex = 0; ChildIndex < ArrangedChildren.Num(); ++ChildIndex )
 	{
-		const FArrangedWidget& SomeChild = ArrangedChildren(ChildIndex);
+		const FArrangedWidget& SomeChild = ArrangedChildren[ChildIndex];
 		if ( MatchRule.IsMatch(SomeChild.Widget) )
 		{
 			OutReversedPath.AddWidget( SomeChild );
@@ -47,7 +47,7 @@ bool FWidgetPath::SearchForWidgetRecursively_Reverse( const MatchRuleType& Match
 
 	for( int32 ChildIndex = ArrangedChildren.Num()-1; ChildIndex >=0 ; --ChildIndex )
 	{
-		const FArrangedWidget& SomeChild = ArrangedChildren(ChildIndex);
+		const FArrangedWidget& SomeChild = ArrangedChildren[ChildIndex];
 		if ( MatchRule.IsMatch(SomeChild.Widget) )
 		{
 			OutReversedPath.AddWidget( SomeChild );
diff --git a/Engine/Source/Runtime/UMG/Private/Components/CanvasPanel.cpp b/Engine/Source/Runtime/UMG/Private/Components/CanvasPanel.cpp
--- a/Engine/Source/Runtime/UMG/Private/Components/CanvasPanel.cpp
+++ b/Engine/Source/Runtime/UMG/Private/Components/CanvasPanel.cpp
@@ -90,9 +90,9 @@ bool UCanvasPanel::GetGeometryForSlot(UCanvasPanelSlot* Slot, FGeometry& Arrange
 
 		for ( int32 ChildIndex = 0; ChildIndex < ArrangedChildren.Num(); ChildIndex++ )
 		{
-			if ( ArrangedChildren(ChildIndex).Widget == Slot->Content->TakeWidget() )
+			if ( ArrangedChildren[ChildIndex].Widget == Slot->Content->TakeWidget() )
 			{
-				ArrangedGeometry = ArrangedChildren(ChildIndex).Geometry;
+				ArrangedGeometry = ArrangedChildren[ChildIndex].Geometry;
 				return true;
 			}
 		}
diff --git a/Engine/Source/Runtime/UMG/Private/Slate/SConstraintCanvas.cpp b/Engine/Source/Runtime/UMG/Private/Slate/SConstraintCanvas.cpp
--- a/Engine/Source/Runtime/UMG/Private/Slate/SConstraintCanvas.cpp
+++ b/Engine/Source/Runtime/UMG/Private/Slate/SConstraintCanvas.cpp
@@ -156,7 +156,7 @@ int32 SConstraintCanvas::OnPaint( const FPaintArgs& Args, const FGeometry& Allot
 
 	for (int32 ChildIndex = 0; ChildIndex < ArrangedChildren.Num(); ++ChildIndex)
 	{
-		FArrangedWidget& CurWidget = ArrangedChildren(ChildIndex);
+		FArrangedWidget& CurWidget = ArrangedChildren[ChildIndex];
 		FSlateRect ChildClipRect = MyClippingRect.IntersectionWith(CurWidget.Geometry.GetClippingRect());
 		const int32 CurWidgetsMaxLayerId = CurWidget.Widget->Paint( Args.WithNewParent(this), CurWidget.Geometry, ChildClipRect, OutDrawElements, MaxLayerId + 1, InWidgetStyle, ShouldBeEnabled(bParentEnabled));
 
