Added FTextSnapshot to "snapshot" an FText at a given point in time.

This can then be used to detect changes in the FText, including live-culture changes (tested by temporarily allowing the editor to perform a live-culture change - some things that are still using FString rather than FText didn't update, but a fair amount of the UI did).

diff --git a/Engine/Source/Runtime/Core/Private/Internationalization/Text.cpp b/Engine/Source/Runtime/Core/Private/Internationalization/Text.cpp
--- a/Engine/Source/Runtime/Core/Private/Internationalization/Text.cpp
+++ b/Engine/Source/Runtime/Core/Private/Internationalization/Text.cpp
@@ -797,4 +797,42 @@ bool FText::IdenticalTo( const FText& Other ) const
 	return DisplayString == Other.DisplayString;
 }
 
+FTextSnapshot::FTextSnapshot()
+	: DisplayStringPtr()
+	, HistoryRevision(INDEX_NONE)
+	, Flags(0)
+{
+}
+
+FTextSnapshot::FTextSnapshot(const FText& InText)
+	: DisplayStringPtr(InText.DisplayString)
+	, HistoryRevision(InText.History.IsValid() ? InText.History->Revision : INDEX_NONE)
+	, Flags(InText.Flags)
+{
+}
+
+bool FTextSnapshot::IdenticalTo(const FText& InText) const
+{
+	// Make sure the string is up-to-date with the current culture
+	// (this usually happens when ToString() is called)
+	InText.Rebuild();
+
+	const int32 InHistoryRevision = InText.History.IsValid() ? InText.History->Revision : INDEX_NONE;
+	return DisplayStringPtr == InText.DisplayString
+		&& HistoryRevision == InHistoryRevision
+		&& Flags == InText.Flags;
+}
+
+bool FTextSnapshot::IsDisplayStringEqualTo(const FText& InText) const
+{
+	// Make sure the string is up-to-date with the current culture
+	// (this usually happens when ToString() is called)
+	InText.Rebuild();
+
+	// We have to assume that the display string has changed if the history of the text has changed
+	// (due to a culture change), as we no longer have the old display string to compare against
+	const int32 InHistoryRevision = InText.History.IsValid() ? InText.History->Revision : INDEX_NONE;
+	return HistoryRevision == InHistoryRevision && DisplayStringPtr.IsValid() && DisplayStringPtr->Equals(InText.ToString(), ESearchCase::CaseSensitive);
+}
+
 #undef LOCTEXT_NAMESPACE
diff --git a/Engine/Source/Runtime/Core/Private/Internationalization/TextHistory.cpp b/Engine/Source/Runtime/Core/Private/Internationalization/TextHistory.cpp
--- a/Engine/Source/Runtime/Core/Private/Internationalization/TextHistory.cpp
+++ b/Engine/Source/Runtime/Core/Private/Internationalization/TextHistory.cpp
@@ -41,10 +41,14 @@ void FTextHistory::SerializeForDisplayString(FArchive& Ar, TSharedRef<FString, E
 
 void FTextHistory::Rebuild(TSharedRef< FString, ESPMode::ThreadSafe > InDisplayString)
 {
-	if(IsOutOfDate())
-	{
-		Revision = FTextLocalizationManager::Get().GetHeadCultureRevision();
+	const bool bIsOutOfDate = IsOutOfDate();
+
+	// FTextHistory_Base will never report being out-of-date, but we need to keep the history revision in sync
+	// with the head culture regardless so that FTextSnapshot::IdenticalTo still works correctly
+	Revision = FTextLocalizationManager::Get().GetHeadCultureRevision();
 
+	if(bIsOutOfDate)
+	{
 		InDisplayString.Get() = FTextInspector::GetDisplayString(ToText(false));
 	}
 }
diff --git a/Engine/Source/Runtime/Core/Private/Internationalization/TextHistory.h b/Engine/Source/Runtime/Core/Private/Internationalization/TextHistory.h
--- a/Engine/Source/Runtime/Core/Private/Internationalization/TextHistory.h
+++ b/Engine/Source/Runtime/Core/Private/Internationalization/TextHistory.h
@@ -21,6 +21,8 @@ enum class ETextHistoryType
 class CORE_API FTextHistory
 {
 public:
+	friend class FTextSnapshot;
+
 	FTextHistory();
 
 	virtual ~FTextHistory() {};
diff --git a/Engine/Source/Runtime/Core/Public/Internationalization/Text.h b/Engine/Source/Runtime/Core/Public/Internationalization/Text.h
--- a/Engine/Source/Runtime/Core/Public/Internationalization/Text.h
+++ b/Engine/Source/Runtime/Core/Public/Internationalization/Text.h
@@ -351,6 +351,7 @@ public:
 	static const FText UnEscapedCloseBraceOutsideOfArgumentBlock;
 	static const FText SerializationFailureError;
 
+	friend class FTextSnapshot;
 	friend class FTextInspector;
 	friend class FInternationalization;
 	friend class UStruct;
@@ -366,6 +367,31 @@ public:
 #endif
 };
 
+/** A snapshot of an FText at a point in time that can be used to detect changes in the FText, including live-culture changes */
+class CORE_API FTextSnapshot
+{
+public:
+	FTextSnapshot();
+
+	explicit FTextSnapshot(const FText& InText);
+
+	/** Check to see whether the given text is identical to the text this snapshot was made from */
+	bool IdenticalTo(const FText& InText) const;
+
+	/** Check to see whether the display string of the given text is identical to the display string this snapshot was made from */
+	bool IsDisplayStringEqualTo(const FText& InText) const;
+
+private:
+	/** A pointer to the visible display string for the FText we took a snapshot of (used for an efficient pointer compare) */
+	TSharedPtr<FString, ESPMode::ThreadSafe> DisplayStringPtr;
+
+	/** Revision index of the history of the FText we took a snapshot of, or INDEX_NONE if there was no history */
+	int32 HistoryRevision;
+
+	/** Flags with various information on what sort of FText we took a snapshot of */
+	int32 Flags;
+};
+
 class CORE_API FTextInspector
 {
 private:
diff --git a/Engine/Source/Runtime/Slate/Private/Widgets/Text/SMultiLineEditableText.cpp b/Engine/Source/Runtime/Slate/Private/Widgets/Text/SMultiLineEditableText.cpp
--- a/Engine/Source/Runtime/Slate/Private/Widgets/Text/SMultiLineEditableText.cpp
+++ b/Engine/Source/Runtime/Slate/Private/Widgets/Text/SMultiLineEditableText.cpp
@@ -264,7 +264,7 @@ void SMultiLineEditableText::Construct( const FArguments& InArgs )
 		// Update the cached BoundText value to prevent it triggering another SetEditableText update again next Tick
 		if (BoundText.IsBound())
 		{
-			BoundTextLastTick = TextToSet;
+			BoundTextLastTick = FTextSnapshot(TextToSet);
 		}
 	}
 
@@ -327,7 +327,7 @@ void SMultiLineEditableText::SetText(const TAttribute< FText >& InText)
 	// Update the cached BoundText value to prevent it triggering another SetEditableText update again next Tick
 	if (BoundText.IsBound())
 	{
-		BoundTextLastTick = TextToSet;
+		BoundTextLastTick = FTextSnapshot(TextToSet);
 	}
 
 	// Update the internal editable text
@@ -1874,10 +1874,10 @@ void SMultiLineEditableText::Tick( const FGeometry& AllottedGeometry, const doub
 		bool bHasSetText = false;
 
 		const FText& TextToSet = BoundText.Get(FText::GetEmpty());
-		if (BoundText.IsBound() && !BoundTextLastTick.IdenticalTo(TextToSet))
+		if (!BoundTextLastTick.IdenticalTo(TextToSet))
 		{
 			// The pointer used by the bound text has changed, however the text may still be the same - check that now
-			if (!BoundTextLastTick.ToString().Equals(TextToSet.ToString(), ESearchCase::CaseSensitive))
+			if (!BoundTextLastTick.IsDisplayStringEqualTo(TextToSet))
 			{
 				// The source text has changed, so update the internal editable text
 				bHasSetText = true;
@@ -1885,7 +1885,7 @@ void SMultiLineEditableText::Tick( const FGeometry& AllottedGeometry, const doub
 			}
 
 			// Update this even if the text is lexically identical, as it will update the pointer compared by IdenticalTo for the next Tick
-			BoundTextLastTick = TextToSet;
+			BoundTextLastTick = FTextSnapshot(TextToSet);
 		}
 
 		if (!bHasSetText && Marshaller->IsDirty())
diff --git a/Engine/Source/Runtime/Slate/Private/Widgets/Text/SRichTextBlock.cpp b/Engine/Source/Runtime/Slate/Private/Widgets/Text/SRichTextBlock.cpp
--- a/Engine/Source/Runtime/Slate/Private/Widgets/Text/SRichTextBlock.cpp
+++ b/Engine/Source/Runtime/Slate/Private/Widgets/Text/SRichTextBlock.cpp
@@ -55,17 +55,17 @@ void SRichTextBlock::Tick( const FGeometry& AllottedGeometry, const double InCur
 
 	bool bRequiresTextUpdate = false;
 	const FText& TextToSet = GetText();
-	if (BoundText.IsBound() && !BoundTextLastTick.IdenticalTo(TextToSet))
+	if (!BoundTextLastTick.IdenticalTo(TextToSet))
 	{
 		// The pointer used by the bound text has changed, however the text may still be the same - check that now
-		if (!BoundTextLastTick.ToString().Equals(TextToSet.ToString(), ESearchCase::CaseSensitive))
+		if (!BoundTextLastTick.IsDisplayStringEqualTo(TextToSet))
 		{
 			// The source text has changed, so update the internal editable text
 			bRequiresTextUpdate = true;
 		}
 
 		// Update this even if the text is lexically identical, as it will update the pointer compared by IdenticalTo for the next Tick
-		BoundTextLastTick = TextToSet;
+		BoundTextLastTick = FTextSnapshot(TextToSet);
 	}
 
 	if (bRequiresTextUpdate || Marshaller->IsDirty())
@@ -141,7 +141,7 @@ void SRichTextBlock::SetText( const TAttribute<FText>& InTextAttr )
 	// Update the cached BoundText value to prevent it triggering another SetText update again next Tick
 	if (BoundText.IsBound())
 	{
-		BoundTextLastTick = InText;
+		BoundTextLastTick = FTextSnapshot(InText);
 	}
 
 	TextLayout->ClearLines();
diff --git a/Engine/Source/Runtime/Slate/Public/Widgets/Text/SMultiLineEditableText.h b/Engine/Source/Runtime/Slate/Public/Widgets/Text/SMultiLineEditableText.h
--- a/Engine/Source/Runtime/Slate/Public/Widgets/Text/SMultiLineEditableText.h
+++ b/Engine/Source/Runtime/Slate/Public/Widgets/Text/SMultiLineEditableText.h
@@ -483,8 +483,8 @@ private:
 	/** The text displayed in this text block */
 	TAttribute<FText> BoundText;
 
-	/** The state of BoundText last Tick() (only used when BoundText is bound to a delegate providing the source text) */
-	FText BoundTextLastTick;
+	/** The state of BoundText last Tick() (used to allow updates when the text is changed) */
+	FTextSnapshot BoundTextLastTick;
 
 	/** The marshaller used to get/set the BoundText text to/from the text layout. */
 	TSharedPtr< ITextLayoutMarshaller > Marshaller;
diff --git a/Engine/Source/Runtime/Slate/Public/Widgets/Text/SRichTextBlock.h b/Engine/Source/Runtime/Slate/Public/Widgets/Text/SRichTextBlock.h
--- a/Engine/Source/Runtime/Slate/Public/Widgets/Text/SRichTextBlock.h
+++ b/Engine/Source/Runtime/Slate/Public/Widgets/Text/SRichTextBlock.h
@@ -138,8 +138,8 @@ private:
 	/** The text displayed in this text block */
 	TAttribute< FText > BoundText;
 
-	/** The state of BoundText last Tick() (only used when BoundText is bound to a delegate providing the source text) */
-	FText BoundTextLastTick;
+	/** The state of BoundText last Tick() (used to allow updates when the text is changed) */
+	FTextSnapshot BoundTextLastTick;
 
 	/** The marshaller used to set the Text text in the text layout. */
 	TSharedPtr< FRichTextLayoutMarshaller > Marshaller;
