Fixed inertial scroll not being cancelled when entering overscroll

The problem here was that if inertial scroll was so strong that it caused an overscroll, and the overscroll came to a rest before the inertial scroll had naturally decelerated, the inertial scroll would kick back in when overscroll came to a rest. This has been fixed by clearing the inertial scroll on tick if it is not valid for it to occur.

Also fixed an issue where scroll samples would not be added when a touch-initated scroll was not moving. This could mean that a fast scroll followed by holding the scroll still would result in the inertial scroll kicking in with a fast velocity when the touch was released. This has been fixed by adding scroll samples in the tick function rather than on mouse/touch moved.

diff --git a/Engine/Source/Runtime/Slate/Private/Widgets/Views/STableViewBase.cpp b/Engine/Source/Runtime/Slate/Private/Widgets/Views/STableViewBase.cpp
--- a/Engine/Source/Runtime/Slate/Private/Widgets/Views/STableViewBase.cpp
+++ b/Engine/Source/Runtime/Slate/Private/Widgets/Views/STableViewBase.cpp
@@ -146,21 +146,47 @@ static FEndOfListResult ComputeOffsetForEndOfList( const FGeometry& ListPanelGeo
 	return FEndOfListResult( OffsetFromEndOfList, ItemsAboveView );
 }
 
-
-void STableViewBase::Tick( const FGeometry& AllottedGeometry, const double InCurrentTime, const float InDeltaTime )
+void STableViewBase::TickInertialScroll( const FGeometry& AllottedGeometry, const double InCurrentTime, const float InDeltaTime )
 {
-	if (ItemsPanel.IsValid())
+	if ( IsRightClickScrolling() )
 	{
-		if ( !IsRightClickScrolling() )
+		// We sample for the inertial scroll on tick rather than on mouse/touch move so
+		// that we still get samples even if the mouse has not moved.
+		if ( CanUseInertialScroll(TickScrollDelta) )
 		{
-			this->InertialScrollManager.UpdateScrollVelocity(InDeltaTime);
-			const float ScrollVelocity = this->InertialScrollManager.GetScrollVelocity();
+			this->InertialScrollManager.AddScrollSample(TickScrollDelta, InCurrentTime);
+		}
+	}
+	else
+	{
+		// If we are not right click scrolling then we can perform inertial scroll
+		this->InertialScrollManager.UpdateScrollVelocity(InDeltaTime);
+		const float ScrollVelocity = this->InertialScrollManager.GetScrollVelocity();
 
-			if ( ScrollVelocity != 0.f && CanUseInertialScroll(ScrollVelocity) )
+		if ( ScrollVelocity != 0.f )
+		{
+			if ( CanUseInertialScroll(ScrollVelocity) )
 			{
 				this->ScrollBy(AllottedGeometry, ScrollVelocity * InDeltaTime, EAllowOverscroll::Yes);
 			}
+			else
+			{
+				this->InertialScrollManager.ClearScrollVelocity();
+			}
+		}
+	}
+
+	TickScrollDelta = 0.f;
+}
 
+void STableViewBase::Tick( const FGeometry& AllottedGeometry, const double InCurrentTime, const float InDeltaTime )
+{
+	if (ItemsPanel.IsValid())
+	{
+		TickInertialScroll(AllottedGeometry, InCurrentTime, InDeltaTime);
+
+		if ( !IsRightClickScrolling() )
+		{
 			// If we are currently in overscroll, the list will need refreshing.
 			// Do this before UpdateOverscroll, as that could cause GetOverscroll() to be 0
 			if ( Overscroll.GetOverscroll() != 0.0f )
@@ -291,11 +317,6 @@ FReply STableViewBase::OnMouseButtonUp( const FGeometry& MyGeometry, const FPoin
 {
 	if ( MouseEvent.GetEffectingButton() == EKeys::RightMouseButton )
 	{
-		if ( IsRightClickScrolling() && Overscroll.GetOverscroll() != 0.f )
-		{
-			// Clear the scroll velocity if are overscrolling
-			InertialScrollManager.ClearScrollVelocity();
-		}
 
 		OnRightMouseButtonUp( MouseEvent.GetScreenSpacePosition() );
 
@@ -335,10 +356,7 @@ FReply STableViewBase::OnMouseMove( const FGeometry& MyGeometry, const FPointerE
 		// the mouse and dragging the view?
 		if( IsRightClickScrolling() )
 		{
-			if ( CanUseInertialScroll( -ScrollByAmount ) )
-			{
-				this->InertialScrollManager.AddScrollSample(-ScrollByAmount, FPlatformTime::Seconds());
-			}
+			TickScrollDelta -= ScrollByAmount;
 
 			const float AmountScrolled = this->ScrollBy( MyGeometry, -ScrollByAmount, EAllowOverscroll::Yes );
 
@@ -435,11 +453,7 @@ FReply STableViewBase::OnTouchMoved( const FGeometry& MyGeometry, const FPointer
 	{
 		const float ScrollByAmount = InTouchEvent.GetCursorDelta().Y / MyGeometry.Scale;
 		AmountScrolledWhileRightMouseDown += FMath::Abs( ScrollByAmount );
-
-		if ( CanUseInertialScroll( ScrollByAmount ) )
-		{
-			this->InertialScrollManager.AddScrollSample( ScrollByAmount, FPlatformTime::Seconds());
-		}
+		TickScrollDelta += ScrollByAmount;
 
 		const float AmountScrolled = this->ScrollBy( MyGeometry, ScrollByAmount, EAllowOverscroll::Yes );
 
@@ -460,11 +474,6 @@ FReply STableViewBase::OnTouchMoved( const FGeometry& MyGeometry, const FPointer
 
 FReply STableViewBase::OnTouchEnded( const FGeometry& MyGeometry, const FPointerEvent& InTouchEvent )
 {
-	if ( IsRightClickScrolling() && Overscroll.GetOverscroll() != 0.f )
-	{
-		// Clear the scroll velocity if are overscrolling
-		InertialScrollManager.ClearScrollVelocity();
-	}
 
 	AmountScrolledWhileRightMouseDown = 0;
 	bStartedTouchInteraction = false;
@@ -541,6 +550,7 @@ STableViewBase::STableViewBase( ETableViewMode::Type InTableViewMode )
 	, ScrollOffset( 0 )
 	, bStartedTouchInteraction( false )
 	, AmountScrolledWhileRightMouseDown( 0 )
+	, TickScrollDelta( 0 )
 	, LastGenerateResults( 0,0,0,false )
 	, bWasAtEndOfList(false)
 	, SelectionMode( ESelectionMode::Multi )
diff --git a/Engine/Source/Runtime/Slate/Public/Widgets/Views/STableViewBase.h b/Engine/Source/Runtime/Slate/Public/Widgets/Views/STableViewBase.h
--- a/Engine/Source/Runtime/Slate/Public/Widgets/Views/STableViewBase.h
+++ b/Engine/Source/Runtime/Slate/Public/Widgets/Views/STableViewBase.h
@@ -222,6 +222,9 @@ protected:
 	/** How much we scrolled while the rmb has been held */
 	float AmountScrolledWhileRightMouseDown;
 
+	/** The amount we have scrolled this tick cycle */
+	float TickScrollDelta;
+
 	/** Information about the widgets we generated during the last regenerate pass */
 	FReGenerateResults LastGenerateResults;
 
@@ -254,6 +257,9 @@ private:
 	/** Check whether the current state of the table warrants inertial scroll by the specified amount */
 	bool CanUseInertialScroll( float ScrollAmount ) const;
 
+	/** Called every tick to update and perform the inertial scroll */
+	void TickInertialScroll( const FGeometry& AllottedGeometry, const double InCurrentTime, const float InDeltaTime );
+
 protected:
 
 	struct SLATE_API FOverscroll
