Added support for right-to-left combo boxes via MenuPlacement_ComboBoxRight.

diff --git a/Engine/Source/Runtime/Slate/Private/Widgets/Input/SMenuAnchor.cpp b/Engine/Source/Runtime/Slate/Private/Widgets/Input/SMenuAnchor.cpp
--- a/Engine/Source/Runtime/Slate/Private/Widgets/Input/SMenuAnchor.cpp
+++ b/Engine/Source/Runtime/Slate/Private/Widgets/Input/SMenuAnchor.cpp
@@ -122,12 +122,14 @@ FGeometry SMenuAnchor::ComputeMenuPlacement( const FGeometry& AllottedGeometry,
 {
 	FVector2D PopupSize = PopupDesiredSize * AllottedGeometry.Scale;
 
-	if ( PlacementMode == MenuPlacement_ComboBox )
+	if ( PlacementMode == MenuPlacement_ComboBox || PlacementMode == MenuPlacement_ComboBoxRight )
 	{
 		PopupSize = FVector2D( FMath::Max( AllottedGeometry.Size.X, PopupDesiredSize.X ), PopupDesiredSize.Y ) * AllottedGeometry.Scale;
 	}
 
-	const float AbsoluteX = AllottedGeometry.AbsolutePosition.X;
+	const float AbsoluteX = (PlacementMode == MenuPlacement_ComboBoxRight)
+		? AllottedGeometry.AbsolutePosition.X + AllottedGeometry.GetDrawSize().X - PopupSize.X
+		: AllottedGeometry.AbsolutePosition.X;
 	const float AbsoluteY = AllottedGeometry.AbsolutePosition.Y;
 	FSlateRect Anchor(AbsoluteX, AbsoluteY, AbsoluteX + AllottedGeometry.Size.X*AllottedGeometry.Scale, AbsoluteY + AllottedGeometry.Size.Y*AllottedGeometry.Scale);
 	EOrientation Orientation = (PlacementMode == MenuPlacement_MenuRight) ? Orient_Horizontal : Orient_Vertical;
@@ -220,7 +222,7 @@ void SMenuAnchor::SetIsOpen( bool InIsOpen, const bool bFocusMenu )
 			const FVector2D SummonLocationSize = MyGeometry.Size;
 
 			FPopupTransitionEffect TransitionEffect( FPopupTransitionEffect::None );
-			if( PlacementMode == MenuPlacement_ComboBox )
+			if( PlacementMode == MenuPlacement_ComboBox || PlacementMode == MenuPlacement_ComboBoxRight )
 			{
 				TransitionEffect = FPopupTransitionEffect( FPopupTransitionEffect::ComboButton );
 				NewWindowSize = FVector2D( FMath::Max(MyGeometry.Size.X, DesiredContentSize.X), DesiredContentSize.Y );
diff --git a/Engine/Source/Runtime/SlateCore/Public/Types/SlateEnums.h b/Engine/Source/Runtime/SlateCore/Public/Types/SlateEnums.h
--- a/Engine/Source/Runtime/SlateCore/Public/Types/SlateEnums.h
+++ b/Engine/Source/Runtime/SlateCore/Public/Types/SlateEnums.h
@@ -120,6 +120,9 @@ enum EMenuPlacement
 	/** Place the menu immediately below the anchor and match is width to the anchor's content */
 	MenuPlacement_ComboBox UMETA(DisplayName="Combo Box"),
 
+	/** Place the menu immediately below the anchor and match is width to the anchor's content. If the width overflows, align with the right edge of the anchor. */
+	MenuPlacement_ComboBoxRight UMETA(DisplayName="Combo Box Right"),
+
 	/** Place the menu to the right of the anchor */
 	MenuPlacement_MenuRight UMETA(DisplayName="Menu Right"),
 
