Merging using UE4-To-UE4-LauncherReleases

diff --git a/Engine/Source/Runtime/Slate/Private/Framework/Text/SlateWidgetRun.cpp b/Engine/Source/Runtime/Slate/Private/Framework/Text/SlateWidgetRun.cpp
--- a/Engine/Source/Runtime/Slate/Private/Framework/Text/SlateWidgetRun.cpp
+++ b/Engine/Source/Runtime/Slate/Private/Framework/Text/SlateWidgetRun.cpp
@@ -6,14 +6,14 @@
 
 #include "SlateWidgetRun.h"
 
-TSharedRef< FSlateWidgetRun > FSlateWidgetRun::Create( const FRunInfo& InRunInfo, const TSharedRef< const FString >& InText, const FWidgetRunInfo& InWidgetInfo )
+TSharedRef< FSlateWidgetRun > FSlateWidgetRun::Create(const TSharedRef<class FTextLayout>& TextLayout, const FRunInfo& InRunInfo, const TSharedRef< const FString >& InText, const FWidgetRunInfo& InWidgetInfo)
 {
-	return MakeShareable( new FSlateWidgetRun( InRunInfo, InText, InWidgetInfo ) );
+	return MakeShareable(new FSlateWidgetRun(TextLayout, InRunInfo, InText, InWidgetInfo));
 }
 
-TSharedRef< FSlateWidgetRun > FSlateWidgetRun::Create( const FRunInfo& InRunInfo, const TSharedRef< const FString >& InText, const FWidgetRunInfo& InWidgetInfo, const FTextRange& InRange )
+TSharedRef< FSlateWidgetRun > FSlateWidgetRun::Create(const TSharedRef<class FTextLayout>& TextLayout, const FRunInfo& InRunInfo, const TSharedRef< const FString >& InText, const FWidgetRunInfo& InWidgetInfo, const FTextRange& InRange)
 {
-	return MakeShareable( new FSlateWidgetRun( InRunInfo, InText, InWidgetInfo, InRange ) );
+	return MakeShareable(new FSlateWidgetRun(TextLayout, InRunInfo, InText, InWidgetInfo, InRange));
 }
 
 FTextRange FSlateWidgetRun::GetTextRange() const 
@@ -53,6 +53,18 @@ TSharedRef< ILayoutBlock > FSlateWidgetRun::CreateBlock( int32 StartIndex, int32
 
 int32 FSlateWidgetRun::OnPaint( const FPaintArgs& Args, const FTextLayout::FLineView& Line, const TSharedRef< ILayoutBlock >& Block, const FTextBlockStyle& DefaultStyle, const FGeometry& AllottedGeometry, const FSlateRect& MyClippingRect, FSlateWindowElementList& OutDrawElements, int32 LayerId, const FWidgetStyle& InWidgetStyle, bool bParentEnabled ) const 
 {
+	const FVector2D DesiredWidgetSize = Info.Widget->GetDesiredSize();
+	if (DesiredWidgetSize != WidgetSize)
+	{
+		WidgetSize = DesiredWidgetSize;
+
+		const TSharedPtr<FTextLayout> TextLayoutPtr = TextLayout.Pin();
+		if (TextLayoutPtr.IsValid())
+		{
+			TextLayoutPtr->DirtyRunLayout(SharedThis(this));
+		}
+	}
+
 	FGeometry WidgetGeometry = AllottedGeometry.MakeChild( Block->GetLocationOffset() * ( 1 / AllottedGeometry.Scale ), Block->GetSize() * ( 1 / AllottedGeometry.Scale ), 1.0f );
 	return Info.Widget->Paint( Args, WidgetGeometry, MyClippingRect, OutDrawElements, LayerId, InWidgetStyle, bParentEnabled );
 }
@@ -85,7 +97,7 @@ void FSlateWidgetRun::Move(const TSharedRef<FString>& NewText, const FTextRange&
 
 TSharedRef<IRun> FSlateWidgetRun::Clone() const
 {
-	return FSlateWidgetRun::Create(RunInfo, Text, Info, Range);
+	return FSlateWidgetRun::Create(TextLayout.Pin().ToSharedRef(), RunInfo, Text, Info, Range);
 }
 
 void FSlateWidgetRun::AppendTextTo(FString& AppendToText) const
@@ -104,34 +116,42 @@ const FRunInfo& FSlateWidgetRun::GetRunInfo() const
 	return RunInfo;
 }
 
-FSlateWidgetRun::FSlateWidgetRun( const FRunInfo& InRunInfo, const TSharedRef< const FString >& InText, const FWidgetRunInfo& InWidgetInfo ) 
-	: RunInfo( InRunInfo )
+FSlateWidgetRun::FSlateWidgetRun(const TSharedRef<class FTextLayout>& TextLayout, const FRunInfo& InRunInfo, const TSharedRef< const FString >& InText, const FWidgetRunInfo& InWidgetInfo)
+	: TextLayout(TextLayout)
+	, RunInfo( InRunInfo )
 	, Text( InText )
 	, Range( 0, Text->Len() )
 	, Info( InWidgetInfo )
 	, Children()
+	, WidgetSize()
 {
 	Info.Widget->SlatePrepass();
-	Children.Add( Info.Widget );
+	WidgetSize = Info.Widget->GetDesiredSize();
+	Children.Add(Info.Widget);
 }
 
-FSlateWidgetRun::FSlateWidgetRun( const FRunInfo& InRunInfo, const TSharedRef< const FString >& InText, const FWidgetRunInfo& InWidgetInfo, const FTextRange& InRange ) 
-	: RunInfo( InRunInfo )
+FSlateWidgetRun::FSlateWidgetRun(const TSharedRef<class FTextLayout>& TextLayout, const FRunInfo& InRunInfo, const TSharedRef< const FString >& InText, const FWidgetRunInfo& InWidgetInfo, const FTextRange& InRange)
+	: TextLayout(TextLayout)
+	, RunInfo(InRunInfo)
 	, Text( InText )
 	, Range( InRange )
 	, Info( InWidgetInfo )
 	, Children()
+	, WidgetSize()
 {
 	Info.Widget->SlatePrepass();
+	WidgetSize = Info.Widget->GetDesiredSize();
 	Children.Add( Info.Widget );
 }
 
 FSlateWidgetRun::FSlateWidgetRun( const FSlateWidgetRun& Run ) 
-	: RunInfo( Run.RunInfo )
+	: TextLayout(Run.TextLayout)
+	, RunInfo(Run.RunInfo)
 	, Text( Run.Text )
 	, Range( Run.Range )
 	, Info( Run.Info )
 	, Children()
+	, WidgetSize(Run.WidgetSize)
 {
 
 }
diff --git a/Engine/Source/Runtime/Slate/Private/Framework/Text/TextDecorators.cpp b/Engine/Source/Runtime/Slate/Private/Framework/Text/TextDecorators.cpp
--- a/Engine/Source/Runtime/Slate/Private/Framework/Text/TextDecorators.cpp
+++ b/Engine/Source/Runtime/Slate/Private/Framework/Text/TextDecorators.cpp
@@ -36,7 +36,7 @@ TSharedRef< ISlateRun > FWidgetDecorator::Create(const TSharedRef<class FTextLay
 		RunInfo.MetaData.Add(Pair.Key, OriginalText.Mid( Pair.Value.BeginIndex, Pair.Value.EndIndex - Pair.Value.BeginIndex));
 	}
 
-	return FSlateWidgetRun::Create( RunInfo, InOutModelText, CreateWidgetDelegate.Execute( RunInfo, Style ), ModelRange );
+	return FSlateWidgetRun::Create(TextLayout, RunInfo, InOutModelText, CreateWidgetDelegate.Execute(RunInfo, Style), ModelRange);
 }
 
 TSharedRef< FImageDecorator > FImageDecorator::Create( FString InRunName, const ISlateStyle* const InStyle )
diff --git a/Engine/Source/Runtime/Slate/Public/Framework/Text/SlateWidgetRun.h b/Engine/Source/Runtime/Slate/Public/Framework/Text/SlateWidgetRun.h
--- a/Engine/Source/Runtime/Slate/Public/Framework/Text/SlateWidgetRun.h
+++ b/Engine/Source/Runtime/Slate/Public/Framework/Text/SlateWidgetRun.h
@@ -30,9 +30,9 @@ public:
 		}
 	};
 
-	static TSharedRef< FSlateWidgetRun > Create( const FRunInfo& InRunInfo, const TSharedRef< const FString >& InText, const FWidgetRunInfo& InWidgetInfo );
+	static TSharedRef< FSlateWidgetRun > Create(const TSharedRef<class FTextLayout>& TextLayout, const FRunInfo& InRunInfo, const TSharedRef< const FString >& InText, const FWidgetRunInfo& InWidgetInfo);
 
-	static TSharedRef< FSlateWidgetRun > Create( const FRunInfo& InRunInfo, const TSharedRef< const FString >& InText, const FWidgetRunInfo& InWidgetInfo, const FTextRange& InRange );
+	static TSharedRef< FSlateWidgetRun > Create(const TSharedRef<class FTextLayout>& TextLayout, const FRunInfo& InRunInfo, const TSharedRef< const FString >& InText, const FWidgetRunInfo& InWidgetInfo, const FTextRange& InRange);
 
 	virtual ~FSlateWidgetRun() {}
 
@@ -70,18 +70,22 @@ public:
 
 private:
 
-	FSlateWidgetRun( const FRunInfo& InRunInfo, const TSharedRef< const FString >& InText, const FWidgetRunInfo& InWidgetInfo );
+	FSlateWidgetRun(const TSharedRef<class FTextLayout>& TextLayout, const FRunInfo& InRunInfo, const TSharedRef< const FString >& InText, const FWidgetRunInfo& InWidgetInfo);
 
-	FSlateWidgetRun( const FRunInfo& InRunInfo, const TSharedRef< const FString >& InText, const FWidgetRunInfo& InWidgetInfo, const FTextRange& InRange );
+	FSlateWidgetRun(const TSharedRef<class FTextLayout>& TextLayout, const FRunInfo& InRunInfo, const TSharedRef< const FString >& InText, const FWidgetRunInfo& InWidgetInfo, const FTextRange& InRange);
 
 	FSlateWidgetRun( const FSlateWidgetRun& Run );
 
 private:
 
+	TWeakPtr<class FTextLayout> TextLayout;
 	FRunInfo RunInfo;
 	TSharedRef< const FString > Text;
 	FTextRange Range;
 	FWidgetRunInfo Info;
 	TArray< TSharedRef<SWidget> > Children;
+
+	mutable FVector2D WidgetSize;
+
 };
 #endif //WITH_FANCY_TEXT
\ No newline at end of file
