Fix handling of menus on OS X when the game thread is waiting on a main thread command.

diff --git a/Engine/Source/Runtime/Slate/Private/Framework/MultiBox/Mac/MacMenu.cpp b/Engine/Source/Runtime/Slate/Private/Framework/MultiBox/Mac/MacMenu.cpp
--- a/Engine/Source/Runtime/Slate/Private/Framework/MultiBox/Mac/MacMenu.cpp
+++ b/Engine/Source/Runtime/Slate/Private/Framework/MultiBox/Mac/MacMenu.cpp
@@ -140,7 +140,7 @@
 					}
 				}
 				return false;
-			});
+			}, InGameRunLoopMode(@[ UE4NilEventMode, UE4ShowEventMode, UE4ResizeEventMode, UE4FullscreenEventMode, UE4CloseEventMode, UE4IMEEventMode ]));
 		}
 
 		return NO;
@@ -233,7 +233,7 @@ CFStringRef FSlateMacMenu::GetMenuItemTitle(const TSharedRef<const FMenuEntryBlo
 
 		CFStringRef Title = FPlatformString::TCHARToCFString(*Label.Get().ToString());
 		return Title;
-	});
+	}, InGameRunLoopMode(@[ UE4NilEventMode, UE4ShowEventMode, UE4ResizeEventMode, UE4FullscreenEventMode, UE4CloseEventMode, UE4IMEEventMode ]));
 }
 
 NSImage* FSlateMacMenu::GetMenuItemIcon(const TSharedRef<const FMenuEntryBlock>& Block)
@@ -263,7 +263,7 @@ NSImage* FSlateMacMenu::GetMenuItemIcon(const TSharedRef<const FMenuEntryBlock>&
 			}
 		}
 		return MenuImage;
-	});
+	}, InGameRunLoopMode(@[ UE4NilEventMode, UE4ShowEventMode, UE4ResizeEventMode, UE4FullscreenEventMode, UE4CloseEventMode, UE4IMEEventMode ]));
 }
 
 void FSlateMacMenu::ExecuteMenuItemAction(const TSharedRef< const class FMenuEntryBlock >& Block)
@@ -313,7 +313,7 @@ CFStringRef FSlateMacMenu::GetMenuItemKeyEquivalent(const TSharedRef<const class
 			return FPlatformString::TCHARToCFString(*KeyString);
 		}
 		return FPlatformString::TCHARToCFString(TEXT(""));
-	});
+	}, InGameRunLoopMode(@[ UE4NilEventMode, UE4ShowEventMode, UE4ResizeEventMode, UE4FullscreenEventMode, UE4CloseEventMode, UE4IMEEventMode ]));
 }
 
 void FSlateMacMenu::UpdateMenu(FMacMenu* Menu)
@@ -392,7 +392,7 @@ void FSlateMacMenu::UpdateMenu(FMacMenu* Menu)
 		{
 			[NSApp setWindowsMenu:Menu];
 		}
-	});
+	}, InGameRunLoopMode(@[ UE4NilEventMode, UE4ShowEventMode, UE4ResizeEventMode, UE4FullscreenEventMode, UE4CloseEventMode, UE4IMEEventMode ]));
 }
 
 bool FSlateMacMenu::IsMenuItemEnabled(const TSharedRef<const class FMenuEntryBlock>& Block)
@@ -414,7 +414,7 @@ bool FSlateMacMenu::IsMenuItemEnabled(const TSharedRef<const class FMenuEntryBlo
 		}
 		
 		return bEnabled;
-	});
+	}, InGameRunLoopMode(@[ UE4NilEventMode, UE4ShowEventMode, UE4ResizeEventMode, UE4FullscreenEventMode, UE4CloseEventMode, UE4IMEEventMode ]));
 
 	return bEnabled;
 }
@@ -438,7 +438,7 @@ int32 FSlateMacMenu::GetMenuItemState(const TSharedRef<const class FMenuEntryBlo
 		}
 		
 		return bIsChecked;
-	});
+	}, InGameRunLoopMode(@[ UE4NilEventMode, UE4ShowEventMode, UE4ResizeEventMode, UE4FullscreenEventMode, UE4CloseEventMode, UE4IMEEventMode ]));
 
 	return bIsChecked ? NSOnState : NSOffState;
 }
@@ -455,5 +455,5 @@ TSharedRef<SWidget> FSlateMacMenu::GetMenuBuilderWidget(FMacMenu* Menu)
 		}
 		TSharedPtr<SWidget> Widget = MenuBuilder.MakeWidget();
 		return Widget;
-	}).ToSharedRef();
+	}, InGameRunLoopMode(@[ UE4NilEventMode, UE4ShowEventMode, UE4ResizeEventMode, UE4FullscreenEventMode, UE4CloseEventMode, UE4IMEEventMode ])).ToSharedRef();
 }
