Fixed crash when re-styling zero-length runs.

Zero-length runs that were contained in a removed section of text were previously causing an assert as they were non-adjacent to subsequent runs. Removing them fixes this.

diff --git a/Engine/Source/Runtime/Slate/Private/Framework/Text/TextLayout.cpp b/Engine/Source/Runtime/Slate/Private/Framework/Text/TextLayout.cpp
--- a/Engine/Source/Runtime/Slate/Private/Framework/Text/TextLayout.cpp
+++ b/Engine/Source/Runtime/Slate/Private/Framework/Text/TextLayout.cpp
@@ -1207,6 +1207,11 @@ bool FTextLayout::RemoveAt( const FTextLocation& Location, int32 Count )
 				break;
 			}
 		}
+		else if(IntersectedRangeToRemove.IsEmpty() && RunRange.IsEmpty() && RemoveTextRange.Contains(RunRange.BeginIndex) && RemoveTextRange.Contains(RunRange.EndIndex))
+		{
+			// empty run that was inside our removed range, safe to remove
+			LineModel.Runs.RemoveAt(RunIndex);
+		}
 	}
 
 	DirtyFlags |= EDirtyState::Layout;
