Added GetSelectedRuns() to SMultiLineEditableText and SMultiLineEditableTextBox

Allows behaviour similar to 'standard' text edtors where multiple selected run styles can be handled (e.g. in word, selecting multiple styles will blank the font selector), rather than just the style under the cursor.

diff --git a/Engine/Source/Runtime/Slate/Private/Widgets/Input/SMultiLineEditableTextBox.cpp b/Engine/Source/Runtime/Slate/Private/Widgets/Input/SMultiLineEditableTextBox.cpp
--- a/Engine/Source/Runtime/Slate/Private/Widgets/Input/SMultiLineEditableTextBox.cpp
+++ b/Engine/Source/Runtime/Slate/Private/Widgets/Input/SMultiLineEditableTextBox.cpp
@@ -240,4 +240,9 @@ TSharedPtr<const IRun> SMultiLineEditableTextBox::GetRunUnderCursor() const
 	return EditableText->GetRunUnderCursor();
 }
 
+const TArray<TSharedRef<const IRun>> SMultiLineEditableTextBox::GetSelectedRuns() const
+{
+	return EditableText->GetSelectedRuns();
+}
+
 #endif //WITH_FANCY_TEXT
diff --git a/Engine/Source/Runtime/Slate/Private/Widgets/Text/SMultiLineEditableText.cpp b/Engine/Source/Runtime/Slate/Private/Widgets/Text/SMultiLineEditableText.cpp
--- a/Engine/Source/Runtime/Slate/Private/Widgets/Text/SMultiLineEditableText.cpp
+++ b/Engine/Source/Runtime/Slate/Private/Widgets/Text/SMultiLineEditableText.cpp
@@ -1340,6 +1340,76 @@ TSharedPtr<const IRun> SMultiLineEditableText::GetRunUnderCursor() const
 	return nullptr;
 }
 
+const TArray<TSharedRef<const IRun>> SMultiLineEditableText::GetSelectedRuns() const
+{
+	TArray<TSharedRef<const IRun>> Runs;
+
+	if(AnyTextSelected())
+	{
+		const TArray<FTextLayout::FLineModel>& Lines = TextLayout->GetLineModels();
+		const FTextLocation CursorInteractionLocation = CursorInfo.GetCursorInteractionLocation();
+		if (Lines.IsValidIndex(SelectionStart.GetValue().GetLineIndex()) && Lines.IsValidIndex(CursorInteractionLocation.GetLineIndex()))
+		{
+			const FTextSelection Selection(SelectionStart.GetValue(), CursorInteractionLocation);
+			const int32 StartLine = Selection.GetBeginning().GetLineIndex();
+			const int32 EndLine = Selection.GetEnd().GetLineIndex();
+
+			// iterate over lines
+			for(int32 LineIndex = StartLine; LineIndex <= EndLine; LineIndex++)
+			{
+				const bool bIsFirstLine = LineIndex == StartLine;
+				const bool bIsLastLine = LineIndex == EndLine;
+
+				const FTextLayout::FLineModel& LineModel = Lines[LineIndex];
+				for (int32 RunIndex = 0; RunIndex < LineModel.Runs.Num(); ++RunIndex)
+				{
+					const FTextLayout::FRunModel& RunModel = LineModel.Runs[RunIndex];
+
+					// check what we should be intersecting with
+					if(!bIsFirstLine && !bIsLastLine)
+					{
+						// whole line is inside the range, so just add the run
+						Runs.Add(RunModel.GetRun());
+					}
+					else 
+					{
+						const FTextRange RunRange = RunModel.GetTextRange();
+						if(bIsFirstLine && !bIsLastLine)
+						{
+							// on first line of multi-line selection
+							const FTextRange IntersectedRange = RunRange.Intersect(FTextRange(Selection.GetBeginning().GetOffset(), LineModel.Text->Len()));
+							if (!IntersectedRange.IsEmpty())
+							{
+								Runs.Add(RunModel.GetRun());
+							}
+						}
+						else if(!bIsFirstLine && bIsLastLine)
+						{
+							// on last line of multi-line selection
+							const FTextRange IntersectedRange = RunRange.Intersect(FTextRange(0, Selection.GetEnd().GetOffset()));
+							if (!IntersectedRange.IsEmpty())
+							{
+								Runs.Add(RunModel.GetRun());
+							}
+						}
+						else
+						{
+							// single line selection
+							const FTextRange IntersectedRange = RunRange.Intersect(FTextRange(Selection.GetBeginning().GetOffset(), Selection.GetEnd().GetOffset()));
+							if (!IntersectedRange.IsEmpty())
+							{
+								Runs.Add(RunModel.GetRun());
+							}
+						}
+					}
+				}
+			}
+		}
+	}
+
+	return Runs;
+}
+
 bool SMultiLineEditableText::CanExecuteSelectAll() const
 {
 	bool bCanExecute = true;
diff --git a/Engine/Source/Runtime/Slate/Public/Widgets/Input/SMultiLineEditableTextBox.h b/Engine/Source/Runtime/Slate/Public/Widgets/Input/SMultiLineEditableTextBox.h
--- a/Engine/Source/Runtime/Slate/Public/Widgets/Input/SMultiLineEditableTextBox.h
+++ b/Engine/Source/Runtime/Slate/Public/Widgets/Input/SMultiLineEditableTextBox.h
@@ -181,6 +181,9 @@ public:
 	/** Get the run currently under the cursor, or null if there is no run currently under the cursor */
 	TSharedPtr<const IRun> GetRunUnderCursor() const;
 
+	/** Get the runs currently that are current selected, some of which may be only partially selected */
+	const TArray<TSharedRef<const IRun>> GetSelectedRuns() const;
+
 protected:
 
 	/** Editable text widget */
diff --git a/Engine/Source/Runtime/Slate/Public/Widgets/Text/SMultiLineEditableText.h b/Engine/Source/Runtime/Slate/Public/Widgets/Text/SMultiLineEditableText.h
--- a/Engine/Source/Runtime/Slate/Public/Widgets/Text/SMultiLineEditableText.h
+++ b/Engine/Source/Runtime/Slate/Public/Widgets/Text/SMultiLineEditableText.h
@@ -121,6 +121,9 @@ public:
 	/** Get the run currently under the cursor, or null if there is no run currently under the cursor */
 	TSharedPtr<const IRun> GetRunUnderCursor() const;
 
+	/** Get the runs currently that are current selected, some of which may be only partially selected */
+	const TArray<TSharedRef<const IRun>> GetSelectedRuns() const;
+
 private:
 	
 	/** 
