TabManager: Main Level Editor can no longer leave the tabwell of origin. It can only be re-arranged with respect to other tabs in that tabwell.

Fixes TTP 344479.

diff --git a/Engine/Source/Runtime/Slate/Private/Framework/Docking/FDockingDragOperation.cpp b/Engine/Source/Runtime/Slate/Private/Framework/Docking/FDockingDragOperation.cpp
--- a/Engine/Source/Runtime/Slate/Private/Framework/Docking/FDockingDragOperation.cpp
+++ b/Engine/Source/Runtime/Slate/Private/Framework/Docking/FDockingDragOperation.cpp
@@ -241,6 +241,42 @@ FVector2D FDockingDragOperation::GetTabGrabOffsetFraction() const
 	return TabGrabOffsetFraction;
 }
 
+
+bool FDockingDragOperation::CanDockInNode(const TSharedRef<SDockingNode>& DockNode, EViaTabwell IsDockingViaTabwell ) const
+{
+	const TSharedRef<FTabManager> TargetTabManager = DockNode->GetDockArea()->GetTabManager();
+	const TSharedRef<FTabManager> TabManagerOfOrigin = this->TabOwnerAreaOfOrigin->GetTabManager();
+
+	if (TabBeingDragged->GetTabRole() == ETabRole::NomadTab)
+	{
+		if ( IsDockingViaTabwell == FDockingDragOperation::DockingViaTabWell )
+		{
+			// Nomad tabs can be docked in in any tab well.
+			return true;
+		}
+		else
+		{
+			return TargetTabManager != FGlobalTabmanager::Get();
+		}
+	}
+	else if (TabBeingDragged->GetTabRole() == ETabRole::MajorTab)
+	{
+		// Major tabs can only be stacked; they should not 
+		// be allowed to split areas. They are also confined to their
+		// tab manager of origin.
+		// The only exception is an empty area, where docking the tab should be really easy.
+		const bool bTabManagerMatches = TargetTabManager == TabManagerOfOrigin;
+		const bool bCanDockInEmptyArea = DockNode->GetNodeType() == SDockingNode::DockArea && StaticCastSharedRef<SDockingArea>(DockNode)->GetChildNodes().Num() == 0;
+		return bTabManagerMatches && (IsDockingViaTabwell == FDockingDragOperation::DockingViaTabWell || bCanDockInEmptyArea);
+	}
+	else
+	{
+		// Most commonly, tabs are confined to their tab manager of origin.
+		return (TargetTabManager == TabManagerOfOrigin);
+	}	
+}
+
+
 FDockingDragOperation::~FDockingDragOperation()
 {
 }
@@ -251,6 +287,7 @@ FDockingDragOperation::FDockingDragOperation( const TSharedRef<SDockTab>& InTabT
 	: TabBeingDragged( InTabToBeDragged )
 	, TabGrabOffsetFraction( InTabGrabOffsetFraction )
 	, TabOwnerAreaOfOrigin( InTabOwnerArea )
+	, TabStackOfOrigin( InTabToBeDragged->GetParent()->GetParentDockTabStack() )
 	, HoveredTabPanelPtr( )
 	, HoveredDockTarget()
 	, LastContentSize( OwnerAreaSize )
diff --git a/Engine/Source/Runtime/Slate/Private/Framework/Docking/FDockingDragOperation.h b/Engine/Source/Runtime/Slate/Private/Framework/Docking/FDockingDragOperation.h
--- a/Engine/Source/Runtime/Slate/Private/Framework/Docking/FDockingDragOperation.h
+++ b/Engine/Source/Runtime/Slate/Private/Framework/Docking/FDockingDragOperation.h
@@ -2,6 +2,8 @@
 
 #pragma once
 
+class SDockingTabStack;
+
 /** A Sample implementation of IDragDropOperation */
 class FDockingDragOperation : public FDragDropOperation
 {
@@ -90,8 +92,10 @@ public:
 	/**
 	 * Create this Drag and Drop Content
 	 *
-	 * @param InWidgetBeingDragged    The Widget being dragged in this drag and drop operation
-	 * @param InTabOwnerStack         The Widget that originally contained what we are dragging
+	 * @param InTabToBeDragged    The tab being dragged (dragged within or outside of a tabwell).
+	 * @param InTabGrabOffset     Where within the tab we grabbed, so we're not dragging by the upper left of the tab.
+	 * @param InTabOwnerArea      The DockArea that owns this tab until it is relocated.
+	 * @param OwnerAreaSize       Size of the DockArea at the time when we start dragging.
 	 *
 	 * @return a new FDockingDragOperation
 	 */
@@ -103,6 +107,16 @@ public:
 	/** @return location where the user grabbed within the tab as a fraction of the tab's size */
 	FVector2D GetTabGrabOffsetFraction() const;
 
+	/**Is this dock tab being placed via a tab well or via a target*/
+	enum EViaTabwell
+	{
+		DockingViaTabWell,
+		DockingViaTarget
+	};
+
+	/** Checks to see if this tab can dock in this node. Some tabs can only dock via the tab well. */
+	bool CanDockInNode( const TSharedRef<SDockingNode>& DockNode, EViaTabwell IsDockingViaTabwell ) const;
+
 	virtual ~FDockingDragOperation();
 
 protected:
@@ -124,6 +138,9 @@ protected:
 	/** The area from which we initially started dragging */
 	TSharedPtr<SDockingArea> TabOwnerAreaOfOrigin;
 
+	/** Tab stack from which we started dragging this tab. */
+	TWeakPtr<class SDockingTabStack> TabStackOfOrigin;
+
 	/** The TabWell over which we are currently hovering */
 	TWeakPtr<class SDockingTabWell> HoveredTabPanelPtr;
 
diff --git a/Engine/Source/Runtime/Slate/Private/Framework/Docking/SDockingArea.cpp b/Engine/Source/Runtime/Slate/Private/Framework/Docking/SDockingArea.cpp
--- a/Engine/Source/Runtime/Slate/Private/Framework/Docking/SDockingArea.cpp
+++ b/Engine/Source/Runtime/Slate/Private/Framework/Docking/SDockingArea.cpp
@@ -10,7 +10,6 @@ void SDockingArea::Construct( const FArguments& InArgs, const TSharedRef<FTabMan
 	InTabManager->GetPrivateApi().OnDockAreaCreated( SharedThis(this) );
 
 	bManageParentWindow = InArgs._ShouldManageParentWindow;
-	bShouldAllowRemovalOfLastTab = InArgs._ShouldAllowRemovalOfLastTab;
 	bIsOverlayVisible = false;
 
 	bIsCenterTargetVisible = false;
@@ -103,7 +102,7 @@ void SDockingArea::OnDragEnter( const FGeometry& MyGeometry, const FDragDropEven
 	TSharedPtr<FDockingDragOperation> DragDropOperation = DragDropEvent.GetOperationAs<FDockingDragOperation>();
 	if ( DragDropOperation.IsValid() )
 	{
-		if ( DragDropOperation->GetTabBeingDragged()->CanDockInNode(SharedThis(this), SDockTab::DockingViaTarget) )
+		if ( DragDropOperation->CanDockInNode(SharedThis(this), FDockingDragOperation::DockingViaTarget) )
 		{
 			ShowCross();
 		}		
@@ -312,16 +311,6 @@ TSharedRef<FTabManager> SDockingArea::GetTabManager() const
 	return MyTabManager.Pin().ToSharedRef();
 }
 
-bool SDockingArea::GetShouldAllowRemovalOfLastTab() const
-{
-	return bShouldAllowRemovalOfLastTab;
-}
-
-void SDockingArea::SetShouldAllowRemovalOfLastTab(bool Allow)
-{
-	bShouldAllowRemovalOfLastTab = Allow;
-}
-
 SDockingNode::ECleanupRetVal SDockingArea::CleanUpNodes()
 {
 	SDockingNode::ECleanupRetVal ReturnValue = SDockingSplitter::CleanUpNodes();
diff --git a/Engine/Source/Runtime/Slate/Private/Framework/Docking/SDockingArea.h b/Engine/Source/Runtime/Slate/Private/Framework/Docking/SDockingArea.h
--- a/Engine/Source/Runtime/Slate/Private/Framework/Docking/SDockingArea.h
+++ b/Engine/Source/Runtime/Slate/Private/Framework/Docking/SDockingArea.h
@@ -14,7 +14,6 @@ public:
 	SLATE_BEGIN_ARGS(SDockingArea)
 		: _ParentWindow()
 		, _ShouldManageParentWindow(true)
-		, _ShouldAllowRemovalOfLastTab(true)
 		, _InitialContent()
 		{
 			// Visible by default, but don't absorb clicks
@@ -29,19 +28,6 @@ public:
 		/** True if this docking area should close the parent window when the last tab in this docking area goes away */
 		SLATE_ARGUMENT( bool, ShouldManageParentWindow )
 
-		/**
-		 * Should the dock area allow removing the last tab in the dock area.  By default this is true, but in cases like the main 
-		 * tab window, we don't want the last tab to ever be dragged out.
-		 */
-		SLATE_ARGUMENT(bool, ShouldAllowRemovalOfLastTab)
-
-
-		/** Should we ignore dragging the last tab in this well? */
-		bool IgnoreDraggingLastTabInWell() const;
-
-		/** Should we ignore dragging the last tab in this well? */
-		TAttribute<bool> bIgnoreDraggingLastTabInWell;
-
 		 /**
 		 * What to put into the DockArea initially. Usually a TabStack, so that some tabs can be added to it.
 		 */
@@ -89,10 +75,6 @@ public:
 
 	TSharedRef<FTabManager> GetTabManager() const;
 
-	bool GetShouldAllowRemovalOfLastTab() const;
-
-	void SetShouldAllowRemovalOfLastTab(bool Allow);
-
 protected:
 	virtual SDockingNode::ECleanupRetVal CleanUpNodes() override;
 
@@ -143,7 +125,4 @@ private:
 
 	/** True when the last tab has been pulled from this area, meaning that this DockArea will not be necessary once that tab finds a new home. */
 	bool bCleanUpUponTabRelocation;
-
-	/** True normally, prevents dragging out the last tab of a dock area if false. */
-	bool bShouldAllowRemovalOfLastTab;
 };
diff --git a/Engine/Source/Runtime/Slate/Private/Framework/Docking/SDockingTabStack.cpp b/Engine/Source/Runtime/Slate/Private/Framework/Docking/SDockingTabStack.cpp
--- a/Engine/Source/Runtime/Slate/Private/Framework/Docking/SDockingTabStack.cpp
+++ b/Engine/Source/Runtime/Slate/Private/Framework/Docking/SDockingTabStack.cpp
@@ -338,7 +338,7 @@ FReply SDockingTabStack::OnDragOver( const FGeometry& MyGeometry, const FDragDro
 	TSharedPtr<FDockingDragOperation> DragDropOperation = DragDropEvent.GetOperationAs<FDockingDragOperation>();
 	if ( DragDropOperation.IsValid() )
 	{	
-		if (DragDropOperation->GetTabBeingDragged()->CanDockInNode(SharedThis(this), SDockTab::DockingViaTarget))
+		if (DragDropOperation->CanDockInNode(SharedThis(this), FDockingDragOperation::DockingViaTarget))
 		{
 			FGeometry OverlayGeometry = this->FindChildGeometry( MyGeometry, OverlayManagement.ContentAreaOverlay.ToSharedRef() );
 	
diff --git a/Engine/Source/Runtime/Slate/Private/Framework/Docking/SDockingTabWell.cpp b/Engine/Source/Runtime/Slate/Private/Framework/Docking/SDockingTabWell.cpp
--- a/Engine/Source/Runtime/Slate/Private/Framework/Docking/SDockingTabWell.cpp
+++ b/Engine/Source/Runtime/Slate/Private/Framework/Docking/SDockingTabWell.cpp
@@ -240,29 +240,38 @@ float SDockingTabWell::ComputeDraggedTabOffset( const FGeometry& MyGeometry, con
 
 FReply SDockingTabWell::StartDraggingTab( TSharedRef<SDockTab> TabToStartDragging, FVector2D InTabGrabOffsetFraction, const FPointerEvent& MouseEvent )
 {
-	Tabs.Remove(TabToStartDragging);
-	// We just removed the foreground tab.
-	ForegroundTabIndex = INDEX_NONE;
-	ParentTabStackPtr.Pin()->OnTabRemoved(TabToStartDragging->GetLayoutIdentifier());
-
-	// Tha tab well keeps track of which tab we are dragging; we treat is specially during rendering and layout.
-	TabBeingDraggedPtr = TabToStartDragging;
-		
-	TabGrabOffsetFraction = InTabGrabOffsetFraction;
+	const bool bCanLeaveTabWell = TabToStartDragging->GetTabManager()->GetPrivateApi().CanTabLeaveTabWell( TabToStartDragging );
 
 	// We are about to start dragging a tab, so make sure its offset is correct
 	this->ChildBeingDraggedOffset = ComputeDraggedTabOffset( MouseEvent.FindGeometry(SharedThis(this)), MouseEvent, InTabGrabOffsetFraction );
 
-	// Start dragging.
-	TSharedRef<FDockingDragOperation> DragDropOperation =
-		FDockingDragOperation::New(
-			TabToStartDragging,
-			InTabGrabOffsetFraction,
-			GetDockArea().ToSharedRef(),
-			ParentTabStackPtr.Pin()->GetTabStackGeometry().Size
-		);
+	// Tha tab well keeps track of which tab we are dragging; we treat is specially during rendering and layout.
+	TabBeingDraggedPtr = TabToStartDragging;	
+	TabGrabOffsetFraction = InTabGrabOffsetFraction;
+	
 
-	return FReply::Handled().BeginDragDrop( DragDropOperation );
+	if (bCanLeaveTabWell)
+	{
+		Tabs.Remove(TabToStartDragging);
+		// We just removed the foreground tab.
+		ForegroundTabIndex = INDEX_NONE;
+		ParentTabStackPtr.Pin()->OnTabRemoved(TabToStartDragging->GetLayoutIdentifier());
+
+		// Start dragging.
+		TSharedRef<FDockingDragOperation> DragDropOperation =
+			FDockingDragOperation::New(
+				TabToStartDragging,
+				InTabGrabOffsetFraction,
+				GetDockArea().ToSharedRef(),
+				ParentTabStackPtr.Pin()->GetTabStackGeometry().Size
+			);
+
+		return FReply::Handled().BeginDragDrop( DragDropOperation );
+	}
+	else
+	{
+		return FReply::Handled().CaptureMouse(SharedThis(this));
+	}
 }
 
 void SDockingTabWell::OnDragEnter( const FGeometry& MyGeometry, const FDragDropEvent& DragDropEvent )
@@ -270,7 +279,7 @@ void SDockingTabWell::OnDragEnter( const FGeometry& MyGeometry, const FDragDropE
 	TSharedPtr<FDockingDragOperation> DragDropOperation = DragDropEvent.GetOperationAs<FDockingDragOperation>();
 	if ( DragDropOperation.IsValid() )
 	{
-		if (DragDropOperation->GetTabBeingDragged()->CanDockInNode(ParentTabStackPtr.Pin().ToSharedRef(), SDockTab::DockingViaTabWell))
+		if (DragDropOperation->CanDockInNode(ParentTabStackPtr.Pin().ToSharedRef(), FDockingDragOperation::DockingViaTabWell))
 		{
 			// The user dragged a tab into this TabWell.
 
@@ -292,16 +301,16 @@ void SDockingTabWell::OnDragLeave( const FDragDropEvent& DragDropEvent )
 	TSharedPtr<FDockingDragOperation> DragDropOperation = DragDropEvent.GetOperationAs<FDockingDragOperation>();
 	if ( DragDropOperation.IsValid() )
 	{
+		TSharedRef<SDockingTabStack> ParentTabStack = ParentTabStackPtr.Pin().ToSharedRef();
+		TSharedPtr<SDockTab> TabBeingDragged = this->TabBeingDraggedPtr.Pin();
 		// Check for TabBeingDraggedPtr validity as it may no longer be valid when dragging tabs in game
-		if (this->TabBeingDraggedPtr.IsValid() && DragDropOperation->GetTabBeingDragged()->CanDockInNode(ParentTabStackPtr.Pin().ToSharedRef(), SDockTab::DockingViaTabWell))
+		if ( TabBeingDragged.IsValid() && DragDropOperation->CanDockInNode(ParentTabStack, FDockingDragOperation::DockingViaTabWell) )
 		{
 			// Update the DragAndDrop operation based on this change.
-			const int32 LastForegroundTabIndex = Tabs.Find(this->TabBeingDraggedPtr.Pin().ToSharedRef());
+			const int32 LastForegroundTabIndex = Tabs.Find(TabBeingDragged.ToSharedRef());
 
 			// The user is pulling a tab out of this TabWell.
-			this->TabBeingDraggedPtr.Pin()->SetParent();
-
-			TSharedPtr<SDockingTabStack> ParentDockNode = ParentTabStackPtr.Pin();
+			TabBeingDragged->SetParent();
 
 			// We are no longer dragging a tab in this tab well, so stop
 			// showing it in the TabWell.
@@ -311,15 +320,15 @@ void SDockingTabWell::OnDragLeave( const FDragDropEvent& DragDropEvent )
 			BringTabToFront( FMath::Max(LastForegroundTabIndex-1, 0) );
 
 			// We may have removed the last tab that this DockNode had.
-			if (Tabs.Num() == 0 )
+			if ( Tabs.Num() == 0 )
 			{
 				// Let the DockNode know that it is no longer needed.
-				ParentDockNode->OnLastTabRemoved();
+				ParentTabStack->OnLastTabRemoved();
 			}
 
 			GetDockArea()->CleanUp( SDockingNode::TabRemoval_DraggedOut );
 
-			const FGeometry& DockNodeGeometry = ParentDockNode->GetTabStackGeometry();
+			const FGeometry& DockNodeGeometry = ParentTabStack->GetTabStackGeometry();
 			DragDropOperation->OnTabWellLeft( SharedThis(this), DockNodeGeometry );
 		}
 	}
@@ -331,7 +340,7 @@ FReply SDockingTabWell::OnDragOver( const FGeometry& MyGeometry, const FDragDrop
 	TSharedPtr<FDockingDragOperation> DragDropOperation = DragDropEvent.GetOperationAs<FDockingDragOperation>();
 	if ( DragDropOperation.IsValid() )
 	{
-		if (DragDropOperation->GetTabBeingDragged()->CanDockInNode(ParentTabStackPtr.Pin().ToSharedRef(), SDockTab::DockingViaTabWell))
+		if (DragDropOperation->CanDockInNode(ParentTabStackPtr.Pin().ToSharedRef(), FDockingDragOperation::DockingViaTabWell))
 		{
 			// We are dragging the tab through a TabWell.
 			// Update the position of the Tab that we are dragging in the panel.
@@ -347,7 +356,7 @@ FReply SDockingTabWell::OnDrop( const FGeometry& MyGeometry, const FDragDropEven
 	TSharedPtr<FDockingDragOperation> DragDropOperation = DragDropEvent.GetOperationAs<FDockingDragOperation>();
 	if (DragDropOperation.IsValid())
 	{
-		if (DragDropOperation->GetTabBeingDragged()->CanDockInNode(ParentTabStackPtr.Pin().ToSharedRef(), SDockTab::DockingViaTabWell))
+		if (DragDropOperation->CanDockInNode(ParentTabStackPtr.Pin().ToSharedRef(), FDockingDragOperation::DockingViaTabWell))
 		{
 			if ( ensure( TabBeingDraggedPtr.IsValid() ) )
 			{
@@ -384,6 +393,35 @@ EWindowZone::Type SDockingTabWell::GetWindowZoneOverride() const
 	return EWindowZone::TitleBar;
 }
 
+FReply SDockingTabWell::OnMouseButtonUp( const FGeometry& MyGeometry, const FPointerEvent& MouseEvent )
+{
+	if (this->HasMouseCapture())
+	{
+		// We are no longer dragging a tab in this tab well, so stop showing it in the TabWell.
+		this->TabBeingDraggedPtr.Reset();
+		return FReply::Handled().ReleaseMouseCapture();
+	}
+	else
+	{
+		return FReply::Unhandled();
+	}
+}
+
+FReply SDockingTabWell::OnMouseMove( const FGeometry& MyGeometry, const FPointerEvent& MouseEvent )
+{
+	if (this->HasMouseCapture())
+	{
+		// We are dragging the tab through a TabWell.
+		// Update the position of the Tab that we are dragging in the panel.
+		this->ChildBeingDraggedOffset = ComputeDraggedTabOffset(MyGeometry, MouseEvent, TabGrabOffsetFraction);
+		return FReply::Handled();
+	}
+	else
+	{
+		return FReply::Unhandled();
+	}
+}
+
 
 void SDockingTabWell::BringTabToFront( int32 TabIndexToActivate )
 {
diff --git a/Engine/Source/Runtime/Slate/Private/Framework/Docking/SDockingTabWell.h b/Engine/Source/Runtime/Slate/Private/Framework/Docking/SDockingTabWell.h
--- a/Engine/Source/Runtime/Slate/Private/Framework/Docking/SDockingTabWell.h
+++ b/Engine/Source/Runtime/Slate/Private/Framework/Docking/SDockingTabWell.h
@@ -95,6 +95,8 @@ public:
 	virtual FReply OnDragOver( const FGeometry& MyGeometry, const FDragDropEvent& DragDropEvent ) override;
 	virtual FReply OnDrop( const FGeometry& MyGeometry, const FDragDropEvent& DragDropEvent ) override;
 	virtual EWindowZone::Type GetWindowZoneOverride() const override;
+	virtual FReply OnMouseButtonUp( const FGeometry& MyGeometry, const FPointerEvent& MouseEvent ) override;
+	virtual FReply OnMouseMove( const FGeometry& MyGeometry, const FPointerEvent& MouseEvent ) override;
 	// End of SWidget interface
 
 private:
diff --git a/Engine/Source/Runtime/Slate/Private/Framework/Docking/TabManager.cpp b/Engine/Source/Runtime/Slate/Private/Framework/Docking/TabManager.cpp
--- a/Engine/Source/Runtime/Slate/Private/Framework/Docking/TabManager.cpp
+++ b/Engine/Source/Runtime/Slate/Private/Framework/Docking/TabManager.cpp
@@ -390,6 +390,11 @@ void FTabManager::FPrivateApi::OnTabManagerClosing()
 	TabManager.OnTabManagerClosing();
 }
 
+bool FTabManager::FPrivateApi::CanTabLeaveTabWell(const TSharedRef<const SDockTab>& TabToTest) const
+{
+	return TabToTest != TabManager.MainNonCloseableTab.Pin();
+}
+
 const TArray< TWeakPtr<SDockingArea> >& FTabManager::FPrivateApi::GetLiveDockAreas() const
 {
 	return TabManager.DockAreas;
@@ -1089,19 +1094,16 @@ TSharedRef<SDockingNode> FTabManager::RestoreArea_Helper( const TSharedRef<FLayo
 		}
 		else
 		{
-			SAssignNew(NewDockAreaWidget, SDockingArea, SharedThis(this), NodeAsArea.ToSharedRef())
+			SAssignNew( NewDockAreaWidget, SDockingArea, SharedThis(this), NodeAsArea.ToSharedRef() )
 
 				// We only want to set a parent window on this dock area, if we need to have title area content
 				// embedded within it.  SDockingArea assumes that if it has a parent window set, then it needs to have
 				// title area content 
-				.ParentWindow(bEmbedTitleAreaContent ? ParentWindow : TSharedPtr<SWindow>())
+				.ParentWindow( bEmbedTitleAreaContent ? ParentWindow : TSharedPtr<SWindow>() )
 
 				// Never manage these windows, even if a parent window is set.  The owner will take care of
 				// destroying these windows.
-				.ShouldManageParentWindow(false)
-
-				// Don't allow the removal of the last tab, because it will leave an orphaned main frame.
-				.ShouldAllowRemovalOfLastTab(OwnerTabPtr.IsValid() ? true : false);
+				.ShouldManageParentWindow( false );
 
 			RestoreSplitterContent( NodeAsArea.ToSharedRef(), NewDockAreaWidget.ToSharedRef(), ParentWindow );
 		}
diff --git a/Engine/Source/Runtime/Slate/Private/Widgets/Docking/SDockTab.cpp b/Engine/Source/Runtime/Slate/Private/Widgets/Docking/SDockTab.cpp
--- a/Engine/Source/Runtime/Slate/Private/Widgets/Docking/SDockTab.cpp
+++ b/Engine/Source/Runtime/Slate/Private/Widgets/Docking/SDockTab.cpp
@@ -29,15 +29,7 @@ FReply SDockTab::OnMouseButtonDown( const FGeometry& MyGeometry, const FPointerE
 			TotalDraggedDistance = 0;
 			ActivateInParent(ETabActivationCause::UserClickedOnTab);
 
-			// If we're the last tab and we're disallowing the removal of the last tab, ignore the attempt to drag.
-			TSharedPtr<SDockingTabWell> ParentWell = ParentPtr.Pin();
-			TSharedPtr<SDockingArea> DockArea = GetDockArea();
-			if ( ParentWell.IsValid() && DockArea.IsValid() && ParentWell->GetNumTabs() == 1 && !DockArea->GetShouldAllowRemovalOfLastTab() )
-			{
-				return FReply::Unhandled();
-			}
-
-			return FReply::Handled().DetectDrag(SharedThis(this), EKeys::LeftMouseButton);
+			return FReply::Handled().DetectDrag( SharedThis(this), EKeys::LeftMouseButton );
 		}
 		else if ( MouseEvent.GetEffectingButton() == EKeys::MiddleMouseButton )
 		{
@@ -670,38 +662,6 @@ void SDockTab::SetTabIcon( const TAttribute<const FSlateBrush*> InTabIcon )
 	TabIcon = InTabIcon;
 }
 
-bool SDockTab::CanDockInNode(const TSharedRef<SDockingNode>& DockNode, EViaTabwell IsDockingViaTabwell ) const
-{
-	const TSharedRef<FTabManager> TargetTabManager = DockNode->GetDockArea()->GetTabManager();
-	if (this->TabRole == ETabRole::NomadTab)
-	{
-		if ( IsDockingViaTabwell == SDockTab::DockingViaTabWell )
-		{
-			// Nomad tabs can be docked in in any tab well.
-			return true;
-		}
-		else
-		{
-			return TargetTabManager != FGlobalTabmanager::Get();
-		}
-	}
-	else if (this->TabRole == ETabRole::MajorTab)
-	{
-		// Major tabs can only be stacked; they should not 
-		// be allowed to split areas. They are also confined to their
-		// tab manager of origin.
-		// The only exception is an empty area, where docking the tab should be really easy.
-		const bool bTabManagerMatches = TargetTabManager == this->GetTabManager();
-		const bool bCanDockInEmptyArea = DockNode->GetNodeType() == SDockingNode::DockArea && StaticCastSharedRef<SDockingArea>(DockNode)->GetChildNodes().Num() == 0;
-		return bTabManagerMatches && (IsDockingViaTabwell == SDockTab::DockingViaTabWell || bCanDockInEmptyArea);
-	}
-	else
-	{
-		// Most commonly, tabs are confined to their tab manager of origin.
-		return (TargetTabManager == this->GetTabManager());
-	}	
-}
-
 
 bool SDockTab::ShouldAutosize() const
 {
diff --git a/Engine/Source/Runtime/Slate/Public/Framework/Docking/TabManager.h b/Engine/Source/Runtime/Slate/Public/Framework/Docking/TabManager.h
--- a/Engine/Source/Runtime/Slate/Public/Framework/Docking/TabManager.h
+++ b/Engine/Source/Runtime/Slate/Public/Framework/Docking/TabManager.h
@@ -466,6 +466,7 @@ class SLATE_API FTabManager : public TSharedFromThis<FTabManager>
 				void OnTabClosing( const TSharedRef<SDockTab>& TabBeingClosed );
 				void OnDockAreaClosing( const TSharedRef<SDockingArea>& DockAreaThatIsClosing );
 				void OnTabManagerClosing();
+				bool CanTabLeaveTabWell(const TSharedRef<const SDockTab>& TabToTest) const;
 				const TArray< TWeakPtr<SDockingArea> >& GetLiveDockAreas() const;
 				/**
 				 * Notify the tab manager that the NewForegroundTab was brought to front and the BackgroundedTab was send to the background as a result.
diff --git a/Engine/Source/Runtime/Slate/Public/Widgets/Docking/SDockTab.h b/Engine/Source/Runtime/Slate/Public/Widgets/Docking/SDockTab.h
--- a/Engine/Source/Runtime/Slate/Public/Widgets/Docking/SDockTab.h
+++ b/Engine/Source/Runtime/Slate/Public/Widgets/Docking/SDockTab.h
@@ -169,16 +169,6 @@ public:
 
 	/** Sets the tab icon */
 	void SetTabIcon( const TAttribute<const FSlateBrush*> InTabIcon );
-	
-	/**Is this dock tab being placed via a tab well or via a target*/
-	enum EViaTabwell
-	{
-		DockingViaTabWell,
-		DockingViaTarget
-	};
-
-	/** Checks to see if this tab can dock in this node. Some tabs can only dock via the tab well. */
-	bool CanDockInNode(const TSharedRef<class SDockingNode>& DockNode, EViaTabwell IsDockingViaTabwell ) const;
 
 	/** Should this tab be sized based on its content. */
 	bool ShouldAutosize() const;
@@ -208,7 +198,7 @@ public:
 	void ActivateInParent(ETabActivationCause::Type InActivationCause);
 
 	/** Set the tab manager that is controlling this tab */
-	void SetTabManager( const TSharedPtr<class FTabManager>& InTabManager );
+	void SetTabManager( const TSharedPtr<FTabManager>& InTabManager );
 
 	/**
 	 * Set the custom code to execute for saving visual state in this tab.
@@ -294,7 +284,7 @@ protected:
 protected:
 
 	/** The tab manager that created this tab. */
-	TWeakPtr<class FTabManager> MyTabManager;
+	TWeakPtr<FTabManager> MyTabManager;
 
 	/** The stuff to show when this tab is selected */
 	TSharedRef<SWidget> Content;
