LockCursor() will not longer assert when it fails; instead we will get an ensure and info about which widget we wanted the lock the cursor to.

diff --git a/Engine/Source/Runtime/Slate/Private/Framework/Application/SlateApplication.cpp b/Engine/Source/Runtime/Slate/Private/Framework/Application/SlateApplication.cpp
--- a/Engine/Source/Runtime/Slate/Private/Framework/Application/SlateApplication.cpp
+++ b/Engine/Source/Runtime/Slate/Private/Framework/Application/SlateApplication.cpp
@@ -1920,36 +1920,42 @@ void FSlateApplication::LockCursor( const TSharedPtr<SWidget>& Widget )
 		{
 			// Get a path to this widget so we know the position and size of its geometry
 			FWidgetPath WidgetPath;
-			GeneratePathToWidgetChecked( Widget.ToSharedRef(), WidgetPath );
-
-			// The last widget in the path should be the widget we are locking the cursor to
-			FArrangedWidget& WidgetGeom = WidgetPath.Widgets[ WidgetPath.Widgets.Num() - 1 ];
-
-			TSharedRef<SWindow> Window = WidgetPath.GetWindow();
-			// Do not attempt to lock the cursor to the window if its not in the foreground.  It would cause annoying side effects
-			if( Window->GetNativeWindow()->IsForegroundWindow() )
+			const bool bFoundWidthToLockTo = GeneratePathToWidgetUnchecked( Widget.ToSharedRef(), WidgetPath );
+			if ( bFoundWidthToLockTo )
 			{
-				check( WidgetGeom.Widget == Widget );
+				// The last widget in the path should be the widget we are locking the cursor to
+				FArrangedWidget& WidgetGeom = WidgetPath.Widgets[WidgetPath.Widgets.Num() - 1];
+
+				TSharedRef<SWindow> Window = WidgetPath.GetWindow();
+				// Do not attempt to lock the cursor to the window if its not in the foreground.  It would cause annoying side effects
+				if (Window->GetNativeWindow()->IsForegroundWindow())
+				{
+					check(WidgetGeom.Widget == Widget);
 
-				FSlateRect SlateClipRect = WidgetGeom.Geometry.GetClippingRect();
+					FSlateRect SlateClipRect = WidgetGeom.Geometry.GetClippingRect();
 
-				// Generate a screen space clip rect based on the widgets geometry
+					// Generate a screen space clip rect based on the widgets geometry
 
-				// Note: We round the upper left coordinate of the clip rect so we guarantee the rect is inside the geometry of the widget.  If we truncated when there is a half pixel we would cause the clip
-				// rect to be half a pixel larger than the geometry and cause the mouse to go outside of the geometry.
-				RECT ClipRect;
-				ClipRect.left = FMath::RoundToInt( SlateClipRect.Left );
-				ClipRect.top = FMath::RoundToInt( SlateClipRect.Top );
-				ClipRect.right = FMath::TruncToInt( SlateClipRect.Right );
-				ClipRect.bottom = FMath::TruncToInt( SlateClipRect.Bottom );
+					// Note: We round the upper left coordinate of the clip rect so we guarantee the rect is inside the geometry of the widget.  If we truncated when there is a half pixel we would cause the clip
+					// rect to be half a pixel larger than the geometry and cause the mouse to go outside of the geometry.
+					RECT ClipRect;
+					ClipRect.left = FMath::RoundToInt(SlateClipRect.Left);
+					ClipRect.top = FMath::RoundToInt(SlateClipRect.Top);
+					ClipRect.right = FMath::TruncToInt(SlateClipRect.Right);
+					ClipRect.bottom = FMath::TruncToInt(SlateClipRect.Bottom);
 
-				// Lock the mouse to the widget
-				PlatformApplication->Cursor->Lock( &ClipRect );
+					// Lock the mouse to the widget
+					PlatformApplication->Cursor->Lock(&ClipRect);
+				}
+			}
+			else
+			{
+				ensureMsgf( false, TEXT("Attempting to LockCursor() to widget but could not find widget %s"), *Widget->ToString() );
 			}
 		}
 		else
 		{
-			// Unlock  the mouse
+			// Unlock the mouse
 			PlatformApplication->Cursor->Lock( NULL );
 		}
 	}
