Convert SMenuAnchor::ComputeMenuPlacement to use local space logic, and FGeometry::MakeChild to render transforms are respected.
If CalculatePopupWindowPosition was given the AllottedGeometry, it could keep the calculations in local space and be render transform aware as it computed the placement.

diff --git a/Engine/Source/Runtime/Slate/Private/Widgets/Input/SMenuAnchor.cpp b/Engine/Source/Runtime/Slate/Private/Widgets/Input/SMenuAnchor.cpp
--- a/Engine/Source/Runtime/Slate/Private/Widgets/Input/SMenuAnchor.cpp
+++ b/Engine/Source/Runtime/Slate/Private/Widgets/Input/SMenuAnchor.cpp
@@ -120,22 +120,26 @@ int32 SMenuAnchor::OnPaint( const FPaintArgs& Args, const FGeometry& AllottedGeo
 
 FGeometry SMenuAnchor::ComputeMenuPlacement( const FGeometry& AllottedGeometry, const FVector2D& PopupDesiredSize, EMenuPlacement PlacementMode )
 {
-	FVector2D PopupSize = PopupDesiredSize * AllottedGeometry.Scale;
-
-	if ( PlacementMode == MenuPlacement_ComboBox || PlacementMode == MenuPlacement_ComboBoxRight )
-	{
-		PopupSize = FVector2D( FMath::Max( AllottedGeometry.Size.X, PopupDesiredSize.X ), PopupDesiredSize.Y ) * AllottedGeometry.Scale;
-	}
-
-	const float AbsoluteX = (PlacementMode == MenuPlacement_ComboBoxRight)
-		? AllottedGeometry.AbsolutePosition.X + AllottedGeometry.GetDrawSize().X - PopupSize.X
-		: AllottedGeometry.AbsolutePosition.X;
-	const float AbsoluteY = AllottedGeometry.AbsolutePosition.Y;
-	FSlateRect Anchor(AbsoluteX, AbsoluteY, AbsoluteX + AllottedGeometry.Size.X*AllottedGeometry.Scale, AbsoluteY + AllottedGeometry.Size.Y*AllottedGeometry.Scale);
-	EOrientation Orientation = (PlacementMode == MenuPlacement_MenuRight) ? Orient_Horizontal : Orient_Vertical;
-	const FVector2D NewPositionDesktopSpace = FSlateApplication::Get().CalculatePopupWindowPosition( Anchor, PopupSize, Orientation );
-
-	return FGeometry( FVector2D::ZeroVector, NewPositionDesktopSpace, PopupSize / AllottedGeometry.Scale, AllottedGeometry.Scale );
+	// Compute the popup size, offset, and anchor rect  in local space
+	const FVector2D PopupSizeLocalSpace = (PlacementMode == MenuPlacement_ComboBox || PlacementMode == MenuPlacement_ComboBoxRight)
+		? FVector2D( FMath::Max( AllottedGeometry.Size.X, PopupDesiredSize.X ), PopupDesiredSize.Y )
+		: PopupDesiredSize;
+	const FVector2D OffsetLocalSpace = (PlacementMode == MenuPlacement_ComboBoxRight) 
+		? FVector2D(AllottedGeometry.GetLocalSize().X - PopupSizeLocalSpace.X, 0.0f) 
+		: FVector2D::ZeroVector;
+	const FSlateRect AnchorLocalSpace = FSlateRect::FromPointAndExtent(OffsetLocalSpace, AllottedGeometry.Size);
+	const EOrientation Orientation = (PlacementMode == MenuPlacement_MenuRight) ? Orient_Horizontal : Orient_Vertical;
+	
+	// ask the application to compute the proper desktop offset for the anchor. This requires the offsets to be in desktop space.
+	const FVector2D NewPositionDesktopSpace = FSlateApplication::Get().CalculatePopupWindowPosition( 
+		TransformRect(AllottedGeometry.GetAccumulatedLayoutTransform(), AnchorLocalSpace), 
+		TransformVector(AllottedGeometry.GetAccumulatedLayoutTransform(), PopupSizeLocalSpace), 
+		Orientation );
+
+	// transform the desktop offset into local space and use that as the layout transform for the child content.
+	return AllottedGeometry.MakeChild(
+		PopupSizeLocalSpace,
+		FSlateLayoutTransform(TransformPoint(Inverse(AllottedGeometry.GetAccumulatedLayoutTransform()), NewPositionDesktopSpace)));
 }
 
 void SMenuAnchor::RequestClosePopupWindow( const TSharedRef<SWindow>& PopupWindow )
