Refactored and optimized MacMenu code

diff --git a/Engine/Source/Runtime/Core/Private/Mac/MacApplication.cpp b/Engine/Source/Runtime/Core/Private/Mac/MacApplication.cpp
--- a/Engine/Source/Runtime/Core/Private/Mac/MacApplication.cpp
+++ b/Engine/Source/Runtime/Core/Private/Mac/MacApplication.cpp
@@ -13,23 +13,23 @@
 
 #include "ModuleManager.h"
 
-NSString* NSWindowDraggingFinished = @"NSWindowDraggingFinished";
+static NSString* NSWindowDraggingFinished = @"NSWindowDraggingFinished";
 
-void HandleNSEvent(NSEvent* Event)
+static void HandleNSEvent(NSEvent* Event)
 {
 	const bool bIsMouseClickOrKeyEvent = [Event type] == NSLeftMouseDown || [Event type] == NSLeftMouseUp
-	|| [Event type] == NSRightMouseDown || [Event type] == NSRightMouseUp
-	|| [Event type] == NSOtherMouseDown || [Event type] == NSOtherMouseUp
-	|| [Event type] == NSKeyDown || ([Event type] == NSKeyUp && !([Event modifierFlags] & NSCommandKeyMask));
-	
-	if( MacApplication )
+									  || [Event type] == NSRightMouseDown || [Event type] == NSRightMouseUp
+									  || [Event type] == NSOtherMouseDown || [Event type] == NSOtherMouseUp
+									  || [Event type] == NSKeyDown || [Event type] == NSKeyUp;
+
+	if (MacApplication)
 	{
-		if( !bIsMouseClickOrKeyEvent || [Event window] == NULL )
+		if (!bIsMouseClickOrKeyEvent || [Event window] == NULL)
 		{
 			FMacEvent::SendToGameRunLoop(Event, EMacEventSendMethod::Async);
 		}
-		
-		if( [Event type] == NSLeftMouseUp )
+
+		if ([Event type] == NSLeftMouseUp)
 		{
 			NSNotification* Notification = [NSNotification notificationWithName:NSWindowDraggingFinished object:[Event window]];
 			FMacEvent::SendToGameRunLoop(Notification, [Event window], EMacEventSendMethod::Async);
diff --git a/Engine/Source/Runtime/Slate/Private/Framework/MultiBox/Mac/MacMenu.cpp b/Engine/Source/Runtime/Slate/Private/Framework/MultiBox/Mac/MacMenu.cpp
--- a/Engine/Source/Runtime/Slate/Private/Framework/MultiBox/Mac/MacMenu.cpp
+++ b/Engine/Source/Runtime/Slate/Private/Framework/MultiBox/Mac/MacMenu.cpp
@@ -5,46 +5,22 @@
 #include "CocoaThread.h"
 
 @interface FMacMenuItem : NSMenuItem
-{
-	TSharedPtr< const FMenuEntryBlock > MenuEntryBlock;
-}
-
+@property (assign) TSharedPtr<const FMenuEntryBlock> MenuEntryBlock;
 - (void)performAction;
-
 @end
 
 @implementation FMacMenuItem
 
 - (id)initWithMenuEntryBlock:(TSharedRef< const FMenuEntryBlock >&)Block
 {
-	CFStringRef Title = FSlateMacMenu::GetMenuItemTitle(Block);
-	uint32 KeyModifiers;
-	CFStringRef KeyEquivalent = FSlateMacMenu::GetMenuItemKeyEquivalent(Block, &KeyModifiers);
-	self = [super initWithTitle:(NSString*)Title action:nil keyEquivalent:(NSString*)KeyEquivalent];
-	[self setKeyEquivalentModifierMask:KeyModifiers];
-	NSImage* MenuImage = FSlateMacMenu::GetMenuItemIcon(Block);
-	if(MenuImage)
-	{
-		[self setImage:MenuImage];
-		[MenuImage release];
-	}
-	CFRelease(Title);
-	CFRelease(KeyEquivalent);
-
-	MenuEntryBlock = Block;
-
-	if (FSlateMacMenu::IsMenuItemEnabled(Block))
-	{
-		[self setTarget:self];
-		[self setAction:@selector(performAction)];
-	}
-
+	self = [super initWithTitle:@"" action:nil keyEquivalent:@""];
+	self.MenuEntryBlock = Block;
 	return self;
 }
 
 - (void)performAction
 {
-	FSlateMacMenu::ExecuteMenuItemAction(MenuEntryBlock.ToSharedRef());
+	FSlateMacMenu::ExecuteMenuItemAction(self.MenuEntryBlock.ToSharedRef());
 }
 
 @end
@@ -53,105 +29,17 @@
 
 - (id)initWithMenuEntryBlock:(TSharedRef< const FMenuEntryBlock >&)Block
 {
-	CFStringRef Title = FSlateMacMenu::GetMenuItemTitle(Block);
-	self = [super initWithTitle:(NSString*)Title];
-	CFRelease(Title);
-
+	self = [super initWithTitle:@""];
 	[self setDelegate:self];
-
-	MenuEntryBlock = Block;
-	
+	self.MenuEntryBlock = Block;
 	return self;
 }
 
-- (void)dealloc
-{
-	[super dealloc];
-}
-
 - (void)menuNeedsUpdate:(NSMenu*)Menu
 {
 	FSlateMacMenu::UpdateMenu(self);
 }
 
-- (BOOL)menuHasKeyEquivalent:(NSMenu*)menu forEvent:(NSEvent*)event target:(id*)target action:(SEL*)action
-{
-	// @todo: FSlateMacMenu::GetMenuBuilderWidget is currently expensive, so to avoid a framerate drop when using, for example, WASD keys, we always return false if no modifier keys are pressed.
-	// That's OK for now, because for a rare case where a menu item has a shortcut without any modifiers, Slate will handle it, but it needs to be fixed.
-	if ([event modifierFlags] == 0 || ![menu isKindOfClass:[FMacMenu class]])
-	{
-		return NO;
-	}
-	else
-	{
-		FMacMenu* MacMenu = (FMacMenu*)menu;
-
-		FText WindowLabel = NSLOCTEXT("MainMenu", "WindowMenu", "Window");
-		bool bIsWindowMenu = (WindowLabel.ToString().Compare(FString([MacMenu title])) == 0);
-
-		unichar Character = toupper([[event characters] characterAtIndex:0]);
-
-		if (bIsWindowMenu)
-		{
-			return ([event modifierFlags] & NSCommandKeyMask) && (Character == 'W' || Character == 'M');
-		}
-		else
-		{
-			return GameThreadReturn(^{
-				TSharedRef< SWidget > Widget = FSlateMacMenu::GetMenuBuilderWidget(MacMenu);
-
-				TSharedRef<const FMultiBox> MultiBox = StaticCastSharedRef<SMultiBoxWidget>(Widget)->GetMultiBox();
-				const TArray<TSharedRef<const FMultiBlock>>& MenuBlocks = MultiBox->GetBlocks();
-
-				for (int32 Index = 0; Index < MenuBlocks.Num(); Index++)
-				{
-					EMultiBlockType::Type Type = MenuBlocks[Index]->GetType();
-
-					if (Type == EMultiBlockType::MenuEntry)
-					{
-						TSharedRef<const FMenuEntryBlock> Block = StaticCastSharedRef<const FMenuEntryBlock>(MenuBlocks[Index]);
-						if (Block->GetAction().IsValid())
-						{
-							const TSharedRef<const FInputGesture>& Gesture = Block->GetAction()->GetActiveGesture();
-							if (!Gesture->GetKeyText().IsEmpty() && Gesture->GetKeyText().ToString()[0] == Character)
-							{
-								uint32 EventModifiers = [event modifierFlags] & (NSCommandKeyMask | NSShiftKeyMask | NSAlternateKeyMask | NSControlKeyMask);
-								uint32 GestureModfiers = 0;
-								if (Gesture->bCtrl)
-								{
-									GestureModfiers |= NSCommandKeyMask;
-								}
-								if (Gesture->bShift)
-								{
-									GestureModfiers |= NSShiftKeyMask;
-								}
-								if (Gesture->bAlt)
-								{
-									GestureModfiers |= NSAlternateKeyMask;
-								}
-								if (Gesture->bCmd)
-								{
-									GestureModfiers |= NSControlKeyMask;
-								}
-
-								return EventModifiers == GestureModfiers;
-							}
-						}
-					}
-				}
-				return false;
-			}, InGameRunLoopMode(@[ UE4NilEventMode, UE4ShowEventMode, UE4ResizeEventMode, UE4FullscreenEventMode, UE4CloseEventMode, UE4IMEEventMode ]));
-		}
-
-		return NO;
-	}
-}
-
-- ( TSharedPtr< const FMenuEntryBlock >& )GetMenuEntryBlock
-{
-	return MenuEntryBlock;
-}
-
 @end
 
 void FSlateMacMenu::UpdateWithMultiBox(const TSharedRef< FMultiBox >& MultiBox)
@@ -164,105 +52,162 @@ void FSlateMacMenu::UpdateWithMultiBox(const TSharedRef< FMultiBox >& MultiBox)
 		{
 			[[NSApp mainMenu] removeItemAtIndex:Index];
 		}
-		
+
 		const TArray<TSharedRef<const FMultiBlock> >& MenuBlocks = MultiBox->GetBlocks();
-		
+
 		for (int32 Index = 0; Index < MenuBlocks.Num(); Index++)
 		{
 			TSharedRef<const FMenuEntryBlock> Block = StaticCastSharedRef<const FMenuEntryBlock>(MenuBlocks[Index]);
 			FMacMenu* Menu = [[FMacMenu alloc] initWithMenuEntryBlock:Block];
-			
+			NSString* Title = FSlateMacMenu::GetMenuItemTitle(Block);
+			[Menu setTitle:Title];
+
 			NSMenuItem* MenuItem = [[NSMenuItem new] autorelease];
-			[MenuItem setTitle:[Menu title]];
+			[MenuItem setTitle:Title];
 			[[NSApp mainMenu] addItem:MenuItem];
 			[MenuItem setSubmenu:Menu];
 		}
-		
-		if([[NSApp mainMenu] itemWithTitle:HelpTitle.ToString().GetNSString()] == nil && HelpMenu)
+
+		if ([[NSApp mainMenu] itemWithTitle:HelpTitle.ToString().GetNSString()] == nil && HelpMenu)
 		{
 			[[NSApp mainMenu] addItem:HelpMenu];
 		}
-		if(HelpMenu)
+
+		if (HelpMenu)
 		{
 			[HelpMenu release];
 		}
 	});
 }
 
-static const TSharedRef<SWidget> FindTextBlockWidget(TSharedRef<SWidget> Content)
+void FSlateMacMenu::UpdateMenu(FMacMenu* Menu)
 {
-	if (Content->GetType() == FName(TEXT("STextBlock")))
-	{
-		return Content;
-	}
-
-	FChildren* Children = Content->GetChildren();
-	int32 NumChildren = Children->Num();
+	GameThreadCall(^{
+		FText WindowLabel = NSLOCTEXT("MainMenu", "WindowMenu", "Window");
+		const bool bIsWindowMenu = (WindowLabel.ToString().Compare(FString([Menu title])) == 0);
 
-	for (int32 Index=0; Index < NumChildren; ++Index)
-	{
-		const TSharedRef<SWidget> Found = FindTextBlockWidget(Children->GetChildAt(Index));
-		if (Found != SNullWidget::NullWidget)
+		if ([[Menu itemArray] count] == 0)
 		{
-			return Found;
+			if (bIsWindowMenu)
+			{
+				NSMenuItem* MinimizeItem = [[[NSMenuItem alloc] initWithTitle:@"Minimize" action:@selector(miniaturize:) keyEquivalent:@"m"] autorelease];
+				NSMenuItem* ZoomItem = [[[NSMenuItem alloc] initWithTitle:@"Zoom" action:@selector(performZoom:) keyEquivalent:@""] autorelease];
+				NSMenuItem* CloseItem = [[[NSMenuItem alloc] initWithTitle:@"Close" action:@selector(performClose:) keyEquivalent:@"w"] autorelease];
+				NSMenuItem* BringAllToFrontItem = [[[NSMenuItem alloc] initWithTitle:@"Bring All to Front" action:@selector(arrangeInFront:) keyEquivalent:@""] autorelease];
+
+				[Menu addItem:MinimizeItem];
+				[Menu addItem:ZoomItem];
+				[Menu addItem:CloseItem];
+				[Menu addItem:[NSMenuItem separatorItem]];
+				[Menu addItem:BringAllToFrontItem];
+				[Menu addItem:[NSMenuItem separatorItem]];
+			}
 		}
-	}
-	return SNullWidget::NullWidget;
-}
 
-CFStringRef FSlateMacMenu::GetMenuItemTitle(const TSharedRef<const FMenuEntryBlock>& Block)
-{
-	return GameThreadReturn(^{
-		TAttribute<FText> Label;
-		if (!Block->LabelOverride.IsBound() && Block->LabelOverride.Get().IsEmpty() && Block->GetAction().IsValid())
-		{
-			Label = Block->GetAction()->GetLabel();
-		}
-		else if (!Block->LabelOverride.Get().IsEmpty())
+		if (!Menu.MultiBox.IsValid())
 		{
-			Label = Block->LabelOverride;
-		}
-		else if (Block->EntryWidget.IsValid())
-		{
-			const TSharedRef<SWidget>& TextBlockWidget = FindTextBlockWidget(Block->EntryWidget.ToSharedRef());
-			if (TextBlockWidget != SNullWidget::NullWidget)
+			const bool bShouldCloseWindowAfterMenuSelection = true;
+			FMenuBuilder MenuBuilder(bShouldCloseWindowAfterMenuSelection, Menu.MenuEntryBlock->GetActionList(), Menu.MenuEntryBlock->Extender);
 			{
-				Label = StaticCastSharedRef<STextBlock>(TextBlockWidget)->GetText();
+				// Have the menu fill its contents
+				Menu.MenuEntryBlock->EntryBuilder.ExecuteIfBound(MenuBuilder);
 			}
+			TSharedRef<SWidget> Widget = MenuBuilder.MakeWidget();
+			Menu.MultiBox = TSharedPtr<const FMultiBox>(StaticCastSharedRef<SMultiBoxWidget>(Widget)->GetMultiBox());
 		}
 
-		CFStringRef Title = FPlatformString::TCHARToCFString(*Label.Get().ToString());
-		return Title;
-	}, InGameRunLoopMode(@[ UE4NilEventMode, UE4ShowEventMode, UE4ResizeEventMode, UE4FullscreenEventMode, UE4CloseEventMode, UE4IMEEventMode ]));
-}
-
-NSImage* FSlateMacMenu::GetMenuItemIcon(const TSharedRef<const FMenuEntryBlock>& Block)
-{
-	return GameThreadReturn(^{
-		NSImage* MenuImage = nil;
-		FSlateIcon Icon;
-		if(Block->IconOverride.IsSet())
-		{
-			Icon = Block->IconOverride;
-		}
-		else if(Block->GetAction().IsValid() && Block->GetAction()->GetIcon().IsSet())
-		{
-			Icon = Block->GetAction()->GetIcon();
-		}
-		if(Icon.IsSet())
+		const TArray<TSharedRef<const FMultiBlock>>& MenuBlocks = Menu.MultiBox->GetBlocks();
+		int32 ItemIndexAdjust = 0;
+		for (int32 Index = 0; Index < MenuBlocks.Num(); Index++)
 		{
-			if(Icon.GetIcon())
+			const int32 ItemIndex = (bIsWindowMenu ? Index + 6 : Index) - ItemIndexAdjust;
+			NSMenuItem* MenuItem = [Menu numberOfItems] > ItemIndex ? [Menu itemAtIndex:ItemIndex] : nil;
+
+			EMultiBlockType::Type Type = MenuBlocks[Index]->GetType();
+			if (Type == EMultiBlockType::MenuEntry)
 			{
-				FSlateBrush const* IconBrush = Icon.GetIcon();
-				FName ResourceName = IconBrush->GetResourceName();
-				MenuImage = [[NSImage alloc] initWithContentsOfFile: ResourceName.ToString().GetNSString()];
-				if(MenuImage)
+				TSharedRef<const FMenuEntryBlock> Block = StaticCastSharedRef<const FMenuEntryBlock>(MenuBlocks[Index]);
+				if (MenuItem && (![MenuItem isKindOfClass:[FMacMenuItem class]] || (Block->bIsSubMenu && [MenuItem submenu] == nil) || (!Block->bIsSubMenu && [MenuItem submenu] != nil)))
+				{
+					[Menu removeItem:MenuItem];
+					MenuItem = nil;
+				}
+				if (!MenuItem)
+				{
+					MenuItem = [[[FMacMenuItem alloc] initWithMenuEntryBlock:Block] autorelease];
+
+					if (Block->bIsSubMenu)
+					{
+						FMacMenu* SubMenu = [[FMacMenu alloc] initWithMenuEntryBlock:Block];
+						[MenuItem setSubmenu:SubMenu];
+					}
+
+					if ([Menu numberOfItems] > ItemIndex)
+					{
+						[Menu insertItem:MenuItem atIndex:ItemIndex];
+					}
+					else
+					{
+						[Menu addItem:MenuItem];
+					}
+				}
+
+				[MenuItem setTitle:FSlateMacMenu::GetMenuItemTitle(Block)];
+
+				uint32 KeyModifiers;
+				NSString* KeyEquivalent = FSlateMacMenu::GetMenuItemKeyEquivalent(Block, &KeyModifiers);
+				[MenuItem setKeyEquivalent:KeyEquivalent];
+				[MenuItem setKeyEquivalentModifierMask:KeyModifiers];
+
+				if (bIsWindowMenu)
+				{
+					NSImage* MenuImage = FSlateMacMenu::GetMenuItemIcon(Block);
+					if(MenuImage)
+					{
+						[MenuItem setImage:MenuImage];
+						[MenuImage release];
+					}
+				}
+				else
 				{
-					[MenuImage setSize:NSMakeSize(16.0f, 16.0f)];
+					[MenuItem setImage:nil];
+				}
+
+				if (FSlateMacMenu::IsMenuItemEnabled(Block))
+				{
+					[MenuItem setTarget:MenuItem];
+					[MenuItem setAction:@selector(performAction)];
+				}
+
+				if (!Block->bIsSubMenu)
+				{
+					[MenuItem setState:FSlateMacMenu::GetMenuItemState(Block)];
 				}
 			}
+			else if (Type == EMultiBlockType::MenuSeparator)
+			{
+				if (MenuItem && ![MenuItem isSeparatorItem])
+				{
+					[Menu removeItem:MenuItem];
+				}
+				else if (!MenuItem)
+				{
+					if ([Menu numberOfItems] > ItemIndex)
+					{
+						[Menu insertItem:[NSMenuItem separatorItem] atIndex:ItemIndex];
+					}
+					else
+					{
+						[Menu addItem:[NSMenuItem separatorItem]];
+					}
+				}
+			}
+			else
+			{
+				// If it's a type we skip, update ItemIndexAdjust so we can properly calculate item's index in NSMenu
+				ItemIndexAdjust++;
+			}
 		}
-		return MenuImage;
 	}, InGameRunLoopMode(@[ UE4NilEventMode, UE4ShowEventMode, UE4ResizeEventMode, UE4FullscreenEventMode, UE4CloseEventMode, UE4IMEEventMode ]));
 }
 
@@ -284,176 +229,144 @@ void FSlateMacMenu::ExecuteMenuItemAction(const TSharedRef< const class FMenuEnt
 	}, NSDefaultRunLoopMode, false);
 }
 
-CFStringRef FSlateMacMenu::GetMenuItemKeyEquivalent(const TSharedRef<const class FMenuEntryBlock>& Block, uint32* OutModifiers)
+static const TSharedRef<SWidget> FindTextBlockWidget(TSharedRef<SWidget> Content)
 {
-	return GameThreadReturn(^{
-		if (Block->GetAction().IsValid())
-		{
-			const TSharedRef<const FInputGesture>& Gesture = Block->GetAction()->GetActiveGesture();
+	if (Content->GetType() == FName(TEXT("STextBlock")))
+	{
+		return Content;
+	}
 
-			*OutModifiers = 0;
-			if (Gesture->bCtrl)
-			{
-				*OutModifiers |= NSCommandKeyMask;
-			}
-			if (Gesture->bShift)
-			{
-				*OutModifiers |= NSShiftKeyMask;
-			}
-			if (Gesture->bAlt)
-			{
-				*OutModifiers |= NSAlternateKeyMask;
-			}
-			if (Gesture->bCmd)
-			{
-				*OutModifiers |= NSControlKeyMask;
-			}
+	FChildren* Children = Content->GetChildren();
+	int32 NumChildren = Children->Num();
 
-			FString KeyString = Gesture->GetKeyText().ToString().ToLower();
-			return FPlatformString::TCHARToCFString(*KeyString);
+	for (int32 Index=0; Index < NumChildren; ++Index)
+	{
+		const TSharedRef<SWidget> Found = FindTextBlockWidget(Children->GetChildAt(Index));
+		if (Found != SNullWidget::NullWidget)
+		{
+			return Found;
 		}
-		return FPlatformString::TCHARToCFString(TEXT(""));
-	}, InGameRunLoopMode(@[ UE4NilEventMode, UE4ShowEventMode, UE4ResizeEventMode, UE4FullscreenEventMode, UE4CloseEventMode, UE4IMEEventMode ]));
+	}
+	return SNullWidget::NullWidget;
 }
 
-void FSlateMacMenu::UpdateMenu(FMacMenu* Menu)
+NSString* FSlateMacMenu::GetMenuItemTitle(const TSharedRef<const FMenuEntryBlock>& Block)
 {
-	GameThreadCall(^{
-		FText WindowLabel = NSLOCTEXT("MainMenu", "WindowMenu", "Window");
-		
-		bool bIsWindowMenu = (WindowLabel.ToString().Compare(FString([Menu title])) == 0);
-		
-		[Menu removeAllItems];
-		
-		if(bIsWindowMenu)
+	TAttribute<FText> Label;
+	if (!Block->LabelOverride.IsBound() && Block->LabelOverride.Get().IsEmpty() && Block->GetAction().IsValid())
+	{
+		Label = Block->GetAction()->GetLabel();
+	}
+	else if (!Block->LabelOverride.Get().IsEmpty())
+	{
+		Label = Block->LabelOverride;
+	}
+	else if (Block->EntryWidget.IsValid())
+	{
+		const TSharedRef<SWidget>& TextBlockWidget = FindTextBlockWidget(Block->EntryWidget.ToSharedRef());
+		if (TextBlockWidget != SNullWidget::NullWidget)
 		{
-			NSMenuItem* MinimizeItem = [[[NSMenuItem alloc] initWithTitle:@"Minimize" action:@selector(miniaturize:) keyEquivalent:@"m"] autorelease];
-			NSMenuItem* ZoomItem = [[[NSMenuItem alloc] initWithTitle:@"Zoom" action:@selector(performZoom:) keyEquivalent:@""] autorelease];
-			NSMenuItem* CloseItem = [[[NSMenuItem alloc] initWithTitle:@"Close" action:@selector(performClose:) keyEquivalent:@"w"] autorelease];
-			NSMenuItem* BringAllToFrontItem = [[[NSMenuItem alloc] initWithTitle:@"Bring All to Front" action:@selector(arrangeInFront:) keyEquivalent:@""] autorelease];
-			
-			[Menu addItem:MinimizeItem];
-			[Menu addItem:ZoomItem];
-			[Menu addItem:CloseItem];
-			[Menu addItem:[NSMenuItem separatorItem]];
-			[Menu addItem:BringAllToFrontItem];
-			[Menu addItem:[NSMenuItem separatorItem]];
+			Label = StaticCastSharedRef<STextBlock>(TextBlockWidget)->GetText();
 		}
+	}
 
-		TSharedRef< SWidget > Widget = GetMenuBuilderWidget(Menu);
-
-		TSharedRef<const FMultiBox> MultiBox = StaticCastSharedRef<SMultiBoxWidget>(Widget)->GetMultiBox();
-		const TArray<TSharedRef<const FMultiBlock>>& MenuBlocks = MultiBox->GetBlocks();
+	return Label.Get().ToString().GetNSString();
+}
 
-		for (int32 Index = 0; Index < MenuBlocks.Num(); Index++)
+NSImage* FSlateMacMenu::GetMenuItemIcon(const TSharedRef<const FMenuEntryBlock>& Block)
+{
+	NSImage* MenuImage = nil;
+	FSlateIcon Icon;
+	if (Block->IconOverride.IsSet())
+	{
+		Icon = Block->IconOverride;
+	}
+	else if (Block->GetAction().IsValid() && Block->GetAction()->GetIcon().IsSet())
+	{
+		Icon = Block->GetAction()->GetIcon();
+	}
+	if (Icon.IsSet())
+	{
+		if (Icon.GetIcon())
 		{
-			EMultiBlockType::Type Type = MenuBlocks[Index]->GetType();
-
-			if (Type == EMultiBlockType::MenuEntry)
+			FSlateBrush const* IconBrush = Icon.GetIcon();
+			FName ResourceName = IconBrush->GetResourceName();
+			MenuImage = [[NSImage alloc] initWithContentsOfFile:ResourceName.ToString().GetNSString()];
+			if (MenuImage)
 			{
-				TSharedRef<const FMenuEntryBlock> Block = StaticCastSharedRef<const FMenuEntryBlock>(MenuBlocks[Index]);
-				FMacMenuItem* MenuItem = [[[FMacMenuItem alloc] initWithMenuEntryBlock:Block] autorelease];
-
-				if (Block->bIsSubMenu)
-				{
-					FMacMenu* SubMenu = [[FMacMenu alloc] initWithMenuEntryBlock:Block];
-					[MenuItem setSubmenu:SubMenu];
-				}
-				else
-				{
-					[MenuItem setState:FSlateMacMenu::GetMenuItemState(Block)];
-				}
-				
-				if(!bIsWindowMenu)
-				{
-					bool bIsWindowSubMenu = false;
-					NSMenu* SuperMenu = [Menu supermenu];
-					while(!bIsWindowSubMenu && SuperMenu)
-					{
-						bIsWindowSubMenu = (WindowLabel.ToString().Compare(FString([SuperMenu title])) == 0);
-						SuperMenu = [SuperMenu supermenu];
-					}
-					
-					if(!bIsWindowSubMenu)
-					{
-						[MenuItem setImage:nil];
-					}
-				}
-
-				[Menu addItem:MenuItem];
-			}
-			else if (Type == EMultiBlockType::MenuSeparator)
-			{
-				[Menu addItem:[NSMenuItem separatorItem]];
+				[MenuImage setSize:NSMakeSize(16.0f, 16.0f)];
 			}
 		}
-		
-		if(bIsWindowMenu)
-		{
-			[NSApp setWindowsMenu:Menu];
-		}
-	}, InGameRunLoopMode(@[ UE4NilEventMode, UE4ShowEventMode, UE4ResizeEventMode, UE4FullscreenEventMode, UE4CloseEventMode, UE4IMEEventMode ]));
+	}
+	return MenuImage;
 }
 
-bool FSlateMacMenu::IsMenuItemEnabled(const TSharedRef<const class FMenuEntryBlock>& Block)
+NSString* FSlateMacMenu::GetMenuItemKeyEquivalent(const TSharedRef<const class FMenuEntryBlock>& Block, uint32* OutModifiers)
 {
-	bool bEnabled = GameThreadReturn(^{
-		TSharedPtr<const FUICommandList> ActionList = Block->GetActionList();
-		TSharedPtr<const FUICommandInfo> Action = Block->GetAction();
-		const FUIAction& DirectActions = Block->GetDirectActions();
+	if (Block->GetAction().IsValid())
+	{
+		const TSharedRef<const FInputGesture>& Gesture = Block->GetAction()->GetActiveGesture();
 
-		bool bEnabled = true;
-		if (ActionList.IsValid() && Action.IsValid())
+		*OutModifiers = 0;
+		if (Gesture->bCtrl)
 		{
-			bEnabled = ActionList->CanExecuteAction(Action.ToSharedRef());
+			*OutModifiers |= NSCommandKeyMask;
 		}
-		else
+		if (Gesture->bShift)
 		{
-			// There is no action list or action associated with this block via a UI command.  Execute any direct action we have
-			bEnabled = DirectActions.CanExecute();
+			*OutModifiers |= NSShiftKeyMask;
+		}
+		if (Gesture->bAlt)
+		{
+			*OutModifiers |= NSAlternateKeyMask;
+		}
+		if (Gesture->bCmd)
+		{
+			*OutModifiers |= NSControlKeyMask;
 		}
-		
-		return bEnabled;
-	}, InGameRunLoopMode(@[ UE4NilEventMode, UE4ShowEventMode, UE4ResizeEventMode, UE4FullscreenEventMode, UE4CloseEventMode, UE4IMEEventMode ]));
 
-	return bEnabled;
+		FString KeyString = Gesture->GetKeyText().ToString().ToLower();
+		return KeyString.GetNSString();
+	}
+	return @"";
 }
 
-int32 FSlateMacMenu::GetMenuItemState(const TSharedRef<const class FMenuEntryBlock>& Block)
+bool FSlateMacMenu::IsMenuItemEnabled(const TSharedRef<const class FMenuEntryBlock>& Block)
 {
-	bool bIsChecked = GameThreadReturn(^{
-		bool bIsChecked = false;
-		TSharedPtr<const FUICommandList> ActionList = Block->GetActionList();
-		TSharedPtr<const FUICommandInfo> Action = Block->GetAction();
-		const FUIAction& DirectActions = Block->GetDirectActions();
+	TSharedPtr<const FUICommandList> ActionList = Block->GetActionList();
+	TSharedPtr<const FUICommandInfo> Action = Block->GetAction();
+	const FUIAction& DirectActions = Block->GetDirectActions();
 
-		if (ActionList.IsValid() && Action.IsValid())
-		{
-			bIsChecked = ActionList->IsChecked(Action.ToSharedRef());
-		}
-		else
-		{
-			// There is no action list or action associated with this block via a UI command.  Execute any direct action we have
-			bIsChecked = DirectActions.IsChecked();
-		}
-		
-		return bIsChecked;
-	}, InGameRunLoopMode(@[ UE4NilEventMode, UE4ShowEventMode, UE4ResizeEventMode, UE4FullscreenEventMode, UE4CloseEventMode, UE4IMEEventMode ]));
+	bool bEnabled = true;
+	if (ActionList.IsValid() && Action.IsValid())
+	{
+		bEnabled = ActionList->CanExecuteAction(Action.ToSharedRef());
+	}
+	else
+	{
+		// There is no action list or action associated with this block via a UI command.  Execute any direct action we have
+		bEnabled = DirectActions.CanExecute();
+	}
 
-	return bIsChecked ? NSOnState : NSOffState;
+	return bEnabled;
 }
 
-TSharedRef<SWidget> FSlateMacMenu::GetMenuBuilderWidget(FMacMenu* Menu)
+int32 FSlateMacMenu::GetMenuItemState(const TSharedRef<const class FMenuEntryBlock>& Block)
 {
-	return GameThreadReturn(^{
-		// @todo: ideally this should be cached and updated only when the menu contents change. See menuHasKeyEquivalent.
-		const bool bShouldCloseWindowAfterMenuSelection = true;
-		FMenuBuilder MenuBuilder(bShouldCloseWindowAfterMenuSelection, [Menu GetMenuEntryBlock]->GetActionList(), [Menu GetMenuEntryBlock]->Extender);
-		{
-			// Have the menu fill its contents
-			[Menu GetMenuEntryBlock]->EntryBuilder.ExecuteIfBound(MenuBuilder);
-		}
-		TSharedPtr<SWidget> Widget = MenuBuilder.MakeWidget();
-		return Widget;
-	}, InGameRunLoopMode(@[ UE4NilEventMode, UE4ShowEventMode, UE4ResizeEventMode, UE4FullscreenEventMode, UE4CloseEventMode, UE4IMEEventMode ])).ToSharedRef();
+	bool bIsChecked = false;
+	TSharedPtr<const FUICommandList> ActionList = Block->GetActionList();
+	TSharedPtr<const FUICommandInfo> Action = Block->GetAction();
+	const FUIAction& DirectActions = Block->GetDirectActions();
+
+	if (ActionList.IsValid() && Action.IsValid())
+	{
+		bIsChecked = ActionList->IsChecked(Action.ToSharedRef());
+	}
+	else
+	{
+		// There is no action list or action associated with this block via a UI command.  Execute any direct action we have
+		bIsChecked = DirectActions.IsChecked();
+	}
+
+	return bIsChecked ? NSOnState : NSOffState;
 }
diff --git a/Engine/Source/Runtime/Slate/Private/Framework/MultiBox/Mac/MacMenu.h b/Engine/Source/Runtime/Slate/Private/Framework/MultiBox/Mac/MacMenu.h
--- a/Engine/Source/Runtime/Slate/Private/Framework/MultiBox/Mac/MacMenu.h
+++ b/Engine/Source/Runtime/Slate/Private/Framework/MultiBox/Mac/MacMenu.h
@@ -5,28 +5,24 @@
 #include "MultiBox.h"
 #include "SMenuEntryBlock.h"
 
-#ifdef __OBJC__
 @interface FMacMenu : NSMenu <NSMenuDelegate>
-{
-	TSharedPtr<const FMenuEntryBlock> MenuEntryBlock;
-}
-- (TSharedPtr<const FMenuEntryBlock>&)GetMenuEntryBlock;
+@property (assign) TSharedPtr<const FMenuEntryBlock> MenuEntryBlock;
+@property (assign) TSharedPtr<const FMultiBox> MultiBox;
 @end
-#else
-class FMacMenu;
-#endif
 
 class SLATE_API FSlateMacMenu
 {
 public:
 
-	static void UpdateWithMultiBox(const TSharedRef< FMultiBox >& MultiBox);
-	static CFStringRef GetMenuItemTitle(const TSharedRef< const class FMenuEntryBlock >& Block);
+	static void UpdateWithMultiBox(const TSharedRef<FMultiBox>& MultiBox);
+	static void UpdateMenu(FMacMenu* Menu);
+	static void ExecuteMenuItemAction(const TSharedRef<const FMenuEntryBlock>& Block);
+
+private:
+
+	static NSString* GetMenuItemTitle(const TSharedRef< const FMenuEntryBlock >& Block);
 	static NSImage* GetMenuItemIcon(const TSharedRef<const FMenuEntryBlock>& Block);
-	static void ExecuteMenuItemAction(const TSharedRef< const class FMenuEntryBlock >& Block);
-	static CFStringRef GetMenuItemKeyEquivalent(const TSharedRef< const class FMenuEntryBlock >& Block, uint32* OutModifiers);
-	static void UpdateMenu(FMacMenu* InMenu);
-	static bool IsMenuItemEnabled(const TSharedRef<const class FMenuEntryBlock>& Block);
-	static int32 GetMenuItemState(const TSharedRef<const class FMenuEntryBlock>& Block);
-	static TSharedRef<SWidget> GetMenuBuilderWidget(FMacMenu* InMenu);
+	static NSString* GetMenuItemKeyEquivalent(const TSharedRef<const FMenuEntryBlock>& Block, uint32* OutModifiers);
+	static bool IsMenuItemEnabled(const TSharedRef<const FMenuEntryBlock>& Block);
+	static int32 GetMenuItemState(const TSharedRef<const FMenuEntryBlock>& Block);
 };
