Made the Slate text layout respect layout transforms
Converted some code which was creating FPaintGeometry the old way (so losing the render transforms) to use the new way instead.

diff --git a/Engine/Source/Runtime/Slate/Private/Framework/Text/SlateHyperlinkRun.cpp b/Engine/Source/Runtime/Slate/Private/Framework/Text/SlateHyperlinkRun.cpp
--- a/Engine/Source/Runtime/Slate/Private/Framework/Text/SlateHyperlinkRun.cpp
+++ b/Engine/Source/Runtime/Slate/Private/Framework/Text/SlateHyperlinkRun.cpp
@@ -75,7 +75,10 @@ int32 FSlateHyperlinkRun::OnPaint( const FPaintArgs& Args, const FTextLayout::FL
 {
 	const TSharedRef< FWidgetLayoutBlock > WidgetBlock = StaticCastSharedRef< FWidgetLayoutBlock >( Block );
 
-	FGeometry WidgetGeometry = AllottedGeometry.MakeChild( Block->GetLocationOffset() * ( 1 / AllottedGeometry.Scale ), Block->GetSize() * ( 1 / AllottedGeometry.Scale ), 1.0f );
+	// The block size and offset values are pre-scaled, so we need to account for that when converting the block offsets into paint geometry
+	const float InverseScale = Inverse(AllottedGeometry.Scale);
+
+	const FGeometry WidgetGeometry = AllottedGeometry.MakeChild(TransformVector(InverseScale, Block->GetSize()), FSlateLayoutTransform(TransformPoint(InverseScale, Block->GetLocationOffset())));
 	return WidgetBlock->GetWidget()->Paint( Args, WidgetGeometry, MyClippingRect, OutDrawElements, LayerId, InWidgetStyle, bParentEnabled );
 }
 
@@ -87,7 +90,13 @@ const TArray< TSharedRef<SWidget> >& FSlateHyperlinkRun::GetChildren()
 void FSlateHyperlinkRun::ArrangeChildren( const TSharedRef< ILayoutBlock >& Block, const FGeometry& AllottedGeometry, FArrangedChildren& ArrangedChildren ) const 
 {
 	const TSharedRef< FWidgetLayoutBlock > WidgetBlock = StaticCastSharedRef< FWidgetLayoutBlock >( Block );
-	ArrangedChildren.AddWidget( AllottedGeometry.MakeChild( WidgetBlock->GetWidget(), Block->GetLocationOffset() * ( 1 / AllottedGeometry.Scale ), Block->GetSize() * ( 1 / AllottedGeometry.Scale ), 1.0f ) );
+	
+	// The block size and offset values are pre-scaled, so we need to account for that when converting the block offsets into paint geometry
+	const float InverseScale = Inverse(AllottedGeometry.Scale);
+
+	ArrangedChildren.AddWidget(
+		AllottedGeometry.MakeChild(WidgetBlock->GetWidget(), TransformVector(InverseScale, Block->GetSize()), FSlateLayoutTransform(TransformPoint(InverseScale, Block->GetLocationOffset())))
+		);
 }
 
 int32 FSlateHyperlinkRun::GetTextIndexAt( const TSharedRef< ILayoutBlock >& Block, const FVector2D& Location, float Scale, ETextHitPoint* const OutHitPoint ) const
diff --git a/Engine/Source/Runtime/Slate/Private/Framework/Text/SlateImageRun.cpp b/Engine/Source/Runtime/Slate/Private/Framework/Text/SlateImageRun.cpp
--- a/Engine/Source/Runtime/Slate/Private/Framework/Text/SlateImageRun.cpp
+++ b/Engine/Source/Runtime/Slate/Private/Framework/Text/SlateImageRun.cpp
@@ -69,11 +69,22 @@ FVector2D FSlateImageRun::GetLocationAt( const TSharedRef< ILayoutBlock >& Block
 
 int32 FSlateImageRun::OnPaint( const FPaintArgs& Args, const FTextLayout::FLineView& Line, const TSharedRef< ILayoutBlock >& Block, const FTextBlockStyle& DefaultStyle, const FGeometry& AllottedGeometry, const FSlateRect& MyClippingRect, FSlateWindowElementList& OutDrawElements, int32 LayerId, const FWidgetStyle& InWidgetStyle, bool bParentEnabled ) const 
 {
+	// The block size and offset values are pre-scaled, so we need to account for that when converting the block offsets into paint geometry
+	const float InverseScale = Inverse(AllottedGeometry.Scale);
+
 	if ( Image->DrawAs != ESlateBrushDrawType::NoDrawType )
 	{
 		const FColor FinalColorAndOpacity( InWidgetStyle.GetColorAndOpacityTint() * Image->GetTint( InWidgetStyle ) );
 		const uint32 DrawEffects = bParentEnabled ? ESlateDrawEffect::None : ESlateDrawEffect::DisabledEffect;
-		FSlateDrawElement::MakeBox(OutDrawElements, ++LayerId, FPaintGeometry( AllottedGeometry.AbsolutePosition + Block->GetLocationOffset(), Block->GetSize(), AllottedGeometry.Scale ), Image, MyClippingRect, DrawEffects, FinalColorAndOpacity );
+		FSlateDrawElement::MakeBox(
+			OutDrawElements, 
+			++LayerId, 
+			AllottedGeometry.ToPaintGeometry(TransformVector(InverseScale, Block->GetSize()), FSlateLayoutTransform(TransformPoint(InverseScale, Block->GetLocationOffset()))), 
+			Image, 
+			MyClippingRect, 
+			DrawEffects, 
+			FinalColorAndOpacity
+			);
 	}
 
 	return LayerId;
diff --git a/Engine/Source/Runtime/Slate/Private/Framework/Text/SlateTextHighlightRunRenderer.cpp b/Engine/Source/Runtime/Slate/Private/Framework/Text/SlateTextHighlightRunRenderer.cpp
--- a/Engine/Source/Runtime/Slate/Private/Framework/Text/SlateTextHighlightRunRenderer.cpp
+++ b/Engine/Source/Runtime/Slate/Private/Framework/Text/SlateTextHighlightRunRenderer.cpp
@@ -16,11 +16,14 @@ int32 FSlateTextHighlightRunRenderer::OnPaint( const FPaintArgs& Args, const FTe
 	FVector2D Location( Block->GetLocationOffset() );
 	Location.Y = Line.Offset.Y;
 
+	// The block size and offset values are pre-scaled, so we need to account for that when converting the block offsets into paint geometry
+	const float InverseScale = Inverse(AllottedGeometry.Scale);
+
 	// Draw the actual highlight rectangle
 	FSlateDrawElement::MakeBox(
 		OutDrawElements,
 		++LayerId,
-		FPaintGeometry( AllottedGeometry.AbsolutePosition + Location, FVector2D( Block->GetSize().X, Line.Size.Y ), AllottedGeometry.Scale ),
+		AllottedGeometry.ToPaintGeometry(TransformVector(InverseScale, FVector2D( Block->GetSize().X, Line.Size.Y )), FSlateLayoutTransform(TransformPoint(InverseScale, Location))), 
 		&DefaultStyle.HighlightShape,
 		MyClippingRect,
 		bParentEnabled ? ESlateDrawEffect::None : ESlateDrawEffect::DisabledEffect,
diff --git a/Engine/Source/Runtime/Slate/Private/Framework/Text/SlateTextLayout.cpp b/Engine/Source/Runtime/Slate/Private/Framework/Text/SlateTextLayout.cpp
--- a/Engine/Source/Runtime/Slate/Private/Framework/Text/SlateTextLayout.cpp
+++ b/Engine/Source/Runtime/Slate/Private/Framework/Text/SlateTextLayout.cpp
@@ -73,10 +73,13 @@ int32 FSlateTextLayout::OnPaint( const FPaintArgs& Args, const FTextBlockStyle&
 			{
 				BlockHue.R += 50.0f;
 
+				// The block size and offset values are pre-scaled, so we need to account for that when converting the block offsets into paint geometry
+				const float InverseScale = Inverse(AllottedGeometry.Scale);
+
 				FSlateDrawElement::MakeBox(
 					OutDrawElements, 
 					BlockDebugLayer,
-					FPaintGeometry( AllottedGeometry.AbsolutePosition + Block->GetLocationOffset(), Block->GetSize(), AllottedGeometry.Scale ),
+					AllottedGeometry.ToPaintGeometry(TransformVector(InverseScale, Block->GetSize()), FSlateLayoutTransform(TransformPoint(InverseScale, Block->GetLocationOffset()))),
 					&DefaultTextStyle.HighlightShape,
 					ClippingRect,
 					DrawEffects,
diff --git a/Engine/Source/Runtime/Slate/Private/Framework/Text/SlateTextRun.cpp b/Engine/Source/Runtime/Slate/Private/Framework/Text/SlateTextRun.cpp
--- a/Engine/Source/Runtime/Slate/Private/Framework/Text/SlateTextRun.cpp
+++ b/Engine/Source/Runtime/Slate/Private/Framework/Text/SlateTextRun.cpp
@@ -74,6 +74,9 @@ int32 FSlateTextRun::OnPaint( const FPaintArgs& Args, const FTextLayout::FLineVi
 	const bool ShouldDropShadow = Style.ShadowOffset.Size() > 0;
 	const FVector2D BlockLocationOffset = Block->GetLocationOffset();
 	const FTextRange BlockRange = Block->GetTextRange();
+	
+	// The block size and offset values are pre-scaled, so we need to account for that when converting the block offsets into paint geometry
+	const float InverseScale = Inverse(AllottedGeometry.Scale);
 
 	// Draw the optional shadow
 	if ( ShouldDropShadow )
@@ -81,7 +84,7 @@ int32 FSlateTextRun::OnPaint( const FPaintArgs& Args, const FTextLayout::FLineVi
 		FSlateDrawElement::MakeText(
 			OutDrawElements,
 			++LayerId,
-			FPaintGeometry( AllottedGeometry.AbsolutePosition + Block->GetLocationOffset() + Style.ShadowOffset, Block->GetSize(), AllottedGeometry.Scale ),
+			AllottedGeometry.ToPaintGeometry(TransformVector(InverseScale, Block->GetSize()), FSlateLayoutTransform(TransformPoint(InverseScale, Block->GetLocationOffset() + Style.ShadowOffset))),
 			Text.Get(),
 			BlockRange.BeginIndex,
 			BlockRange.EndIndex,
@@ -96,7 +99,7 @@ int32 FSlateTextRun::OnPaint( const FPaintArgs& Args, const FTextLayout::FLineVi
 	FSlateDrawElement::MakeText(
 		OutDrawElements,
 		++LayerId,
-		FPaintGeometry( AllottedGeometry.AbsolutePosition + Block->GetLocationOffset(), Block->GetSize(), AllottedGeometry.Scale ),
+		AllottedGeometry.ToPaintGeometry(TransformVector(InverseScale, Block->GetSize()), FSlateLayoutTransform(TransformPoint(InverseScale, Block->GetLocationOffset()))),
 		Text.Get(),
 		BlockRange.BeginIndex,
 		BlockRange.EndIndex,
diff --git a/Engine/Source/Runtime/Slate/Private/Framework/Text/SlateWidgetRun.cpp b/Engine/Source/Runtime/Slate/Private/Framework/Text/SlateWidgetRun.cpp
--- a/Engine/Source/Runtime/Slate/Private/Framework/Text/SlateWidgetRun.cpp
+++ b/Engine/Source/Runtime/Slate/Private/Framework/Text/SlateWidgetRun.cpp
@@ -53,6 +53,9 @@ TSharedRef< ILayoutBlock > FSlateWidgetRun::CreateBlock( int32 StartIndex, int32
 
 int32 FSlateWidgetRun::OnPaint( const FPaintArgs& Args, const FTextLayout::FLineView& Line, const TSharedRef< ILayoutBlock >& Block, const FTextBlockStyle& DefaultStyle, const FGeometry& AllottedGeometry, const FSlateRect& MyClippingRect, FSlateWindowElementList& OutDrawElements, int32 LayerId, const FWidgetStyle& InWidgetStyle, bool bParentEnabled ) const 
 {
+	// The block size and offset values are pre-scaled, so we need to account for that when converting the block offsets into paint geometry
+	const float InverseScale = Inverse(AllottedGeometry.Scale);
+
 	const FVector2D DesiredWidgetSize = Info.Widget->GetDesiredSize();
 	if (DesiredWidgetSize != WidgetSize)
 	{
@@ -65,7 +68,7 @@ int32 FSlateWidgetRun::OnPaint( const FPaintArgs& Args, const FTextLayout::FLine
 		}
 	}
 
-	FGeometry WidgetGeometry = AllottedGeometry.MakeChild( Block->GetLocationOffset() * ( 1 / AllottedGeometry.Scale ), Block->GetSize() * ( 1 / AllottedGeometry.Scale ), 1.0f );
+	const FGeometry WidgetGeometry = AllottedGeometry.MakeChild(TransformVector(InverseScale, Block->GetSize()), FSlateLayoutTransform(TransformPoint(InverseScale, Block->GetLocationOffset())));
 	return Info.Widget->Paint( Args, WidgetGeometry, MyClippingRect, OutDrawElements, LayerId, InWidgetStyle, bParentEnabled );
 }
 
@@ -76,7 +79,12 @@ const TArray< TSharedRef<SWidget> >& FSlateWidgetRun::GetChildren()
 
 void FSlateWidgetRun::ArrangeChildren( const TSharedRef< ILayoutBlock >& Block, const FGeometry& AllottedGeometry, FArrangedChildren& ArrangedChildren ) const
 {
-	ArrangedChildren.AddWidget( AllottedGeometry.MakeChild( Info.Widget, Block->GetLocationOffset() * ( 1 / AllottedGeometry.Scale ), Block->GetSize() * ( 1 / AllottedGeometry.Scale ), 1.0f ) );
+	// The block size and offset values are pre-scaled, so we need to account for that when converting the block offsets into paint geometry
+	const float InverseScale = Inverse(AllottedGeometry.Scale);
+
+	ArrangedChildren.AddWidget(
+		AllottedGeometry.MakeChild(Info.Widget, TransformVector(InverseScale, Block->GetSize()), FSlateLayoutTransform(TransformPoint(InverseScale, Block->GetLocationOffset())))
+		);
 }
 
 int32 FSlateWidgetRun::GetTextIndexAt( const TSharedRef< ILayoutBlock >& Block, const FVector2D& Location, float Scale, ETextHitPoint* const OutHitPoint ) const
diff --git a/Engine/Source/Runtime/Slate/Private/Widgets/Text/SMultiLineEditableText.cpp b/Engine/Source/Runtime/Slate/Private/Widgets/Text/SMultiLineEditableText.cpp
--- a/Engine/Source/Runtime/Slate/Private/Widgets/Text/SMultiLineEditableText.cpp
+++ b/Engine/Source/Runtime/Slate/Private/Widgets/Text/SMultiLineEditableText.cpp
@@ -71,6 +71,9 @@ int32 SMultiLineEditableText::FCursorLineHighlighter::OnPaint(const FPaintArgs&
 	const float CursorWidth = FTextEditHelper::CalculateCaretWidth(FontMaxCharHeight);
 	const double CurrentTime = FSlateApplication::Get().GetCurrentTime();
 
+	// The block size and offset values are pre-scaled, so we need to account for that when converting the block offsets into paint geometry
+	const float InverseScale = Inverse(AllottedGeometry.Scale);
+
 	// The cursor is always visible (i.e. not blinking) when we're interacting with it; otherwise it might get lost.
 	const bool bForceCursorVisible = (CurrentTime - CursorInfo->GetLastCursorInteractionTime()) < EditableTextDefs::CaretBlinkPauseTime;
 	float CursorOpacity = (bForceCursorVisible)
@@ -88,7 +91,7 @@ int32 SMultiLineEditableText::FCursorLineHighlighter::OnPaint(const FPaintArgs&
 	FSlateDrawElement::MakeBox(
 		OutDrawElements,
 		LayerId,
-		FPaintGeometry(AllottedGeometry.AbsolutePosition + Location + OptionalWidth, FVector2D(CursorWidth * AllottedGeometry.Scale, Size.Y), AllottedGeometry.Scale),
+		AllottedGeometry.ToPaintGeometry(TransformVector(InverseScale, FVector2D(CursorWidth * AllottedGeometry.Scale, Size.Y)), FSlateLayoutTransform(TransformPoint(InverseScale, Location + OptionalWidth))),
 		CursorBrush,
 		MyClippingRect,
 		bParentEnabled ? ESlateDrawEffect::None : ESlateDrawEffect::DisabledEffect,
@@ -111,6 +114,9 @@ int32 SMultiLineEditableText::FTextCompositionHighlighter::OnPaint( const FPaint
 	const FVector2D Location(Line.Offset.X + OffsetX, Line.Offset.Y);
 	const FVector2D Size(Width, Line.TextSize.Y);
 
+	// The block size and offset values are pre-scaled, so we need to account for that when converting the block offsets into paint geometry
+	const float InverseScale = Inverse(AllottedGeometry.Scale);
+
 	if (Size.X)
 	{
 		const FLinearColor LineColorAndOpacity = InWidgetStyle.GetForegroundColor();
@@ -122,7 +128,7 @@ int32 SMultiLineEditableText::FTextCompositionHighlighter::OnPaint( const FPaint
 		FSlateDrawElement::MakeBox(
 			OutDrawElements,
 			++LayerId,
-			FPaintGeometry(AllottedGeometry.AbsolutePosition + Location, Size, AllottedGeometry.Scale),
+			AllottedGeometry.ToPaintGeometry(TransformVector(InverseScale, Size), FSlateLayoutTransform(TransformPoint(InverseScale, Location))),
 			CompositionBrush,
 			MyClippingRect,
 			bParentEnabled ? ESlateDrawEffect::None : ESlateDrawEffect::DisabledEffect,
@@ -151,6 +157,9 @@ int32 SMultiLineEditableText::FTextSelectionRunRenderer::OnPaint( const FPaintAr
 	const FLinearColor UnfocusedSelectionColorAdjustment(-0.01f, -0.01f, 0.01f);
 	const FColor SelectionBackgroundColorAndOpacity = ((FLinearColor::White - DefaultStyle.ColorAndOpacity.GetColor(InWidgetStyle))*0.5f + (bHasKeyboardFocus ? FocusedSelectionColorAdjustment : UnfocusedSelectionColorAdjustment)) * InWidgetStyle.GetColorAndOpacityTint();
 
+	// The block size and offset values are pre-scaled, so we need to account for that when converting the block offsets into paint geometry
+	const float InverseScale = Inverse(AllottedGeometry.Scale);
+
 	const float HighlightWidth = Block->GetSize().X;
 	if (HighlightWidth)
 	{
@@ -158,7 +167,7 @@ int32 SMultiLineEditableText::FTextSelectionRunRenderer::OnPaint( const FPaintAr
 		FSlateDrawElement::MakeBox(
 			OutDrawElements,
 			++LayerId,
-			FPaintGeometry(AllottedGeometry.AbsolutePosition + Location, FVector2D(HighlightWidth, Line.Size.Y), AllottedGeometry.Scale),
+			AllottedGeometry.ToPaintGeometry(TransformVector(InverseScale, FVector2D(HighlightWidth, Line.Size.Y)), FSlateLayoutTransform(TransformPoint(InverseScale, Location))),
 			&DefaultStyle.HighlightShape,
 			MyClippingRect,
 			bParentEnabled ? ESlateDrawEffect::None : ESlateDrawEffect::DisabledEffect,
