Mac drag and drop improvements. A window won't activate on mouse down if the cursor is over a draggable item. Added FMacPlatformMisc::ActivateApplication().

diff --git a/Engine/Source/Runtime/Core/Private/Mac/CocoaTextView.cpp b/Engine/Source/Runtime/Core/Private/Mac/CocoaTextView.cpp
--- a/Engine/Source/Runtime/Core/Private/Mac/CocoaTextView.cpp
+++ b/Engine/Source/Runtime/Core/Private/Mac/CocoaTextView.cpp
@@ -136,6 +136,13 @@
 	}
 }
 
+// Returning YES allows SlateApplication to control if window should be activated on mouse down.
+// This is used for drag and drop when we don't want to activate if the mouse cursor is over a draggable item.
+- (BOOL)shouldDelayWindowOrderingForEvent:(NSEvent*)Event
+{
+	return YES;
+}
+
 - (void)activateInputMethod:(const TSharedRef<ITextInputMethodContext>&)InContext
 {
 	if (IMMContext.IsValid())
diff --git a/Engine/Source/Runtime/Core/Private/Mac/MacPlatformMisc.cpp b/Engine/Source/Runtime/Core/Private/Mac/MacPlatformMisc.cpp
--- a/Engine/Source/Runtime/Core/Private/Mac/MacPlatformMisc.cpp
+++ b/Engine/Source/Runtime/Core/Private/Mac/MacPlatformMisc.cpp
@@ -342,6 +342,13 @@ void FMacPlatformMisc::UpdateWindowMenu()
 	[WindowMenu addItem:BringAllToFrontItem];
 }
 
+void FMacPlatformMisc::ActivateApplication()
+{
+	MainThreadCall(^{
+		[NSApp activateIgnoringOtherApps:YES];
+	});
+}
+
 bool FMacPlatformMisc::ControlScreensaver(EScreenSaverAction Action)
 {
 	static uint32 IOPMNoSleepAssertion = 0;
diff --git a/Engine/Source/Runtime/Core/Public/Mac/MacPlatformMisc.h b/Engine/Source/Runtime/Core/Public/Mac/MacPlatformMisc.h
--- a/Engine/Source/Runtime/Core/Public/Mac/MacPlatformMisc.h
+++ b/Engine/Source/Runtime/Core/Public/Mac/MacPlatformMisc.h
@@ -136,6 +136,8 @@ struct CORE_API FMacPlatformMisc : public FGenericPlatformMisc
 
 	static void UpdateWindowMenu();
 
+	static void ActivateApplication();
+
 	/**
 	 * Returns whether the platform is running on battery power or not.
 	 */
diff --git a/Engine/Source/Runtime/Slate/Private/Framework/Application/SlateApplication.cpp b/Engine/Source/Runtime/Slate/Private/Framework/Application/SlateApplication.cpp
--- a/Engine/Source/Runtime/Slate/Private/Framework/Application/SlateApplication.cpp
+++ b/Engine/Source/Runtime/Slate/Private/Framework/Application/SlateApplication.cpp
@@ -3284,8 +3284,9 @@ bool FSlateApplication::ProcessMouseButtonDownEvent( const TSharedPtr< FGenericW
 			LOG_EVENT( EEventLog::MouseButtonDown, Reply );
 
 			// If none of the widgets requested keyboard focus to be set (or set the keyboard focus explicitly), set it to the leaf-most widget under the mouse.
+			// On Mac we prevent the OS from activating the window on mouse down, so we have full control and can activate only if there's nothing draggable under the mouse cursor.
 			const bool bFocusChangedByEventHandler = PreviouslyFocusedWidget != GetKeyboardFocusedWidget();
-			if (!Reply.GetFocusRecepient().IsValid() && !bFocusChangedByEventHandler)
+			if ((!Reply.GetFocusRecepient().IsValid() && !bFocusChangedByEventHandler) || (PLATFORM_MAC && !DragDetector.DetectDragForWidget.IsValid()))
 			{
 				// The event handler for OnMouseButton down may have altered the widget hierarchy.
 				// Refresh the previously-cached widget path.
@@ -3302,6 +3303,14 @@ bool FSlateApplication::ProcessMouseButtonDownEvent( const TSharedPtr< FGenericW
 						SetKeyboardFocus( NewFocusedWidgetPath, EKeyboardFocusCause::Mouse );
 					}
 				}
+
+#if PLATFORM_MAC
+				if (WidgetsUnderCursor.TopLevelWindow.IsValid() && !DragDetector.DetectDragForWidget.IsValid())
+				{
+					WidgetsUnderCursor.TopLevelWindow->BringToFront(true);
+					FPlatformMisc::ActivateApplication();
+				}
+#endif
 			}
 		}
 
