Slate : Fixed MenuPlacement_AboveAnchor not being respected.

diff --git a/Engine/Source/Runtime/Slate/Private/Widgets/Input/SMenuAnchor.cpp b/Engine/Source/Runtime/Slate/Private/Widgets/Input/SMenuAnchor.cpp
--- a/Engine/Source/Runtime/Slate/Private/Widgets/Input/SMenuAnchor.cpp
+++ b/Engine/Source/Runtime/Slate/Private/Widgets/Input/SMenuAnchor.cpp
@@ -124,9 +124,11 @@ FGeometry SMenuAnchor::ComputeMenuPlacement( const FGeometry& AllottedGeometry,
 	const FVector2D PopupSizeLocalSpace = (PlacementMode == MenuPlacement_ComboBox || PlacementMode == MenuPlacement_ComboBoxRight)
 		? FVector2D( FMath::Max( AllottedGeometry.Size.X, PopupDesiredSize.X ), PopupDesiredSize.Y )
 		: PopupDesiredSize;
-	const FVector2D OffsetLocalSpace = (PlacementMode == MenuPlacement_ComboBoxRight) 
-		? FVector2D(AllottedGeometry.GetLocalSize().X - PopupSizeLocalSpace.X, 0.0f) 
-		: FVector2D::ZeroVector;
+	const FVector2D OffsetLocalSpace = (PlacementMode == MenuPlacement_AboveAnchor)
+		? FVector2D( 0.0f, -AllottedGeometry.GetLocalSize().Y - PopupSizeLocalSpace.Y )
+		: (PlacementMode == MenuPlacement_ComboBoxRight) 
+			? FVector2D(AllottedGeometry.GetLocalSize().X - PopupSizeLocalSpace.X, 0.0f) 
+			: FVector2D::ZeroVector;
 	const FSlateRect AnchorLocalSpace = FSlateRect::FromPointAndExtent(OffsetLocalSpace, AllottedGeometry.Size);
 	const EOrientation Orientation = (PlacementMode == MenuPlacement_MenuRight) ? Orient_Horizontal : Orient_Vertical;
 	
