Slate: removed legacy hittest mode.

diff --git a/Engine/Source/Runtime/Slate/Private/Framework/Application/SlateApplication.cpp b/Engine/Source/Runtime/Slate/Private/Framework/Application/SlateApplication.cpp
--- a/Engine/Source/Runtime/Slate/Private/Framework/Application/SlateApplication.cpp
+++ b/Engine/Source/Runtime/Slate/Private/Framework/Application/SlateApplication.cpp
@@ -237,62 +237,6 @@ void FPopupSupport::SendNotifications( const FWidgetPath& WidgetsUnderCursor )
 }
 
 
-/**
- * Given a widget-geometry pair and the mouse cursor location in absolute coordinates populate OutWidgetsUnderCursor with Candidate's children if they are
- * under the cursor. If the Candidate is under the cursor return true, otherwise return false.
- *
- * @param Candidate                 The widget geometry pair being tested against the cursor location.
- * @param InAbsoluteCursorLocation  The cursor location in absolute space (absolute space should match that of WidgetGeometry) 
- * @param OutWidgetsUnderCursor     Widgets to be populated by the Candidate's children if the candidate is under the cursor.
- * @param bIgnoreEnabledStatus		Allows hit tests to run on disabled widgets. 
- *
- * @returns true if the candidate is under the cursor. false otherwise.
- */
-static bool LocateWidgetsUnderCursor_Helper( FArrangedWidget& Candidate, FVector2D InAbsoluteCursorLocation, FArrangedChildren& OutWidgetsUnderCursor, bool bIgnoreEnabledStatus )
-{
-	const bool bCandidateUnderCursor = 
-		// Candidate is physically under the cursor
-		Candidate.Geometry.IsUnderLocation( InAbsoluteCursorLocation ) && 
-		// Candidate actually considers itself hit by this test
-		(Candidate.Widget->OnHitTest( Candidate.Geometry, InAbsoluteCursorLocation ));
-
-	bool bHitAnyWidget = false;
-	if( bCandidateUnderCursor )
-	{
-		// The candidate widget is under the mouse
-		OutWidgetsUnderCursor.AddWidget(Candidate);
-
-		// Check to see if we were asked to still allow children to be hit test visible
-		bool bHitChildWidget = false;
-		if( ( Candidate.Widget->GetVisibility().AreChildrenHitTestVisible() ) != 0 )
-		{
-			FArrangedChildren ArrangedChildren(OutWidgetsUnderCursor.GetFilter());
-			Candidate.Widget->ArrangeChildren(Candidate.Geometry, ArrangedChildren);
-
-			// A widget's children are implicitly Z-ordered from first to last
-			for (int32 ChildIndex = ArrangedChildren.Num()-1; !bHitChildWidget && ChildIndex >= 0; --ChildIndex)
-			{
-				FArrangedWidget& SomeChild = ArrangedChildren[ChildIndex];
-				bHitChildWidget = (SomeChild.Widget->IsEnabled() || bIgnoreEnabledStatus) && LocateWidgetsUnderCursor_Helper( SomeChild, InAbsoluteCursorLocation, OutWidgetsUnderCursor, bIgnoreEnabledStatus );
-			}
-		}
-
-		// If we hit a child widget or we hit our candidate widget then we'll append our widgets
-		const bool bHitCandidateWidget = Candidate.Widget->GetVisibility().IsHitTestVisible();
-		bHitAnyWidget = bHitChildWidget || bHitCandidateWidget;
-		if( !bHitAnyWidget )
-		{
-			// No child widgets were hit, and even though the cursor was over our candidate widget, the candidate
-			// widget was not hit-testable, so we won't report it
-			check( OutWidgetsUnderCursor.Last() == Candidate );
-			OutWidgetsUnderCursor.Remove( OutWidgetsUnderCursor.Num() - 1 );
-		}
-	}
-
-	return bHitAnyWidget;
-}
-
-
 void FSlateApplication::Create()
 {
 	EKeys::Initialize();
@@ -495,19 +439,8 @@ FWidgetPath FSlateApplication::LocateWindowUnderMouse( FVector2D ScreenspaceMous
 
 		if ( Window->IsVisible() && AcceptsInput && Window->IsScreenspaceMouseWithin(ScreenspaceMouseCoordinate) && !bPrevWindowWasModal )
 		{
-			if ( !SWidget::UseLegacyHittest() )
-			{
-				const TArray<FArrangedWidget> ArrangedWidgets = HittestGrid->GetBubblePath( ScreenspaceMouseCoordinate, bIgnoreEnabledStatus );
-				return FWidgetPath( ArrangedWidgets );
-			}
-			else
-			{
-				FArrangedWidget WindowWidgetGeometry( Window, Window->GetWindowGeometryInScreen() );
-				if ( Window->IsEnabled() && LocateWidgetsUnderCursor_Helper(WindowWidgetGeometry, ScreenspaceMouseCoordinate, OutWidgetPath, bIgnoreEnabledStatus) )		
-				{			
-					return FWidgetPath( Windows[WindowIndex], OutWidgetPath );
-				}
-			}
+			const TArray<FArrangedWidget> ArrangedWidgets = HittestGrid->GetBubblePath( ScreenspaceMouseCoordinate, bIgnoreEnabledStatus );
+			return FWidgetPath( ArrangedWidgets );
 		}
 	}
 
@@ -719,7 +652,7 @@ void FSlateApplication::PrivateDrawWindows( TSharedPtr<SWindow> DrawOnlyThisWind
 		SCOPE_CYCLE_COUNTER( STAT_SlateDrawWindowTime );
 
 		const bool bClearHittestGrid = !DrawOnlyThisWindow.IsValid();
-		if ( !SWidget::UseLegacyHittest( ) && bClearHittestGrid )
+		if ( bClearHittestGrid )
 		{
 			HittestGrid->BeginFrame( VirtualDesktopRect );
 		}
diff --git a/Engine/Source/Runtime/SlateCore/Private/Widgets/SWidget.cpp b/Engine/Source/Runtime/SlateCore/Private/Widgets/SWidget.cpp
--- a/Engine/Source/Runtime/SlateCore/Private/Widgets/SWidget.cpp
+++ b/Engine/Source/Runtime/SlateCore/Private/Widgets/SWidget.cpp
@@ -5,11 +5,6 @@
 DECLARE_CYCLE_STAT(TEXT("OnPaint"), STAT_SlateOnPaint, STATGROUP_Slate);
 DECLARE_CYCLE_STAT(TEXT("ArrangeChildren"), STAT_SlateArrangeChildren, STATGROUP_Slate);
 
-int32 SWidget::bUseLegacyHittest = false;
-FAutoConsoleVariableRef SWidget::CVarUseLegacyHittest(
-	TEXT("Slate.UseLegacyHittest"),
-	SWidget::bUseLegacyHittest,
-	TEXT("Should slate fall-back to using the legacy hittest path that does not support child widgets outside their parents' bounds.") );
 
 SWidget::SWidget()
 	: CreatedInFile( TEXT("") )
@@ -524,17 +519,9 @@ void SWidget::SetDebugInfo( const ANSICHAR* InType, const ANSICHAR* InFile, int3
 int32 SWidget::Paint(const FPaintArgs& Args, const FGeometry& AllottedGeometry, const FSlateRect& MyClippingRect, FSlateWindowElementList& OutDrawElements, int32 LayerId, const FWidgetStyle& InWidgetStyle, bool bParentEnabled) const
 {
 	SCOPE_CYCLE_COUNTER(STAT_SlateOnPaint);
-	if ( !bUseLegacyHittest )
-	{
-		FPaintArgs UpdatedArgs = Args.RecordHittestGeometry( this, AllottedGeometry, MyClippingRect );
-		return OnPaint(UpdatedArgs, AllottedGeometry, MyClippingRect, OutDrawElements, LayerId, InWidgetStyle, bParentEnabled);
-	}
-	else
-	{
-		return OnPaint(Args, AllottedGeometry, MyClippingRect, OutDrawElements, LayerId, InWidgetStyle, bParentEnabled);
-	}
 
-	
+	FPaintArgs UpdatedArgs = Args.RecordHittestGeometry( this, AllottedGeometry, MyClippingRect );
+	return OnPaint(UpdatedArgs, AllottedGeometry, MyClippingRect, OutDrawElements, LayerId, InWidgetStyle, bParentEnabled);
 }
 
 void SWidget::ArrangeChildren(const FGeometry& AllottedGeometry, FArrangedChildren& ArrangedChildren) const
@@ -542,9 +529,3 @@ void SWidget::ArrangeChildren(const FGeometry& AllottedGeometry, FArrangedChildr
 	SCOPE_CYCLE_COUNTER(STAT_SlateArrangeChildren);
 	OnArrangeChildren(AllottedGeometry, ArrangedChildren);
 }
-
-
-bool SWidget::UseLegacyHittest()
-{
-	return bUseLegacyHittest != 0;
-}
diff --git a/Engine/Source/Runtime/SlateCore/Public/Widgets/SWidget.h b/Engine/Source/Runtime/SlateCore/Public/Widgets/SWidget.h
--- a/Engine/Source/Runtime/SlateCore/Public/Widgets/SWidget.h
+++ b/Engine/Source/Runtime/SlateCore/Public/Widgets/SWidget.h
@@ -817,13 +817,4 @@ private:
 
 	/** Is this widget hovered? */
 	bool bIsHovered;
-
-	// @see UseLegacyHittest()
-	static int32 bUseLegacyHittest;
-	static FAutoConsoleVariableRef CVarUseLegacyHittest;
-
-public:
-
-	/** For perf testing. Old hittest schema does not support widgets outside their parents' bounds. */
-	static bool UseLegacyHittest();
 };
