Level editor tab can now be reordered again.

diff --git a/Engine/Source/Runtime/Slate/Private/Framework/Docking/SDockingTabWell.cpp b/Engine/Source/Runtime/Slate/Private/Framework/Docking/SDockingTabWell.cpp
--- a/Engine/Source/Runtime/Slate/Private/Framework/Docking/SDockingTabWell.cpp
+++ b/Engine/Source/Runtime/Slate/Private/Framework/Docking/SDockingTabWell.cpp
@@ -22,7 +22,7 @@ SDockingTabWell::SDockingTabWell()
 void SDockingTabWell::Construct( const FArguments& InArgs )
 {
 	ForegroundTabIndex = INDEX_NONE;
-	TabBeingDraggedPtr = TWeakPtr<SDockTab>(NULL);
+	TabBeingDraggedPtr = TSharedPtr<SDockTab>(nullptr);
 	ChildBeingDraggedOffset = 0.0f;
 	TabGrabOffsetFraction = FVector2D::ZeroVector;
 		
@@ -75,7 +75,7 @@ void SDockingTabWell::OnArrangeChildren( const FGeometry& AllottedGeometry, FArr
 
 
 	// The tab that is being dragged by the user, if any.
-	TSharedPtr<SDockTab> TabBeingDragged = TabBeingDraggedPtr.Pin();
+	TSharedPtr<SDockTab> TabBeingDragged = TabBeingDraggedPtr;
 		
 	const int32 NumChildren = Tabs.Num();
 
@@ -170,7 +170,7 @@ FVector2D SDockingTabWell::ComputeDesiredSize() const
 		DesiredSizeResult.Y = FMath::Max(SomeTabDesiredSize.Y, DesiredSizeResult.Y);
 	}
 
-	TSharedPtr<SDockTab> TabBeingDragged = TabBeingDraggedPtr.Pin();
+	TSharedPtr<SDockTab> TabBeingDragged = TabBeingDraggedPtr;
 	if ( TabBeingDragged.IsValid() )
 	{
 		const FVector2D SomeTabDesiredSize = TabBeingDragged->GetDesiredSize();
@@ -207,7 +207,7 @@ FVector2D SDockingTabWell::ComputeChildSize( const FGeometry& AllottedGeometry )
 	// We will make our choice based on the first tab we encounter.
 	TSharedPtr<SDockTab> FirstTab = (NumChildren > 0)
 		? Tabs[0]
-		: TabBeingDraggedPtr.Pin();
+		: TabBeingDraggedPtr;
 
 	// If there are no tabs in this tabwell, assume minor tabs.
 	FVector2D MaxTabSize(0,0);
@@ -236,6 +236,14 @@ float SDockingTabWell::ComputeDraggedTabOffset( const FGeometry& MyGeometry, con
 }
 
 
+int32 SDockingTabWell::ComputeChildDropIndex(const FGeometry& MyGeometry, const TSharedRef<SDockTab>& TabBeingDragged)
+{
+	const float ChildWidth = ComputeChildSize(MyGeometry).X;
+	const float ChildWidthWithOverlap = ChildWidth - TabBeingDragged->GetOverlapWidth();
+	const float DraggedChildCenter = ChildBeingDraggedOffset + ChildWidth / 2;
+	const int32 DropLocationIndex = FMath::Clamp(static_cast<int32>(DraggedChildCenter / ChildWidthWithOverlap), 0, Tabs.Num());
+	return DropLocationIndex;
+}
 
 
 FReply SDockingTabWell::StartDraggingTab( TSharedRef<SDockTab> TabToStartDragging, FVector2D InTabGrabOffsetFraction, const FPointerEvent& MouseEvent )
@@ -248,11 +256,11 @@ FReply SDockingTabWell::StartDraggingTab( TSharedRef<SDockTab> TabToStartDraggin
 	// Tha tab well keeps track of which tab we are dragging; we treat is specially during rendering and layout.
 	TabBeingDraggedPtr = TabToStartDragging;	
 	TabGrabOffsetFraction = InTabGrabOffsetFraction;
+	Tabs.Remove(TabToStartDragging);
 	
 
 	if (bCanLeaveTabWell)
 	{
-		Tabs.Remove(TabToStartDragging);
 		// We just removed the foreground tab.
 		ForegroundTabIndex = INDEX_NONE;
 		ParentTabStackPtr.Pin()->OnTabRemoved(TabToStartDragging->GetLayoutIdentifier());
@@ -302,7 +310,7 @@ void SDockingTabWell::OnDragLeave( const FDragDropEvent& DragDropEvent )
 	if ( DragDropOperation.IsValid() )
 	{
 		TSharedRef<SDockingTabStack> ParentTabStack = ParentTabStackPtr.Pin().ToSharedRef();
-		TSharedPtr<SDockTab> TabBeingDragged = this->TabBeingDraggedPtr.Pin();
+		TSharedPtr<SDockTab> TabBeingDragged = this->TabBeingDraggedPtr;
 		// Check for TabBeingDraggedPtr validity as it may no longer be valid when dragging tabs in game
 		if ( TabBeingDragged.IsValid() && DragDropOperation->CanDockInNode(ParentTabStack, FDockingDragOperation::DockingViaTabWell) )
 		{
@@ -361,13 +369,10 @@ FReply SDockingTabWell::OnDrop( const FGeometry& MyGeometry, const FDragDropEven
 			if ( ensure( TabBeingDraggedPtr.IsValid() ) )
 			{
 				// We dropped a Tab into this TabWell.
+				const TSharedRef<SDockTab> TabBeingDragged = TabBeingDraggedPtr.ToSharedRef();
 
-				// Figure out where in this TabWell to drop the Tab.
-				const float ChildWidth = ComputeChildSize(MyGeometry).X;
-				const TSharedRef<SDockTab> TabBeingDragged = TabBeingDraggedPtr.Pin().ToSharedRef();
-				const float ChildWidthWithOverlap = ChildWidth - TabBeingDragged->GetOverlapWidth();
-				const float DraggedChildCenter = ChildBeingDraggedOffset + ChildWidth / 2;
-				const int32 DropLocationIndex = FMath::Clamp( static_cast<int32>(DraggedChildCenter/ChildWidthWithOverlap), 0, Tabs.Num() );
+				// Figure out where in this TabWell to drop the Tab.				
+				const int32 DropLocationIndex = ComputeChildDropIndex(MyGeometry, TabBeingDragged);
 
 				ensure( DragDropOperation->GetTabBeingDragged().ToSharedRef() == TabBeingDragged );
 
@@ -397,8 +402,12 @@ FReply SDockingTabWell::OnMouseButtonUp( const FGeometry& MyGeometry, const FPoi
 {
 	if (this->HasMouseCapture())
 	{
-		// We are no longer dragging a tab in this tab well, so stop showing it in the TabWell.
+		const TSharedRef<SDockTab> TabBeingDragged = TabBeingDraggedPtr.ToSharedRef();
 		this->TabBeingDraggedPtr.Reset();
+		const int32 DropLocationIndex = ComputeChildDropIndex(MyGeometry, TabBeingDragged);
+		Tabs.Insert( TabBeingDragged, DropLocationIndex );
+		BringTabToFront(TabBeingDragged);
+		// We are no longer dragging a tab in this tab well, so stop showing it in the TabWell.
 		return FReply::Handled().ReleaseMouseCapture();
 	}
 	else
@@ -487,7 +496,7 @@ TSharedPtr<SDockTab> SDockingTabWell::GetForegroundTab() const
 {
 	if (TabBeingDraggedPtr.IsValid())
 	{
-		return TabBeingDraggedPtr.Pin();
+		return TabBeingDraggedPtr;
 	}
 	return (Tabs.Num() > 0 && ForegroundTabIndex > INDEX_NONE) ? Tabs[ForegroundTabIndex] : TSharedPtr<SDockTab>();
 }
diff --git a/Engine/Source/Runtime/Slate/Private/Framework/Docking/SDockingTabWell.h b/Engine/Source/Runtime/Slate/Private/Framework/Docking/SDockingTabWell.h
--- a/Engine/Source/Runtime/Slate/Private/Framework/Docking/SDockingTabWell.h
+++ b/Engine/Source/Runtime/Slate/Private/Framework/Docking/SDockingTabWell.h
@@ -112,6 +112,12 @@ private:
 	 */
 	float ComputeDraggedTabOffset( const FGeometry& MyGeometry, const FPointerEvent& MouseEvent, const FVector2D& TabGrabOffsetFraction ) const;
 
+	/**
+	 * The user is dropping a tab on this TabWell, and the TabWell's geometry is currently MyGeometry.
+	 * @return the index of the slot into which the tab shoudl go.
+	 */
+	int32 ComputeChildDropIndex(const FGeometry& MyGeometry, const TSharedRef<SDockTab>& TabBeingDragged);
+
 	/** The tabs in this TabWell */
 	TSlotlessChildren< SDockTab > Tabs;
 
@@ -119,7 +125,7 @@ private:
 	TWeakPtr<class SDockingTabStack> ParentTabStackPtr;
 
 	/** The Tab being dragged through the TabWell, if there is one */
-	TWeakPtr<SDockTab> TabBeingDraggedPtr;
+	TSharedPtr<SDockTab> TabBeingDraggedPtr;
 
 	/** The offset of the Tab being dragged through this panel */
 	float ChildBeingDraggedOffset;
