UMG/Slate - The progress bar now has clean clipping lines when displaying the progress image.  It no longer scales the paint geometry, it instead scales the clipping rect, it does this for everything except for FillFromCenter, where the only reasonable thing to do is scale the paint geometry, otherwise a sphere image would be filled using a sqaure.  Additionally, I've added Fill directions for TopToBottom and BottomToTop.

diff --git a/Engine/Source/Runtime/Slate/Private/Widgets/Notifications/SProgressBar.cpp b/Engine/Source/Runtime/Slate/Private/Widgets/Notifications/SProgressBar.cpp
--- a/Engine/Source/Runtime/Slate/Private/Widgets/Notifications/SProgressBar.cpp
+++ b/Engine/Source/Runtime/Slate/Private/Widgets/Notifications/SProgressBar.cpp
@@ -113,20 +113,25 @@ int32 SProgressBar::OnPaint( const FPaintArgs& Args, const FGeometry& AllottedGe
 	
 	if( ProgressFraction.IsSet() )
 	{
+		const float ClampedFraction = FMath::Clamp(ProgressFraction.GetValue(), 0.0f, 1.0f);
+
 		switch (BarFillType)
 		{
 			case EProgressBarFillType::RightToLeft:
 			{
-				const float ClampedFraction = FMath::Clamp( ProgressFraction.GetValue(), 0.0f, 1.0f  );
+				const float LeftOffset = ForegroundClippingRect.GetSize().X * (1.0 - ClampedFraction);
+				FSlateRect ClippingArea = ForegroundClippingRect;
+				ClippingArea.Left += LeftOffset;
+
 				// Draw Fill
 				FSlateDrawElement::MakeBox(
 					OutDrawElements,
 					RetLayerId++,
 					AllottedGeometry.ToPaintGeometry(
-						FVector2D( AllottedGeometry.Size.X - (AllottedGeometry.Size.X * ( ClampedFraction )) , 0.0f),
-						FVector2D( AllottedGeometry.Size.X * ( ClampedFraction ) , AllottedGeometry.Size.Y )),
+						FVector2D::ZeroVector,
+						FVector2D( AllottedGeometry.Size.X, AllottedGeometry.Size.Y )),
 					CurrentFillImage,
-					ForegroundClippingRect,
+					ClippingArea,
 					DrawEffects,
 					FillColorAndOpacitySRGB
 					);
@@ -134,7 +139,6 @@ int32 SProgressBar::OnPaint( const FPaintArgs& Args, const FGeometry& AllottedGe
 			}
 			case EProgressBarFillType::FillFromCenter:
 			{
-				const float ClampedFraction = FMath::Clamp( ProgressFraction.GetValue(), 0.0f, 1.0f  );
 				// Draw Fill
 				FSlateDrawElement::MakeBox(
 					OutDrawElements,
@@ -149,26 +153,68 @@ int32 SProgressBar::OnPaint( const FPaintArgs& Args, const FGeometry& AllottedGe
 					);
 				break;
 			}
+			case EProgressBarFillType::TopToBottom:
+			{
+				const float BottomOffset = ForegroundClippingRect.GetSize().Y * ( 1.0 - ClampedFraction );
+				FSlateRect ClippingArea = ForegroundClippingRect;
+				ClippingArea.Bottom -= BottomOffset;
+
+				// Draw Fill
+				FSlateDrawElement::MakeBox(
+					OutDrawElements,
+					RetLayerId++,
+					AllottedGeometry.ToPaintGeometry(
+						FVector2D::ZeroVector,
+						FVector2D( AllottedGeometry.Size.X, AllottedGeometry.Size.Y )),
+					CurrentFillImage,
+					ClippingArea,
+					DrawEffects,
+					FillColorAndOpacitySRGB
+					);
+				break;
+			}
+			case EProgressBarFillType::BottomToTop:
+			{
+				const float TopOffset = ForegroundClippingRect.GetSize().Y * ( 1.0 - ClampedFraction );
+				FSlateRect ClippingArea = ForegroundClippingRect;
+				ClippingArea.Top += TopOffset;
+
+				// Draw Fill
+				FSlateDrawElement::MakeBox(
+					OutDrawElements,
+					RetLayerId++,
+					AllottedGeometry.ToPaintGeometry(
+						FVector2D::ZeroVector,
+						FVector2D( AllottedGeometry.Size.X, AllottedGeometry.Size.Y )),
+					CurrentFillImage,
+					ClippingArea,
+					DrawEffects,
+					FillColorAndOpacitySRGB
+					);
+				break;
+			}
 			case EProgressBarFillType::LeftToRight:
 			default:
 			{
-				const float ClampedFraction = FMath::Clamp( ProgressFraction.GetValue(), 0.0f, 1.0f  );
+				const float RightOffset = ForegroundClippingRect.GetSize().X * ( 1.0 - ClampedFraction );
+				FSlateRect ClippingArea = ForegroundClippingRect;
+				ClippingArea.Right -= RightOffset;
+
 				// Draw Fill
 				FSlateDrawElement::MakeBox(
 					OutDrawElements,
 					RetLayerId++,
 					AllottedGeometry.ToPaintGeometry(
 						FVector2D::ZeroVector,
-						FVector2D( AllottedGeometry.Size.X * ( ClampedFraction ) , AllottedGeometry.Size.Y )),
+						FVector2D( AllottedGeometry.Size.X, AllottedGeometry.Size.Y )),
 					CurrentFillImage,
-					ForegroundClippingRect,
+					ClippingArea,
 					DrawEffects,
 					FillColorAndOpacitySRGB
 					);
 				break;
 			}
 		}
-		
 	}
 	else
 	{
diff --git a/Engine/Source/Runtime/Slate/Public/Widgets/Notifications/SProgressBar.h b/Engine/Source/Runtime/Slate/Public/Widgets/Notifications/SProgressBar.h
--- a/Engine/Source/Runtime/Slate/Public/Widgets/Notifications/SProgressBar.h
+++ b/Engine/Source/Runtime/Slate/Public/Widgets/Notifications/SProgressBar.h
@@ -7,7 +7,7 @@
 /**
  * SProgressBar Fill Type 
  */
-UENUM()
+UENUM(BlueprintType)
 namespace EProgressBarFillType
 {
 	enum Type
@@ -18,6 +18,10 @@ namespace EProgressBarFillType
 		RightToLeft,
 		// will fill up from the center to the outer edges
 		FillFromCenter,
+		// will fill up from the top to the the bottom
+		TopToBottom,
+		// will fill up from the bottom to the the top
+		BottomToTop,
 	};
 }
 
