UMG/Slate - Intersecting a scaled allocated geometry with the clip rect to determine the actual clip rect to use for the progress bar, previous code caused the clip rect to move the start of the progress.

diff --git a/Engine/Source/Runtime/Slate/Private/Widgets/Notifications/SProgressBar.cpp b/Engine/Source/Runtime/Slate/Private/Widgets/Notifications/SProgressBar.cpp
--- a/Engine/Source/Runtime/Slate/Private/Widgets/Notifications/SProgressBar.cpp
+++ b/Engine/Source/Runtime/Slate/Private/Widgets/Notifications/SProgressBar.cpp
@@ -119,9 +119,8 @@ int32 SProgressBar::OnPaint( const FPaintArgs& Args, const FGeometry& AllottedGe
 		{
 			case EProgressBarFillType::RightToLeft:
 			{
-				const float LeftOffset = ForegroundClippingRect.GetSize().X * (1.0 - ClampedFraction);
-				FSlateRect ClippingArea = ForegroundClippingRect;
-				ClippingArea.Left += LeftOffset;
+				FSlateRect ClippedAllotedGeometry = FSlateRect(AllottedGeometry.AbsolutePosition, AllottedGeometry.AbsolutePosition + AllottedGeometry.Size * AllottedGeometry.Scale);
+				ClippedAllotedGeometry.Left = ClippedAllotedGeometry.Right - ClippedAllotedGeometry.GetSize().X * ClampedFraction;
 
 				// Draw Fill
 				FSlateDrawElement::MakeBox(
@@ -131,7 +130,7 @@ int32 SProgressBar::OnPaint( const FPaintArgs& Args, const FGeometry& AllottedGe
 						FVector2D::ZeroVector,
 						FVector2D( AllottedGeometry.Size.X, AllottedGeometry.Size.Y )),
 					CurrentFillImage,
-					ClippingArea,
+					ForegroundClippingRect.IntersectionWith(ClippedAllotedGeometry),
 					DrawEffects,
 					FillColorAndOpacitySRGB
 					);
@@ -155,9 +154,8 @@ int32 SProgressBar::OnPaint( const FPaintArgs& Args, const FGeometry& AllottedGe
 			}
 			case EProgressBarFillType::TopToBottom:
 			{
-				const float BottomOffset = ForegroundClippingRect.GetSize().Y * ( 1.0 - ClampedFraction );
-				FSlateRect ClippingArea = ForegroundClippingRect;
-				ClippingArea.Bottom -= BottomOffset;
+				FSlateRect ClippedAllotedGeometry = FSlateRect(AllottedGeometry.AbsolutePosition, AllottedGeometry.AbsolutePosition + AllottedGeometry.Size * AllottedGeometry.Scale);
+				ClippedAllotedGeometry.Bottom = ClippedAllotedGeometry.Top + ClippedAllotedGeometry.GetSize().Y * ClampedFraction;
 
 				// Draw Fill
 				FSlateDrawElement::MakeBox(
@@ -167,7 +165,7 @@ int32 SProgressBar::OnPaint( const FPaintArgs& Args, const FGeometry& AllottedGe
 						FVector2D::ZeroVector,
 						FVector2D( AllottedGeometry.Size.X, AllottedGeometry.Size.Y )),
 					CurrentFillImage,
-					ClippingArea,
+					ForegroundClippingRect.IntersectionWith(ClippedAllotedGeometry),
 					DrawEffects,
 					FillColorAndOpacitySRGB
 					);
@@ -175,9 +173,8 @@ int32 SProgressBar::OnPaint( const FPaintArgs& Args, const FGeometry& AllottedGe
 			}
 			case EProgressBarFillType::BottomToTop:
 			{
-				const float TopOffset = ForegroundClippingRect.GetSize().Y * ( 1.0 - ClampedFraction );
-				FSlateRect ClippingArea = ForegroundClippingRect;
-				ClippingArea.Top += TopOffset;
+				FSlateRect ClippedAllotedGeometry = FSlateRect(AllottedGeometry.AbsolutePosition, AllottedGeometry.AbsolutePosition + AllottedGeometry.Size * AllottedGeometry.Scale);
+				ClippedAllotedGeometry.Top = ClippedAllotedGeometry.Bottom - ClippedAllotedGeometry.GetSize().Y * ClampedFraction;
 
 				// Draw Fill
 				FSlateDrawElement::MakeBox(
@@ -187,7 +184,7 @@ int32 SProgressBar::OnPaint( const FPaintArgs& Args, const FGeometry& AllottedGe
 						FVector2D::ZeroVector,
 						FVector2D( AllottedGeometry.Size.X, AllottedGeometry.Size.Y )),
 					CurrentFillImage,
-					ClippingArea,
+					ForegroundClippingRect.IntersectionWith(ClippedAllotedGeometry),
 					DrawEffects,
 					FillColorAndOpacitySRGB
 					);
@@ -196,9 +193,8 @@ int32 SProgressBar::OnPaint( const FPaintArgs& Args, const FGeometry& AllottedGe
 			case EProgressBarFillType::LeftToRight:
 			default:
 			{
-				const float RightOffset = ForegroundClippingRect.GetSize().X * ( 1.0 - ClampedFraction );
-				FSlateRect ClippingArea = ForegroundClippingRect;
-				ClippingArea.Right -= RightOffset;
+				FSlateRect ClippedAllotedGeometry = FSlateRect(AllottedGeometry.AbsolutePosition, AllottedGeometry.AbsolutePosition + AllottedGeometry.Size * AllottedGeometry.Scale);
+				ClippedAllotedGeometry.Right = ClippedAllotedGeometry.Left + ClippedAllotedGeometry.GetSize().X * ClampedFraction;
 
 				// Draw Fill
 				FSlateDrawElement::MakeBox(
@@ -208,7 +204,7 @@ int32 SProgressBar::OnPaint( const FPaintArgs& Args, const FGeometry& AllottedGe
 						FVector2D::ZeroVector,
 						FVector2D( AllottedGeometry.Size.X, AllottedGeometry.Size.Y )),
 					CurrentFillImage,
-					ClippingArea,
+					ForegroundClippingRect.IntersectionWith(ClippedAllotedGeometry),
 					DrawEffects,
 					FillColorAndOpacitySRGB
 					);
