Slate: fixed bug with MouseCapture setting. Previously, if the target widget was not successful, its parent would get mouse capture, which is completely wrong!

diff --git a/Engine/Source/Runtime/Slate/Private/Framework/Application/SlateApplication.cpp b/Engine/Source/Runtime/Slate/Private/Framework/Application/SlateApplication.cpp
--- a/Engine/Source/Runtime/Slate/Private/Framework/Application/SlateApplication.cpp
+++ b/Engine/Source/Runtime/Slate/Private/Framework/Application/SlateApplication.cpp
@@ -108,14 +108,26 @@ void FSlateApplication::MouseCaptorHelper::SetMouseCaptor( const FWidgetPath& Ev
 	if ( Widget.IsValid() )
 	{
 		TSharedRef< SWidget > WidgetRef = Widget.ToSharedRef();
-		MouseCaptorWeakPath = EventPath.GetPathDownTo( WidgetRef );
+		FWidgetPath NewMouseCaptorPath = EventPath.GetPathDownTo( WidgetRef );
 
-		if ( !MouseCaptorWeakPath.IsValid() )
+		const auto IsPathToCaptorFound = []( const FWidgetPath& PathToTest, const TSharedRef<SWidget>& WidgetToFind )
+		{
+			return PathToTest.Widgets.Num() > 0 && PathToTest.Widgets.Last().Widget == WidgetToFind;
+		};
+
+		if ( IsPathToCaptorFound( NewMouseCaptorPath, WidgetRef ) )
+		{
+			MouseCaptorWeakPath = NewMouseCaptorPath;
+		}
+		else
 		{
 			// If the target widget wasn't found on the event path then start the search from the root
-			FWidgetPath NewMouseCaptorPath = EventPath.GetPathDownTo( EventPath.Widgets[0].Widget );
+			NewMouseCaptorPath = EventPath.GetPathDownTo( EventPath.Widgets[0].Widget );
 			NewMouseCaptorPath.ExtendPathTo( FWidgetMatcher( WidgetRef ) );
-			MouseCaptorWeakPath = NewMouseCaptorPath;
+			
+			MouseCaptorWeakPath = IsPathToCaptorFound( NewMouseCaptorPath, WidgetRef )
+				? NewMouseCaptorPath
+				: FWeakWidgetPath();
 		}
 	}
 }
