UMG - Tweaking the slider for UMG, holding onto the style pointer now.

diff --git a/Engine/Source/Runtime/Slate/Private/Widgets/Input/SSlider.cpp b/Engine/Source/Runtime/Slate/Private/Widgets/Input/SSlider.cpp
--- a/Engine/Source/Runtime/Slate/Private/Widgets/Input/SSlider.cpp
+++ b/Engine/Source/Runtime/Slate/Private/Widgets/Input/SSlider.cpp
@@ -7,8 +7,8 @@ void SSlider::Construct( const SSlider::FArguments& InDeclaration )
 {
 	check(InDeclaration._Style);
 
-	DisabledHandleImage = &InDeclaration._Style->DisabledThumbImage;
-	NormalHandleImage = &InDeclaration._Style->NormalThumbImage;
+	Style = InDeclaration._Style;
+
 	IndentHandle = InDeclaration._IndentHandle;
 	LockedAttribute = InDeclaration._Locked;
 	Orientation = InDeclaration._Orientation;
@@ -21,7 +21,6 @@ void SSlider::Construct( const SSlider::FArguments& InDeclaration )
 	OnValueChanged = InDeclaration._OnValueChanged;
 }
 
-
 int32 SSlider::OnPaint( const FPaintArgs& Args, const FGeometry& AllottedGeometry, const FSlateRect& MyClippingRect, FSlateWindowElementList& OutDrawElements, int32 LayerId, const FWidgetStyle& InWidgetStyle, bool bParentEnabled ) const
 {
 #if SLATE_HD_STATS
@@ -39,8 +38,8 @@ int32 SSlider::OnPaint( const FPaintArgs& Args, const FGeometry& AllottedGeometr
 	FVector2D SliderEndPoint;
 
 	// calculate slider geometry as if it's a horizontal slider (we'll rotate it later if it's vertical)
-	const FVector2D HalfHandleSize = 0.5f * NormalHandleImage->ImageSize;
-	const float Indentation = IndentHandle.Get() ? NormalHandleImage->ImageSize.X : 0.0f;
+	const FVector2D HalfHandleSize = 0.5f * Style->NormalThumbImage.ImageSize;
+	const float Indentation = IndentHandle.Get() ? Style->NormalThumbImage.ImageSize.X : 0.0f;
 
 	const float SliderLength = AllottedWidth - Indentation;
 	const float SliderHandleOffset = ValueAttribute.Get() * SliderLength;
@@ -102,8 +101,8 @@ int32 SSlider::OnPaint( const FPaintArgs& Args, const FGeometry& AllottedGeometr
 	FSlateDrawElement::MakeBox( 
 		OutDrawElements,
 		LayerId,
-		SliderGeometry.ToPaintGeometry(HandleTopLeftPoint, NormalHandleImage->ImageSize),
-		LockedAttribute.Get() ? DisabledHandleImage : NormalHandleImage,
+		SliderGeometry.ToPaintGeometry(HandleTopLeftPoint, Style->NormalThumbImage.ImageSize),
+		LockedAttribute.Get() ? &Style->DisabledThumbImage : &Style->NormalThumbImage,
 		RotatedClippingRect,
 		DrawEffects,
 		SliderHandleColor.Get().GetColor(InWidgetStyle) * InWidgetStyle.GetColorAndOpacityTint()
@@ -112,25 +111,23 @@ int32 SSlider::OnPaint( const FPaintArgs& Args, const FGeometry& AllottedGeometr
 	return LayerId;
 }
 
-
 FVector2D SSlider::ComputeDesiredSize() const
 {
 	static const FVector2D SSliderDesiredSize(16.0f, 16.0f);
 
-	if (NormalHandleImage == nullptr)
+	if ( Style == nullptr )
 	{
 		return SSliderDesiredSize;
 	}
 
 	if (Orientation == Orient_Vertical)
 	{
-		return FVector2D(NormalHandleImage->ImageSize.Y, SSliderDesiredSize.Y);
+		return FVector2D(Style->NormalThumbImage.ImageSize.Y, SSliderDesiredSize.Y);
 	}
 
-	return FVector2D(SSliderDesiredSize.X, NormalHandleImage->ImageSize.Y);
+	return FVector2D(SSliderDesiredSize.X, Style->NormalThumbImage.ImageSize.Y);
 }
 
-
 FReply SSlider::OnMouseButtonDown( const FGeometry& MyGeometry, const FPointerEvent& MouseEvent )
 {
 	if ((MouseEvent.GetEffectingButton() == EKeys::LeftMouseButton) && !LockedAttribute.Get())
@@ -144,7 +141,6 @@ FReply SSlider::OnMouseButtonDown( const FGeometry& MyGeometry, const FPointerEv
 	return FReply::Unhandled();
 }
 
-
 FReply SSlider::OnMouseButtonUp( const FGeometry& MyGeometry, const FPointerEvent& MouseEvent )
 {
 	if ((MouseEvent.GetEffectingButton() == EKeys::LeftMouseButton) && HasMouseCapture())
@@ -158,7 +154,6 @@ FReply SSlider::OnMouseButtonUp( const FGeometry& MyGeometry, const FPointerEven
 	return FReply::Unhandled();
 }
 
-
 FReply SSlider::OnMouseMove( const FGeometry& MyGeometry, const FPointerEvent& MouseEvent )
 {
 	if (this->HasMouseCapture() && !LockedAttribute.Get())
@@ -172,7 +167,6 @@ FReply SSlider::OnMouseMove( const FGeometry& MyGeometry, const FPointerEvent& M
 	return FReply::Unhandled();
 }
 
-
 void SSlider::CommitValue(float NewValue)
 {
 	if (!ValueAttribute.IsBound())
@@ -183,11 +177,10 @@ void SSlider::CommitValue(float NewValue)
 	OnValueChanged.ExecuteIfBound(NewValue);
 }
 
-
 float SSlider::PositionToValue( const FGeometry& MyGeometry, const FVector2D& AbsolutePosition )
 {
 	const FVector2D LocalPosition = MyGeometry.AbsoluteToLocal(AbsolutePosition);
-	const float Indentation = IndentHandle.Get() ? NormalHandleImage->ImageSize.X : 0.0f;
+	const float Indentation = IndentHandle.Get() ? Style->NormalThumbImage.ImageSize.X : 0.0f;
 
 	float RelativeValue;
 
@@ -215,25 +208,25 @@ void SSlider::SetValue(const TAttribute<float>& InValueAttribute)
 
 void SSlider::SetIndentHandle(const TAttribute<bool>& InIndentHandle)
 {
-    IndentHandle = InIndentHandle;
+	IndentHandle = InIndentHandle;
 }
 
 void SSlider::SetLocked(const TAttribute<bool>& InLocked)
 {
-    LockedAttribute = InLocked;
+	LockedAttribute = InLocked;
 }
 
 void SSlider::SetOrientation(EOrientation InOrientation)
 {
-    Orientation = InOrientation;
+	Orientation = InOrientation;
 }
 
 void SSlider::SetSliderBarColor(FSlateColor InSliderBarColor)
 {
-    SliderBarColor = InSliderBarColor;
+	SliderBarColor = InSliderBarColor;
 }
 
 void SSlider::SetSliderHandleColor(FSlateColor InSliderHandleColor)
 {
-    SliderHandleColor = InSliderHandleColor;
+	SliderHandleColor = InSliderHandleColor;
 }
diff --git a/Engine/Source/Runtime/Slate/Public/Widgets/Input/SSlider.h b/Engine/Source/Runtime/Slate/Public/Widgets/Input/SSlider.h
--- a/Engine/Source/Runtime/Slate/Public/Widgets/Input/SSlider.h
+++ b/Engine/Source/Runtime/Slate/Public/Widgets/Input/SSlider.h
@@ -68,20 +68,20 @@ public:
 
 	/** See the Value attribute */
 	void SetValue(const TAttribute<float>& InValueAttribute);
-    
-    /** See the IndentHandle attribute */
+	
+	/** See the IndentHandle attribute */
 	void SetIndentHandle(const TAttribute<bool>& InIndentHandle);
-    
-    /** See the Locked attribute */
+	
+	/** See the Locked attribute */
 	void SetLocked(const TAttribute<bool>& InLocked);
-    
-    /** See the Orientation attribute */
+	
+	/** See the Orientation attribute */
 	void SetOrientation(EOrientation InOrientation);
-    
-    /** See the SliderBarColor attribute */
+	
+	/** See the SliderBarColor attribute */
 	void SetSliderBarColor(FSlateColor InSliderBarColor);
-    
-    /** See the SliderHandleColor attribute */
+	
+	/** See the SliderHandleColor attribute */
 	void SetSliderHandleColor(FSlateColor InSliderHandleColor);
 
 public:
@@ -114,11 +114,8 @@ protected:
 
 private:
 
-	// Holds the image to use when the slider handle is in its disabled state.
-	const FSlateBrush* DisabledHandleImage;
-
-	// Holds the image to use when the slider handle is in its normal state.
-	const FSlateBrush* NormalHandleImage;
+	// Holds the style passed to the widget upon construction.
+	const FSliderStyle* Style;
 
 	// Holds a flag indicating whether the slideable area should be indented to fit the handle.
 	TAttribute<bool> IndentHandle;
diff --git a/Engine/Source/Runtime/UMG/Private/Components/Slider.cpp b/Engine/Source/Runtime/UMG/Private/Components/Slider.cpp
--- a/Engine/Source/Runtime/UMG/Private/Components/Slider.cpp
+++ b/Engine/Source/Runtime/UMG/Private/Components/Slider.cpp
@@ -13,11 +13,15 @@ USlider::USlider(const FPostConstructInitializeProperties& PCIP)
 	Orientation = EOrientation::Orient_Horizontal;
 	SliderBarColor = FLinearColor::White;
 	SliderHandleColor = FLinearColor::White;
+
+	SSlider::FArguments Defaults;
+	WidgetStyle = *Defaults._Style;
 }
 
 TSharedRef<SWidget> USlider::RebuildWidget()
 {
 	MySlider = SNew(SSlider)
+		.Style(&WidgetStyle)
 		.OnMouseCaptureBegin(BIND_UOBJECT_DELEGATE(FSimpleDelegate, HandleOnMouseCaptureBegin))
 		.OnMouseCaptureEnd(BIND_UOBJECT_DELEGATE(FSimpleDelegate, HandleOnMouseCaptureEnd))
 		.OnValueChanged(BIND_UOBJECT_DELEGATE(FOnFloatValueChanged, HandleOnValueChanged));
diff --git a/Engine/Source/Runtime/UMG/Public/Components/Slider.h b/Engine/Source/Runtime/UMG/Public/Components/Slider.h
--- a/Engine/Source/Runtime/UMG/Public/Components/Slider.h
+++ b/Engine/Source/Runtime/UMG/Public/Components/Slider.h
@@ -16,8 +16,9 @@ class UMG_API USlider : public UWidget
 
 public:
 	
-	///** The style used to draw the slider. */
-	//SLATE_STYLE_ARGUMENT(FSliderStyle, Style)
+	/** The progress bar style */
+	UPROPERTY(EditAnywhere, BlueprintReadWrite, Category="Style", meta=( DisplayName="Style" ))
+	FSliderStyle WidgetStyle;
 
 	/** The volume value to display. */
 	UPROPERTY(EditDefaultsOnly, Category=Appearance)
