Cleaned up hyperlink tooltip API a little

Post-code review, removed accessors & moved delegates into constructor.

diff --git a/Engine/Source/Editor/IntroTutorials/Private/STutorialEditableText.cpp b/Engine/Source/Editor/IntroTutorials/Private/STutorialEditableText.cpp
--- a/Engine/Source/Editor/IntroTutorials/Private/STutorialEditableText.cpp
+++ b/Engine/Source/Editor/IntroTutorials/Private/STutorialEditableText.cpp
@@ -465,7 +465,9 @@ FReply STutorialEditableText::HandleInsertHyperLinkClicked()
 			RunInfo, 
 			MakeShareable(new FString(Name.ToString())), 
 			FEditorStyle::Get().GetWidgetStyle<FHyperlinkStyle>(FName(TEXT("Tutorials.Content.Hyperlink"))), 
-			CurrentHyperlinkType->OnClickedDelegate
+			CurrentHyperlinkType->OnClickedDelegate,
+			CurrentHyperlinkType->TooltipDelegate,
+			CurrentHyperlinkType->TooltipTextDelegate
 			);
 
 		// @todo: if the rich editable text box allowed us to replace a run that we found under the cursor (or returned a non-const
diff --git a/Engine/Source/Runtime/AppFramework/Private/Widgets/Testing/STestSuite.cpp b/Engine/Source/Runtime/AppFramework/Private/Widgets/Testing/STestSuite.cpp
--- a/Engine/Source/Runtime/AppFramework/Private/Widgets/Testing/STestSuite.cpp
+++ b/Engine/Source/Runtime/AppFramework/Private/Widgets/Testing/STestSuite.cpp
@@ -2895,7 +2895,9 @@ protected:
 			RunInfo, 
 			MakeShareable(new FString(Name.ToString())), 
 			FTestStyle::Get().GetWidgetStyle<FHyperlinkStyle>(FName(TEXT("RichText.Editor.Hyperlink"))), 
-			OnHyperlinkClicked
+			OnHyperlinkClicked,
+			FSlateHyperlinkRun::FOnGenerateTooltip(),
+			FSlateHyperlinkRun::FOnGetTooltipText()
 			);
 		RichEditableTextBox->InsertRunAtCursor(HyperlinkRun);
 
diff --git a/Engine/Source/Runtime/Slate/Private/Framework/Text/SlateHyperlinkRun.cpp b/Engine/Source/Runtime/Slate/Private/Framework/Text/SlateHyperlinkRun.cpp
--- a/Engine/Source/Runtime/Slate/Private/Framework/Text/SlateHyperlinkRun.cpp
+++ b/Engine/Source/Runtime/Slate/Private/Framework/Text/SlateHyperlinkRun.cpp
@@ -6,14 +6,14 @@
 
 #include "SlateHyperlinkRun.h"
 
-TSharedRef< FSlateHyperlinkRun > FSlateHyperlinkRun::Create( const FRunInfo& InRunInfo, const TSharedRef< const FString >& InText, const FHyperlinkStyle& InStyle, FOnClick NavigateDelegate )
+TSharedRef< FSlateHyperlinkRun > FSlateHyperlinkRun::Create( const FRunInfo& InRunInfo, const TSharedRef< const FString >& InText, const FHyperlinkStyle& InStyle, FOnClick NavigateDelegate, FOnGenerateTooltip InTooltipDelegate, FOnGetTooltipText InTooltipTextDelegate )
 {
-	return MakeShareable( new FSlateHyperlinkRun( InRunInfo, InText, InStyle, NavigateDelegate ) );
+	return MakeShareable( new FSlateHyperlinkRun( InRunInfo, InText, InStyle, NavigateDelegate, InTooltipDelegate, InTooltipTextDelegate ) );
 }
 
-TSharedRef< FSlateHyperlinkRun > FSlateHyperlinkRun::Create( const FRunInfo& InRunInfo, const TSharedRef< const FString >& InText, const FHyperlinkStyle& InStyle, FOnClick NavigateDelegate, const FTextRange& InRange )
+TSharedRef< FSlateHyperlinkRun > FSlateHyperlinkRun::Create( const FRunInfo& InRunInfo, const TSharedRef< const FString >& InText, const FHyperlinkStyle& InStyle, FOnClick NavigateDelegate, FOnGenerateTooltip InTooltipDelegate, FOnGetTooltipText InTooltipTextDelegate, const FTextRange& InRange )
 {
-	return MakeShareable( new FSlateHyperlinkRun( InRunInfo, InText, InStyle, NavigateDelegate, InRange ) );
+	return MakeShareable( new FSlateHyperlinkRun( InRunInfo, InText, InStyle, NavigateDelegate, InTooltipDelegate, InTooltipTextDelegate, InRange ) );
 }
 
 FTextRange FSlateHyperlinkRun::GetTextRange() const 
@@ -163,7 +163,7 @@ void FSlateHyperlinkRun::Move(const TSharedRef<FString>& NewText, const FTextRan
 
 TSharedRef<IRun> FSlateHyperlinkRun::Clone() const
 {
-	return FSlateHyperlinkRun::Create(RunInfo, Text, Style, NavigateDelegate, Range);
+	return FSlateHyperlinkRun::Create(RunInfo, Text, Style, NavigateDelegate, TooltipDelegate, TooltipTextDelegate, Range);
 }
 
 void FSlateHyperlinkRun::AppendTextTo(FString& AppendToText) const
@@ -184,24 +184,28 @@ const FRunInfo& FSlateHyperlinkRun::GetRunInfo() const
 	return RunInfo;
 }
 
-FSlateHyperlinkRun::FSlateHyperlinkRun( const FRunInfo& InRunInfo, const TSharedRef< const FString >& InText, const FHyperlinkStyle& InStyle, FOnClick InNavigateDelegate ) 
+FSlateHyperlinkRun::FSlateHyperlinkRun( const FRunInfo& InRunInfo, const TSharedRef< const FString >& InText, const FHyperlinkStyle& InStyle, FOnClick InNavigateDelegate, FOnGenerateTooltip InTooltipDelegate, FOnGetTooltipText InTooltipTextDelegate ) 
 	: RunInfo( InRunInfo )
 	, Text( InText )
 	, Range( 0, Text->Len() )
 	, Style( InStyle )
 	, NavigateDelegate( InNavigateDelegate )
+	, TooltipDelegate( InTooltipDelegate )
+	, TooltipTextDelegate( InTooltipTextDelegate )
 	, ViewModel( MakeShareable( new FSlateHyperlinkRun::FWidgetViewModel() ) )
 	, Children()
 {
 
 }
 
-FSlateHyperlinkRun::FSlateHyperlinkRun( const FRunInfo& InRunInfo, const TSharedRef< const FString >& InText, const FHyperlinkStyle& InStyle, FOnClick InNavigateDelegate, const FTextRange& InRange ) 
+FSlateHyperlinkRun::FSlateHyperlinkRun( const FRunInfo& InRunInfo, const TSharedRef< const FString >& InText, const FHyperlinkStyle& InStyle, FOnClick InNavigateDelegate, FOnGenerateTooltip InTooltipDelegate, FOnGetTooltipText InTooltipTextDelegate, const FTextRange& InRange ) 
 	: RunInfo( InRunInfo )
 	, Text( InText )
 	, Range( InRange )
 	, Style( InStyle )
 	, NavigateDelegate( InNavigateDelegate )
+	, TooltipDelegate( InTooltipDelegate )
+	, TooltipTextDelegate( InTooltipTextDelegate )
 	, ViewModel( MakeShareable( new FSlateHyperlinkRun::FWidgetViewModel() ) )
 	, Children()
 {
@@ -214,6 +218,8 @@ FSlateHyperlinkRun::FSlateHyperlinkRun( const FSlateHyperlinkRun& Run )
 	, Range( Run.Range )
 	, Style( Run.Style )
 	, NavigateDelegate( Run.NavigateDelegate )
+	, TooltipDelegate( Run.TooltipDelegate )
+	, TooltipTextDelegate( Run.TooltipTextDelegate )
 	, ViewModel( MakeShareable( new FSlateHyperlinkRun::FWidgetViewModel() ) )
 	, Children()
 {
diff --git a/Engine/Source/Runtime/Slate/Private/Framework/Text/TextDecorators.cpp b/Engine/Source/Runtime/Slate/Private/Framework/Text/TextDecorators.cpp
--- a/Engine/Source/Runtime/Slate/Private/Framework/Text/TextDecorators.cpp
+++ b/Engine/Source/Runtime/Slate/Private/Framework/Text/TextDecorators.cpp
@@ -146,10 +146,6 @@ TSharedRef< ISlateRun > FHyperlinkDecorator::Create(const TSharedRef<class FText
 		RunInfo.MetaData.Add(Pair.Key, OriginalText.Mid( Pair.Value.BeginIndex, Pair.Value.EndIndex - Pair.Value.BeginIndex));
 	}
 
-	TSharedRef<FSlateHyperlinkRun> HyperlinkRun = FSlateHyperlinkRun::Create( RunInfo, InOutModelText, Style->GetWidgetStyle<FHyperlinkStyle>( FName( *StyleName ) ), NavigateDelegate, ModelRange );
-	HyperlinkRun->OnGenerateToolTip() = ToolTipDelegate;
-	HyperlinkRun->OnGetToolTipText() = ToolTipTextDelegate;
-
-	return HyperlinkRun;
+	return FSlateHyperlinkRun::Create( RunInfo, InOutModelText, Style->GetWidgetStyle<FHyperlinkStyle>( FName( *StyleName ) ), NavigateDelegate, ToolTipDelegate, ToolTipTextDelegate, ModelRange );
 }
 #endif //WITH_FANCY_TEXT
\ No newline at end of file
diff --git a/Engine/Source/Runtime/Slate/Public/Framework/Text/SlateHyperlinkRun.h b/Engine/Source/Runtime/Slate/Public/Framework/Text/SlateHyperlinkRun.h
--- a/Engine/Source/Runtime/Slate/Public/Framework/Text/SlateHyperlinkRun.h
+++ b/Engine/Source/Runtime/Slate/Public/Framework/Text/SlateHyperlinkRun.h
@@ -32,9 +32,9 @@ public:
 
 public:
 
-	static TSharedRef< FSlateHyperlinkRun > Create( const FRunInfo& InRunInfo, const TSharedRef< const FString >& InText, const FHyperlinkStyle& InStyle, FOnClick NavigateDelegate );
+	static TSharedRef< FSlateHyperlinkRun > Create( const FRunInfo& InRunInfo, const TSharedRef< const FString >& InText, const FHyperlinkStyle& InStyle, FOnClick NavigateDelegate, FOnGenerateTooltip InTooltipDelegate, FOnGetTooltipText InTooltipTextDelegate );
 																																	 
-	static TSharedRef< FSlateHyperlinkRun > Create( const FRunInfo& InRunInfo, const TSharedRef< const FString >& InText, const FHyperlinkStyle& InStyle, FOnClick NavigateDelegate, const FTextRange& InRange );
+	static TSharedRef< FSlateHyperlinkRun > Create( const FRunInfo& InRunInfo, const TSharedRef< const FString >& InText, const FHyperlinkStyle& InStyle, FOnClick NavigateDelegate, FOnGenerateTooltip InTooltipDelegate, FOnGetTooltipText InTooltipTextDelegate, const FTextRange& InRange );
 
 public:
 
@@ -75,15 +75,11 @@ public:
 
 	virtual const FRunInfo& GetRunInfo() const override;
 
-	FOnGenerateTooltip& OnGenerateToolTip() { return TooltipDelegate; }
-
-	FOnGetTooltipText& OnGetToolTipText() { return TooltipTextDelegate; }
-
 private:
 
-	FSlateHyperlinkRun( const FRunInfo& InRunInfo, const TSharedRef< const FString >& InText, const FHyperlinkStyle& InStyle, FOnClick InNavigateDelegate );
+	FSlateHyperlinkRun( const FRunInfo& InRunInfo, const TSharedRef< const FString >& InText, const FHyperlinkStyle& InStyle, FOnClick InNavigateDelegate, FOnGenerateTooltip InTooltipDelegate, FOnGetTooltipText InTooltipTextDelegate );
 																										 
-	FSlateHyperlinkRun( const FRunInfo& InRunInfo, const TSharedRef< const FString >& InText, const FHyperlinkStyle& InStyle, FOnClick InNavigateDelegate, const FTextRange& InRange );
+	FSlateHyperlinkRun( const FRunInfo& InRunInfo, const TSharedRef< const FString >& InText, const FHyperlinkStyle& InStyle, FOnClick InNavigateDelegate, FOnGenerateTooltip InTooltipDelegate, FOnGetTooltipText InTooltipTextDelegate, const FTextRange& InRange );
 
 	FSlateHyperlinkRun( const FSlateHyperlinkRun& Run );
 
