Fixes for drag and drop on Mac

diff --git a/Engine/Source/Runtime/Slate/Private/Framework/Application/SlateApplication.cpp b/Engine/Source/Runtime/Slate/Private/Framework/Application/SlateApplication.cpp
--- a/Engine/Source/Runtime/Slate/Private/Framework/Application/SlateApplication.cpp
+++ b/Engine/Source/Runtime/Slate/Private/Framework/Application/SlateApplication.cpp
@@ -3231,7 +3231,7 @@ bool FSlateApplication::ProcessMouseButtonDownEvent( const TSharedPtr< FGenericW
 			// If none of the widgets requested keyboard focus to be set (or set the keyboard focus explicitly), set it to the leaf-most widget under the mouse.
 			// On Mac we prevent the OS from activating the window on mouse down, so we have full control and can activate only if there's nothing draggable under the mouse cursor.
 			const bool bFocusChangedByEventHandler = PreviouslyFocusedWidget != GetKeyboardFocusedWidget();
-			if ((!Reply.GetFocusRecepient().IsValid() && !bFocusChangedByEventHandler) || (PLATFORM_MAC && !DragDetector.DetectDragForWidget.IsValid()))
+			if ((!Reply.GetFocusRecepient().IsValid() || (PLATFORM_MAC && !DragDetector.DetectDragForWidget.IsValid())) && !bFocusChangedByEventHandler)
 			{
 				// The event handler for OnMouseButton down may have altered the widget hierarchy.
 				// Refresh the previously-cached widget path.
@@ -3252,8 +3252,10 @@ bool FSlateApplication::ProcessMouseButtonDownEvent( const TSharedPtr< FGenericW
 #if PLATFORM_MAC
 				if (WidgetsUnderCursor.TopLevelWindow.IsValid() && !DragDetector.DetectDragForWidget.IsValid())
 				{
+					MouseCaptorHelper Captor = MouseCaptor;
 					WidgetsUnderCursor.TopLevelWindow->BringToFront(true);
 					FPlatformMisc::ActivateApplication();
+					MouseCaptor = Captor;
 				}
 #endif
 			}
