Fixed right click menu in Matinee window on Mac (one more side effect of drag and drop changes)

diff --git a/Engine/Source/Runtime/Slate/Private/Framework/Application/SlateApplication.cpp b/Engine/Source/Runtime/Slate/Private/Framework/Application/SlateApplication.cpp
--- a/Engine/Source/Runtime/Slate/Private/Framework/Application/SlateApplication.cpp
+++ b/Engine/Source/Runtime/Slate/Private/Framework/Application/SlateApplication.cpp
@@ -3250,7 +3250,7 @@ bool FSlateApplication::ProcessMouseButtonDownEvent( const TSharedPtr< FGenericW
 				}
 
 #if PLATFORM_MAC
-				if (WidgetsUnderCursor.TopLevelWindow.IsValid() && !DragDetector.DetectDragForWidget.IsValid())
+				if (WidgetsUnderCursor.TopLevelWindow.IsValid() && !DragDetector.DetectDragForWidget.IsValid() && MenuStack.GetNumStackLevels() == 0)
 				{
 					MouseCaptorHelper Captor = MouseCaptor;
 					WidgetsUnderCursor.TopLevelWindow->BringToFront(true);
