Fixed crash when pressing return in a read-only text box

https://answers.unrealengine.com/questions/98180/crash-in-stextedithelper.html

FScopedTextTransaction asserts when created for a read-only text box, as they can't be modified. Anything that creates a transaction is supposed to make sure that the text box isn't read-only.

diff --git a/Engine/Source/Runtime/Slate/Private/Framework/Text/TextEditHelper.cpp b/Engine/Source/Runtime/Slate/Private/Framework/Text/TextEditHelper.cpp
--- a/Engine/Source/Runtime/Slate/Private/Framework/Text/TextEditHelper.cpp
+++ b/Engine/Source/Runtime/Slate/Private/Framework/Text/TextEditHelper.cpp
@@ -156,7 +156,7 @@ FReply FTextEditHelper::OnKeyDown( const FKeyboardEvent& InKeyboardEvent, const
 
 		return FReply::Handled();
 	}
-	else if( Key == EKeys::Enter )
+	else if( Key == EKeys::Enter && !TextEditor->GetIsReadOnly() )
 	{
 		FScopedTextTransaction TextTransaction(TextEditor);
 
@@ -188,7 +188,7 @@ FReply FTextEditHelper::OnKeyDown( const FKeyboardEvent& InKeyboardEvent, const
 	}	
 	else if( Key == EKeys::Escape )
 	{
-		return TextEditor->OnEscape();;
+		return TextEditor->OnEscape();
 	}
 
 	// @Todo: Slate keybindings support more than one set of keys. 
