Slate - Adding setters for the new justification/lineheight/margin options for STextBlock.

diff --git a/Engine/Source/Runtime/Slate/Private/Widgets/Text/STextBlock.cpp b/Engine/Source/Runtime/Slate/Private/Widgets/Text/STextBlock.cpp
--- a/Engine/Source/Runtime/Slate/Private/Widgets/Text/STextBlock.cpp
+++ b/Engine/Source/Runtime/Slate/Private/Widgets/Text/STextBlock.cpp
@@ -242,6 +242,21 @@ void STextBlock::SetMinDesiredWidth(const TAttribute<float>& InMinDesiredWidth)
 	MinDesiredWidth = InMinDesiredWidth;
 }
 
+void STextBlock::SetLineHeightPercentage(const TAttribute<float>& InLineHeightPercentage)
+{
+	LineHeightPercentage = InLineHeightPercentage;
+}
+
+void STextBlock::SetMargin(const TAttribute<FMargin>& InMargin)
+{
+	Margin = InMargin;
+}
+
+void STextBlock::SetJustification(const TAttribute<ETextJustify::Type>& InJustification)
+{
+	Justification = InJustification;
+}
+
 FTextBlockStyle STextBlock::GetComputedTextStyle() const
 {
 	FTextBlockStyle ComputedStyle = *TextStyle;
diff --git a/Engine/Source/Runtime/Slate/Public/Widgets/Text/STextBlock.h b/Engine/Source/Runtime/Slate/Public/Widgets/Text/STextBlock.h
--- a/Engine/Source/Runtime/Slate/Public/Widgets/Text/STextBlock.h
+++ b/Engine/Source/Runtime/Slate/Public/Widgets/Text/STextBlock.h
@@ -157,6 +157,15 @@ public:
 	/** See MinDesiredWidth attribute */
 	void SetMinDesiredWidth(const TAttribute<float>& InMinDesiredWidth);
 
+	/** See LineHeightPercentage attribute */
+	void SetLineHeightPercentage(const TAttribute<float>& InLineHeightPercentage);
+
+	/** See Margin attribute */
+	void SetMargin(const TAttribute<FMargin>& InMargin);
+
+	/** See Justification attribute */
+	void SetJustification(const TAttribute<ETextJustify::Type>& InJustification);
+
 	// SWidget interface
 	virtual int32 OnPaint( const FPaintArgs& Args, const FGeometry& AllottedGeometry, const FSlateRect& MyClippingRect, FSlateWindowElementList& OutDrawElements, int32 LayerId, const FWidgetStyle& InWidgetStyle, bool bParentEnabled ) const override;
 	virtual FReply OnMouseButtonDoubleClick( const FGeometry& InMyGeometry, const FPointerEvent& InMouseEvent ) override;
