Added Justification attribute to SInlineEditableTextBlock

This is just a pass-through attribute to the underlying STextBlock

diff --git a/Engine/Source/Runtime/Slate/Private/Widgets/Text/SInlineEditableTextBlock.cpp b/Engine/Source/Runtime/Slate/Private/Widgets/Text/SInlineEditableTextBlock.cpp
--- a/Engine/Source/Runtime/Slate/Private/Widgets/Text/SInlineEditableTextBlock.cpp
+++ b/Engine/Source/Runtime/Slate/Private/Widgets/Text/SInlineEditableTextBlock.cpp
@@ -32,6 +32,7 @@ void SInlineEditableTextBlock::Construct( const FArguments& InArgs )
 			.HighlightText( InArgs._HighlightText )
 			.ToolTipText( InArgs._ToolTipText )
 			.WrapTextAt( InArgs._WrapTextAt )
+			.Justification( InArgs._Justification )
 			.LineBreakPolicy( InArgs._LineBreakPolicy )
 		]
 	];
diff --git a/Engine/Source/Runtime/Slate/Public/Widgets/Text/SInlineEditableTextBlock.h b/Engine/Source/Runtime/Slate/Public/Widgets/Text/SInlineEditableTextBlock.h
--- a/Engine/Source/Runtime/Slate/Public/Widgets/Text/SInlineEditableTextBlock.h
+++ b/Engine/Source/Runtime/Slate/Public/Widgets/Text/SInlineEditableTextBlock.h
@@ -20,6 +20,7 @@ class SLATE_API SInlineEditableTextBlock: public SCompoundWidget
 		, _ShadowColorAndOpacity()
 		, _HighlightText()
 		, _WrapTextAt(0.0f)
+		, _Justification(ETextJustify::Left)
 		, _LineBreakPolicy()
 		, _IsReadOnly(false)
 	{}
@@ -48,6 +49,9 @@ class SLATE_API SInlineEditableTextBlock: public SCompoundWidget
 		/** Whether text wraps onto a new line when it's length exceeds this width; if this value is zero or negative, no wrapping occurs. */
 		SLATE_ATTRIBUTE( float, WrapTextAt )
 
+		/** How the text should be aligned with the margin. */
+		SLATE_ATTRIBUTE( ETextJustify::Type, Justification )
+
 		/** The iterator to use to detect appropriate soft-wrapping points for lines (or null to use the default) */
 		SLATE_ARGUMENT( TSharedPtr<IBreakIterator>, LineBreakPolicy )
 
