Fixed "flickering" issue caused by SMultiLineEditableText calculating invalid scroll offsets

This only happened when SMultiLineEditableText was being used without scrollbars (as the scrollbars were clamping the offset).

If the content was smaller than the actual size, the ratio calculated would be greater than 1, which would produce bad offsets - the offset is now only calculated when the ratio is <1, otherwise it's just set to 0.

diff --git a/Engine/Source/Runtime/Slate/Private/Widgets/Text/SMultiLineEditableText.cpp b/Engine/Source/Runtime/Slate/Private/Widgets/Text/SMultiLineEditableText.cpp
--- a/Engine/Source/Runtime/Slate/Private/Widgets/Text/SMultiLineEditableText.cpp
+++ b/Engine/Source/Runtime/Slate/Private/Widgets/Text/SMultiLineEditableText.cpp
@@ -2080,7 +2080,7 @@ void SMultiLineEditableText::Tick( const FGeometry& AllottedGeometry, const doub
 
 		// If this text box has no size, do not compute a view fraction because it will be wrong and causes pop in when the size is available
 		const float ViewFraction = (VisibleSize > 0.0f && ContentSize > 0.0f) ? VisibleSize / ContentSize : 1;
-		const float ViewOffset = (ContentSize > 0.0f) ? FMath::Clamp<float>(ScrollOffset.X / ContentSize, 0.0f, 1.0f - ViewFraction) : 0.0f;
+		const float ViewOffset = (ContentSize > 0.0f && ViewFraction < 1.0f) ? FMath::Clamp<float>(ScrollOffset.X / ContentSize, 0.0f, 1.0f - ViewFraction) : 0.0f;
 	
 		// Update the scrollbar with the clamped version of the offset
 		ScrollOffset.X = ViewOffset * ContentSize;
@@ -2102,7 +2102,7 @@ void SMultiLineEditableText::Tick( const FGeometry& AllottedGeometry, const doub
 
 		// If this text box has no size, do not compute a view fraction because it will be wrong and causes pop in when the size is available
 		const float ViewFraction = (VisibleSize > 0.0f && ContentSize > 0.0f) ? VisibleSize / ContentSize : 1;
-		const float ViewOffset = (ContentSize > 0.0f) ? FMath::Clamp<float>(ScrollOffset.Y / ContentSize, 0.0f, 1.0f - ViewFraction) : 0.0f;
+		const float ViewOffset = (ContentSize > 0.0f && ViewFraction < 1.0f) ? FMath::Clamp<float>(ScrollOffset.Y / ContentSize, 0.0f, 1.0f - ViewFraction) : 0.0f;
 	
 		// Update the scrollbar with the clamped version of the offset
 		ScrollOffset.Y = ViewOffset * ContentSize;
