Undid some clearly incorrect metadata tagging in SDockingTabStack.
Fixed indentation in slate declarations.

diff --git a/Engine/Source/Runtime/Slate/Private/Framework/Docking/SDockingTabStack.cpp b/Engine/Source/Runtime/Slate/Private/Framework/Docking/SDockingTabStack.cpp
--- a/Engine/Source/Runtime/Slate/Private/Framework/Docking/SDockingTabStack.cpp
+++ b/Engine/Source/Runtime/Slate/Private/Framework/Docking/SDockingTabStack.cpp
@@ -114,100 +114,94 @@ void SDockingTabStack::Construct( const FArguments& InArgs, const TSharedRef<FTa
 		.Padding(5.0f, 0.0f, 0.0f, 0.0f)
 		.VAlign((VAlign_Center));
 
-	FTagMetaData MetaData(TEXT("Undefined"));
-	if (Tabs.Num() > 0)
-	{
-		MetaData.Tag = *Tabs[0].TabId.ToString();
-	}
 	ChildSlot
 	[
 		SNew(SVerticalBox)
-			.Visibility( EVisibility::SelfHitTestInvisible )
+		.Visibility( EVisibility::SelfHitTestInvisible )
 
 		+ SVerticalBox::Slot()
-			.AutoHeight()
+		.AutoHeight()
+		[
+			// tab well area
+			SNew(SBorder)
+			.Visibility(this, &SDockingTabStack::GetTabWellVisibility)
+			.DesiredSizeScale(this, &SDockingTabStack::GetTabWellScale)
+			.BorderImage(FCoreStyle::Get().GetBrush("NoBorder"))
+			.VAlign(VAlign_Bottom)
+			.OnMouseButtonDown(this, &SDockingTabStack::TabWellRightClicked)
+			.Padding(0.0f)
 			[
-				// tab well area
-				SNew(SBorder)
-					.Visibility(this, &SDockingTabStack::GetTabWellVisibility)
-					.DesiredSizeScale(this, &SDockingTabStack::GetTabWellScale)
-					.BorderImage(FCoreStyle::Get().GetBrush("NoBorder"))
-					.VAlign(VAlign_Bottom)
-					.OnMouseButtonDown(this, &SDockingTabStack::TabWellRightClicked)
-					.Padding(0.0f)
-					[
-						SNew(SVerticalBox)
-							.Visibility(EVisibility::SelfHitTestInvisible)
-
-						+ SVerticalBox::Slot()
-							.AutoHeight()
-							.Expose(TitleBarSlot)
-
-						+ SVerticalBox::Slot()
-							.AutoHeight()
-							[
-								SNew(SImage)
-									.Image(this, &SDockingTabStack::GetTabWellBrush)
-							]
-					]
+				SNew(SVerticalBox)
+				.Visibility(EVisibility::SelfHitTestInvisible)
+
+				+ SVerticalBox::Slot()
+				.AutoHeight()
+				.Expose(TitleBarSlot)
+
+				+ SVerticalBox::Slot()
+				.AutoHeight()
+				[
+					SNew(SImage)
+						.Image(this, &SDockingTabStack::GetTabWellBrush)
+				]
 			]
+		]
 
 		+ SVerticalBox::Slot()
-			.FillHeight(1.0f)
+		.FillHeight(1.0f)
+		[
+			// tab content area
+			SAssignNew(OverlayManagement.ContentAreaOverlay, SOverlay)
+
+			+ SOverlay::Slot()
 			[
-				// tab content area
-				SAssignNew(OverlayManagement.ContentAreaOverlay, SOverlay)
-
-				+ SOverlay::Slot()
-					[
-						// content goes here
-						SAssignNew(ContentSlot, SBorder)
-							.BorderImage(this, &SDockingTabStack::GetContentAreaBrush)
-							.Padding(this, &SDockingTabStack::GetContentPadding)
-							.AddMetaData<FTagMetaData>(MetaData)
-							[
-								SNew(STextBlock)
-									.Text(LOCTEXT("EmptyTabMessage", "Empty Tab!"))
-							]
-					]
-
-				+ SOverlay::Slot()
-					.Padding(0.0f)
-					.HAlign(HAlign_Left)
-					.VAlign(VAlign_Top)
-					[
-						// unhide tab well button (yellow triangle)
-						SNew(SButton)
-							.ButtonStyle(UnhideTabWellButtonStyle)
-							.OnClicked(this, &SDockingTabStack::UnhideTabWell)
-							.ContentPadding(0.0f)
-							.Visibility(this, &SDockingTabStack::GetUnhideButtonVisibility)
-							.DesiredSizeScale(this, &SDockingTabStack::GetUnhideTabWellButtonScale)
-							.ButtonColorAndOpacity(this, &SDockingTabStack::GetUnhideTabWellButtonOpacity)
-							[
-								// button should be big enough to show its own image
-								SNew(SSpacer)
-									.Size(UnhideTabWellButtonStyle->Normal.ImageSize)
-							]
-					]
+				// content goes here
+				SAssignNew(ContentSlot, SBorder)
+				.BorderImage(this, &SDockingTabStack::GetContentAreaBrush)
+				.Padding(this, &SDockingTabStack::GetContentPadding)
+				[
+					SNew(STextBlock)
+						.Text(LOCTEXT("EmptyTabMessage", "Empty Tab!"))
+				]
+			]
+
+			+ SOverlay::Slot()
+			.Padding(0.0f)
+			.HAlign(HAlign_Left)
+			.VAlign(VAlign_Top)
+			[
+				// unhide tab well button (yellow triangle)
+				SNew(SButton)
+				.ButtonStyle(UnhideTabWellButtonStyle)
+				.OnClicked(this, &SDockingTabStack::UnhideTabWell)
+				.ContentPadding(0.0f)
+				.Visibility(this, &SDockingTabStack::GetUnhideButtonVisibility)
+				.DesiredSizeScale(this, &SDockingTabStack::GetUnhideTabWellButtonScale)
+				.ButtonColorAndOpacity(this, &SDockingTabStack::GetUnhideTabWellButtonOpacity)
+				[
+					// button should be big enough to show its own image
+					SNew(SSpacer)
+						.Size(UnhideTabWellButtonStyle->Normal.ImageSize)
+				]
+			]
 
 #if DEBUG_TAB_MANAGEMENT
-				+ SOverlay::Slot()
-					.HAlign(HAlign_Left)
-					.VAlign(VAlign_Top)
-					[
-						SNew(SBorder)
-							.BorderImage(FCoreStyle::Get().GetBrush("Docking.Border"))
-							.BorderBackgroundColor(FLinearColor(1.0f, 0.5f, 0.0f, 0.75f))
-							.Visibility(EVisibility::HitTestInvisible)
-							[
-								SNew(STextBlock)
-									.Text(this, &SDockingTabStack::ShowPersistentTabs)
-									.ShadowOffset(FVector2D::UnitVector)
-							]
-					]
-#endif
+			+ SOverlay::Slot()
+			.HAlign(HAlign_Left)
+			.VAlign(VAlign_Top)
+			[
+				SNew(SBorder)
+				.BorderImage(FCoreStyle::Get().GetBrush("Docking.Border"))
+				.BorderBackgroundColor(FLinearColor(1.0f, 0.5f, 0.0f, 0.75f))
+				.Visibility(EVisibility::HitTestInvisible)
+				[
+					SNew(STextBlock)
+					.Text(this, &SDockingTabStack::ShowPersistentTabs)
+					.ShadowOffset(FVector2D::UnitVector)
+				]
 			]
+#endif
+		]
 	];
 
 	if (bIsDocumentArea)
