Add OnClick event for SListView codereview(Nick.Atamas)

diff --git a/Engine/Source/Runtime/Slate/Public/Framework/SlateDelegates.h b/Engine/Source/Runtime/Slate/Public/Framework/SlateDelegates.h
--- a/Engine/Source/Runtime/Slate/Public/Framework/SlateDelegates.h
+++ b/Engine/Source/Runtime/Slate/Public/Framework/SlateDelegates.h
@@ -182,6 +182,13 @@ public:
 		bool
 	)
 
+	/** Called when the user clicks on an item in a tree or list */
+	DECLARE_DELEGATE_OneParam(
+		FOnMouseButtonClick,
+			/** param: The item that was double-clicked on */
+			ArgumentType
+			)
+
 	/** Called when the user double-clicks on an item in a tree or list */
 	DECLARE_DELEGATE_OneParam(
 		FOnMouseButtonDoubleClick,
diff --git a/Engine/Source/Runtime/Slate/Public/Framework/Views/ITypedTableView.h b/Engine/Source/Runtime/Slate/Public/Framework/Views/ITypedTableView.h
--- a/Engine/Source/Runtime/Slate/Public/Framework/Views/ITypedTableView.h
+++ b/Engine/Source/Runtime/Slate/Public/Framework/Views/ITypedTableView.h
@@ -147,6 +147,9 @@ public:
 	/** Called when the item is right-clicked when not right mouse scrolling */
 	virtual void Private_OnItemRightClicked( ItemType TheItem, const FPointerEvent& MouseEvent ) = 0;
 
+	/** Called when the item is clicked, and returns whether or not the event was handled */
+	virtual bool Private_OnItemClicked(ItemType TheItem) = 0;
+
 	/** Called when the item is double-clicked, and returns whether or not the event was handled */
 	virtual bool Private_OnItemDoubleClicked( ItemType TheItem ) = 0;
 
diff --git a/Engine/Source/Runtime/Slate/Public/Widgets/Views/SListView.h b/Engine/Source/Runtime/Slate/Public/Widgets/Views/SListView.h
--- a/Engine/Source/Runtime/Slate/Public/Widgets/Views/SListView.h
+++ b/Engine/Source/Runtime/Slate/Public/Widgets/Views/SListView.h
@@ -38,6 +38,7 @@ public:
 	typedef typename TSlateDelegates< ItemType >::FOnGenerateRow FOnGenerateRow;
 	typedef typename TSlateDelegates< ItemType >::FOnItemScrolledIntoView FOnItemScrolledIntoView;
 	typedef typename TSlateDelegates< NullableItemType >::FOnSelectionChanged FOnSelectionChanged;
+	typedef typename TSlateDelegates< ItemType >::FOnMouseButtonClick FOnMouseButtonClick;
 	typedef typename TSlateDelegates< ItemType >::FOnMouseButtonDoubleClick FOnMouseButtonDoubleClick;
 
 public:
@@ -57,6 +58,7 @@ public:
 		, _ListItemsSource( static_cast<const TArray<ItemType>*>(nullptr) ) //@todo Slate Syntax: Initializing from nullptr without a cast
 		, _ItemHeight(16)
 		, _OnContextMenuOpening()
+		, _OnMouseButtonClick()
 		, _OnMouseButtonDoubleClick()
 		, _OnSelectionChanged()
 		, _SelectionMode(ESelectionMode::Multi)
@@ -74,6 +76,8 @@ public:
 
 		SLATE_EVENT( FOnContextMenuOpening, OnContextMenuOpening )
 
+		SLATE_EVENT(FOnMouseButtonClick, OnMouseButtonClick)
+
 		SLATE_EVENT( FOnMouseButtonDoubleClick, OnMouseButtonDoubleClick )
 
 		SLATE_EVENT( FOnSelectionChanged, OnSelectionChanged )
@@ -102,6 +106,7 @@ public:
 
 		this->ItemsSource = InArgs._ListItemsSource;
 		this->OnContextMenuOpening = InArgs._OnContextMenuOpening;
+		this->OnClick = InArgs._OnMouseButtonClick;
 		this->OnDoubleClick = InArgs._OnMouseButtonDoubleClick;
 		this->OnSelectionChanged = InArgs._OnSelectionChanged;
 		this->SelectionMode = InArgs._SelectionMode;
@@ -658,6 +663,16 @@ public:
 		this->OnRightMouseButtonUp( MouseEvent.GetScreenSpacePosition() );
 	}
 
+	virtual bool Private_OnItemClicked(ItemType TheItem) override
+	{
+		if (OnClick.ExecuteIfBound(TheItem))
+		{
+			return true;	// Handled
+		}
+
+		return false;	// Not handled
+	}
+	
 	virtual bool Private_OnItemDoubleClicked( ItemType TheItem ) override
 	{
 		if( OnDoubleClick.ExecuteIfBound( TheItem ) )
@@ -1322,6 +1337,9 @@ protected:
 	/** Delegate to invoke when selection changes. */
 	FOnSelectionChanged OnSelectionChanged;
 
+	/** Caled when the user clicks on an element int he list view with the left mouse button */
+	FOnMouseButtonClick OnClick;
+
 	/** Called when the user double-clicks on an element in the list view with the left mouse button */
 	FOnMouseButtonDoubleClick OnDoubleClick;
 
diff --git a/Engine/Source/Runtime/Slate/Public/Widgets/Views/STableRow.h b/Engine/Source/Runtime/Slate/Public/Widgets/Views/STableRow.h
--- a/Engine/Source/Runtime/Slate/Public/Widgets/Views/STableRow.h
+++ b/Engine/Source/Runtime/Slate/Public/Widgets/Views/STableRow.h
@@ -218,7 +218,7 @@ public:
 						OwnerWidget->Private_SetItemSelection( *MyItem, true, true );
 						ChangedSelectionOnMouseDown = true;
 					}
-
+					
 					return FReply::Handled()
 						.DetectDrag( SharedThis(this), EKeys::LeftMouseButton )
 						.SetKeyboardFocus( OwnerWidget->AsWidget(), EKeyboardFocusCause::Mouse )
@@ -351,6 +351,12 @@ public:
 					}
 				}
 
+				const ItemType* MyItem = OwnerWidget->Private_ItemFromWidget(this);
+				if (MyItem && OwnerWidget->Private_OnItemClicked(*MyItem))
+				{
+					Reply = FReply::Handled().ReleaseMouseCapture();
+				}
+
 				if ( ChangedSelectionOnMouseDown )
 				{
 					OwnerWidget->Private_SignalSelectionChanged(ESelectInfo::OnMouseClick);
