Slate - Passing the owner listview to the item generator so on errors generating widgets we can print the tostring widget information to the output in the debugger.

diff --git a/Engine/Source/Runtime/Slate/Public/Widgets/Views/SListView.h b/Engine/Source/Runtime/Slate/Public/Widgets/Views/SListView.h
--- a/Engine/Source/Runtime/Slate/Public/Widgets/Views/SListView.h
+++ b/Engine/Source/Runtime/Slate/Public/Widgets/Views/SListView.h
@@ -157,6 +157,7 @@ public:
 		, ItemsSource( nullptr )
 		, ItemToScrollIntoView( NullableItemType(nullptr) )
 		, ItemToNotifyWhenInView( NullableItemType(nullptr) ) 
+		, WidgetGenerator( this )
 	{ }
 
 public:
@@ -393,6 +394,11 @@ private:
 	class FWidgetGenerator
 	{
 	public:
+		FWidgetGenerator(SListView<ItemType>* OwnerList)
+			: OwnerList(OwnerList)
+		{
+		}
+
 		/**
 		 * Find a widget for this item if it has already been constructed.
 		 *
@@ -465,8 +471,8 @@ private:
 				}				
 			}
 
-			checkf( ItemToWidgetMap.Num() == WidgetMapToItem.Num(), TEXT("ItemToWidgetMap length (%d) does not match WidgetMapToItem length (%d)"), ItemToWidgetMap.Num(), WidgetMapToItem.Num() );
-			checkf( WidgetMapToItem.Num() == ItemsWithGeneratedWidgets.Num(), TEXT("WidgetMapToItem length (%d) does not match ItemsWithGeneratedWidgets length (%d). This is often because the same item is in the list more than once."), WidgetMapToItem.Num(), ItemsWithGeneratedWidgets.Num() );
+			checkf(ItemToWidgetMap.Num() == WidgetMapToItem.Num(), TEXT("ItemToWidgetMap length (%d) does not match WidgetMapToItem length (%d).  %s"), ItemToWidgetMap.Num(), WidgetMapToItem.Num(), *OwnerList->ToString());
+			checkf(WidgetMapToItem.Num() == ItemsWithGeneratedWidgets.Num(), TEXT("WidgetMapToItem length (%d) does not match ItemsWithGeneratedWidgets length (%d). This is often because the same item is in the list more than once.  %s"), WidgetMapToItem.Num(), ItemsWithGeneratedWidgets.Num(), *OwnerList->ToString());
 			ItemsToBeCleanedUp.Reset();
 		}
 
@@ -494,6 +500,9 @@ private:
 			ItemsToBeCleanedUp.Reset();
 		}
 
+		/** We store a pointer to the owner list for error purposes, so when asserts occur we can report which list it happened for. */
+		SListView<ItemType>* OwnerList;
+
 		/** Map of DataItems to corresponding SWidgets */
 		TMap< ItemType, TSharedRef<ITableRow> > ItemToWidgetMap;
 
