Adding image support to tutorial rich text

Improved text layout support when inserting runs/text or splitting lines on runs that were non-text (images or widgets). The text layout now inserts an extra text run when splitting a non-text run, which avoids issues where text was either being inserted into a non-text run (and vanishing), or an image run was being cloned (and appearing twice).

This also fixes the cursor movement in the multiline editable text when selecting over images or widgets (the cursor would jump to the start of the document as GetTextIndexAt hadn't been implemented. Additionally, it also fixes an issue where Measure was always trying to place the cursor at the end of an image run (ignoring the values of BeginIndex and EndIndex) which made the cursor offset draw in the wrong place.

These changes required the text layout to be able to create a default text run, which involved refactoring the text marshallers as the Slate text run now knows about the default text style, taking that responsibility away from the marshallers

Added tutorial-specific image decorator that accepts a content-relative or engine relative image path.

Added button to tutorial rich text editor to add new images.

All previously-imported images should still 'work'.

diff --git a/Engine/Source/Editor/IntroTutorials/Private/STutorialEditableText.cpp b/Engine/Source/Editor/IntroTutorials/Private/STutorialEditableText.cpp
--- a/Engine/Source/Editor/IntroTutorials/Private/STutorialEditableText.cpp
+++ b/Engine/Source/Editor/IntroTutorials/Private/STutorialEditableText.cpp
@@ -4,6 +4,8 @@
 #include "STutorialEditableText.h"
 #include "RichTextLayoutMarshaller.h"
 #include "TutorialText.h"
+#include "TutorialImageDecorator.h"
+#include "DesktopPlatformModule.h"
 
 #define LOCTEXT_NAMESPACE "STutorialEditableText"
 
@@ -27,8 +29,7 @@ void STutorialEditableText::Construct(const FArguments& InArgs)
 
 	TSharedRef<FRichTextLayoutMarshaller> RichTextMarshaller = FRichTextLayoutMarshaller::Create(
 		TArray<TSharedRef<ITextDecorator>>(), 
-		&FEditorStyle::Get(), 
-		FEditorStyle::Get().GetWidgetStyle<FTextBlockStyle>("Tutorials.Content.Text")
+		&FEditorStyle::Get()
 		);
 
 
@@ -236,6 +237,22 @@ void STutorialEditableText::Construct(const FArguments& InArgs)
 						]
 					]
 				]
+
+				+SHorizontalBox::Slot()
+				.Padding(FMargin(4.0f, 0.0f, 0.0f, 0.0f))
+				.AutoWidth()
+				[
+					SNew(SButton)
+					.ToolTipText(LOCTEXT("ImageButtonTooltip", "Insert Image"))
+					.ButtonStyle(FEditorStyle::Get(), "TutorialEditableText.Toolbar.Button")
+					.OnClicked(this, &STutorialEditableText::HandleImageButtonClicked)
+					.ContentPadding(1.0f)
+					.Content()
+					[
+						SNew(SImage)
+						.Image(FEditorStyle::Get().GetBrush("TutorialEditableText.Toolbar.ImageImage"))
+					]
+				]
 			]
 		]
 	];
@@ -547,4 +564,51 @@ EVisibility STutorialEditableText::GetExcerptVisibility() const
 	return CurrentHyperlinkType.IsValid() && CurrentHyperlinkType->Type == EHyperlinkType::UDN ? EVisibility::Visible : EVisibility::Collapsed;
 }
 
+FReply STutorialEditableText::HandleImageButtonClicked()
+{
+	IDesktopPlatform* DesktopPlatform = FDesktopPlatformModule::Get();
+	if ( DesktopPlatform )
+	{
+		TArray<FString> OutFiles;
+		const FString Extension(TEXT("png"));
+		const FString Filter = FString::Printf(TEXT("%s files (*.%s)|*.%s"), *Extension, *Extension, *Extension);
+		const FString DefaultPath = FEditorDirectories::Get().GetLastDirectory(ELastDirectory::GENERIC_IMPORT);
+
+		TSharedPtr<SWindow> ParentWindow = FSlateApplication::Get().FindWidgetWindow(AsShared());
+		void* ParentWindowHandle = (ParentWindow.IsValid() && ParentWindow->GetNativeWindow().IsValid()) ? ParentWindow->GetNativeWindow()->GetOSWindowHandle() : nullptr;
+
+		if (DesktopPlatform->OpenFileDialog(ParentWindowHandle, FText::Format(LOCTEXT("ImagePickerDialogTitle", "Choose a {0} file"), FText::FromString(Extension)).ToString(), DefaultPath, TEXT(""), Filter, EFileDialogFlags::None, OutFiles))
+		{
+			check(OutFiles.Num() == 1);
+
+			FRunInfo RunInfo(TEXT("img"));
+
+			// the path to the image needs to be stored either as a 'long package name' version of itself
+			// (minus png extension) or as a literal (base dir relative) path.
+			FString ContentPath;
+			if(FPackageName::TryConvertFilenameToLongPackageName(OutFiles[0], ContentPath))
+			{
+				RunInfo.MetaData.Add(TEXT("src"), ContentPath);
+			}
+			else
+			{
+				RunInfo.MetaData.Add(TEXT("src"), OutFiles[0]);
+			}
+			
+			TSharedRef<FSlateImageRun> ImageRun = FSlateImageRun::Create(
+				RunInfo, 
+				MakeShareable(new FString(TEXT("\x200B"))), // Zero-Width Breaking Space
+				FName(*FTutorialImageDecorator::GetPathToImage(OutFiles[0])),
+				0
+				);
+
+			RichEditableTextBox->InsertRunAtCursor(ImageRun);
+
+			FEditorDirectories::Get().SetLastDirectory(ELastDirectory::GENERIC_IMPORT, FPaths::GetPath(OutFiles[0]));	
+		}
+	}
+
+	return FReply::Handled();
+}
+
 #undef LOCTEXT_NAMESPACE
\ No newline at end of file
diff --git a/Engine/Source/Editor/IntroTutorials/Private/STutorialEditableText.h b/Engine/Source/Editor/IntroTutorials/Private/STutorialEditableText.h
--- a/Engine/Source/Editor/IntroTutorials/Private/STutorialEditableText.h
+++ b/Engine/Source/Editor/IntroTutorials/Private/STutorialEditableText.h
@@ -70,6 +70,8 @@ protected:
 
 	EVisibility GetExcerptVisibility() const;
 
+	FReply HandleImageButtonClicked();
+
 protected:
 	TSharedPtr<SMultiLineEditableTextBox> RichEditableTextBox;
 
diff --git a/Engine/Source/Editor/IntroTutorials/Private/TutorialImageDecorator.cpp b/Engine/Source/Editor/IntroTutorials/Private/TutorialImageDecorator.cpp
new file mode 100644
--- /dev/null
+++ b/Engine/Source/Editor/IntroTutorials/Private/TutorialImageDecorator.cpp
@@ -0,0 +1,66 @@
+// Copyright 1998-2014 Epic Games, Inc. All Rights Reserved.
+
+#include "IntroTutorialsPrivatePCH.h"
+#include "TutorialImageDecorator.h"
+
+TSharedRef< FTutorialImageDecorator > FTutorialImageDecorator::Create()
+{
+	return MakeShareable( new FTutorialImageDecorator() );
+}
+
+FTutorialImageDecorator::FTutorialImageDecorator()
+{
+}
+
+bool FTutorialImageDecorator::Supports( const FTextRunParseResults& RunParseResult, const FString& Text ) const
+{
+	return ( RunParseResult.Name == TEXT("img") );
+}
+
+TSharedRef< ISlateRun > FTutorialImageDecorator::Create(const TSharedRef<class FTextLayout>& TextLayout, const FTextRunParseResults& RunParseResult, const FString& OriginalText, const TSharedRef< FString >& InOutModelText, const ISlateStyle* Style)
+{
+	const FTextRange* const BrushNameRange = RunParseResult.MetaData.Find( TEXT("src") );
+
+	FTextRange ModelRange;
+	ModelRange.BeginIndex = InOutModelText->Len();
+	*InOutModelText += TEXT('\x200B'); // Zero-Width Breaking Space
+	ModelRange.EndIndex = InOutModelText->Len();
+
+	FRunInfo RunInfo( RunParseResult.Name );
+	for(const TPair<FString, FTextRange>& Pair : RunParseResult.MetaData)
+	{
+		RunInfo.MetaData.Add(Pair.Key, OriginalText.Mid( Pair.Value.BeginIndex, Pair.Value.EndIndex - Pair.Value.BeginIndex));
+	}
+
+	if ( BrushNameRange != NULL )
+	{
+		const FString BrushNameString = OriginalText.Mid(BrushNameRange->BeginIndex, BrushNameRange->EndIndex - BrushNameRange->BeginIndex);
+
+		return FSlateImageRun::Create( RunInfo, InOutModelText, *GetPathToImage(BrushNameString), 0, ModelRange );
+	}
+
+	const FInlineTextImageStyle& ImageStyle = FInlineTextImageStyle::GetDefault();
+	return FSlateImageRun::Create( RunInfo, InOutModelText, &ImageStyle.Image, ImageStyle.Baseline, ModelRange );
+}
+
+FString FTutorialImageDecorator::GetPathToImage(const FString& InSrcMetaData)
+{
+	// images are content-relative if they are created from a project
+	FString FullPath;
+	if(FPackageName::IsValidLongPackageName(InSrcMetaData))
+	{
+		FullPath = FPackageName::LongPackageNameToFilename(InSrcMetaData);
+		if(FPaths::GetExtension(FullPath).Len() == 0)
+		{
+			FullPath += TEXT(".png");
+		}
+	}
+	else
+	{
+		// otherwise they are raw paths relative to the engine
+		FullPath = InSrcMetaData;
+		FPaths::MakePathRelativeTo(FullPath, FPlatformProcess::BaseDir());
+	}
+
+	return FullPath;
+}
\ No newline at end of file
diff --git a/Engine/Source/Editor/IntroTutorials/Private/TutorialImageDecorator.h b/Engine/Source/Editor/IntroTutorials/Private/TutorialImageDecorator.h
new file mode 100644
--- /dev/null
+++ b/Engine/Source/Editor/IntroTutorials/Private/TutorialImageDecorator.h
@@ -0,0 +1,24 @@
+// Copyright 1998-2014 Epic Games, Inc. All Rights Reserved.
+#pragma once
+
+#include "TextDecorators.h"
+
+class FTutorialImageDecorator : public ITextDecorator
+{
+public:
+
+	static TSharedRef< FTutorialImageDecorator > Create();
+	virtual ~FTutorialImageDecorator() {}
+
+public:
+
+	virtual bool Supports( const FTextRunParseResults& RunParseResult, const FString& Text ) const override;
+
+	virtual TSharedRef< ISlateRun > Create(const TSharedRef<class FTextLayout>& TextLayout, const FTextRunParseResults& RunParseResult, const FString& OriginalText, const TSharedRef< FString >& InOutModelText, const ISlateStyle* Style) override;
+
+	static FString GetPathToImage(const FString& InSrcMetaData);
+
+private:
+
+	FTutorialImageDecorator();
+};
\ No newline at end of file
diff --git a/Engine/Source/Editor/IntroTutorials/Private/TutorialText.cpp b/Engine/Source/Editor/IntroTutorials/Private/TutorialText.cpp
--- a/Engine/Source/Editor/IntroTutorials/Private/TutorialText.cpp
+++ b/Engine/Source/Editor/IntroTutorials/Private/TutorialText.cpp
@@ -9,6 +9,7 @@
 #include "Editor/MainFrame/Public/MainFrame.h"
 #include "EditorTutorial.h"
 #include "SourceCodeNavigation.h"
+#include "TutorialImageDecorator.h"
 
 #define LOCTEXT_NAMESPACE "TutorialText"
 
@@ -265,6 +266,7 @@ void FTutorialText::GetRichTextDecorators(TArray< TSharedRef< class ITextDecorat
 	}
 
 	OutDecorators.Add(FTextStyleDecorator::Create());
+	OutDecorators.Add(FTutorialImageDecorator::Create());
 }
 
 void FTutorialText::Initialize()
diff --git a/Engine/Source/Runtime/AppFramework/Private/Widgets/Testing/STestSuite.cpp b/Engine/Source/Runtime/AppFramework/Private/Widgets/Testing/STestSuite.cpp
--- a/Engine/Source/Runtime/AppFramework/Private/Widgets/Testing/STestSuite.cpp
+++ b/Engine/Source/Runtime/AppFramework/Private/Widgets/Testing/STestSuite.cpp
@@ -2381,8 +2381,7 @@ public:
 		// The rich-text marshaller will be used with the rich-text editor...
 		TSharedRef<FRichTextLayoutMarshaller> RichTextMarshaller = FRichTextLayoutMarshaller::Create(
 			TArray<TSharedRef<ITextDecorator>>(), 
-			&FTestStyle::Get(), 
-			FTestStyle::Get().GetWidgetStyle<FTextBlockStyle>("RichText.Editor.Text")
+			&FTestStyle::Get()
 			);
 
 		// ... so we also need to add some decorators to handle the things we want to demo
diff --git a/Engine/Source/Runtime/Slate/Private/Framework/Text/PlainTextLayoutMarshaller.cpp b/Engine/Source/Runtime/Slate/Private/Framework/Text/PlainTextLayoutMarshaller.cpp
--- a/Engine/Source/Runtime/Slate/Private/Framework/Text/PlainTextLayoutMarshaller.cpp
+++ b/Engine/Source/Runtime/Slate/Private/Framework/Text/PlainTextLayoutMarshaller.cpp
@@ -2,12 +2,13 @@
 
 #include "SlatePrivatePCH.h"
 #include "PlainTextLayoutMarshaller.h"
+#include "SlateTextLayout.h"
 
 #if WITH_FANCY_TEXT
 
-TSharedRef< FPlainTextLayoutMarshaller > FPlainTextLayoutMarshaller::Create(FTextBlockStyle InDefaultTextStyle)
+TSharedRef< FPlainTextLayoutMarshaller > FPlainTextLayoutMarshaller::Create()
 {
-	return MakeShareable(new FPlainTextLayoutMarshaller(MoveTemp(InDefaultTextStyle)));
+	return MakeShareable(new FPlainTextLayoutMarshaller());
 }
 
 FPlainTextLayoutMarshaller::~FPlainTextLayoutMarshaller()
@@ -16,6 +17,8 @@ FPlainTextLayoutMarshaller::~FPlainTextLayoutMarshaller()
 
 void FPlainTextLayoutMarshaller::SetText(const FString& SourceString, FTextLayout& TargetTextLayout)
 {
+	const FTextBlockStyle& DefaultTextStyle = static_cast<FSlateTextLayout&>(TargetTextLayout).GetDefaultTextStyle();
+
 	TArray<FTextRange> LineRanges;
 	FTextRange::CalculateLineRangesFromString(SourceString, LineRanges);
 
@@ -35,8 +38,7 @@ void FPlainTextLayoutMarshaller::GetText(FString& TargetString, const FTextLayou
 	SourceTextLayout.GetAsText(TargetString);
 }
 
-FPlainTextLayoutMarshaller::FPlainTextLayoutMarshaller(FTextBlockStyle InDefaultTextStyle)
-	: FSlateTextLayoutMarshaller(MoveTemp(InDefaultTextStyle))
+FPlainTextLayoutMarshaller::FPlainTextLayoutMarshaller()
 {
 }
 
diff --git a/Engine/Source/Runtime/Slate/Private/Framework/Text/RichTextLayoutMarshaller.cpp b/Engine/Source/Runtime/Slate/Private/Framework/Text/RichTextLayoutMarshaller.cpp
--- a/Engine/Source/Runtime/Slate/Private/Framework/Text/RichTextLayoutMarshaller.cpp
+++ b/Engine/Source/Runtime/Slate/Private/Framework/Text/RichTextLayoutMarshaller.cpp
@@ -3,17 +3,18 @@
 #include "SlatePrivatePCH.h"
 #include "RichTextLayoutMarshaller.h"
 #include "RichTextMarkupProcessing.h"
+#include "SlateTextLayout.h"
 
 #if WITH_FANCY_TEXT
 
-TSharedRef< FRichTextLayoutMarshaller > FRichTextLayoutMarshaller::Create(TArray< TSharedRef< ITextDecorator > > InDecorators, const ISlateStyle* const InDecoratorStyleSet, FTextBlockStyle InDefaultTextStyle)
+TSharedRef< FRichTextLayoutMarshaller > FRichTextLayoutMarshaller::Create(TArray< TSharedRef< ITextDecorator > > InDecorators, const ISlateStyle* const InDecoratorStyleSet)
 {
-	return MakeShareable(new FRichTextLayoutMarshaller(FDefaultRichTextMarkupParser::Create(), FDefaultRichTextMarkupWriter::Create(), MoveTemp(InDecorators), InDecoratorStyleSet, MoveTemp(InDefaultTextStyle)));
+	return MakeShareable(new FRichTextLayoutMarshaller(FDefaultRichTextMarkupParser::Create(), FDefaultRichTextMarkupWriter::Create(), MoveTemp(InDecorators), InDecoratorStyleSet));
 }
 
-TSharedRef< FRichTextLayoutMarshaller > FRichTextLayoutMarshaller::Create(TSharedPtr< IRichTextMarkupParser > InParser, TSharedPtr< IRichTextMarkupWriter > InWriter, TArray< TSharedRef< ITextDecorator > > InDecorators, const ISlateStyle* const InDecoratorStyleSet, FTextBlockStyle InDefaultTextStyle)
+TSharedRef< FRichTextLayoutMarshaller > FRichTextLayoutMarshaller::Create(TSharedPtr< IRichTextMarkupParser > InParser, TSharedPtr< IRichTextMarkupWriter > InWriter, TArray< TSharedRef< ITextDecorator > > InDecorators, const ISlateStyle* const InDecoratorStyleSet)
 {
-	return MakeShareable(new FRichTextLayoutMarshaller(MoveTemp(InParser), MoveTemp(InWriter), MoveTemp(InDecorators), InDecoratorStyleSet, MoveTemp(InDefaultTextStyle)));
+	return MakeShareable(new FRichTextLayoutMarshaller(MoveTemp(InParser), MoveTemp(InWriter), MoveTemp(InDecorators), InDecoratorStyleSet));
 }
 
 FRichTextLayoutMarshaller::~FRichTextLayoutMarshaller()
@@ -22,6 +23,8 @@ FRichTextLayoutMarshaller::~FRichTextLayoutMarshaller()
 
 void FRichTextLayoutMarshaller::SetText(const FString& SourceString, FTextLayout& TargetTextLayout)
 {
+	const FTextBlockStyle& DefaultTextStyle = static_cast<FSlateTextLayout&>(TargetTextLayout).GetDefaultTextStyle();
+
 	TArray<FTextLineParseResults> LineParseResultsArray;
 	FString ProcessedString;
 	Parser->Process(LineParseResultsArray, SourceString, ProcessedString);
@@ -102,9 +105,8 @@ void FRichTextLayoutMarshaller::GetText(FString& TargetString, const FTextLayout
 	Writer->Write(WriterLines, TargetString);
 }
 
-FRichTextLayoutMarshaller::FRichTextLayoutMarshaller(TSharedPtr< IRichTextMarkupParser > InParser, TSharedPtr< IRichTextMarkupWriter > InWriter, TArray< TSharedRef< ITextDecorator > > InDecorators, const ISlateStyle* const InDecoratorStyleSet, FTextBlockStyle InDefaultTextStyle)
-	: FSlateTextLayoutMarshaller(MoveTemp(InDefaultTextStyle))
-	, Parser(MoveTemp(InParser))
+FRichTextLayoutMarshaller::FRichTextLayoutMarshaller(TSharedPtr< IRichTextMarkupParser > InParser, TSharedPtr< IRichTextMarkupWriter > InWriter, TArray< TSharedRef< ITextDecorator > > InDecorators, const ISlateStyle* const InDecoratorStyleSet)
+	: Parser(MoveTemp(InParser))
 	, Writer(MoveTemp(InWriter))
 	, Decorators(MoveTemp(InDecorators))
 	, InlineDecorators()
diff --git a/Engine/Source/Runtime/Slate/Private/Framework/Text/SlateHyperlinkRun.cpp b/Engine/Source/Runtime/Slate/Private/Framework/Text/SlateHyperlinkRun.cpp
--- a/Engine/Source/Runtime/Slate/Private/Framework/Text/SlateHyperlinkRun.cpp
+++ b/Engine/Source/Runtime/Slate/Private/Framework/Text/SlateHyperlinkRun.cpp
@@ -184,6 +184,11 @@ const FRunInfo& FSlateHyperlinkRun::GetRunInfo() const
 	return RunInfo;
 }
 
+ERunAttributes FSlateHyperlinkRun::GetRunAttributes() const
+{
+	return ERunAttributes::SupportsText;
+}
+
 FSlateHyperlinkRun::FSlateHyperlinkRun( const FRunInfo& InRunInfo, const TSharedRef< const FString >& InText, const FHyperlinkStyle& InStyle, FOnClick InNavigateDelegate, FOnGenerateTooltip InTooltipDelegate, FOnGetTooltipText InTooltipTextDelegate ) 
 	: RunInfo( InRunInfo )
 	, Text( InText )
diff --git a/Engine/Source/Runtime/Slate/Private/Framework/Text/SlateImageRun.cpp b/Engine/Source/Runtime/Slate/Private/Framework/Text/SlateImageRun.cpp
--- a/Engine/Source/Runtime/Slate/Private/Framework/Text/SlateImageRun.cpp
+++ b/Engine/Source/Runtime/Slate/Private/Framework/Text/SlateImageRun.cpp
@@ -26,6 +26,16 @@ TSharedRef< FSlateImageRun > FSlateImageRun::Create( const FRunInfo& InRunInfo,
 	return MakeShareable( new FSlateImageRun( InRunInfo, InText, InImage, InBaseline, InRange ) );
 }
 
+TSharedRef< FSlateImageRun > FSlateImageRun::Create( const FRunInfo& InRunInfo, const TSharedRef< const FString >& InText, FName InDynamicBrushName, int16 InBaseline )
+{
+	return MakeShareable( new FSlateImageRun( InRunInfo, InText, InDynamicBrushName, InBaseline ) );
+}
+
+TSharedRef< FSlateImageRun > FSlateImageRun::Create( const FRunInfo& InRunInfo, const TSharedRef< const FString >& InText, FName InDynamicBrushName, int16 InBaseline, const FTextRange& InRange )
+{
+	return MakeShareable( new FSlateImageRun( InRunInfo, InText, InDynamicBrushName, InBaseline, InRange ) );
+}
+
 FSlateImageRun::FSlateImageRun( const FRunInfo& InRunInfo, const TSharedRef< const FString >& InText, const FSlateBrush* InImage, int16 InBaseline, const FTextRange& InRange ) 
 	: RunInfo( InRunInfo )
 	, Text( InText )
@@ -46,6 +56,30 @@ FSlateImageRun::FSlateImageRun( const FRunInfo& InRunInfo, const TSharedRef< con
 	check( Image != nullptr);
 }
 
+FSlateImageRun::FSlateImageRun( const FRunInfo& InRunInfo, const TSharedRef< const FString >& InText, FName InDynamicBrushName, int16 InBaseline ) 
+	: RunInfo( InRunInfo )
+	, Text( InText )
+	, Range( 0, Text->Len() )
+	, Image( nullptr )
+	, Baseline( InBaseline )
+{
+	FIntPoint Size = FSlateApplication::Get().GetRenderer()->GenerateDynamicImageResource(InDynamicBrushName);
+	DynamicBrush = MakeShareable(new FSlateDynamicImageBrush( InDynamicBrushName, FVector2D(Size.X, Size.Y) ) );
+	Image = DynamicBrush.Get();
+}
+
+FSlateImageRun::FSlateImageRun( const FRunInfo& InRunInfo, const TSharedRef< const FString >& InText, FName InDynamicBrushName, int16 InBaseline, const FTextRange& InRange ) 
+	: RunInfo( InRunInfo )
+	, Text( InText )
+	, Range( InRange )
+	, Image( nullptr )
+	, Baseline( InBaseline )
+{
+	FIntPoint Size = FSlateApplication::Get().GetRenderer()->GenerateDynamicImageResource(InDynamicBrushName);
+	DynamicBrush = MakeShareable(new FSlateDynamicImageBrush( InDynamicBrushName, FVector2D(Size.X, Size.Y) ) );
+	Image = DynamicBrush.Get();
+}
+
 const TArray< TSharedRef<SWidget> >& FSlateImageRun::GetChildren() 
 {
 	static TArray< TSharedRef<SWidget> > NoChildren;
@@ -59,7 +93,34 @@ void FSlateImageRun::ArrangeChildren( const TSharedRef< ILayoutBlock >& Block, c
 
 int32 FSlateImageRun::GetTextIndexAt( const TSharedRef< ILayoutBlock >& Block, const FVector2D& Location, float Scale, ETextHitPoint* const OutHitPoint ) const
 {
-	return INDEX_NONE;
+	// An image should always contain a single character (a breaking space)
+	check(Range.Len() == 1);
+
+	const FVector2D& BlockOffset = Block->GetLocationOffset();
+	const FVector2D& BlockSize = Block->GetSize();
+
+	const float Left = BlockOffset.X;
+	const float Top = BlockOffset.Y;
+	const float Right = BlockOffset.X + BlockSize.X;
+	const float Bottom = BlockOffset.Y + BlockSize.Y;
+
+	const bool ContainsPoint = Location.X >= Left && Location.X < Right && Location.Y >= Top && Location.Y < Bottom;
+
+	if ( !ContainsPoint )
+	{
+		return INDEX_NONE;
+	}
+
+	const FVector2D ScaledImageSize = Image->ImageSize * Scale;
+	const int32 Index = (Location.X <= (Left + (ScaledImageSize.X * 0.5f))) ? Range.BeginIndex : Range.EndIndex;
+	
+	if (OutHitPoint)
+	{
+		const FTextRange BlockRange = Block->GetTextRange();
+		*OutHitPoint = (Index == BlockRange.EndIndex) ? ETextHitPoint::RightGutter : ETextHitPoint::WithinText;
+	}
+
+	return Index;
 }
 
 FVector2D FSlateImageRun::GetLocationAt( const TSharedRef< ILayoutBlock >& Block, int32 Offset, float Scale ) const
@@ -102,7 +163,11 @@ int8 FSlateImageRun::GetKerning( int32 CurrentIndex, float Scale ) const
 
 FVector2D FSlateImageRun::Measure( int32 BeginIndex, int32 EndIndex, float Scale ) const 
 {
-	check( BeginIndex == Range.BeginIndex && EndIndex == Range.EndIndex );
+	if ( EndIndex - BeginIndex == 0 )
+	{
+		return FVector2D( 0, GetMaxHeight( Scale ) );
+	}
+
 	return Image->ImageSize * Scale;
 }
 
@@ -134,18 +199,25 @@ void FSlateImageRun::Move(const TSharedRef<FString>& NewText, const FTextRange&
 
 TSharedRef<IRun> FSlateImageRun::Clone() const
 {
-	return FSlateImageRun::Create(RunInfo, Text, Image, Baseline, Range);
+	TSharedRef<FSlateImageRun> NewRun = FSlateImageRun::Create(RunInfo, Text, Image, Baseline, Range);
+
+	// make sure we copy the dynamic brush so it doesnt get released
+	NewRun->DynamicBrush = DynamicBrush;
+
+	return NewRun;
 }
 
-void FSlateImageRun::AppendTextTo(FString& Text) const
+void FSlateImageRun::AppendTextTo(FString& AppendToText) const
 {
-	// Do nothing
+	AppendToText.Append(**Text + Range.BeginIndex, Range.Len());
 }
 
 void FSlateImageRun::AppendTextTo(FString& AppendToText, const FTextRange& PartialRange) const
 {
 	check(Range.BeginIndex <= PartialRange.BeginIndex);
 	check(Range.EndIndex >= PartialRange.EndIndex);
+
+	AppendToText.Append(**Text + PartialRange.BeginIndex, PartialRange.Len());
 }
 
 const FRunInfo& FSlateImageRun::GetRunInfo() const
@@ -153,4 +225,9 @@ const FRunInfo& FSlateImageRun::GetRunInfo() const
 	return RunInfo;
 }
 
+ERunAttributes FSlateImageRun::GetRunAttributes() const
+{
+	return ERunAttributes::None;
+}
+
 #endif //WITH_FANCY_TEXT
diff --git a/Engine/Source/Runtime/Slate/Private/Framework/Text/SlateTextLayout.cpp b/Engine/Source/Runtime/Slate/Private/Framework/Text/SlateTextLayout.cpp
--- a/Engine/Source/Runtime/Slate/Private/Framework/Text/SlateTextLayout.cpp
+++ b/Engine/Source/Runtime/Slate/Private/Framework/Text/SlateTextLayout.cpp
@@ -6,16 +6,17 @@
 
 #include "SlateTextLayout.h"
 
-TSharedRef< FSlateTextLayout > FSlateTextLayout::Create()
+TSharedRef< FSlateTextLayout > FSlateTextLayout::Create(FTextBlockStyle InDefaultTextStyle)
 {
-	TSharedRef< FSlateTextLayout > Layout = MakeShareable( new FSlateTextLayout() );
+	TSharedRef< FSlateTextLayout > Layout = MakeShareable( new FSlateTextLayout(MoveTemp(InDefaultTextStyle)) );
 	Layout->AggregateChildren();
 
 	return Layout;
 }
 
-FSlateTextLayout::FSlateTextLayout()
+FSlateTextLayout::FSlateTextLayout(FTextBlockStyle InDefaultTextStyle)
 	: Children()
+	, DefaultTextStyle(MoveTemp(InDefaultTextStyle))
 {
 
 }
@@ -40,7 +41,7 @@ void FSlateTextLayout::ArrangeChildren( const FGeometry& AllottedGeometry, FArra
 	}
 }
 
-int32 FSlateTextLayout::OnPaint( const FPaintArgs& Args, const FTextBlockStyle& DefaultTextStyle, const FGeometry& AllottedGeometry, const FSlateRect& MyClippingRect, FSlateWindowElementList& OutDrawElements, int32 LayerId, const FWidgetStyle& InWidgetStyle, bool bParentEnabled ) const
+int32 FSlateTextLayout::OnPaint( const FPaintArgs& Args, const FGeometry& AllottedGeometry, const FSlateRect& MyClippingRect, FSlateWindowElementList& OutDrawElements, int32 LayerId, const FWidgetStyle& InWidgetStyle, bool bParentEnabled ) const
 {
 	const FSlateRect ClippingRect = AllottedGeometry.GetClippingRect().IntersectionWith(MyClippingRect);
 	const ESlateDrawEffect::Type DrawEffects = bParentEnabled ? ESlateDrawEffect::None : ESlateDrawEffect::DisabledEffect;
@@ -133,6 +134,16 @@ void FSlateTextLayout::EndLayout()
 	AggregateChildren();
 }
 
+void FSlateTextLayout::SetDefaultTextStyle(FTextBlockStyle InDefaultTextStyle)
+{
+	DefaultTextStyle = MoveTemp(InDefaultTextStyle);
+}
+
+const FTextBlockStyle& FSlateTextLayout::GetDefaultTextStyle() const
+{
+	return DefaultTextStyle;
+}
+
 void FSlateTextLayout::AggregateChildren()
 {
 	Children.Empty();
@@ -155,5 +166,9 @@ void FSlateTextLayout::AggregateChildren()
 	}
 }
 
+TSharedRef<IRun> FSlateTextLayout::CreateDefaultTextRun(const TSharedRef<FString>& NewText, const FTextRange& NewRange) const
+{
+	return FSlateTextRun::Create(FRunInfo(), NewText, DefaultTextStyle, NewRange);
+}
 
-#endif //WITH_FANCY_TEXT
\ No newline at end of file
+#endif //WITH_FANCY_TEXT
diff --git a/Engine/Source/Runtime/Slate/Private/Framework/Text/SlateTextLayoutMarshaller.cpp b/Engine/Source/Runtime/Slate/Private/Framework/Text/SlateTextLayoutMarshaller.cpp
deleted file mode 100644
diff --git a/Engine/Source/Runtime/Slate/Private/Framework/Text/SlateTextRun.cpp b/Engine/Source/Runtime/Slate/Private/Framework/Text/SlateTextRun.cpp
--- a/Engine/Source/Runtime/Slate/Private/Framework/Text/SlateTextRun.cpp
+++ b/Engine/Source/Runtime/Slate/Private/Framework/Text/SlateTextRun.cpp
@@ -201,6 +201,11 @@ const FRunInfo& FSlateTextRun::GetRunInfo() const
 	return RunInfo;
 }
 
+ERunAttributes FSlateTextRun::GetRunAttributes() const
+{
+	return ERunAttributes::SupportsText;
+}
+
 FSlateTextRun::FSlateTextRun( const FRunInfo& InRunInfo, const TSharedRef< const FString >& InText, const FTextBlockStyle& InStyle ) 
 	: RunInfo( InRunInfo )
 	, Text( InText )
diff --git a/Engine/Source/Runtime/Slate/Private/Framework/Text/SlateWidgetRun.cpp b/Engine/Source/Runtime/Slate/Private/Framework/Text/SlateWidgetRun.cpp
--- a/Engine/Source/Runtime/Slate/Private/Framework/Text/SlateWidgetRun.cpp
+++ b/Engine/Source/Runtime/Slate/Private/Framework/Text/SlateWidgetRun.cpp
@@ -38,6 +38,11 @@ int16 FSlateWidgetRun::GetMaxHeight( float Scale ) const
 
 FVector2D FSlateWidgetRun::Measure( int32 StartIndex, int32 EndIndex, float Scale ) const 
 {
+	if ( EndIndex - StartIndex == 0 )
+	{
+		return FVector2D( 0, GetMaxHeight( Scale ) );
+	}
+
 	return Info.Size.Get( Info.Widget->GetDesiredSize() ) * Scale;
 }
 
@@ -89,7 +94,34 @@ void FSlateWidgetRun::ArrangeChildren( const TSharedRef< ILayoutBlock >& Block,
 
 int32 FSlateWidgetRun::GetTextIndexAt( const TSharedRef< ILayoutBlock >& Block, const FVector2D& Location, float Scale, ETextHitPoint* const OutHitPoint ) const
 {
-	return INDEX_NONE;
+	// A widget should always contain a single character (a breaking space)
+	check(Range.Len() == 1);
+
+	const FVector2D& BlockOffset = Block->GetLocationOffset();
+	const FVector2D& BlockSize = Block->GetSize();
+
+	const float Left = BlockOffset.X;
+	const float Top = BlockOffset.Y;
+	const float Right = BlockOffset.X + BlockSize.X;
+	const float Bottom = BlockOffset.Y + BlockSize.Y;
+
+	const bool ContainsPoint = Location.X >= Left && Location.X < Right && Location.Y >= Top && Location.Y < Bottom;
+
+	if ( !ContainsPoint )
+	{
+		return INDEX_NONE;
+	}
+
+	const FVector2D ScaledWidgetSize = Info.Widget->GetDesiredSize() * Scale;
+	const int32 Index = (Location.X <= (Left + (ScaledWidgetSize.X * 0.5f))) ? Range.BeginIndex : Range.EndIndex;
+	
+	if (OutHitPoint)
+	{
+		const FTextRange BlockRange = Block->GetTextRange();
+		*OutHitPoint = (Index == BlockRange.EndIndex) ? ETextHitPoint::RightGutter : ETextHitPoint::WithinText;
+	}
+
+	return Index;
 }
 
 FVector2D FSlateWidgetRun::GetLocationAt( const TSharedRef< ILayoutBlock >& Block, int32 Offset, float Scale ) const
@@ -110,13 +142,15 @@ TSharedRef<IRun> FSlateWidgetRun::Clone() const
 
 void FSlateWidgetRun::AppendTextTo(FString& AppendToText) const
 {
-	//Do nothing
+	AppendToText.Append(**Text + Range.BeginIndex, Range.Len());
 }
 
 void FSlateWidgetRun::AppendTextTo(FString& AppendToText, const FTextRange& PartialRange) const
 {
 	check(Range.BeginIndex <= PartialRange.BeginIndex);
 	check(Range.EndIndex >= PartialRange.EndIndex);
+
+	AppendToText.Append(**Text + PartialRange.BeginIndex, PartialRange.Len());
 }
 
 const FRunInfo& FSlateWidgetRun::GetRunInfo() const
@@ -124,6 +158,11 @@ const FRunInfo& FSlateWidgetRun::GetRunInfo() const
 	return RunInfo;
 }
 
+ERunAttributes FSlateWidgetRun::GetRunAttributes() const
+{
+	return ERunAttributes::None;
+}
+
 FSlateWidgetRun::FSlateWidgetRun(const TSharedRef<class FTextLayout>& TextLayout, const FRunInfo& InRunInfo, const TSharedRef< const FString >& InText, const FWidgetRunInfo& InWidgetInfo)
 	: TextLayout(TextLayout)
 	, RunInfo( InRunInfo )
diff --git a/Engine/Source/Runtime/Slate/Private/Framework/Text/SyntaxHighlighterTextLayoutMarshaller.cpp b/Engine/Source/Runtime/Slate/Private/Framework/Text/SyntaxHighlighterTextLayoutMarshaller.cpp
--- a/Engine/Source/Runtime/Slate/Private/Framework/Text/SyntaxHighlighterTextLayoutMarshaller.cpp
+++ b/Engine/Source/Runtime/Slate/Private/Framework/Text/SyntaxHighlighterTextLayoutMarshaller.cpp
@@ -39,9 +39,8 @@ bool FSyntaxHighlighterTextLayoutMarshaller::IsSyntaxHighlightingEnabled() const
 	return bSyntaxHighlightingEnabled;
 }
 
-FSyntaxHighlighterTextLayoutMarshaller::FSyntaxHighlighterTextLayoutMarshaller(TSharedPtr< FSyntaxTokenizer > InTokenizer, FTextBlockStyle InDefaultTextStyle)
-	: FPlainTextLayoutMarshaller(MoveTemp(InDefaultTextStyle))
-	, Tokenizer(MoveTemp(InTokenizer))
+FSyntaxHighlighterTextLayoutMarshaller::FSyntaxHighlighterTextLayoutMarshaller(TSharedPtr< FSyntaxTokenizer > InTokenizer)
+	: Tokenizer(MoveTemp(InTokenizer))
 	, bSyntaxHighlightingEnabled(true)
 {
 }
@@ -174,7 +173,7 @@ void FRichTextSyntaxHighlighterTextLayoutMarshaller::ParseTokens(const FString&
 }
 
 FRichTextSyntaxHighlighterTextLayoutMarshaller::FRichTextSyntaxHighlighterTextLayoutMarshaller(TSharedPtr< FSyntaxTokenizer > InTokenizer, const FSyntaxTextStyle& InSyntaxTextStyle)
-	: FSyntaxHighlighterTextLayoutMarshaller(MoveTemp(InTokenizer), InSyntaxTextStyle.NormalTextStyle)
+	: FSyntaxHighlighterTextLayoutMarshaller(MoveTemp(InTokenizer))
 	, SyntaxTextStyle(InSyntaxTextStyle)
 {
 }
diff --git a/Engine/Source/Runtime/Slate/Private/Framework/Text/TextLayout.cpp b/Engine/Source/Runtime/Slate/Private/Framework/Text/TextLayout.cpp
--- a/Engine/Source/Runtime/Slate/Private/Framework/Text/TextLayout.cpp
+++ b/Engine/Source/Runtime/Slate/Private/Framework/Text/TextLayout.cpp
@@ -901,7 +901,31 @@ bool FTextLayout::InsertAt(const FTextLocation& Location, TCHAR Character)
 			check(RunIsAfterInsertLocation == false);
 			RunIsAfterInsertLocation = true;
 
-			RunModel.SetTextRange( FTextRange( RunRange.BeginIndex, RunRange.EndIndex + 1 ) );
+			if ((RunModel.GetRun()->GetRunAttributes() & ERunAttributes::SupportsText) != ERunAttributes::None)
+			{
+				RunModel.SetTextRange( FTextRange( RunRange.BeginIndex, RunRange.EndIndex + 1 ) );
+			}
+			else
+			{
+				// Non-text runs are supposed to have a single dummy character in them
+				check(RunRange.Len() == 1);
+
+				// This run doesn't support text, so we need to insert a new text run before or after the current run depending on the insertion point
+				const bool bIsInsertingToTheLeft = InsertLocation == RunRange.BeginIndex;
+				if (bIsInsertingToTheLeft)
+				{
+					// Insert the new text run to the left of the non-text run
+					TSharedRef<IRun> NewTextRun = CreateDefaultTextRun( LineModel.Text, FTextRange( RunRange.BeginIndex, RunRange.BeginIndex + 1 ) );
+					RunModel.SetTextRange( FTextRange( RunRange.BeginIndex + 1, RunRange.EndIndex + 1 ) );
+					LineModel.Runs.Insert( NewTextRun, RunIndex++ );
+				}
+				else
+				{
+					// Insert the new text run to the right of the non-text run
+					TSharedRef<IRun> NewTextRun = CreateDefaultTextRun( LineModel.Text, FTextRange( RunRange.EndIndex, RunRange.EndIndex + 1 ) );
+					LineModel.Runs.Insert( NewTextRun, ++RunIndex );
+				}
+			}
 		}
 		else if (RunIsAfterInsertLocation)
 		{
@@ -942,7 +966,31 @@ bool FTextLayout::InsertAt(const FTextLocation& Location, const FString& Text)
 			check(RunIsAfterInsertLocation == false);
 			RunIsAfterInsertLocation = true;
 
-			RunModel.SetTextRange( FTextRange( RunRange.BeginIndex, RunRange.EndIndex + Text.Len() ) );
+			if ((RunModel.GetRun()->GetRunAttributes() & ERunAttributes::SupportsText) != ERunAttributes::None)
+			{
+				RunModel.SetTextRange( FTextRange( RunRange.BeginIndex, RunRange.EndIndex + Text.Len() ) );
+			}
+			else
+			{
+				// Non-text runs are supposed to have a single dummy character in them
+				check(RunRange.Len() == 1);
+
+				// This run doesn't support text, so we need to insert a new text run before or after the current run depending on the insertion point
+				const bool bIsInsertingToTheLeft = InsertLocation == RunRange.BeginIndex;
+				if (bIsInsertingToTheLeft)
+				{
+					// Insert the new text run to the left of the non-text run
+					TSharedRef<IRun> NewTextRun = CreateDefaultTextRun( LineModel.Text, FTextRange( RunRange.BeginIndex, RunRange.BeginIndex + Text.Len() ) );
+					RunModel.SetTextRange( FTextRange( RunRange.BeginIndex + 1, RunRange.EndIndex + Text.Len() ) );
+					LineModel.Runs.Insert( NewTextRun, RunIndex++ );
+				}
+				else
+				{
+					// Insert the new text run to the right of the non-text run
+					TSharedRef<IRun> NewTextRun = CreateDefaultTextRun( LineModel.Text, FTextRange( RunRange.EndIndex, RunRange.EndIndex + Text.Len() ) );
+					LineModel.Runs.Insert( NewTextRun, ++RunIndex );
+				}
+			}
 		}
 		else if (RunIsAfterInsertLocation)
 		{
@@ -986,27 +1034,59 @@ bool FTextLayout::InsertAt(const FTextLocation& Location, TSharedRef<IRun> InRun
 			check(RunIsAfterInsertLocation == false);
 			RunIsAfterInsertLocation = true;
 
+			const int32 InsertLocationEnd = InsertLocation + NewRunText.Len();
+
 			// This run contains the insertion point, so we need to split it
-			const TSharedRef<IRun> LeftRun = Run;
-			const TSharedRef<IRun> RightRun = Run->Clone();
+			TSharedPtr<IRun> LeftRun;
+			TSharedPtr<IRun> RightRun;
+			if ((Run->GetRunAttributes() & ERunAttributes::SupportsText) != ERunAttributes::None)
+			{
+				LeftRun = Run->Clone();
+				LeftRun->SetTextRange(FTextRange(RunRange.BeginIndex, InsertLocation));
 
-			// Update the text ranges for all of the runs
-			const int32 InsertLocationEnd = InsertLocation + NewRunText.Len();
-			LeftRun->SetTextRange( FTextRange( RunRange.BeginIndex, InsertLocation ) );
-			InRun->Move( LineModel.Text, FTextRange( InsertLocation, InsertLocationEnd ) );
-			RightRun->SetTextRange( FTextRange( InsertLocationEnd, RunRange.EndIndex + NewRunText.Len() ) );
+				RightRun = Run;
+				RightRun->SetTextRange(FTextRange(InsertLocationEnd, RunRange.EndIndex + NewRunText.Len()));
+			}
+			else
+			{
+				// Non-text runs are supposed to have a single dummy character in them
+				check(RunRange.Len() == 1);
+
+				// This run doesn't support text, so we need to insert a new text run before or after the current run depending on the insertion point
+				const bool bIsInsertingToTheLeft = InsertLocation == RunRange.BeginIndex;
+				if (bIsInsertingToTheLeft)
+				{
+					// Insert the new text run to the left of the non-text run
+					LeftRun = CreateDefaultTextRun(LineModel.Text, FTextRange(RunRange.BeginIndex, InsertLocation));
+				
+					RightRun = Run;
+					RightRun->SetTextRange(FTextRange(InsertLocationEnd, RunRange.EndIndex + NewRunText.Len()));
+				}
+				else
+				{
+					// Insert the new text run to the right of the non-text run
+					LeftRun = Run;
+
+					RightRun = CreateDefaultTextRun(LineModel.Text, FTextRange(InsertLocationEnd, RunRange.EndIndex + NewRunText.Len()));
+				}
+			}
+
+			InRun->Move(LineModel.Text, FTextRange(InsertLocation, InsertLocationEnd));
+
+			// Remove the old run (it may get re-added again as the right hand run)
+			LineModel.Runs.RemoveAt(RunIndex--);
 
 			// Insert the new runs at the correct place, and then skip over these new array entries
 			const bool LeftRunHasText = !LeftRun->GetTextRange().IsEmpty();
 			const bool RightRunHasText = !RightRun->GetTextRange().IsEmpty();
-			if (!LeftRunHasText)
+			if (LeftRunHasText)
 			{
-				LineModel.Runs.RemoveAt(RunIndex--);
+				LineModel.Runs.Insert(LeftRun.ToSharedRef(), ++RunIndex);
 			}
 			LineModel.Runs.Insert(InRun, ++RunIndex);
 			if (RightRunHasText || bAlwaysKeepRightRun)
 			{
-				LineModel.Runs.Insert(RightRun, ++RunIndex);
+				LineModel.Runs.Insert(RightRun.ToSharedRef(), ++RunIndex);
 			}
 		}
 		else if (RunIsAfterInsertLocation)
@@ -1097,14 +1177,43 @@ bool FTextLayout::SplitLineAt(const FTextLocation& Location)
 			check(RunIsToTheLeftOfTheBreakLocation == true);
 			RunIsToTheLeftOfTheBreakLocation = false;
 
-			const TSharedRef<IRun> LeftRun = Run;
-			const TSharedRef<IRun> RightRun = Run->Clone();
-			
-			LeftRun->Move(LeftLineModel.Text, FTextRange(RunRange.BeginIndex, LeftLineModel.Text->Len()));
-			RightRun->Move(RightLineModel.Text, FTextRange(0, RunRange.EndIndex - LeftLineModel.Text->Len()));
+			TSharedPtr<IRun> LeftRun;
+			TSharedPtr<IRun> RightRun;
+			if ((Run->GetRunAttributes() & ERunAttributes::SupportsText) != ERunAttributes::None)
+			{
+				LeftRun = Run->Clone();
+				LeftRun->Move(LeftLineModel.Text, FTextRange(RunRange.BeginIndex, LeftLineModel.Text->Len()));
+				
+				RightRun = Run;
+				RightRun->Move(RightLineModel.Text, FTextRange(0, RunRange.EndIndex - LeftLineModel.Text->Len()));
+			}
+			else
+			{
+				// Non-text runs are supposed to have a single dummy character in them
+				check(RunRange.Len() == 1);
 
-			LeftLineModel.Runs.Add(FRunModel(LeftRun));
-			RightLineModel.Runs.Add(FRunModel(RightRun));
+				// This run doesn't support text, so we need to insert a new text run before or after the current run depending on the insertion point
+				const bool bIsInsertingToTheLeft = BreakLocation == RunRange.BeginIndex;
+				if (bIsInsertingToTheLeft)
+				{
+					// Insert the new text run to the left of the non-text run
+					LeftRun = CreateDefaultTextRun(LeftLineModel.Text, FTextRange(RunRange.BeginIndex, LeftLineModel.Text->Len()));
+				
+					RightRun = Run;
+					RightRun->Move(RightLineModel.Text, FTextRange(0, RunRange.EndIndex - LeftLineModel.Text->Len()));
+				}
+				else
+				{
+					// Insert the new text run to the right of the non-text run
+					LeftRun = Run;
+					LeftRun->Move(LeftLineModel.Text, FTextRange(RunRange.BeginIndex, LeftLineModel.Text->Len()));
+
+					RightRun = CreateDefaultTextRun(RightLineModel.Text, FTextRange(0, RunRange.EndIndex - LeftLineModel.Text->Len()));
+				}
+			}
+
+			LeftLineModel.Runs.Add(FRunModel(LeftRun.ToSharedRef()));
+			RightLineModel.Runs.Add(FRunModel(RightRun.ToSharedRef()));
 		}
 		else if (RunIsToTheLeftOfTheBreakLocation)
 		{
@@ -1161,7 +1270,7 @@ bool FTextLayout::RemoveAt( const FTextLocation& Location, int32 Count )
 		const FTextRange RunRange = RunModel.GetTextRange();
 
 		const FTextRange IntersectedRangeToRemove = RunRange.Intersect(RemoveTextRange);
-		if (IntersectedRangeToRemove.IsEmpty() && RunRange.EndIndex > RemoveTextRange.EndIndex)
+		if (IntersectedRangeToRemove.IsEmpty() && RunRange.BeginIndex >= RemoveTextRange.EndIndex)
 		{
 			// The whole run is contained to the right of the removal range, just adjust its range by the amount of text that was removed
 			FTextRange NewRange = RunRange;
@@ -1880,4 +1989,4 @@ int32 FTextLayout::FTextOffsetLocations::GetTextLength() const
 		return OffsetEntry.FlatStringIndex + OffsetEntry.DocumentLineLength;
 	}
 	return 0;
-}
+}
\ No newline at end of file
diff --git a/Engine/Source/Runtime/Slate/Private/Widgets/Text/SMultiLineEditableText.cpp b/Engine/Source/Runtime/Slate/Private/Widgets/Text/SMultiLineEditableText.cpp
--- a/Engine/Source/Runtime/Slate/Private/Widgets/Text/SMultiLineEditableText.cpp
+++ b/Engine/Source/Runtime/Slate/Private/Widgets/Text/SMultiLineEditableText.cpp
@@ -265,13 +265,13 @@ void SMultiLineEditableText::Construct( const FArguments& InArgs )
 	Marshaller = InArgs._Marshaller;
 	if (!Marshaller.IsValid())
 	{
-		Marshaller = FPlainTextLayoutMarshaller::Create(TextStyle);
+		Marshaller = FPlainTextLayoutMarshaller::Create();
 	}
 
 	CursorLineHighlighter = FCursorLineHighlighter::Create(&CursorInfo);
 	TextCompositionHighlighter = FTextCompositionHighlighter::Create();
 	TextSelectionRunRenderer = FTextSelectionRunRenderer::Create();
-	TextLayout = FSlateTextLayout::Create();
+	TextLayout = FSlateTextLayout::Create(TextStyle);
 
 	BoundText = InArgs._Text;
 
@@ -2127,7 +2127,7 @@ int32 SMultiLineEditableText::OnPaint( const FPaintArgs& Args, const FGeometry&
 	// Note: This is done here rather than in Tick(), because Tick() doesn't get called while resizing windows, but OnPaint() does
 	CachedSize = AllottedGeometry.Size;
 
-	LayerId = TextLayout->OnPaint( Args.WithNewParent(this), TextStyle, AllottedGeometry, MyClippingRect, OutDrawElements, LayerId, InWidgetStyle, ShouldBeEnabled( bParentEnabled ) );
+	LayerId = TextLayout->OnPaint( Args.WithNewParent(this), AllottedGeometry, MyClippingRect, OutDrawElements, LayerId, InWidgetStyle, ShouldBeEnabled( bParentEnabled ) );
 
 	return LayerId;
 }
diff --git a/Engine/Source/Runtime/Slate/Private/Widgets/Text/SRichTextBlock.cpp b/Engine/Source/Runtime/Slate/Private/Widgets/Text/SRichTextBlock.cpp
--- a/Engine/Source/Runtime/Slate/Private/Widgets/Text/SRichTextBlock.cpp
+++ b/Engine/Source/Runtime/Slate/Private/Widgets/Text/SRichTextBlock.cpp
@@ -31,13 +31,13 @@ void SRichTextBlock::Construct( const FArguments& InArgs )
 			Parser = FDefaultRichTextMarkupParser::Create();
 		}
 
-		TSharedRef<FRichTextLayoutMarshaller> Marshaller = FRichTextLayoutMarshaller::Create(Parser, nullptr, InArgs._Decorators, InArgs._DecoratorStyleSet, TextStyle);
+		TSharedRef<FRichTextLayoutMarshaller> Marshaller = FRichTextLayoutMarshaller::Create(Parser, nullptr, InArgs._Decorators, InArgs._DecoratorStyleSet);
 		for ( const TSharedRef< ITextDecorator >& Decorator : InArgs.InlineDecorators )
 		{
 			Marshaller->AppendInlineDecorator( Decorator );
 		}
 
-		TextLayoutCache = FTextBlockLayout::Create(Marshaller, nullptr);
+		TextLayoutCache = FTextBlockLayout::Create(TextStyle, Marshaller, nullptr);
 	}
 }
 
diff --git a/Engine/Source/Runtime/Slate/Private/Widgets/Text/STextBlock.cpp b/Engine/Source/Runtime/Slate/Private/Widgets/Text/STextBlock.cpp
--- a/Engine/Source/Runtime/Slate/Private/Widgets/Text/STextBlock.cpp
+++ b/Engine/Source/Runtime/Slate/Private/Widgets/Text/STextBlock.cpp
@@ -32,7 +32,7 @@ void STextBlock::Construct( const FArguments& InArgs )
 #if WITH_FANCY_TEXT
 
 	// We use a dummy style here (as it may not be safe to call the delegates used to compute the style), but the correct style is set by ComputeDesiredSize
-	TextLayoutCache = FTextBlockLayout::Create(FPlainTextLayoutMarshaller::Create(FTextBlockStyle::GetDefault()), InArgs._LineBreakPolicy);
+	TextLayoutCache = FTextBlockLayout::Create(FTextBlockStyle::GetDefault(), FPlainTextLayoutMarshaller::Create(), InArgs._LineBreakPolicy);
 
 #endif//WITH_FANCY_TEXT
 }
diff --git a/Engine/Source/Runtime/Slate/Private/Widgets/Text/TextBlockLayout.cpp b/Engine/Source/Runtime/Slate/Private/Widgets/Text/TextBlockLayout.cpp
--- a/Engine/Source/Runtime/Slate/Private/Widgets/Text/TextBlockLayout.cpp
+++ b/Engine/Source/Runtime/Slate/Private/Widgets/Text/TextBlockLayout.cpp
@@ -3,18 +3,18 @@
 #include "SlatePrivatePCH.h"
 #include "TextBlockLayout.h"
 #include "SlateTextLayout.h"
-#include "SlateTextLayoutMarshaller.h"
+#include "BaseTextLayoutMarshaller.h"
 #include "SlateTextHighlightRunRenderer.h"
 
 #if WITH_FANCY_TEXT
 
-TSharedRef<FTextBlockLayout> FTextBlockLayout::Create(TSharedRef<FSlateTextLayoutMarshaller> InMarshaller, TSharedPtr<IBreakIterator> InLineBreakPolicy)
+TSharedRef<FTextBlockLayout> FTextBlockLayout::Create(FTextBlockStyle InDefaultTextStyle, TSharedRef<FBaseTextLayoutMarshaller> InMarshaller, TSharedPtr<IBreakIterator> InLineBreakPolicy)
 {
-	return MakeShareable(new FTextBlockLayout(InMarshaller, InLineBreakPolicy));
+	return MakeShareable(new FTextBlockLayout(MoveTemp(InDefaultTextStyle), InMarshaller, InLineBreakPolicy));
 }
 
-FTextBlockLayout::FTextBlockLayout(TSharedRef<FSlateTextLayoutMarshaller> InMarshaller, TSharedPtr<IBreakIterator> InLineBreakPolicy)
-	: TextLayout(FSlateTextLayout::Create())
+FTextBlockLayout::FTextBlockLayout(FTextBlockStyle InDefaultTextStyle, TSharedRef<FBaseTextLayoutMarshaller> InMarshaller, TSharedPtr<IBreakIterator> InLineBreakPolicy)
+	: TextLayout(FSlateTextLayout::Create(MoveTemp(InDefaultTextStyle)))
 	, Marshaller(InMarshaller)
 	, TextHighlighter(FSlateTextHighlightRunRenderer::Create())
 	, CachedSize(ForceInitToZero)
@@ -33,7 +33,8 @@ FVector2D FTextBlockLayout::ComputeDesiredSize(const FWidgetArgs& InWidgetArgs,
 	// Has the style used for this text block changed?
 	if(!IsStyleUpToDate(InTextStyle))
 	{
-		Marshaller->SetDefaultTextStyle(InTextStyle); // will make marshaller dirty
+		TextLayout->SetDefaultTextStyle(InTextStyle);
+		Marshaller->MakeDirty(); // will regenerate the text using the new default style
 	}
 
 	{
@@ -87,7 +88,7 @@ int32 FTextBlockLayout::OnPaint(const FWidgetArgs& InWidgetArgs, const FPaintArg
 
 	TextLayout->UpdateIfNeeded();
 
-	return TextLayout->OnPaint(InPaintArgs, Marshaller->GetDefaultTextStyle(), InAllottedGeometry, InClippingRect, OutDrawElements, LayerId, InWidgetStyle, bParentEnabled);
+	return TextLayout->OnPaint(InPaintArgs, InAllottedGeometry, InClippingRect, OutDrawElements, LayerId, InWidgetStyle, bParentEnabled);
 }
 
 void FTextBlockLayout::DirtyLayout()
@@ -152,7 +153,7 @@ void FTextBlockLayout::UpdateTextHighlights(const FText& InHighlightText)
 
 bool FTextBlockLayout::IsStyleUpToDate(const FTextBlockStyle& NewStyle) const
 {
-	const FTextBlockStyle& CurrentStyle = Marshaller->GetDefaultTextStyle();
+	const FTextBlockStyle& CurrentStyle = TextLayout->GetDefaultTextStyle();
 
 	return (CurrentStyle.Font == NewStyle.Font)
 		&& (CurrentStyle.ColorAndOpacity == NewStyle.ColorAndOpacity)
diff --git a/Engine/Source/Runtime/Slate/Public/Framework/Text/IRun.h b/Engine/Source/Runtime/Slate/Public/Framework/Text/IRun.h
--- a/Engine/Source/Runtime/Slate/Public/Framework/Text/IRun.h
+++ b/Engine/Source/Runtime/Slate/Public/Framework/Text/IRun.h
@@ -20,6 +20,22 @@ struct SLATE_API FRunInfo
 	TMap< FString, FString > MetaData;
 };
 
+/** Attributes that a run can have */
+enum class ERunAttributes : uint8
+{
+	/**
+	 * This run has no special attributes
+	 */
+	None = 0,
+
+	/**
+	 * This run supports text, and can have new text inserted into it
+	 * Note that even a run which doesn't support text may contain text (likely a breaking space character), however that text should be considered immutable
+	 */
+	SupportsText = 1<<0,
+};
+ENUM_CLASS_FLAGS(ERunAttributes);
+
 class SLATE_API IRun
 {
 public:
@@ -50,4 +66,6 @@ public:
 
 	virtual const FRunInfo& GetRunInfo() const = 0;
 
+	virtual ERunAttributes GetRunAttributes() const = 0;
+
 };
diff --git a/Engine/Source/Runtime/Slate/Public/Framework/Text/PlainTextLayoutMarshaller.h b/Engine/Source/Runtime/Slate/Public/Framework/Text/PlainTextLayoutMarshaller.h
--- a/Engine/Source/Runtime/Slate/Public/Framework/Text/PlainTextLayoutMarshaller.h
+++ b/Engine/Source/Runtime/Slate/Public/Framework/Text/PlainTextLayoutMarshaller.h
@@ -3,16 +3,16 @@
 
 #if WITH_FANCY_TEXT
 
-#include "SlateTextLayoutMarshaller.h"
+#include "BaseTextLayoutMarshaller.h"
 
 /**
  * Get/set the raw text to/from a text layout as plain text
  */
-class SLATE_API FPlainTextLayoutMarshaller : public FSlateTextLayoutMarshaller
+class SLATE_API FPlainTextLayoutMarshaller : public FBaseTextLayoutMarshaller
 {
 public:
 
-	static TSharedRef< FPlainTextLayoutMarshaller > Create(FTextBlockStyle InDefaultTextStyle);
+	static TSharedRef< FPlainTextLayoutMarshaller > Create();
 
 	virtual ~FPlainTextLayoutMarshaller();
 	
@@ -22,7 +22,7 @@ public:
 
 protected:
 
-	FPlainTextLayoutMarshaller(FTextBlockStyle InDefaultTextStyle);
+	FPlainTextLayoutMarshaller();
 
 };
 
diff --git a/Engine/Source/Runtime/Slate/Public/Framework/Text/RichTextLayoutMarshaller.h b/Engine/Source/Runtime/Slate/Public/Framework/Text/RichTextLayoutMarshaller.h
--- a/Engine/Source/Runtime/Slate/Public/Framework/Text/RichTextLayoutMarshaller.h
+++ b/Engine/Source/Runtime/Slate/Public/Framework/Text/RichTextLayoutMarshaller.h
@@ -3,7 +3,7 @@
 
 #if WITH_FANCY_TEXT
 
-#include "SlateTextLayoutMarshaller.h"
+#include "BaseTextLayoutMarshaller.h"
 #include "ITextDecorator.h"
 #include "IRichTextMarkupParser.h"
 #include "IRichTextMarkupWriter.h"
@@ -11,12 +11,12 @@
 /**
  * Get/set the raw text to/from a text layout as rich text
  */
-class SLATE_API FRichTextLayoutMarshaller : public FSlateTextLayoutMarshaller
+class SLATE_API FRichTextLayoutMarshaller : public FBaseTextLayoutMarshaller
 {
 public:
 
-	static TSharedRef< FRichTextLayoutMarshaller > Create(TArray< TSharedRef< ITextDecorator > > InDecorators, const ISlateStyle* const InDecoratorStyleSet, FTextBlockStyle InDefaultTextStyle);
-	static TSharedRef< FRichTextLayoutMarshaller > Create(TSharedPtr< IRichTextMarkupParser > InParser, TSharedPtr< IRichTextMarkupWriter > InWriter, TArray< TSharedRef< ITextDecorator > > InDecorators, const ISlateStyle* const InDecoratorStyleSet, FTextBlockStyle InDefaultTextStyle);
+	static TSharedRef< FRichTextLayoutMarshaller > Create(TArray< TSharedRef< ITextDecorator > > InDecorators, const ISlateStyle* const InDecoratorStyleSet);
+	static TSharedRef< FRichTextLayoutMarshaller > Create(TSharedPtr< IRichTextMarkupParser > InParser, TSharedPtr< IRichTextMarkupWriter > InWriter, TArray< TSharedRef< ITextDecorator > > InDecorators, const ISlateStyle* const InDecoratorStyleSet);
 
 	virtual ~FRichTextLayoutMarshaller();
 
@@ -35,7 +35,7 @@ public:
 
 protected:
 
-	FRichTextLayoutMarshaller(TSharedPtr< IRichTextMarkupParser > InParser, TSharedPtr< IRichTextMarkupWriter > InWriter, TArray< TSharedRef< ITextDecorator > > InDecorators, const ISlateStyle* const InDecoratorStyleSet, FTextBlockStyle InDefaultTextStyle);
+	FRichTextLayoutMarshaller(TSharedPtr< IRichTextMarkupParser > InParser, TSharedPtr< IRichTextMarkupWriter > InWriter, TArray< TSharedRef< ITextDecorator > > InDecorators, const ISlateStyle* const InDecoratorStyleSet);
 
 	TSharedPtr< ITextDecorator > TryGetDecorator(const FString& Line, const FTextRunParseResults& TextRun) const;
 
diff --git a/Engine/Source/Runtime/Slate/Public/Framework/Text/SlateHyperlinkRun.h b/Engine/Source/Runtime/Slate/Public/Framework/Text/SlateHyperlinkRun.h
--- a/Engine/Source/Runtime/Slate/Public/Framework/Text/SlateHyperlinkRun.h
+++ b/Engine/Source/Runtime/Slate/Public/Framework/Text/SlateHyperlinkRun.h
@@ -75,6 +75,8 @@ public:
 
 	virtual const FRunInfo& GetRunInfo() const override;
 
+	virtual ERunAttributes GetRunAttributes() const override;
+
 private:
 
 	FSlateHyperlinkRun( const FRunInfo& InRunInfo, const TSharedRef< const FString >& InText, const FHyperlinkStyle& InStyle, FOnClick InNavigateDelegate, FOnGenerateTooltip InTooltipDelegate, FOnGetTooltipText InTooltipTextDelegate );
diff --git a/Engine/Source/Runtime/Slate/Public/Framework/Text/SlateImageRun.h b/Engine/Source/Runtime/Slate/Public/Framework/Text/SlateImageRun.h
--- a/Engine/Source/Runtime/Slate/Public/Framework/Text/SlateImageRun.h
+++ b/Engine/Source/Runtime/Slate/Public/Framework/Text/SlateImageRun.h
@@ -9,6 +9,8 @@ public:
 
 	static TSharedRef< FSlateImageRun > Create( const FRunInfo& InRunInfo, const TSharedRef< const FString >& InText, const FSlateBrush* InImage, int16 InBaseline );
 	static TSharedRef< FSlateImageRun > Create( const FRunInfo& InRunInfo, const TSharedRef< const FString >& InText, const FSlateBrush* InImage, int16 InBaseline, const FTextRange& InRange );
+	static TSharedRef< FSlateImageRun > Create( const FRunInfo& InRunInfo, const TSharedRef< const FString >& InText, FName InDynamicBrushName, int16 InBaseline );
+	static TSharedRef< FSlateImageRun > Create( const FRunInfo& InRunInfo, const TSharedRef< const FString >& InText, FName InDynamicBrushName, int16 InBaseline, const FTextRange& InRange );
 
 public:
 
@@ -43,17 +45,23 @@ public:
 	virtual void Move(const TSharedRef<FString>& NewText, const FTextRange& NewRange) override;
 	virtual TSharedRef<IRun> Clone() const override;
 
-	virtual void AppendTextTo(FString& Text) const override;
+	virtual void AppendTextTo(FString& AppendToText) const override;
 	virtual void AppendTextTo(FString& AppendToText, const FTextRange& PartialRange) const override;
 
 	virtual const FRunInfo& GetRunInfo() const override;
 
+	virtual ERunAttributes GetRunAttributes() const override;
+
 private:
 
 	FSlateImageRun( const FRunInfo& InRunInfo, const TSharedRef< const FString >& InText, const FSlateBrush* InImage, int16 InBaseline );
 
 	FSlateImageRun( const FRunInfo& InRunInfo, const TSharedRef< const FString >& InText, const FSlateBrush* InImage, int16 InBaseline, const FTextRange& InRange );
 
+	FSlateImageRun( const FRunInfo& InRunInfo, const TSharedRef< const FString >& InText, FName InDynamicBrushName, int16 InBaseline );
+
+	FSlateImageRun( const FRunInfo& InRunInfo, const TSharedRef< const FString >& InText, FName InDynamicBrushName, int16 InBaseline, const FTextRange& InRange );
+
 private:
 
 	FRunInfo RunInfo;
@@ -61,6 +69,7 @@ private:
 	FTextRange Range;
 	const FSlateBrush* Image;
 	int16 Baseline;
+	TSharedPtr< FSlateDynamicImageBrush > DynamicBrush;
 };
 
 #endif //WITH_FANCY_TEXT
\ No newline at end of file
diff --git a/Engine/Source/Runtime/Slate/Public/Framework/Text/SlateTextLayout.h b/Engine/Source/Runtime/Slate/Public/Framework/Text/SlateTextLayout.h
--- a/Engine/Source/Runtime/Slate/Public/Framework/Text/SlateTextLayout.h
+++ b/Engine/Source/Runtime/Slate/Public/Framework/Text/SlateTextLayout.h
@@ -7,13 +7,13 @@ class SLATE_API FSlateTextLayout : public FTextLayout
 {
 public:
 
-	static TSharedRef< FSlateTextLayout > Create();
+	static TSharedRef< FSlateTextLayout > Create(FTextBlockStyle InDefaultTextStyle);
 
 	FChildren* GetChildren();
 
 	void ArrangeChildren( const FGeometry& AllottedGeometry, FArrangedChildren& ArrangedChildren ) const;
 
-	int32 OnPaint( const FPaintArgs& Args, const FTextBlockStyle& DefaultTextStyle, const FGeometry& AllottedGeometry, const FSlateRect& MyClippingRect, FSlateWindowElementList& OutDrawElements, int32 LayerId, const FWidgetStyle& InWidgetStyle, bool bParentEnabled ) const;
+	int32 OnPaint( const FPaintArgs& Args, const FGeometry& AllottedGeometry, const FSlateRect& MyClippingRect, FSlateWindowElementList& OutDrawElements, int32 LayerId, const FWidgetStyle& InWidgetStyle, bool bParentEnabled ) const;
 
 	//void AddLine( const TSharedRef< FString >& Text, const FTextBlockStyle& TextStyle )
 	//{
@@ -24,19 +24,27 @@ public:
 
 	virtual void EndLayout() override;
 
+	void SetDefaultTextStyle(FTextBlockStyle InDefaultTextStyle);
+	const FTextBlockStyle& GetDefaultTextStyle() const;
+
 
 protected:
 
-	FSlateTextLayout();
+	FSlateTextLayout(FTextBlockStyle InDefaultTextStyle);
 
 	int32 OnPaintHighlights(const FPaintArgs& Args, const FTextLayout::FLineView& LineView, const TArray<FLineViewHighlight>& Highlights, const FTextBlockStyle& DefaultTextStyle, const FGeometry& AllottedGeometry, const FSlateRect& ClippingRect, FSlateWindowElementList& OutDrawElements, const int32 LayerId, const FWidgetStyle& InWidgetStyle, bool bParentEnabled) const;
 
 	void AggregateChildren();
 
+	virtual TSharedRef<IRun> CreateDefaultTextRun(const TSharedRef<FString>& NewText, const FTextRange& NewRange) const override;
+
 private:
 
 	TSlotlessChildren< SWidget > Children;
 
+	/** Default style used by the TextLayout */
+	FTextBlockStyle DefaultTextStyle;
+
 	friend class FSlateTextLayoutFactory;
 };
 
diff --git a/Engine/Source/Runtime/Slate/Public/Framework/Text/SlateTextLayoutMarshaller.h b/Engine/Source/Runtime/Slate/Public/Framework/Text/SlateTextLayoutMarshaller.h
deleted file mode 100644
diff --git a/Engine/Source/Runtime/Slate/Public/Framework/Text/SlateTextRun.h b/Engine/Source/Runtime/Slate/Public/Framework/Text/SlateTextRun.h
--- a/Engine/Source/Runtime/Slate/Public/Framework/Text/SlateTextRun.h
+++ b/Engine/Source/Runtime/Slate/Public/Framework/Text/SlateTextRun.h
@@ -46,6 +46,8 @@ public:
 
 	virtual const FRunInfo& GetRunInfo() const override;
 
+	virtual ERunAttributes GetRunAttributes() const override;
+
 private:
 
 	FSlateTextRun( const FRunInfo& InRunInfo, const TSharedRef< const FString >& InText, const FTextBlockStyle& InStyle );
diff --git a/Engine/Source/Runtime/Slate/Public/Framework/Text/SlateWidgetRun.h b/Engine/Source/Runtime/Slate/Public/Framework/Text/SlateWidgetRun.h
--- a/Engine/Source/Runtime/Slate/Public/Framework/Text/SlateWidgetRun.h
+++ b/Engine/Source/Runtime/Slate/Public/Framework/Text/SlateWidgetRun.h
@@ -68,6 +68,8 @@ public:
 
 	virtual const FRunInfo& GetRunInfo() const override;
 
+	virtual ERunAttributes GetRunAttributes() const override;
+
 private:
 
 	FSlateWidgetRun(const TSharedRef<class FTextLayout>& TextLayout, const FRunInfo& InRunInfo, const TSharedRef< const FString >& InText, const FWidgetRunInfo& InWidgetInfo);
diff --git a/Engine/Source/Runtime/Slate/Public/Framework/Text/SyntaxHighlighterTextLayoutMarshaller.h b/Engine/Source/Runtime/Slate/Public/Framework/Text/SyntaxHighlighterTextLayoutMarshaller.h
--- a/Engine/Source/Runtime/Slate/Public/Framework/Text/SyntaxHighlighterTextLayoutMarshaller.h
+++ b/Engine/Source/Runtime/Slate/Public/Framework/Text/SyntaxHighlighterTextLayoutMarshaller.h
@@ -27,7 +27,7 @@ protected:
 
 	virtual void ParseTokens(const FString& SourceString, FTextLayout& TargetTextLayout, TArray<FSyntaxTokenizer::FTokenizedLine> TokenizedLines) = 0;
 
-	FSyntaxHighlighterTextLayoutMarshaller(TSharedPtr< FSyntaxTokenizer > InTokenizer, FTextBlockStyle InDefaultTextStyle);
+	FSyntaxHighlighterTextLayoutMarshaller(TSharedPtr< FSyntaxTokenizer > InTokenizer);
 
 	/** Tokenizer used to style the text */
 	TSharedPtr< FSyntaxTokenizer > Tokenizer;
diff --git a/Engine/Source/Runtime/Slate/Public/Framework/Text/TextLayout.h b/Engine/Source/Runtime/Slate/Public/Framework/Text/TextLayout.h
--- a/Engine/Source/Runtime/Slate/Public/Framework/Text/TextLayout.h
+++ b/Engine/Source/Runtime/Slate/Public/Framework/Text/TextLayout.h
@@ -414,6 +414,11 @@ protected:
 	*/
 	virtual void EndLayout();
 
+	/**
+	 * Called to generate a new empty text run for this text layout
+	 */
+	virtual TSharedRef<IRun> CreateDefaultTextRun(const TSharedRef<FString>& NewText, const FTextRange& NewRange) const = 0;
+
 private:
 
 	void FlowLayout();
diff --git a/Engine/Source/Runtime/Slate/Public/Widgets/Text/TextBlockLayout.h b/Engine/Source/Runtime/Slate/Public/Widgets/Text/TextBlockLayout.h
--- a/Engine/Source/Runtime/Slate/Public/Widgets/Text/TextBlockLayout.h
+++ b/Engine/Source/Runtime/Slate/Public/Widgets/Text/TextBlockLayout.h
@@ -3,7 +3,7 @@
 #pragma once
 
 class FSlateTextLayout;
-class FSlateTextLayoutMarshaller;
+class FBaseTextLayoutMarshaller;
 class ISlateRunRenderer;
 
 #if WITH_FANCY_TEXT
@@ -42,7 +42,7 @@ public:
 		const TAttribute<ETextJustify::Type>& Justification;
 	};
 
-	static TSharedRef<FTextBlockLayout> Create(TSharedRef<FSlateTextLayoutMarshaller> InMarshaller, TSharedPtr<IBreakIterator> InLineBreakPolicy);
+	static TSharedRef<FTextBlockLayout> Create(FTextBlockStyle InDefaultTextStyle, TSharedRef<FBaseTextLayoutMarshaller> InMarshaller, TSharedPtr<IBreakIterator> InLineBreakPolicy);
 
 	/**
 	 * Get the computed desired size for this layout, updating the internal cache as required
@@ -82,13 +82,13 @@ private:
 	/** Calculate the wrapping width based on the given fixed wrap width, and whether we're auto-wrapping */
 	float CalculateWrappingWidth(const FWidgetArgs& InWidgetArgs) const;
 
-	FTextBlockLayout(TSharedRef<FSlateTextLayoutMarshaller> InMarshaller, TSharedPtr<IBreakIterator> InLineBreakPolicy);
+	FTextBlockLayout(FTextBlockStyle InDefaultTextStyle, TSharedRef<FBaseTextLayoutMarshaller> InMarshaller, TSharedPtr<IBreakIterator> InLineBreakPolicy);
 
 	/** In control of the layout and wrapping of the text */
 	TSharedPtr<FSlateTextLayout> TextLayout;
 
 	/** The marshaller used to get/set the text to/from the text layout. */
-	TSharedPtr<FSlateTextLayoutMarshaller> Marshaller;
+	TSharedPtr<FBaseTextLayoutMarshaller> Marshaller;
 
 	/** Used to render the current highlights in the text layout */
 	TSharedPtr<ISlateRunRenderer> TextHighlighter;
