Fixed "flickering" when selecting the text inside a hyperlink in the multiline editable text

Because FTextSelectionRunRenderer is a run renderer, it causes the hyperlink run to get split into multiple blocks when the selection is changed, and because FSlateHyperlinkRun contains a nested widget, it was possible that the widgets inside the new blocks hadn't had a Slate prepass called on them before they were presented for painting (causing a flicker while the STextBlock widget inside the hyperlink was empty).

The FSlateHyperlinkRun now calls SlatePrepass() itself when the blocks are created. This is a similar solution to FSlateWidgetRun, and ensures that the hyperlink is ready to be painted, thus removing the flickering.

diff --git a/Engine/Source/Runtime/Slate/Private/Framework/Text/SlateHyperlinkRun.cpp b/Engine/Source/Runtime/Slate/Private/Framework/Text/SlateHyperlinkRun.cpp
--- a/Engine/Source/Runtime/Slate/Private/Framework/Text/SlateHyperlinkRun.cpp
+++ b/Engine/Source/Runtime/Slate/Private/Framework/Text/SlateHyperlinkRun.cpp
@@ -82,6 +82,10 @@ TSharedRef< ILayoutBlock > FSlateHyperlinkRun::CreateBlock( int32 StartIndex, in
 		.ToolTip( ToolTip )
 		.ToolTipText( ToolTipText )
 		.OnNavigate( this, &FSlateHyperlinkRun::OnNavigate );
+	
+	// We need to do a prepass here as CreateBlock can be called after the main Slate prepass has been run, 
+	// which can result in the hyperlink widget not being correctly setup before it is painted
+	Widget->SlatePrepass();
 
 	Children.Add( Widget );
 
