Bring back tooltips on Linux (TTP #344052).

- Also bring back caching the position immediately, not waiting on window moved event.

diff --git a/Engine/Source/Runtime/Slate/Private/Framework/Application/SlateApplication.cpp b/Engine/Source/Runtime/Slate/Private/Framework/Application/SlateApplication.cpp
--- a/Engine/Source/Runtime/Slate/Private/Framework/Application/SlateApplication.cpp
+++ b/Engine/Source/Runtime/Slate/Private/Framework/Application/SlateApplication.cpp
@@ -309,7 +309,7 @@ FSlateApplication::FSlateApplication()
 	, OnExitRequested()
 	, EventLogger( TSharedPtr<IEventLogger>() )
 	, NumExternalModalWindowsActive( 0 )
-	, bAllowToolTips( PLATFORM_LINUX ? false : true )	// @todo: support tooltip windows on Linux (TTP #344052)
+	, bAllowToolTips( true )
 	, ToolTipDelay( 0.15f )
 	, ToolTipFadeInDuration( 0.1f )
 	, ToolTipSummonTime( 0.0 )
diff --git a/Engine/Source/Runtime/SlateCore/Private/Widgets/SWindow.cpp b/Engine/Source/Runtime/SlateCore/Private/Widgets/SWindow.cpp
--- a/Engine/Source/Runtime/SlateCore/Private/Widgets/SWindow.cpp
+++ b/Engine/Source/Runtime/SlateCore/Private/Widgets/SWindow.cpp
@@ -697,6 +697,8 @@ void SWindow::MoveWindowTo( FVector2D NewPosition )
 	if (NativeWindow.IsValid())
 	{
 		NativeWindow->MoveWindowTo( FMath::TruncToInt(NewPosition.X), FMath::TruncToInt(NewPosition.Y) );
+		// cache position immediately since the "windows moved" event may be delayed until next frame.
+		SetCachedScreenPosition(NewPosition);
 	}
 	else
 	{
