SlateApplication - Add support for multitouch mouse capture.

diff --git a/Engine/Source/Editor/InputBindingEditor/Private/Widgets/SGestureEditBox.cpp b/Engine/Source/Editor/InputBindingEditor/Private/Widgets/SGestureEditBox.cpp
--- a/Engine/Source/Editor/InputBindingEditor/Private/Widgets/SGestureEditBox.cpp
+++ b/Engine/Source/Editor/InputBindingEditor/Private/Widgets/SGestureEditBox.cpp
@@ -105,7 +105,7 @@ FText SGestureEditBox::GetNotificationMessage() const
 
 void SGestureEditBox::OnGestureEditorLostFocus()
 {
-	if( (!GestureAcceptButton.IsValid() || FSlateApplication::Get().GetMouseCaptor() != GestureAcceptButton) && !GestureEditor->IsTyping() )
+	if( (!GestureAcceptButton.IsValid() || GestureAcceptButton->HasMouseCapture() == false) && !GestureEditor->IsTyping() )
 	{
 		if( GestureEditor->IsEditing() && GestureEditor->IsEditedGestureValid() && !GestureEditor->HasConflict() )
 		{
diff --git a/Engine/Source/Editor/UnrealEd/Private/PackageAutoSaver.cpp b/Engine/Source/Editor/UnrealEd/Private/PackageAutoSaver.cpp
--- a/Engine/Source/Editor/UnrealEd/Private/PackageAutoSaver.cpp
+++ b/Engine/Source/Editor/UnrealEd/Private/PackageAutoSaver.cpp
@@ -372,7 +372,7 @@ bool FPackageAutoSaver::CanAutoSave() const
 	const bool bPlayWorldValid = GUnrealEd->PlayWorld != nullptr;
 	const bool bAnyMenusVisible	= FSlateApplication::Get().AnyMenusVisible();
 	const bool bAutomationTesting = GIsAutomationTesting;
-	const bool bIsInteratcting = FSlateApplication::Get().GetMouseCaptor().IsValid() || GUnrealEd->IsUserInteracting() || bDidTypeInATextBlockRecently;
+	const bool bIsInteratcting = FSlateApplication::Get().HasAnyMouseCaptor() || GUnrealEd->IsUserInteracting() || bDidTypeInATextBlockRecently;
 	const bool bHasGameOrProjectLoaded = FApp::HasGameName();
 
 	return (bAutosaveEnabled && !bSlowTask && !bInterpEditMode && !bPlayWorldValid && !bAnyMenusVisible && !bAutomationTesting && !bIsInteratcting && !GIsDemoMode && bHasGameOrProjectLoaded);
diff --git a/Engine/Source/Runtime/Engine/Private/Slate/SceneViewport.cpp b/Engine/Source/Runtime/Engine/Private/Slate/SceneViewport.cpp
--- a/Engine/Source/Runtime/Engine/Private/Slate/SceneViewport.cpp
+++ b/Engine/Source/Runtime/Engine/Private/Slate/SceneViewport.cpp
@@ -44,7 +44,7 @@ FSceneViewport::~FSceneViewport()
 
 bool FSceneViewport::HasMouseCapture() const
 {
-	return FSlateApplication::Get().GetMouseCaptor() == ViewportWidget.Pin();
+	return ViewportWidget.IsValid() && ViewportWidget.Pin()->HasMouseCapture();
 }
 
 bool FSceneViewport::HasFocus() const
diff --git a/Engine/Source/Runtime/Slate/Private/Framework/Application/SlateApplication.cpp b/Engine/Source/Runtime/Slate/Private/Framework/Application/SlateApplication.cpp
--- a/Engine/Source/Runtime/Slate/Private/Framework/Application/SlateApplication.cpp
+++ b/Engine/Source/Runtime/Slate/Private/Framework/Application/SlateApplication.cpp
@@ -69,28 +69,56 @@ TAutoConsoleVariable<int32> TargetFrameRateForResponsiveness(
 	TEXT( "Minimum sustained average frame rate required before we consider the editor to be \"responsive\" for a smooth UI experience" ) );
 
 //////////////////////////////////////////////////////////////////////////
-bool FSlateApplication::MouseCaptorHelper::IsValid() const
+bool FSlateApplication::MouseCaptorHelper::HasCapture() const
 {
-	return MouseCaptorWeakPath.IsValid();
+	for (auto PointerPathPair : PointerIndexToMouseCaptorWeakPathMap)
+	{
+		if (PointerPathPair.Value.IsValid())
+		{
+			return true;
+		}
+	}
+	return false;
+}
+
+bool FSlateApplication::MouseCaptorHelper::HasCaptureForPointerIndex(uint32 PointerIndex) const
+{
+	const FWeakWidgetPath* MouseCaptorWeakPath = PointerIndexToMouseCaptorWeakPathMap.Find(PointerIndex);
+	return MouseCaptorWeakPath && MouseCaptorWeakPath->IsValid();
 }
 
-TSharedPtr< SWidget > FSlateApplication::MouseCaptorHelper::ToSharedWidget() const
+TSharedPtr< SWidget > FSlateApplication::MouseCaptorHelper::ToSharedWidget(uint32 PointerIndex) const
 {
 	// If the path is valid then get the last widget, this is the current mouse captor
 	TSharedPtr< SWidget > SharedWidgetPtr;
-	if ( MouseCaptorWeakPath.IsValid() )
+	const FWeakWidgetPath* MouseCaptorWeakPath = PointerIndexToMouseCaptorWeakPathMap.Find(PointerIndex);
+	if (MouseCaptorWeakPath && MouseCaptorWeakPath->IsValid() )
 	{
-		TWeakPtr< SWidget > WeakWidgetPtr = MouseCaptorWeakPath.GetLastWidget();
+		TWeakPtr< SWidget > WeakWidgetPtr = MouseCaptorWeakPath->GetLastWidget();
 		SharedWidgetPtr = WeakWidgetPtr.Pin();
 	}
 
 	return SharedWidgetPtr;
 }
 
-TSharedPtr< SWidget > FSlateApplication::MouseCaptorHelper::ToSharedWindow()
+TArray<TSharedRef<SWidget>> FSlateApplication::MouseCaptorHelper::ToSharedWidgets() const
+{
+	TArray<TSharedRef<SWidget>> Widgets;
+	for (auto IndexPathPair : PointerIndexToMouseCaptorWeakPathMap)
+	{
+		TSharedPtr<SWidget> LastWidget = IndexPathPair.Value.GetLastWidget().Pin();
+		if (LastWidget.IsValid())
+		{
+			Widgets.Add(LastWidget.ToSharedRef());
+		}
+	}
+	return Widgets;
+}
+
+TSharedPtr< SWidget > FSlateApplication::MouseCaptorHelper::ToSharedWindow(uint32 PointerIndex)
 {
 	// if the path is valid then we can get the window the current mouse captor belongs to
-	FWidgetPath MouseCaptorPath = ToWidgetPath();
+	FWidgetPath MouseCaptorPath = ToWidgetPath(PointerIndex);
 	if ( MouseCaptorPath.IsValid() )
 	{
 		return MouseCaptorPath.GetWindow();
@@ -99,11 +127,11 @@ TSharedPtr< SWidget > FSlateApplication::MouseCaptorHelper::ToSharedWindow()
 	return TSharedPtr< SWidget >();
 }
 
-void FSlateApplication::MouseCaptorHelper::SetMouseCaptor( const FWidgetPath& EventPath, TSharedPtr< SWidget > Widget )
+void FSlateApplication::MouseCaptorHelper::SetMouseCaptor(uint32 PointerIndex, const FWidgetPath& EventPath, TSharedPtr< SWidget > Widget)
 {
 	// Caller is trying to set a new mouse captor, so invalidate the current one - when the function finishes
 	// it still may not have a valid captor widget, this is ok
-	Invalidate();
+	InvalidateCaptureForPointer(PointerIndex);
 
 	if ( Widget.IsValid() )
 	{
@@ -115,6 +143,7 @@ void FSlateApplication::MouseCaptorHelper::SetMouseCaptor( const FWidgetPath& Ev
 			return PathToTest.Widgets.Num() > 0 && PathToTest.Widgets.Last().Widget == WidgetToFind;
 		};
 
+		FWeakWidgetPath MouseCaptorWeakPath;
 		if ( IsPathToCaptorFound( NewMouseCaptorPath, WidgetRef ) )
 		{
 			MouseCaptorWeakPath = NewMouseCaptorPath;
@@ -129,42 +158,76 @@ void FSlateApplication::MouseCaptorHelper::SetMouseCaptor( const FWidgetPath& Ev
 				? NewMouseCaptorPath
 				: FWeakWidgetPath();
 		}
+
+		if (MouseCaptorWeakPath.IsValid())
+		{
+			PointerIndexToMouseCaptorWeakPathMap.Add(PointerIndex, MouseCaptorWeakPath);
+		}
+	}
+}
+
+void FSlateApplication::MouseCaptorHelper::InvalidateCaptureForAllPointers()
+{
+	TArray<uint32> PointerIndices;
+	PointerIndexToMouseCaptorWeakPathMap.GenerateKeyArray(PointerIndices);
+	for (auto Index : PointerIndices)
+	{
+		InvalidateCaptureForPointer(Index);
 	}
 }
 
-void FSlateApplication::MouseCaptorHelper::Invalidate()
+void FSlateApplication::MouseCaptorHelper::InvalidateCaptureForPointer(uint32 PointerIndex)
 {
-	InformCurrentCaptorOfCaptureLoss();
-	MouseCaptorWeakPath = FWeakWidgetPath();
+	InformCurrentCaptorOfCaptureLoss(PointerIndex);
+	PointerIndexToMouseCaptorWeakPathMap.Remove(PointerIndex);
 }
 
-FWidgetPath FSlateApplication::MouseCaptorHelper::ToWidgetPath( FWeakWidgetPath::EInterruptedPathHandling::Type InterruptedPathHandling )
+FWidgetPath FSlateApplication::MouseCaptorHelper::ToWidgetPath(uint32 PointerIndex, FWeakWidgetPath::EInterruptedPathHandling::Type InterruptedPathHandling)
 {
 	FWidgetPath WidgetPath;
-	if ( MouseCaptorWeakPath.IsValid() )
+	const FWeakWidgetPath* MouseCaptorWeakPath = PointerIndexToMouseCaptorWeakPathMap.Find(PointerIndex);
+	if (MouseCaptorWeakPath && MouseCaptorWeakPath->IsValid() )
 	{
-		if ( MouseCaptorWeakPath.ToWidgetPath( WidgetPath, InterruptedPathHandling ) == FWeakWidgetPath::EPathResolutionResult::Truncated )
+		if ( MouseCaptorWeakPath->ToWidgetPath( WidgetPath, InterruptedPathHandling ) == FWeakWidgetPath::EPathResolutionResult::Truncated )
 		{
 			// If the path was truncated then it means this widget is no longer part of the active set,
 			// so we make sure to invalidate its capture
-			Invalidate();
+			InvalidateCaptureForPointer(PointerIndex);
 		}
 	}
 
 	return WidgetPath;
 }
 
-FWeakWidgetPath FSlateApplication::MouseCaptorHelper::ToWeakPath() const
+TArray<FWidgetPath> FSlateApplication::MouseCaptorHelper::ToWidgetPaths()
+{
+	TArray<FWidgetPath> WidgetPaths;
+	TArray<uint32> PointerIndices;
+	PointerIndexToMouseCaptorWeakPathMap.GenerateKeyArray(PointerIndices);
+	for (auto Index : PointerIndices)
+	{
+		WidgetPaths.Add(ToWidgetPath(Index));
+	}
+	return WidgetPaths;
+}
+
+FWeakWidgetPath FSlateApplication::MouseCaptorHelper::ToWeakPath(uint32 PointerIndex) const
 {
-	return MouseCaptorWeakPath;
+	const FWeakWidgetPath* MouseCaptorWeakPath = PointerIndexToMouseCaptorWeakPathMap.Find(PointerIndex);
+	if (MouseCaptorWeakPath)
+	{
+		return *MouseCaptorWeakPath;
+	}
+	return FWeakWidgetPath();
 }
 
-void FSlateApplication::MouseCaptorHelper::InformCurrentCaptorOfCaptureLoss() const
+void FSlateApplication::MouseCaptorHelper::InformCurrentCaptorOfCaptureLoss(uint32 PointerIndex) const
 {
 	// if we have a path to a widget then it is the current mouse captor and needs to know it has lost capture
-	if ( MouseCaptorWeakPath.IsValid() )
+	const FWeakWidgetPath* MouseCaptorWeakPath = PointerIndexToMouseCaptorWeakPathMap.Find(PointerIndex);
+	if (MouseCaptorWeakPath && MouseCaptorWeakPath->IsValid() )
 	{
-		TWeakPtr< SWidget > WeakWidgetPtr = MouseCaptorWeakPath.GetLastWidget();
+		TWeakPtr< SWidget > WeakWidgetPtr = MouseCaptorWeakPath->GetLastWidget();
 		TSharedPtr< SWidget > SharedWidgetPtr = WeakWidgetPtr.Pin();
 		if ( SharedWidgetPtr.IsValid() )
 		{
@@ -176,7 +239,7 @@ void FSlateApplication::MouseCaptorHelper::InformCurrentCaptorOfCaptureLoss() co
 //////////////////////////////////////////////////////////////////////////
 void FPopupSupport::RegisterClickNotification( const TSharedRef<SWidget>& NotifyWhenClickedOutsideMe, const FOnClickedOutside& InNotification )
 {
-	// If the subscriber or a zone object is descroyed, the subscription is
+	// If the subscriber or a zone object is destroyed, the subscription is
 	// no longer active. Clean it up here so that consumers of this API have an
 	// easy time with resource management.
 	struct { void operator()( TArray<FClickSubscriber>& Notifications ) {
@@ -355,6 +418,7 @@ FSlateApplication::FSlateApplication()
 	FGenericCommands::Register();
 
 	NormalExecutionGetter.BindRaw( this, &FSlateApplication::IsNormalExecution );
+	PointerIndexLastPositionMap.Add(CursorPointerIndex, FVector2D::ZeroVector);
 }
 
 FSlateApplication::~FSlateApplication()
@@ -419,6 +483,11 @@ FVector2D FSlateApplication::GetCursorPos() const
 	return FVector2D( 0, 0 );
 }
 
+FVector2D FSlateApplication::GetLastCursorPos() const
+{
+	return PointerIndexLastPositionMap[CursorPointerIndex];
+}
+
 void FSlateApplication::SetCursorPos( const FVector2D& MouseCoordinate )
 {
 	if ( PlatformApplication->Cursor.IsValid() )
@@ -752,7 +821,7 @@ void FSlateApplication::Tick()
 
 	// When Slate captures the mouse, it is up to us to set the cursor 
 	// because the OS assumes that we own the mouse.
-	if (MouseCaptor.IsValid())
+	if (MouseCaptor.HasCapture())
 	{
 		QueryCursor();
 	}
@@ -899,7 +968,7 @@ void FSlateApplication::ThrottleApplicationBasedOnMouseMovement()
 				if( TimeSinceLastUserInteraction > MinTimeSinceButtonPressToThrottle )
 				{
 					// If a widget has the mouse captured, then we won't bother throttling
-					if( !MouseCaptor.IsValid() )
+					if( !MouseCaptor.HasCapture() )
 					{
 						// If there is no Slate window under the mouse, then we won't engage throttling
 						if( LocateWindowUnderMouse( GetCursorPos(), GetInteractiveTopLevelWindows() ).IsValid() )
@@ -1384,10 +1453,9 @@ void FSlateApplication::SetKeyboardFocus( const TSharedPtr< SWidget >& OptionalW
 
 void FSlateApplication::ResetToDefaultInputSettings()
 {
-	if (MouseCaptor.IsValid())
+	if (MouseCaptor.HasCapture())
 	{
-		FWidgetPath MouseCaptorPath = MouseCaptor.ToWidgetPath();
-		ProcessReply( MouseCaptorPath, FReply::Handled().ReleaseMouseCapture(), NULL, NULL );
+		MouseCaptor.InvalidateCaptureForAllPointers();
 	}
 
 	for (int32 UserIndex = 0; UserIndex < ARRAY_COUNT(JoystickCaptorWeakPaths); UserIndex++)
@@ -1415,7 +1483,7 @@ void* FSlateApplication::GetMouseCaptureWindow( void ) const
 
 void FSlateApplication::ReleaseMouseCapture()
 {
-	MouseCaptor.Invalidate();
+	MouseCaptor.InvalidateCaptureForAllPointers();
 }
 
 
@@ -1708,10 +1776,11 @@ void FSlateApplication::ProcessReply( const FWidgetPath& CurrentEventPath, const
 
 	// Release mouse capture if requested or if we are starting a drag and drop.
 	// Make sure to only clobber WidgetsUnderCursor if we actually had a mouse capture.
-	if ( MouseCaptor.IsValid() && (TheReply.ShouldReleaseMouse() || bStartingDragDrop) )
+	uint32 PointerIndex = InMouseEvent != nullptr ? InMouseEvent->GetPointerIndex() : CursorPointerIndex;
+	if (MouseCaptor.HasCaptureForPointerIndex(PointerIndex) && (TheReply.ShouldReleaseMouse() || bStartingDragDrop) )
 	{
-		WidgetsUnderCursorLastEvent = MouseCaptor.ToWeakPath();
-		MouseCaptor.Invalidate();
+		WidgetsUnderCursorLastEvent = MouseCaptor.ToWeakPath(PointerIndex);
+		MouseCaptor.InvalidateCaptureForPointer(PointerIndex);
 	}
 
 
@@ -1766,9 +1835,9 @@ void FSlateApplication::ProcessReply( const FWidgetPath& CurrentEventPath, const
 	
 	TSharedPtr<SWidget> RequestedMouseCaptor = TheReply.GetMouseCaptor();
 	// Do not capture the mouse if we are also starting a drag and drop.
-	if ( RequestedMouseCaptor.IsValid() && !bStartingDragDrop )
+	if( RequestedMouseCaptor.IsValid() && !bStartingDragDrop )
 	{
-		MouseCaptor.SetMouseCaptor( CurrentEventPath, RequestedMouseCaptor );
+		MouseCaptor.SetMouseCaptor(PointerIndex, CurrentEventPath, RequestedMouseCaptor );
 	}
 	
 	if( CurrentEventPath.IsValid() && ( TheReply.ShouldReleaseMouse() || RequestedMouseCaptor.IsValid() ) )
@@ -1819,7 +1888,7 @@ void FSlateApplication::ProcessReply( const FWidgetPath& CurrentEventPath, const
 	if( RequestedMousePos.IsSet() )
 	{
 		const FVector2D Position = RequestedMousePos.GetValue();
-		LastCursorPosition = Position;
+		PointerIndexLastPositionMap.Add(CursorPointerIndex, Position);
 		SetCursorPos( Position );
 	}
 
@@ -1929,9 +1998,9 @@ void FSlateApplication::QueryCursor()
 			const TSharedPtr<SWindow> ActiveModalWindow = GetActiveModalWindow();
 
 			// Query widgets with mouse capture for the cursor
-			if (MouseCaptor.IsValid())
+			if (MouseCaptor.HasCaptureForPointerIndex(CursorPointerIndex))
 			{
-				FWidgetPath MouseCaptorPath = MouseCaptor.ToWidgetPath();
+				FWidgetPath MouseCaptorPath = MouseCaptor.ToWidgetPath(CursorPointerIndex);
 				if ( MouseCaptorPath.IsValid() )
 				{
 					TSharedRef< SWindow > CaptureWindow = MouseCaptorPath.GetWindow();
@@ -1954,7 +2023,9 @@ void FSlateApplication::QueryCursor()
 				FScopedSwitchWorldHack SwitchWorld( WidgetsToQueryForCursor );
 
 				const FVector2D CurrentCursorPosition = GetCursorPos();
+				const FVector2D LastCursorPosition = GetLastCursorPos();
 				const FPointerEvent CursorEvent(
+					CursorPointerIndex,
 					CurrentCursorPosition,
 					LastCursorPosition,
 					CurrentCursorPosition - LastCursorPosition,
@@ -2179,7 +2250,7 @@ void FSlateApplication::UpdateToolTip( bool AllowSpawningOfNewToolTips )
 		if ( ( ActiveToolTipPtr.IsValid() && !ActiveToolTipPtr->IsInteractive() ) || ( NewToolTip.IsValid() && NewToolTip != ActiveToolTip.Pin() ) )
 		{
 			// Keep track of where we want tool tips to be positioned
-			DesiredToolTipLocation = LastCursorPosition + SlateDefs::ToolTipOffsetFromMouse;
+			DesiredToolTipLocation = GetLastCursorPos() + SlateDefs::ToolTipOffsetFromMouse;
 		}
 	}
 
@@ -2447,19 +2518,22 @@ bool FSlateApplication::IsWindowInDestroyQueue(TSharedRef<SWindow> Window) const
 
 void FSlateApplication::SynthesizeMouseMove()
 {
-	FPointerEvent MouseEvent(
-		LastCursorPosition, 
-		LastCursorPosition, 
-		PressedMouseButtons, 
-		EKeys::Invalid, 
-		0, 
-		PlatformApplication->GetModifierKeys()
-	);
-
-	ProcessMouseMoveEvent(MouseEvent);
+	for (auto IndexLocationPairIterator = PointerIndexLastPositionMap.CreateIterator(); IndexLocationPairIterator; ++IndexLocationPairIterator)
+	{
+		// TODO: Use correct user index instead of 0.
+		FPointerEvent MouseEvent
+		(
+			0,
+			IndexLocationPairIterator.Key(),
+			IndexLocationPairIterator.Value(),
+			IndexLocationPairIterator.Value(), 
+			false
+		);
+		ProcessMouseMoveEvent(MouseEvent);
+	}
 }
 
-void FSlateApplication::OnLogSlateEvent(EEventLog::Type Event, const FString& AdditionalContent)
+void FSlateApplication::OnLogSlateEvent(EEventLog::Type Elovent, const FString& AdditionalContent)
 {
 #if LOG_SLATE_EVENTS
 	if (EventLogger.IsValid())
@@ -2802,11 +2876,22 @@ TSharedPtr< SWidget > FSlateApplication::GetKeyboardFocusedWidget() const
 }
 
 
-TSharedPtr< SWidget > FSlateApplication::GetMouseCaptor() const
+bool FSlateApplication::HasAnyMouseCaptor() const
 {
-	return MouseCaptor.ToSharedWidget();
+	return MouseCaptor.HasCapture();
 }
 
+bool FSlateApplication::HasMouseCapture(const TSharedPtr<const SWidget> Widget) const
+{
+	for (auto CaptureWidget : MouseCaptor.ToSharedWidgets())
+	{
+		if (Widget == CaptureWidget)
+		{
+			return true;
+		}
+	}
+	return false;
+}
 
 bool FSlateApplication::HasFocusedDescendants( const TSharedRef< const SWidget >& Widget ) const
 {
@@ -3131,8 +3216,9 @@ bool FSlateApplication::OnMouseDown( const TSharedPtr< FGenericWindow >& Platfor
 	FKey Key = TranslateMouseButtonToKey( Button );
 
 	FPointerEvent MouseEvent(
-		PlatformApplication->Cursor->GetPosition(), 	
-		LastCursorPosition,
+		CursorPointerIndex,
+		GetCursorPos(),
+		GetLastCursorPos(),
 		PressedMouseButtons,
 		Key,
 		0,
@@ -3156,9 +3242,9 @@ bool FSlateApplication::ProcessMouseButtonDownEvent( const TSharedPtr< FGenericW
 	if ( !IsDragDropping() )
 	{
 		FReply Reply = FReply::Unhandled();
-		if (MouseCaptor.IsValid())
+		if (MouseCaptor.HasCaptureForPointerIndex(MouseEvent.GetPointerIndex()))
 		{
-			FWidgetPath MouseCaptorPath = MouseCaptor.ToWidgetPath();
+			FWidgetPath MouseCaptorPath = MouseCaptor.ToWidgetPath(MouseEvent.GetPointerIndex());
 			FArrangedWidget& MouseCaptorWidget = MouseCaptorPath.Widgets.Last();
 
 			// Switch worlds widgets in the current path
@@ -3278,8 +3364,7 @@ bool FSlateApplication::ProcessMouseButtonDownEvent( const TSharedPtr< FGenericW
 		}
 	}
 
-	LastCursorPosition = MouseEvent.GetScreenSpacePosition();
-
+	PointerIndexLastPositionMap.Add(MouseEvent.GetPointerIndex(), MouseEvent.GetScreenSpacePosition());
 	return true;
 }
 
@@ -3294,8 +3379,9 @@ bool FSlateApplication::OnMouseDoubleClick( const TSharedPtr< FGenericWindow >&
 	FKey Key = TranslateMouseButtonToKey( Button );
 
 	FPointerEvent MouseEvent(
-		GetCursorPos(), 	
-		LastCursorPosition,
+		CursorPointerIndex,
+		GetCursorPos(),
+		GetLastCursorPos(),
 		PressedMouseButtons,
 		Key,
 		0,
@@ -3333,8 +3419,7 @@ bool FSlateApplication::ProcessMouseButtonDoubleClickEvent( const TSharedPtr< FG
 
 	LOG_EVENT( EEventLog::MouseButtonDoubleClick, Reply );
 
-	LastCursorPosition = InMouseEvent.GetScreenSpacePosition();
-
+	PointerIndexLastPositionMap.Add(InMouseEvent.GetPointerIndex(), InMouseEvent.GetScreenSpacePosition());
 	return Reply.IsEventHandled();
 }
 
@@ -3350,8 +3435,9 @@ bool FSlateApplication::OnMouseUp( const EMouseButtons::Type Button )
 	FKey Key = TranslateMouseButtonToKey( Button );
 
 	FPointerEvent MouseEvent(
-		GetCursorPos(), 	
-		LastCursorPosition,
+		CursorPointerIndex,
+		GetCursorPos(),
+		GetLastCursorPos(),
 		PressedMouseButtons,
 		Key,
 		0,
@@ -3373,9 +3459,9 @@ bool FSlateApplication::ProcessMouseButtonUpEvent( FPointerEvent& MouseEvent )
 		DragDetector = FDragDetector();
 	}
 
-	if (MouseCaptor.IsValid())
+	if (MouseCaptor.HasCaptureForPointerIndex(MouseEvent.GetPointerIndex()))
 	{
-		FWidgetPath MouseCaptorPath = MouseCaptor.ToWidgetPath();
+		FWidgetPath MouseCaptorPath = MouseCaptor.ToWidgetPath(MouseEvent.GetPointerIndex());
 		if ( ensureMsg(MouseCaptorPath.Widgets.Num() > 0, TEXT("A window had a widget with mouse capture. That entire window has been dismissed before the mouse up could be processed.")) )
 		{
 			FArrangedWidget& MouseCaptorWidget = MouseCaptorPath.Widgets.Last();
@@ -3460,7 +3546,8 @@ bool FSlateApplication::OnMouseWheel( const float Delta )
 	const FVector2D CurrentCursorPosition = GetCursorPos();
 
 	FPointerEvent MouseWheelEvent(
-		CurrentCursorPosition, 	
+		CursorPointerIndex,
+		CurrentCursorPosition,
 		CurrentCursorPosition,
 		PressedMouseButtons,
 		EKeys::Invalid,
@@ -3485,8 +3572,8 @@ bool FSlateApplication::ProcessMouseWheelOrGestureEvent( FPointerEvent& InWheelE
 	
 	// NOTE: We intentionally don't reset LastUserInteractionTimeForThrottling here so that the UI can be responsive while scrolling
 
-	FWidgetPath EventPath = ( MouseCaptor.IsValid() )
-		? MouseCaptor.ToWidgetPath()
+	FWidgetPath EventPath = ( MouseCaptor.HasCaptureForPointerIndex(InWheelEvent.GetPointerIndex()) )
+		? MouseCaptor.ToWidgetPath(InWheelEvent.GetPointerIndex())
 		: LocateWindowUnderMouse( InWheelEvent.GetScreenSpacePosition(), GetInteractiveTopLevelWindows() );
 
 	InWheelEvent.SetEventPath( EventPath );
@@ -3536,10 +3623,11 @@ bool FSlateApplication::OnMouseMove()
 
 	bool Result = true;
 	const FVector2D CurrentCursorPosition = GetCursorPos();
-
+	const FVector2D LastCursorPosition = GetLastCursorPos();
 	if ( LastCursorPosition != CurrentCursorPosition )
 	{
 		FPointerEvent MouseEvent(
+			CursorPointerIndex,
 			CurrentCursorPosition,
 			LastCursorPosition,
 			PressedMouseButtons,
@@ -3559,8 +3647,9 @@ bool FSlateApplication::OnRawMouseMove( const int32 X, const int32 Y )
 	if ( X != 0 || Y != 0 )
 	{
 		FPointerEvent MouseEvent(
-			GetCursorPos(), 	
-			LastCursorPosition,
+			CursorPointerIndex,
+			GetCursorPos(),
+			GetLastCursorPos(),
 			FVector2D( X, Y ), 
 			PressedMouseButtons,
 			PlatformApplication->GetModifierKeys()
@@ -3685,9 +3774,9 @@ bool FSlateApplication::ProcessMouseMoveEvent( FPointerEvent& MouseEvent )
 
 
 	FWidgetPath MouseCaptorPath;
-	if (MouseCaptor.IsValid())
+	if (MouseCaptor.HasCaptureForPointerIndex(MouseEvent.GetPointerIndex()))
 	{
-		MouseCaptorPath = MouseCaptor.ToWidgetPath(FWeakWidgetPath::EInterruptedPathHandling::ReturnInvalid);
+		MouseCaptorPath = MouseCaptor.ToWidgetPath(MouseEvent.GetPointerIndex(), FWeakWidgetPath::EInterruptedPathHandling::ReturnInvalid);
 	}
 
 	if (MouseCaptorPath.IsValid())
@@ -3790,8 +3879,7 @@ bool FSlateApplication::ProcessMouseMoveEvent( FPointerEvent& MouseEvent )
 
 	WidgetsUnderCursorLastEvent = FWeakWidgetPath( WidgetsUnderCursor );
 
-	LastCursorPosition = MouseEvent.GetScreenSpacePosition();
-
+	PointerIndexLastPositionMap.Add(MouseEvent.GetPointerIndex(), MouseEvent.GetScreenSpacePosition());
 	return bHandled;
 }
 
@@ -3976,6 +4064,7 @@ bool FSlateApplication::OnTouchGesture( EGestureEvent::Type GestureType, const F
 	);
 	
 	FPointerEvent MouseWheelEvent(
+		CursorPointerIndex,
 		CurrentCursorPosition,
 		CurrentCursorPosition,
 		PressedMouseButtons,
@@ -4007,11 +4096,13 @@ void FSlateApplication::ProcessTouchStartedEvent( const TSharedPtr< FGenericWind
 
 bool FSlateApplication::OnTouchMoved( const FVector2D& Location, int32 TouchIndex, int32 ControllerId )
 {
+	FVector2D LastLocation = PointerIndexLastPositionMap.Contains(TouchIndex) ? 
+		PointerIndexLastPositionMap[TouchIndex] : Location;
 	FPointerEvent PointerEvent(
 		ControllerId,
 		TouchIndex,
 		Location,
-		LastCursorPosition,
+		LastLocation,
 		true);
 	ProcessTouchMovedEvent(PointerEvent);
 
@@ -4402,11 +4493,13 @@ EDropEffect::Type FSlateApplication::OnDragEnter( const TSharedRef< SWindow >& W
 	DragIsHandled = false;
 
 	const FVector2D CurrentCursorPosition = GetCursorPos();
+	const FVector2D LastCursorPosition = GetLastCursorPos();
 
 	// Tell slate to enter drag and drop mode.
 	// Make a faux mouse event for slate, so we can initiate a drag and drop.
 	FDragDropEvent DragDropEvent(
 		FPointerEvent(
+		CursorPointerIndex,
 		CurrentCursorPosition,
 		LastCursorPosition,
 		PressedMouseButtons,
@@ -4416,8 +4509,6 @@ EDropEffect::Type FSlateApplication::OnDragEnter( const TSharedRef< SWindow >& W
 		DragDropOperation
 	);
 
-	LastCursorPosition = CurrentCursorPosition;
-
 	ProcessDragEnterEvent( Window, DragDropEvent );
 	return EDropEffect::None;
 }
@@ -4435,6 +4526,7 @@ bool FSlateApplication::ProcessDragEnterEvent( TSharedRef<SWindow> WindowEntered
 	FReply TriggerDragDropReply = FReply::Handled().BeginDragDrop( DragDropEvent.GetOperation().ToSharedRef() );
 	ProcessReply( WidgetsUnderCursor, TriggerDragDropReply, &WidgetsUnderCursor, &DragDropEvent );
 
+	PointerIndexLastPositionMap.Add(DragDropEvent.GetPointerIndex(), DragDropEvent.GetScreenSpacePosition());
 	return true;
 }
 
@@ -4447,10 +4539,12 @@ EDropEffect::Type FSlateApplication::OnDragOver( const TSharedPtr< FGenericWindo
 		bool MouseMoveHandled = true;
 		FVector2D CursorMovementDelta( 0, 0 );
 		const FVector2D CurrentCursorPosition = GetCursorPos();
+		const FVector2D LastCursorPosition = GetLastCursorPos();
 
 		if ( LastCursorPosition != CurrentCursorPosition )
 		{
 			FPointerEvent MouseEvent(
+				CursorPointerIndex,
 				CurrentCursorPosition,
 				LastCursorPosition,
 				PressedMouseButtons,
@@ -4491,8 +4585,9 @@ EDropEffect::Type FSlateApplication::OnDragDrop( const TSharedPtr< FGenericWindo
 	if ( IsDragDropping() )
 	{
 		FPointerEvent MouseEvent(
-			GetCursorPos(), 	
-			LastCursorPosition,
+			CursorPointerIndex,
+			GetCursorPos(),
+			GetLastCursorPos(),
 			PressedMouseButtons,
 			EKeys::LeftMouseButton,
 			0,
diff --git a/Engine/Source/Runtime/Slate/Private/Framework/Layout/ScrollyZoomy.cpp b/Engine/Source/Runtime/Slate/Private/Framework/Layout/ScrollyZoomy.cpp
--- a/Engine/Source/Runtime/Slate/Private/Framework/Layout/ScrollyZoomy.cpp
+++ b/Engine/Source/Runtime/Slate/Private/Framework/Layout/ScrollyZoomy.cpp
@@ -118,7 +118,7 @@ FReply FScrollyZoomy::OnMouseMove( const TSharedRef<SWidget> MyWidget, IScrollab
 			FReply Reply = FReply::Handled();
 
 			// Capture the mouse if we need to
-			if (FSlateApplication::Get().GetMouseCaptor() != MyWidget)
+			if (MyWidget->HasMouseCapture() == false)
 			{
 				Reply.CaptureMouse( MyWidget ).UseHighPrecisionMouseMovement( MyWidget );
 				SoftwareCursorPosition = MyGeometry.AbsoluteToLocal( MouseEvent.GetScreenSpacePosition() );
@@ -141,7 +141,7 @@ FReply FScrollyZoomy::OnMouseMove( const TSharedRef<SWidget> MyWidget, IScrollab
 
 void FScrollyZoomy::OnMouseLeave( const TSharedRef<SWidget> MyWidget, const FPointerEvent& MouseEvent )
 {
-	if (FSlateApplication::Get().GetMouseCaptor() != MyWidget)
+	if (MyWidget->HasMouseCapture() == false)
 	{
 		// No longer scrolling (unless we have mouse capture)
 		AmountScrolledWhileRightMouseDown = 0;
diff --git a/Engine/Source/Runtime/Slate/Private/Widgets/Layout/SScrollBox.cpp b/Engine/Source/Runtime/Slate/Private/Widgets/Layout/SScrollBox.cpp
--- a/Engine/Source/Runtime/Slate/Private/Widgets/Layout/SScrollBox.cpp
+++ b/Engine/Source/Runtime/Slate/Private/Widgets/Layout/SScrollBox.cpp
@@ -477,7 +477,7 @@ FReply SScrollBox::OnMouseMove( const FGeometry& MyGeometry, const FPointerEvent
 			FReply Reply = FReply::Handled();
 
 			// Capture the mouse if we need to
-			if( FSlateApplication::Get().GetMouseCaptor().Get() != this )
+			if(this->HasMouseCapture() == false)
 			{
 				Reply.CaptureMouse( AsShared() ).UseHighPrecisionMouseMovement( AsShared() );
 				SoftwareCursorPosition = MyGeometry.AbsoluteToLocal( MouseEvent.GetScreenSpacePosition() );
@@ -499,7 +499,7 @@ FReply SScrollBox::OnMouseMove( const FGeometry& MyGeometry, const FPointerEvent
 
 void SScrollBox::OnMouseLeave( const FPointerEvent& MouseEvent )
 {
-	if( FSlateApplication::Get().GetMouseCaptor().Get() != this )
+	if(this->HasMouseCapture() == false)
 	{
 		// No longer scrolling (unless we have mouse capture)
 		AmountScrolledWhileRightMouseDown = 0;
diff --git a/Engine/Source/Runtime/Slate/Private/Widgets/Layout/SSplitter.cpp b/Engine/Source/Runtime/Slate/Private/Widgets/Layout/SSplitter.cpp
--- a/Engine/Source/Runtime/Slate/Private/Widgets/Layout/SSplitter.cpp
+++ b/Engine/Source/Runtime/Slate/Private/Widgets/Layout/SSplitter.cpp
@@ -777,7 +777,7 @@ FReply SSplitter2x2::OnMouseMove( const FGeometry& MyGeometry, const FPointerEve
 {
 	const FVector2D LocalMousePos = MyGeometry.AbsoluteToLocal( MouseEvent.GetScreenSpacePosition() );
 
-	if ( bIsResizing && FSlateApplication::Get().GetMouseCaptor() == AsShared() )
+	if ( bIsResizing && this->HasMouseCapture() )
 	{	
 		TArray<FLayoutGeometry> LayoutChildren = ArrangeChildrenForLayout(MyGeometry);
 		ResizeChildren( MyGeometry, LayoutChildren, LocalMousePos );
diff --git a/Engine/Source/Runtime/Slate/Private/Widgets/SViewport.cpp b/Engine/Source/Runtime/Slate/Private/Widgets/SViewport.cpp
--- a/Engine/Source/Runtime/Slate/Private/Widgets/SViewport.cpp
+++ b/Engine/Source/Runtime/Slate/Private/Widgets/SViewport.cpp
@@ -79,6 +79,7 @@ int32 SViewport::OnPaint( const FPaintArgs& Args, const FGeometry& AllottedGeome
 		// @todo Slate: why are we calling OnCursorQuery in here?
 		FCursorReply Reply = ViewportInterfacePin->OnCursorQuery( AllottedGeometry,
 			FPointerEvent(
+				FSlateApplicationBase::CursorPointerIndex,
 				CursorPosScreenSpace,
 				CursorPosScreenSpace,
 				FVector2D::ZeroVector,
diff --git a/Engine/Source/Runtime/Slate/Private/Widgets/Text/SInlineEditableTextBlock.cpp b/Engine/Source/Runtime/Slate/Private/Widgets/Text/SInlineEditableTextBlock.cpp
--- a/Engine/Source/Runtime/Slate/Private/Widgets/Text/SInlineEditableTextBlock.cpp
+++ b/Engine/Source/Runtime/Slate/Private/Widgets/Text/SInlineEditableTextBlock.cpp
@@ -73,7 +73,7 @@ bool SInlineEditableTextBlock::SupportsKeyboardFocus() const
 
 void SInlineEditableTextBlock::EnterEditingMode()
 {
-	if(!bIsReadOnly.Get() && !FSlateApplication::Get().GetMouseCaptor().IsValid() )
+	if(!bIsReadOnly.Get() && FSlateApplication::Get().HasAnyMouseCaptor() == false)
 	{
 		if(TextBlock->GetVisibility() == EVisibility::Visible)
 		{
diff --git a/Engine/Source/Runtime/Slate/Private/Widgets/Views/STableViewBase.cpp b/Engine/Source/Runtime/Slate/Private/Widgets/Views/STableViewBase.cpp
--- a/Engine/Source/Runtime/Slate/Private/Widgets/Views/STableViewBase.cpp
+++ b/Engine/Source/Runtime/Slate/Private/Widgets/Views/STableViewBase.cpp
@@ -364,7 +364,7 @@ FReply STableViewBase::OnMouseMove( const FGeometry& MyGeometry, const FPointerE
 
 			// The mouse moved enough that we're now dragging the view. Capture the mouse
 			// so the user does not have to stay within the bounds of the list while dragging.
-			if( FSlateApplication::Get().GetMouseCaptor().Get() != this )
+			if(this->HasMouseCapture() == false)
 			{
 				Reply.CaptureMouse( AsShared() ).UseHighPrecisionMouseMovement( AsShared() );
 				SoftwareCursorPosition = MyGeometry.AbsoluteToLocal( MouseEvent.GetScreenSpacePosition() );
@@ -388,7 +388,7 @@ FReply STableViewBase::OnMouseMove( const FGeometry& MyGeometry, const FPointerE
 void STableViewBase::OnMouseLeave( const FPointerEvent& MouseEvent )
 {
 	bStartedTouchInteraction = false;
-	if( FSlateApplication::Get().GetMouseCaptor().Get() != this )
+	if(this->HasMouseCapture() == false)
 	{
 		// No longer scrolling (unless we have mouse capture)
 		AmountScrolledWhileRightMouseDown = 0;
diff --git a/Engine/Source/Runtime/Slate/Public/Framework/Application/SlateApplication.h b/Engine/Source/Runtime/Slate/Public/Framework/Application/SlateApplication.h
--- a/Engine/Source/Runtime/Slate/Public/Framework/Application/SlateApplication.h
+++ b/Engine/Source/Runtime/Slate/Public/Framework/Application/SlateApplication.h
@@ -628,6 +628,8 @@ public:
 
 protected:
 
+	virtual bool HasMouseCapture(const TSharedPtr<const SWidget> Widget) const override;
+
 	/** 
 	 * Ticks a slate window and all of its children
 	 *
@@ -937,11 +939,13 @@ public:
 
 	virtual FVector2D GetCursorPos( ) const override;
 
+	virtual FVector2D GetLastCursorPos( ) const override;
+
 	virtual FVector2D GetCursorSize( ) const override;
 	
 	virtual TSharedPtr<SWidget> GetKeyboardFocusedWidget( ) const override;
 
-	virtual TSharedPtr<SWidget> GetMouseCaptor( ) const override;
+	virtual bool HasAnyMouseCaptor( ) const override;
 
 	virtual FSlateRect GetPreferredWorkArea( ) const override;
 
@@ -1073,6 +1077,7 @@ private:
 	FSlateApplication();
 
 private:
+
 	/**
 	 * Will be invoked when the size of the geometry of the virtual
 	 * desktop changes (e.g. resolution change or monitors re-arranged)
@@ -1120,45 +1125,67 @@ private:
 	class MouseCaptorHelper
 	{
 	public:
-		/** Returns true if the weak path is valid (i.e. more than zero widgets) */
-		bool IsValid() const;
+		/**
+		 * Returns whether or not there are any active pointer captures.
+		 */
+		bool HasCapture() const;
+
+		/**
+		 * Returns whether or not the particular PointerIndex has capture.
+		 */
+		bool HasCaptureForPointerIndex(uint32 PointerIndex) const;
 
 		/**
-		 * Sets a new mouse captor widget, invalidating the previous one if any and calling
+		 * Sets a new mouse captor widget for a specific pointer index, invalidating the previous one if any and calling
 		 * its OnMouseCaptureLost() handler.
 		 *
-		 * @param EventPath	The path to the event.
-		 * @param Widget	The widget that wants to capture the mouse.
+		 * @param PointerIndex	The index of the pointer which initiated the capture.
+		 * @param EventPath		The path to the event.
+		 * @param Widget		The widget that wants to capture the mouse.
 		 */
-		void SetMouseCaptor( const FWidgetPath& EventPath, TSharedPtr< SWidget > Widget );
+		void SetMouseCaptor(uint32 PointerIndex, const FWidgetPath& EventPath, TSharedPtr< SWidget > Widget );
+
+		/** Invalidates all current mouse captors. Calls OnMouseCaptureLost() on the current mouse captor if one exists */
+		void InvalidateCaptureForAllPointers();
 
-		/** Invalidates the current mouse captor. Calls OnMouseCaptureLost() on the current mouse captor if one exists */
-		void Invalidate();
+		/** Invalidates a specific mouse captor. Calls OnMouseCaptureLost() on the specific mouse captor if one exists */
+		void InvalidateCaptureForPointer(uint32 PointIndex);
 
 		/**
-		 * Retrieves a resolved FWidgetPath from the current, internal weak path, if possible.
+		 * Retrieves a resolved FWidgetPath for a specific pointer index, if possible.
 		 * If the path is invalid or truncated (i.e. the widget is no longer relevant) then the current mouse captor's
 		 * OnMouseCaptureLost() handler is called and the path is invalidated.
 		 *
+		 * @param PointerIndex				The index of the pointer which has capture.
 		 * @param InterruptedPathHandling	How to handled incomplete paths. "Truncate" will return a partial path, "ReturnInvalid" will return an empty path.
 		 */
-		FWidgetPath ToWidgetPath( FWeakWidgetPath::EInterruptedPathHandling::Type InterruptedPathHandling = FWeakWidgetPath::EInterruptedPathHandling::Truncate );
+		FWidgetPath ToWidgetPath(uint32 PointerIndex, FWeakWidgetPath::EInterruptedPathHandling::Type InterruptedPathHandling = FWeakWidgetPath::EInterruptedPathHandling::Truncate );
 
-		/** Retrieves the currently set weak path to the current mouse captor */
-		FWeakWidgetPath ToWeakPath() const;
+		/**
+		 * Retrieves an array of the resolved widget paths for any active captures.
+		 */
+		TArray<FWidgetPath> ToWidgetPaths();
 
-		/* Walks the weak path and retrieves the widget that is set as the current mouse captor */
-		TSharedPtr< SWidget > ToSharedWidget() const;
+		/** Retrieves the weak path for a current mouse captor with a specific pointer index */
+		FWeakWidgetPath ToWeakPath(uint32 PointerIndex) const;
+
+		/* Walks the weak path and retrieves the widget that is set as the current mouse captor with a specific pointer index */
+		TSharedPtr< SWidget > ToSharedWidget(uint32 PointerIndex) const;
+
+		/*
+		 * Retrieves an array of shared widget pointers for the active mouse captures.
+		 */
+		TArray<TSharedRef<SWidget>> ToSharedWidgets() const;
 
-		/* Walks the weak path and retrieves the window that the current mouse captor widget belongs to */
-		TSharedPtr< SWidget > ToSharedWindow();
+		/* Walks the weak path and retrieves the window for the widget belonging to the mouse captor with the specified pointer index */
+		TSharedPtr< SWidget > ToSharedWindow(uint32 PointerIndex);
 
 	protected:
-		/** If there is a current mouse captor, will call its OnMouseCaptureLost() handler */
-		void InformCurrentCaptorOfCaptureLoss() const;
+		/** Call the OnMouseCaptureLost() handler for the widget captured by the specific pointer index */
+		void InformCurrentCaptorOfCaptureLoss(uint32 PointerIndex) const;
 
-		/** The weak pointer path to the current mouse captor */
-		FWeakWidgetPath MouseCaptorWeakPath;
+		/** A map of pointer idices to weak widget paths for the active mouse captures */
+		TMap<uint32, FWeakWidgetPath> PointerIndexToMouseCaptorWeakPathMap;
 	};
 	/** The current mouse captor for the application, if any. */
 	MouseCaptorHelper MouseCaptor;
@@ -1334,7 +1361,7 @@ private:
 	/** When an drag and drop is happening, we keep track of whether slate knew what to do with the payload on last mouse move */
 	bool DragIsHandled;
 
-	FVector2D LastCursorPosition;
+	TMap<uint32, FVector2D> PointerIndexLastPositionMap;
 
 	/**
 	 * Virtual keyboard text field
diff --git a/Engine/Source/Runtime/SlateCore/Private/Application/SlateApplicationBase.cpp b/Engine/Source/Runtime/SlateCore/Private/Application/SlateApplicationBase.cpp
--- a/Engine/Source/Runtime/SlateCore/Private/Application/SlateApplicationBase.cpp
+++ b/Engine/Source/Runtime/SlateCore/Private/Application/SlateApplicationBase.cpp
@@ -8,3 +8,4 @@
 
 TSharedPtr<FSlateApplicationBase> FSlateApplicationBase::CurrentBaseApplication = nullptr;
 TSharedPtr<GenericApplication> FSlateApplicationBase::PlatformApplication = nullptr;
+const uint32 FSlateApplicationBase::CursorPointerIndex = 1000;
diff --git a/Engine/Source/Runtime/SlateCore/Private/Widgets/SWidget.cpp b/Engine/Source/Runtime/SlateCore/Private/Widgets/SWidget.cpp
--- a/Engine/Source/Runtime/SlateCore/Private/Widgets/SWidget.cpp
+++ b/Engine/Source/Runtime/SlateCore/Private/Widgets/SWidget.cpp
@@ -320,7 +320,7 @@ bool SWidget::HasFocusedDescendants() const
 
 bool SWidget::HasMouseCapture() const
 {
-	return (FSlateApplicationBase::Get().GetMouseCaptor().Get() == this);
+	return FSlateApplicationBase::Get().HasMouseCapture(SharedThis(this));
 }
 
 
diff --git a/Engine/Source/Runtime/SlateCore/Public/Application/SlateApplicationBase.h b/Engine/Source/Runtime/SlateCore/Public/Application/SlateApplicationBase.h
--- a/Engine/Source/Runtime/SlateCore/Public/Application/SlateApplicationBase.h
+++ b/Engine/Source/Runtime/SlateCore/Public/Application/SlateApplicationBase.h
@@ -25,6 +25,7 @@ public:
  */
 class SLATECORE_API FSlateApplicationBase
 {
+	friend class SWidget;
 public:
 
 	/**
@@ -100,12 +101,19 @@ public:
 	virtual const double GetCurrentTime( ) const = 0;
 
 	/**
-	 * Gets the last known position of the cursor.
+	 * Gets the current position of the cursor.
 	 *
 	 * @return Cursor position.
 	 */
 	virtual FVector2D GetCursorPos( ) const = 0;
 
+	/**
+	* Gets the last known position of the cursor.
+	*
+	* @return Cursor position.
+	*/
+	virtual FVector2D GetLastCursorPos( ) const = 0;
+
 	/**
 	 * Gets the size of the cursor..
 	 *
@@ -128,11 +136,11 @@ public:
 	virtual TSharedPtr< SWidget > GetKeyboardFocusedWidget( ) const = 0;
 
 	/**
-	 * Gets the Widget that currently captures the mouse.
+	 * Gets whether or not a widget has captured the mouse.
 	 *
-	 * @return The captor widget, or nullptr if no widget captured the mouse.
+	 * @return True if one or more widgets have capture, otherwise false.
 	 */
-	virtual TSharedPtr< SWidget > GetMouseCaptor( ) const = 0;
+	virtual bool HasAnyMouseCaptor( ) const = 0;
 
 	/**
 	 * Gets the platform application.
@@ -241,6 +249,7 @@ public:
 	virtual bool SetKeyboardFocus( const FWidgetPath& InFocusPath, const EKeyboardFocusCause::Type InCause ) = 0;
 
 public:
+	const static uint32 CursorPointerIndex;
 
 	/**
 	 * Returns the current instance of the application. The application should have been initialized before
@@ -264,6 +273,15 @@ public:
 		return CurrentBaseApplication.IsValid();
 	}
 
+protected:
+
+	/**
+	* Gets whether or not a particular widget has mouse capture.
+	*
+	* @return True if the widget has mouse capture, otherwise false.
+	*/
+	virtual bool HasMouseCapture(const TSharedPtr<const SWidget> Widget) const = 0;
+
 protected:
 
 	// Holds the Slate renderer used to render this application.
diff --git a/Engine/Source/Runtime/SlateCore/Public/Input/Events.h b/Engine/Source/Runtime/SlateCore/Public/Input/Events.h
--- a/Engine/Source/Runtime/SlateCore/Public/Input/Events.h
+++ b/Engine/Source/Runtime/SlateCore/Public/Input/Events.h
@@ -447,6 +447,7 @@ public:
 
 	/** Events are immutable once constructed. */
 	FPointerEvent(
+		uint32 InPointerIndex,
 		const FVector2D& InScreenSpacePosition,
 		const FVector2D& InLastScreenSpacePosition,
 		const TSet<FKey>& InPressedButtons,
@@ -461,7 +462,7 @@ public:
 		, PressedButtons(InPressedButtons)
 		, EffectingButton(InEffectingButton)
 		, UserIndex(0)
-		, PointerIndex(0)
+		, PointerIndex(InPointerIndex)
 		, TouchpadIndex(0)
 		, bIsTouchEvent(false)
 		, GestureType(EGestureEvent::None)
@@ -470,6 +471,7 @@ public:
 
 	/** A constructor for raw mouse events */
 	FPointerEvent(
+		uint32 InPointerIndex,
 		const FVector2D& InScreenSpacePosition,
 		const FVector2D& InLastScreenSpacePosition,
 		const FVector2D& InDelta,
@@ -482,7 +484,7 @@ public:
 		, CursorDelta(InDelta)
 		, PressedButtons(InPressedButtons)
 		, UserIndex(0)
-		, PointerIndex(0)
+		, PointerIndex(InPointerIndex)
 		, TouchpadIndex(0)
 		, bIsTouchEvent(false)
 		, GestureType(EGestureEvent::None)
