- Added IVirtualKeyboardEntry
- Virtual keyboard entry automatically supported for editable text
- Spinbox defaults to numeric keyboard and filters entry properly
- SVirtualKeyboard should function as before for any pre-existing use cases

diff --git a/Engine/Source/Runtime/Core/Public/GenericPlatform/GenericPlatformMisc.h b/Engine/Source/Runtime/Core/Public/GenericPlatform/GenericPlatformMisc.h
--- a/Engine/Source/Runtime/Core/Public/GenericPlatform/GenericPlatformMisc.h
+++ b/Engine/Source/Runtime/Core/Public/GenericPlatform/GenericPlatformMisc.h
@@ -908,6 +908,13 @@ public:
 
 	/** @return Memory representing a true type or open type font provided by the platform as a default font for unreal to consume; empty array if the default font failed to load. */
 	static TArray<uint8> GetSystemFontBytes();
+	/**
+	* Returns whether the platform wants to use a touch screen for a virtual keyboard.
+	*/
+	static bool GetRequiresVirtualKeyboard()
+	{
+		return PLATFORM_HAS_TOUCH_MAIN_SCREEN;
+	}
 
 	/**
 	 * Returns whether the given platform feature is currently available (for instance, Metal is only available in IOS8 and with A7 devices)
diff --git a/Engine/Source/Runtime/Launch/Private/Android/AndroidJNI.cpp b/Engine/Source/Runtime/Launch/Private/Android/AndroidJNI.cpp
--- a/Engine/Source/Runtime/Launch/Private/Android/AndroidJNI.cpp
+++ b/Engine/Source/Runtime/Launch/Private/Android/AndroidJNI.cpp
@@ -15,7 +15,7 @@ JavaVM* GJavaVM;
 jobject GJavaGlobalThis = NULL;
 
 // Pointer to target widget for virtual keyboard contents
-static SVirtualKeyboardEntry *VirtualKeyboardWidget = NULL;
+static IVirtualKeyboardEntry *VirtualKeyboardWidget = NULL;
 
 extern FString GFilePathBase;
 extern FString GFontPathBase;
@@ -194,7 +194,7 @@ void AndroidThunkCpp_ShowConsoleWindow()
 	}
 }
 
-void AndroidThunkCpp_ShowVirtualKeyboardInput(TSharedPtr<SVirtualKeyboardEntry> TextWidget, int32 InputType, const FString& Label, const FString& Contents)
+void AndroidThunkCpp_ShowVirtualKeyboardInput(TSharedPtr<IVirtualKeyboardEntry> TextWidget, int32 InputType, const FString& Label, const FString& Contents)
 {
 	if (JNIEnv* Env = GetJavaEnv())
 	{
@@ -219,7 +219,7 @@ extern "C" void Java_com_epicgames_ue4_GameActivity_nativeVirtualKeyboardResult(
 		if (VirtualKeyboardWidget != NULL)
 		{
 			const char* javaChars = jenv->GetStringUTFChars(contents, 0);
-			VirtualKeyboardWidget->SetText(FText::FromString(FString(UTF8_TO_TCHAR(javaChars))));
+			VirtualKeyboardWidget->SetTextFromVirtualKeyboard(FText::FromString(FString(UTF8_TO_TCHAR(javaChars))));
 
 			//Release the string
 			jenv->ReleaseStringUTFChars(contents, javaChars);
diff --git a/Engine/Source/Runtime/Slate/Private/Framework/Application/SlateApplication.cpp b/Engine/Source/Runtime/Slate/Private/Framework/Application/SlateApplication.cpp
--- a/Engine/Source/Runtime/Slate/Private/Framework/Application/SlateApplication.cpp
+++ b/Engine/Source/Runtime/Slate/Private/Framework/Application/SlateApplication.cpp
@@ -2716,14 +2716,14 @@ const FSlateBrush* FSlateApplication::GetAppIcon() const
 }
 
 
-void FSlateApplication::ShowKeyboard( bool bShow, TSharedPtr<SVirtualKeyboardEntry> TextEntryWidget )
+void FSlateApplication::ShowVirtualKeyboard( bool bShow, TSharedPtr<IVirtualKeyboardEntry> TextEntryWidget )
 {
-	if(SlateTextField == NULL)
+	if(SlateTextField == nullptr)
 	{
 		SlateTextField = new FPlatformTextField();
 	}
 
-	SlateTextField->ShowKeyboard(bShow, TextEntryWidget);
+	SlateTextField->ShowVirtualKeyboard(bShow, TextEntryWidget);
 }
 
 FSlateRect FSlateApplication::GetPreferredWorkArea() const
diff --git a/Engine/Source/Runtime/Slate/Private/Framework/Text/Android/AndroidPlatformTextField.cpp b/Engine/Source/Runtime/Slate/Private/Framework/Text/Android/AndroidPlatformTextField.cpp
--- a/Engine/Source/Runtime/Slate/Private/Framework/Text/Android/AndroidPlatformTextField.cpp
+++ b/Engine/Source/Runtime/Slate/Private/Framework/Text/Android/AndroidPlatformTextField.cpp
@@ -17,13 +17,13 @@
 #define TYPE_TEXT_FLAG_NO_SUGGESTIONS		0x00080000
 
 
-void FAndroidPlatformTextField::ShowKeyboard(bool bShow, TSharedPtr<SVirtualKeyboardEntry> TextEntryWidget)
+void FAndroidPlatformTextField::ShowVirtualKeyboard(bool bShow, TSharedPtr<IVirtualKeyboardEntry> TextEntryWidget)
 {
 	if(bShow)
 	{
 		// Set the EditBox inputType based on keyboard type
 		int32 InputType;
-		switch (TextEntryWidget->GetKeyboardType())
+		switch (TextEntryWidget->GetVirtualKeyboardType())
 		{
 		case EKeyboardType::Keyboard_Number:
 			InputType = TYPE_CLASS_NUMBER | TYPE_TEXT_VARIATION_NORMAL;
@@ -47,7 +47,7 @@ void FAndroidPlatformTextField::ShowKeyboard(bool bShow, TSharedPtr<SVirtualKeyb
 		InputType |= TYPE_TEXT_FLAG_NO_SUGGESTIONS;
 
 		// Show alert for input
-		extern void AndroidThunkCpp_ShowVirtualKeyboardInput(TSharedPtr<SVirtualKeyboardEntry>, int32, const FString&, const FString&);
+		extern void AndroidThunkCpp_ShowVirtualKeyboardInput(TSharedPtr<IVirtualKeyboardEntry>, int32, const FString&, const FString&);
 		AndroidThunkCpp_ShowVirtualKeyboardInput(TextEntryWidget, InputType, TextEntryWidget->GetHintText().ToString(), TextEntryWidget->GetText().ToString());
 	}
 }
diff --git a/Engine/Source/Runtime/Slate/Private/Framework/Text/IOS/IOSPlatformTextField.cpp b/Engine/Source/Runtime/Slate/Private/Framework/Text/IOS/IOSPlatformTextField.cpp
--- a/Engine/Source/Runtime/Slate/Private/Framework/Text/IOS/IOSPlatformTextField.cpp
+++ b/Engine/Source/Runtime/Slate/Private/Framework/Text/IOS/IOSPlatformTextField.cpp
@@ -4,22 +4,22 @@
 
 
 FIOSPlatformTextField::FIOSPlatformTextField()
-	: TextField( NULL )
+	: TextField( nullptr )
 {
 }
 
 FIOSPlatformTextField::~FIOSPlatformTextField()
 {
-	if(TextField != NULL)
+	if(TextField != nullptr)
 	{
 		[TextField release];
-		TextField = NULL;
+		TextField = nullptr;
 	}
 }
 
-void FIOSPlatformTextField::ShowKeyboard(bool bShow, TSharedPtr<SVirtualKeyboardEntry> TextEntryWidget)
+void FIOSPlatformTextField::ShowVirtualKeyboard(bool bShow, TSharedPtr<IVirtualKeyboardEntry> TextEntryWidget)
 {
-	if(TextField == NULL)
+	if(TextField == nullptr)
 	{
 		TextField = [SlateTextField alloc];
 	}
@@ -29,14 +29,14 @@ void FIOSPlatformTextField::ShowKeyboard(bool bShow, TSharedPtr<SVirtualKeyboard
 		// these functions must be run on the main thread
 		dispatch_async(dispatch_get_main_queue(),^ {
 			[TextField show: TextEntryWidget];
-		} );
+		});
 	}
 }
 
 
 @implementation SlateTextField
 
--(void)show:(TSharedPtr<SVirtualKeyboardEntry>)InTextWidget
+-(void)show:(TSharedPtr<IVirtualKeyboardEntry>)InTextWidget
 {
 	TextWidget = InTextWidget;
 
@@ -47,7 +47,7 @@ void FIOSPlatformTextField::ShowKeyboard(bool bShow, TSharedPtr<SVirtualKeyboard
 													otherButtonTitles:NSLocalizedString(@"OK", nil), nil];
 
 	// give the UIAlertView a style so a UITextField is created
-	switch(TextWidget->GetKeyboardType())
+	switch(TextWidget->GetVirtualKeyboardType())
 	{
 		case EKeyboardType::Keyboard_Password:
 			AlertView.alertViewStyle = UIAlertViewStyleSecureTextInput;
@@ -61,31 +61,31 @@ void FIOSPlatformTextField::ShowKeyboard(bool bShow, TSharedPtr<SVirtualKeyboard
 			break;
 	}
 
-	UITextField* TextField = [AlertView textFieldAtIndex:0];
-	TextField.clearsOnBeginEditing = NO;
-	TextField.clearsOnInsertion = NO;
-	TextField.autocorrectionType = UITextAutocorrectionTypeNo;
-	TextField.autocapitalizationType = UITextAutocapitalizationTypeNone;
-	TextField.text = [NSString stringWithFString:TextWidget->GetText().ToString()];
-	TextField.placeholder = [NSString stringWithFString:TextWidget->GetText().ToString()];
+	UITextField* AlertTextField = [AlertView textFieldAtIndex: 0];
+	AlertTextField.clearsOnBeginEditing = NO;
+	AlertTextField.clearsOnInsertion = NO;
+	AlertTextField.autocorrectionType = UITextAutocorrectionTypeNo;
+	AlertTextField.autocapitalizationType = UITextAutocapitalizationTypeNone;
+	AlertTextField.text = [NSString stringWithFString : TextWidget->GetText().ToString()];
+	AlertTextField.placeholder = [NSString stringWithFString : TextWidget->GetHintText().ToString()];
 
 	// set up the keyboard styles not supported in the AlertViewStyle styles
-	switch(TextWidget->GetKeyboardType())
+	switch (TextWidget->GetVirtualKeyboardType())
 	{
-		case EKeyboardType::Keyboard_Email:
-			TextField.keyboardType = UIKeyboardTypeEmailAddress;
-			break;
-		case EKeyboardType::Keyboard_Number:
-			TextField.keyboardType = UIKeyboardTypeNumberPad;
-			break;
-		case EKeyboardType::Keyboard_Web:
-			TextField.keyboardType = UIKeyboardTypeURL;
-			break;
-		case EKeyboardType::Keyboard_Default:
-		case EKeyboardType::Keyboard_Password:
-		default:
-			// nothing to do, UIAlertView style handles these keyboard types
-			break;
+	case EKeyboardType::Keyboard_Email:
+		AlertTextField.keyboardType = UIKeyboardTypeEmailAddress;
+		break;
+	case EKeyboardType::Keyboard_Number:
+		AlertTextField.keyboardType = UIKeyboardTypeDecimalPad;
+		break;
+	case EKeyboardType::Keyboard_Web:
+		AlertTextField.keyboardType = UIKeyboardTypeURL;
+		break;
+	case EKeyboardType::Keyboard_Default:
+	case EKeyboardType::Keyboard_Password:
+	default:
+		// nothing to do, UIAlertView style handles these keyboard types
+		break;
 	}
 
 	[AlertView show];
@@ -98,8 +98,8 @@ void FIOSPlatformTextField::ShowKeyboard(bool bShow, TSharedPtr<SVirtualKeyboard
 	// index 1 is the OK button
 	if(buttonIndex == 1)
 	{
-		UITextField* TextField = [alertView textFieldAtIndex:0];
-		TextWidget->SetText(FText::FromString(TextField.text));
+		UITextField* AlertTextField = [alertView textFieldAtIndex: 0];
+		TextWidget->SetTextFromVirtualKeyboard(FText::FromString(AlertTextField.text));
 	}
 }
 
diff --git a/Engine/Source/Runtime/Slate/Private/Widgets/Input/SEditableText.cpp b/Engine/Source/Runtime/Slate/Private/Widgets/Input/SEditableText.cpp
--- a/Engine/Source/Runtime/Slate/Private/Widgets/Input/SEditableText.cpp
+++ b/Engine/Source/Runtime/Slate/Private/Widgets/Input/SEditableText.cpp
@@ -20,7 +20,8 @@ SEditableText::SEditableText()
 	  bHasDragSelectedSinceMouseDown( false ),
 	  CurrentUndoLevel( INDEX_NONE ),
 	  bIsChangingText( false ),
-	  UICommandList( new FUICommandList() )
+	  UICommandList( new FUICommandList() ),
+	  bTextChangedByVirtualKeyboard(false)
 {
 	// Setup springs
 	FFloatSpring1D::FSpringConfig SpringConfig;
@@ -79,6 +80,7 @@ void SEditableText::Construct( const FArguments& InArgs )
 	OnTextCommitted = InArgs._OnTextCommitted;
 	MinDesiredWidth = InArgs._MinDesiredWidth;
 	SelectAllTextOnCommit = InArgs._SelectAllTextOnCommit;
+	VirtualKeyboardType = IsPassword.Get() ? EKeyboardType::Keyboard_Password : InArgs._VirtualKeyboardType;
 
 	// Map UI commands to delegates which are called when the command should be executed
 	UICommandList->MapAction( FGenericCommands::Get().Undo,
@@ -121,7 +123,6 @@ void SEditableText::Construct( const FArguments& InArgs )
 	}
 }
 
-
 /**
  * Sets the text string currently being edited 
  *
@@ -129,34 +130,58 @@ void SEditableText::Construct( const FArguments& InArgs )
  */
 void SEditableText::SetText( const TAttribute< FText >& InNewText )
 {
-	const bool HasTextChanged = HasKeyboardFocus() ? !InNewText.Get().EqualTo( EditedText ) : !InNewText.Get().EqualTo( Text.Get() );
-
 	// NOTE: This will unbind any getter that is currently assigned to the Text attribute!  You should only be calling
 	//       SetText() if you know what you're doing.
 	Text = InNewText;
 
-	// Update the edited text, if the user isn't already editing it
-	if( HasKeyboardFocus() && !InNewText.Get().EqualTo(EditedText))
+	const bool bHasTextChanged = HasKeyboardFocus() ? !InNewText.Get().EqualTo(EditedText) : !InNewText.Get().EqualTo(Text.Get());
+
+	if (bHasTextChanged)
 	{
-		ClearSelection();
-		EditedText = FText();
+		// Update the edited text, if the user isn't already editing it
+		if (HasKeyboardFocus())
+		{
+			// Update the edited text, if the user isn't already editing it
+			ClearSelection();
+			EditedText = FText();
 
-		// Load the new text
-		LoadText();
+			// Load the new text
+			LoadText();
+
+			// Move the cursor to the end of the string
+			SetCaretPosition(EditedText.ToString().Len());
 
-		// Move the cursor to the end of the string
-		SetCaretPosition( EditedText.ToString().Len() );
+			// Let outsiders know that the text being edited has been changed
+			OnTextChanged.ExecuteIfBound(EditedText);
+		}
+		else
+		{
+			// Let outsiders know that the text content has been changed
+			OnTextChanged.ExecuteIfBound(Text.Get());
+		}
+	}
+}
 
+void SEditableText::SetTextFromVirtualKeyboard(const FText& InNewText)
+{
+	// Only set the text if the text attribute doesn't have a getter binding (otherwise it would be blown away).
+	// If it is bound, we'll assume that OnTextCommitted will handle the update.
+	if (!Text.IsBound())
+	{
+		Text.Set(InNewText);
 	}
 
-	if( HasTextChanged )
+	if (!InNewText.EqualTo(EditedText))
 	{
-		// Let outsiders know that the text content has been changed
-		OnTextChanged.ExecuteIfBound( HasKeyboardFocus() ? EditedText : Text.Get() );
+		EditedText = InNewText;
+
+		// This method is called from the main thread (i.e. not the game thread) of the device with the virtual keyboard
+		// This causes the app to crash on those devices, so we're using polling here to ensure delegates are
+		// fired on the game thread in Tick.
+		bTextChangedByVirtualKeyboard = true;
 	}
 }
 
-
 void SEditableText::RestoreOriginalText()
 {
 	if( HasTextChangedFromOriginal() )
@@ -170,7 +195,6 @@ void SEditableText::RestoreOriginalText()
 	}
 }
 
-
 bool SEditableText::HasTextChangedFromOriginal() const
 {
 	return ( !IsReadOnly.Get() && !EditedText.EqualTo(OriginalText)  );
@@ -195,6 +219,14 @@ void SEditableText::Tick( const FGeometry& AllottedGeometry, const double InCurr
 	// Call parent implementation.
 	SLeafWidget::Tick( AllottedGeometry, InCurrentTime, InDeltaTime );
 
+	// Poll to see if the text has been updated by a virtual keyboard
+	if (bTextChangedByVirtualKeyboard)
+	{
+		//OnTextChanged.ExecuteIfBound(EditedText);
+		OnEnter();
+		bTextChangedByVirtualKeyboard = false;
+	}
+
 	if(TextInputMethodChangeNotifier.IsValid() && TextInputMethodContext.IsValid() && TextInputMethodContext->CachedGeometry != AllottedGeometry)
 	{
 		TextInputMethodContext->CachedGeometry = AllottedGeometry;
@@ -1467,13 +1499,12 @@ bool SEditableText::SupportsKeyboardFocus() const
 	return true;
 }
 
-
 FReply SEditableText::OnKeyboardFocusReceived( const FGeometry& MyGeometry, const FKeyboardFocusEvent& InKeyboardFocusEvent )
 {
 	// Skip the focus received code if it's due to the context menu closing
 	if ( !ContextMenuWindow.IsValid() )
 	{
-		// Don't unselect text with focus is set because another widget lost focus, as that may be in response to a
+		// Don't unselect text when focus is set because another widget lost focus, as that may be in response to a
 		// dismissed context menu where the user clicked an option to select text.
 		if( InKeyboardFocusEvent.GetCause() != EKeyboardFocusCause::OtherWidgetLostFocus )
 		{
@@ -1512,10 +1543,19 @@ FReply SEditableText::OnKeyboardFocusReceived( const FGeometry& MyGeometry, cons
 			}
 		}
 
-		ITextInputMethodSystem* const TextInputMethodSystem = FSlateApplication::Get().GetTextInputMethodSystem();
-		if(TextInputMethodSystem)
+		FSlateApplication& SlateApplication = FSlateApplication::Get();
+		if (FPlatformMisc::GetRequiresVirtualKeyboard())
 		{
-			TextInputMethodSystem->ActivateContext(TextInputMethodContext.ToSharedRef());
+			// @TODO: Create ITextInputMethodSystem derivations for mobile
+			SlateApplication.ShowVirtualKeyboard(true, SharedThis(this));
+		}
+		else
+		{
+			ITextInputMethodSystem* const TextInputMethodSystem = SlateApplication.GetTextInputMethodSystem();
+			if (TextInputMethodSystem)
+			{
+				TextInputMethodSystem->ActivateContext(TextInputMethodContext.ToSharedRef());
+			}
 		}
 	}
 
@@ -1528,18 +1568,29 @@ void SEditableText::OnKeyboardFocusLost( const FKeyboardFocusEvent& InKeyboardFo
 	// Skip the focus lost code if it's due to the context menu opening
 	if ( !ContextMenuWindow.IsValid() )
 	{
-		ITextInputMethodSystem* const TextInputMethodSystem = FSlateApplication::Get().GetTextInputMethodSystem();
-		if(TextInputMethodSystem)
+		FSlateApplication& SlateApplication = FSlateApplication::Get();
+		if (FPlatformMisc::GetRequiresVirtualKeyboard())
 		{
-			TextInputMethodSystem->DeactivateContext(TextInputMethodContext.ToSharedRef());
+			SlateApplication.ShowVirtualKeyboard(false);
+		}
+		else
+		{
+			ITextInputMethodSystem* const TextInputMethodSystem = SlateApplication.GetTextInputMethodSystem();
+			if (TextInputMethodSystem)
+			{
+				TextInputMethodSystem->DeactivateContext(TextInputMethodContext.ToSharedRef());
+			}
 		}
 
-		// When focus is lost let anyone who is interested that text was committed
-		if( InKeyboardFocusEvent.GetCause() != EKeyboardFocusCause::WindowActivate  ) //Don't clear selection when activating a new window or cant copy and paste on right click )
+		// Clear selection unless activating a new window (otherwise can't copy and past on right click)
+		if (InKeyboardFocusEvent.GetCause() != EKeyboardFocusCause::WindowActivate)
 		{
 			ClearSelection();
-		}	
+		}
 
+		//Always clear the local undo chain on commit.
+		ClearUndoStates();
+		
 		// See if user explicitly tabbed away or moved focus
 		ETextCommit::Type TextAction;
 		switch ( InKeyboardFocusEvent.GetCause() )
@@ -1558,9 +1609,7 @@ void SEditableText::OnKeyboardFocusLost( const FKeyboardFocusEvent& InKeyboardFo
 			break;
 		}
 
-		//Always clear the local undo chain on commit.
-		ClearUndoStates();
-
+		// When focus is lost let anyone who is interested know that text was committed
 		OnTextCommitted.ExecuteIfBound( EditedText, TextAction);
 	}
 }
diff --git a/Engine/Source/Runtime/Slate/Private/Widgets/Input/SVirtualKeyboardEntry.cpp b/Engine/Source/Runtime/Slate/Private/Widgets/Input/SVirtualKeyboardEntry.cpp
--- a/Engine/Source/Runtime/Slate/Private/Widgets/Input/SVirtualKeyboardEntry.cpp
+++ b/Engine/Source/Runtime/Slate/Private/Widgets/Input/SVirtualKeyboardEntry.cpp
@@ -36,22 +36,24 @@ void SVirtualKeyboardEntry::Construct( const FArguments& InArgs )
  *
  * @param  InNewText  The new text string
  */
-void SVirtualKeyboardEntry::SetText( const TAttribute< FText >& InNewText )
+void SVirtualKeyboardEntry::SetTextFromVirtualKeyboard( const FText& InNewText )
 {
-	EditedText = InNewText.Get();
-
-	// Don't set text if the text attribute has a 'getter' binding on it, otherwise we'd blow away
-	// that binding.  If there is a getter binding, then we'll assume it will provide us with
-	// updated text after we've fired our 'text changed' callbacks
-	if( !Text.IsBound() )
+	// Only set the text if the text attribute doesn't have a getter binding (otherwise it would be blown away).
+	// If it is bound, we'll assume that OnTextChanged will handle the update.
+	if (!Text.IsBound())
 	{
-		Text.Set ( EditedText );
+		Text.Set(InNewText);
 	}
 
-	bNeedsUpdate = true;
+	if (!InNewText.EqualTo(EditedText))
+	{
+		EditedText = InNewText;
 
-	// Let outsiders know that the text content has been changed
-	//OnTextChanged.ExecuteIfBound( HasKeyboardFocus() ? EditedText : Text.Get() );
+		// This method is called from the main thread (i.e. not the game thread) of the device with the virtual keyboard
+		// This causes the app to crash on those devices, so we're using polling here to ensure delegates are
+		// fired on the game thread in Tick.
+		bNeedsUpdate = true;
+	}
 }
 
 /**
@@ -61,10 +63,7 @@ void SVirtualKeyboardEntry::RestoreOriginalText()
 {
 	if( HasTextChangedFromOriginal() )
 	{
-		SetText(OriginalText);
-
-		// Let outsiders know that the text content has been changed
-		OnTextChanged.ExecuteIfBound( EditedText );
+		SetTextFromVirtualKeyboard(OriginalText);
 	}
 }
 
@@ -245,7 +244,7 @@ FReply SVirtualKeyboardEntry::OnKeyboardFocusReceived( const FGeometry& MyGeomet
 
 	int32 CaretPosition = EditedText.ToString().Len();
 	FSlateApplication& CurrentApp = FSlateApplication::Get();
-	CurrentApp.ShowKeyboard(true, SharedThis(this));
+	CurrentApp.ShowVirtualKeyboard(true, SharedThis(this));
 
 	return FReply::Handled();
 }
@@ -276,7 +275,7 @@ void SVirtualKeyboardEntry::OnKeyboardFocusLost( const FKeyboardFocusEvent& InKe
 	}
 
 	FSlateApplication& CurrentApp = FSlateApplication::Get();
-	CurrentApp.ShowKeyboard(false);
+	CurrentApp.ShowVirtualKeyboard(false);
 
 	OnTextCommitted.ExecuteIfBound( EditedText, TextAction );
 }
diff --git a/Engine/Source/Runtime/Slate/Private/Widgets/Text/SMultiLineEditableText.cpp b/Engine/Source/Runtime/Slate/Private/Widgets/Text/SMultiLineEditableText.cpp
--- a/Engine/Source/Runtime/Slate/Private/Widgets/Text/SMultiLineEditableText.cpp
+++ b/Engine/Source/Runtime/Slate/Private/Widgets/Text/SMultiLineEditableText.cpp
@@ -202,6 +202,7 @@ SMultiLineEditableText::SMultiLineEditableText()
 	, bIsChangingText(false)
 	, IsReadOnly(false)
 	, UICommandList(new FUICommandList())
+	, bTextChangedByVirtualKeyboard(false)
 {
 	
 }
@@ -356,6 +357,25 @@ void SMultiLineEditableText::SetText(const TAttribute< FText >& InText)
 	}
 }
 
+void SMultiLineEditableText::SetTextFromVirtualKeyboard(const FText& InNewText)
+{
+	// Only set the text if the text attribute doesn't have a getter binding (otherwise it would be blown away).
+	// If it is bound, we'll assume that OnTextCommitted will handle the update.
+	if (!BoundText.IsBound())
+	{
+		BoundText.Set(InNewText);
+	}
+
+	// Update the internal editable text
+	if (SetEditableText(InNewText))
+	{
+		// This method is called from the main thread (i.e. not the game thread) of the device with the virtual keyboard
+		// This causes the app to crash on those devices, so we're using polling here to ensure delegates are
+		// fired on the game thread in Tick.
+		bTextChangedByVirtualKeyboard = true;
+	}
+}
+
 bool SMultiLineEditableText::SetEditableText(const FText& TextToSet, const bool bForce)
 {
 	const FText EditedText = GetEditableText();
@@ -428,12 +448,22 @@ FReply SMultiLineEditableText::OnKeyboardFocusReceived( const FGeometry& MyGeome
 	// Skip the focus received code if it's due to the context menu closing
 	if ( !ContextMenuWindow.IsValid() )
 	{
-		ITextInputMethodSystem* const TextInputMethodSystem = FSlateApplication::Get().GetTextInputMethodSystem();
-		if(TextInputMethodSystem)
+		FSlateApplication& SlateApplication = FSlateApplication::Get();
+		if (FPlatformMisc::GetRequiresVirtualKeyboard())
 		{
-			TextInputMethodSystem->ActivateContext(TextInputMethodContext.ToSharedRef());
+			// @TODO: Create ITextInputMethodSystem derivations for mobile
+			SlateApplication.ShowVirtualKeyboard(true, SharedThis(this));
+		}
+		else
+		{
+			ITextInputMethodSystem* const TextInputMethodSystem = SlateApplication.GetTextInputMethodSystem();
+			if (TextInputMethodSystem)
+			{
+				TextInputMethodSystem->ActivateContext(TextInputMethodContext.ToSharedRef());
+			}
 		}
 
+
 		UpdateCursorHighlight();
 		return SWidget::OnKeyboardFocusReceived( MyGeometry, InKeyboardFocusEvent );
 	}
@@ -446,10 +476,18 @@ void SMultiLineEditableText::OnKeyboardFocusLost( const FKeyboardFocusEvent& InK
 	// Skip the focus lost code if it's due to the context menu opening
 	if (!ContextMenuWindow.IsValid())
 	{
-		ITextInputMethodSystem* const TextInputMethodSystem = FSlateApplication::Get().GetTextInputMethodSystem();
-		if (TextInputMethodSystem)
+		FSlateApplication& SlateApplication = FSlateApplication::Get();
+		if (FPlatformMisc::GetRequiresVirtualKeyboard())
 		{
-			TextInputMethodSystem->DeactivateContext(TextInputMethodContext.ToSharedRef());
+			SlateApplication.ShowVirtualKeyboard(false);
+		}
+		else
+		{
+			ITextInputMethodSystem* const TextInputMethodSystem = SlateApplication.GetTextInputMethodSystem();
+			if (TextInputMethodSystem)
+			{
+				TextInputMethodSystem->DeactivateContext(TextInputMethodContext.ToSharedRef());
+			}
 		}
 
 		// When focus is lost let anyone who is interested that text was committed
@@ -2018,6 +2056,13 @@ void SMultiLineEditableText::LoadText()
 
 void SMultiLineEditableText::Tick( const FGeometry& AllottedGeometry, const double InCurrentTime, const float InDeltaTime )
 {
+	if (bTextChangedByVirtualKeyboard)
+	{
+		// Let outsiders know that the text content has been changed
+		OnTextCommitted.ExecuteIfBound(GetEditableText(), ETextCommit::OnUserMovedFocus);
+		bTextChangedByVirtualKeyboard = false;
+	}
+
 	if(TextInputMethodChangeNotifier.IsValid() && TextInputMethodContext.IsValid() && TextInputMethodContext->CachedGeometry != AllottedGeometry)
 	{
 		TextInputMethodContext->CachedGeometry = AllottedGeometry;
diff --git a/Engine/Source/Runtime/Slate/Public/Framework/Application/IPlatformTextField.h b/Engine/Source/Runtime/Slate/Public/Framework/Application/IPlatformTextField.h
--- a/Engine/Source/Runtime/Slate/Public/Framework/Application/IPlatformTextField.h
+++ b/Engine/Source/Runtime/Slate/Public/Framework/Application/IPlatformTextField.h
@@ -2,14 +2,14 @@
 
 #pragma once
 
-class SVirtualKeyboardEntry;
+class IVirtualKeyboardEntry;
 
 class IPlatformTextField
 {
 public:
 	virtual ~IPlatformTextField() {};
 
-	virtual void ShowKeyboard(bool bShow, TSharedPtr<SVirtualKeyboardEntry> TextEntryWidget) = 0;
+	virtual void ShowVirtualKeyboard(bool bShow, TSharedPtr<IVirtualKeyboardEntry> TextEntryWidget) = 0;
 
 private:
 
diff --git a/Engine/Source/Runtime/Slate/Public/Framework/Application/SlateApplication.h b/Engine/Source/Runtime/Slate/Public/Framework/Application/SlateApplication.h
--- a/Engine/Source/Runtime/Slate/Public/Framework/Application/SlateApplication.h
+++ b/Engine/Source/Runtime/Slate/Public/Framework/Application/SlateApplication.h
@@ -580,8 +580,9 @@ public:
 	 * Shows or hides an onscreen keyboard
 	 *
 	 * @param bShow	true to show the keyboard, false to hide it
+	 * @param TextEntryWidget The widget that will receive the input from the virtual keyboard
 	 */
-	void ShowKeyboard( bool bShow, TSharedPtr<SVirtualKeyboardEntry> TextEntryWidget = NULL );
+	void ShowVirtualKeyboard( bool bShow, TSharedPtr<IVirtualKeyboardEntry> TextEntryWidget = nullptr );
 
 	/** Get the work area that has the largest intersection with the specified rectangle */
 	FSlateRect GetWorkArea( const FSlateRect& InRect ) const;
diff --git a/Engine/Source/Runtime/Slate/Public/Framework/Text/Android/AndroidPlatformTextField.h b/Engine/Source/Runtime/Slate/Public/Framework/Text/Android/AndroidPlatformTextField.h
--- a/Engine/Source/Runtime/Slate/Public/Framework/Text/Android/AndroidPlatformTextField.h
+++ b/Engine/Source/Runtime/Slate/Public/Framework/Text/Android/AndroidPlatformTextField.h
@@ -7,7 +7,7 @@
 class FAndroidPlatformTextField : public IPlatformTextField
 {
 public:
-	virtual void ShowKeyboard(bool bShow, TSharedPtr<SVirtualKeyboardEntry> TextEntryWidget) override;
+	virtual void ShowVirtualKeyboard(bool bShow, TSharedPtr<IVirtualKeyboardEntry> TextEntryWidget) override;
 
 private:
 //	SlateTextField* TextField;
diff --git a/Engine/Source/Runtime/Slate/Public/Framework/Text/GenericPlatformTextField.h b/Engine/Source/Runtime/Slate/Public/Framework/Text/GenericPlatformTextField.h
--- a/Engine/Source/Runtime/Slate/Public/Framework/Text/GenericPlatformTextField.h
+++ b/Engine/Source/Runtime/Slate/Public/Framework/Text/GenericPlatformTextField.h
@@ -7,7 +7,7 @@
 class FGenericPlatformTextField : public IPlatformTextField
 {
 public:
-	virtual void ShowKeyboard(bool bShow, TSharedPtr<SVirtualKeyboardEntry> TextEntryWidget) override {};
+	virtual void ShowVirtualKeyboard(bool bShow, TSharedPtr<IVirtualKeyboardEntry> TextEntryWidget) override {};
 
 };
 
diff --git a/Engine/Source/Runtime/Slate/Public/Framework/Text/IOS/IOSPlatformTextField.h b/Engine/Source/Runtime/Slate/Public/Framework/Text/IOS/IOSPlatformTextField.h
--- a/Engine/Source/Runtime/Slate/Public/Framework/Text/IOS/IOSPlatformTextField.h
+++ b/Engine/Source/Runtime/Slate/Public/Framework/Text/IOS/IOSPlatformTextField.h
@@ -14,7 +14,7 @@ public:
 	FIOSPlatformTextField();
 	virtual ~FIOSPlatformTextField();
 
-	virtual void ShowKeyboard(bool bShow, TSharedPtr<SVirtualKeyboardEntry> TextEntryWidget) override;
+	virtual void ShowVirtualKeyboard(bool bShow, TSharedPtr<IVirtualKeyboardEntry> TextEntryWidget) override;
 
 private:
 	SlateTextField* TextField;
@@ -25,9 +25,9 @@ typedef FIOSPlatformTextField FPlatformTextField;
 
 @interface SlateTextField : NSObject<UIAlertViewDelegate>
 {
-	TSharedPtr<SVirtualKeyboardEntry> TextWidget;
+	TSharedPtr<IVirtualKeyboardEntry> TextWidget;
 }
 
--(void)show:(TSharedPtr<SVirtualKeyboardEntry>)InTextWidget;
+-(void)show:(TSharedPtr<IVirtualKeyboardEntry>)InTextWidget;
 
 @end
diff --git a/Engine/Source/Runtime/Slate/Public/Widgets/Input/IVirtualKeyboardEntry.h b/Engine/Source/Runtime/Slate/Public/Widgets/Input/IVirtualKeyboardEntry.h
new file mode 100644
--- /dev/null
+++ b/Engine/Source/Runtime/Slate/Public/Widgets/Input/IVirtualKeyboardEntry.h
@@ -0,0 +1,53 @@
+// Copyright 1998-2014 Epic Games, Inc. All Rights Reserved.
+
+#pragma once
+
+// @todo - hook up keyboard types
+enum EKeyboardType
+{
+	Keyboard_Default,
+	Keyboard_Number,
+	Keyboard_Web,
+	Keyboard_Email,
+	Keyboard_Password,
+};
+
+class SLATE_API IVirtualKeyboardEntry
+{
+
+public:
+	/**
+	* Sets the text to that entered by the virtual keyboard
+	*
+	* @param  InNewText  The new text
+	*/
+	virtual void SetTextFromVirtualKeyboard(const FText& InNewText) = 0;
+	
+	/**
+	* Returns the text.
+	*
+	* @return  Text
+	*/
+	virtual const FText& GetText() const = 0;
+
+	/**
+	* Returns the hint text.
+	*
+	* @return  HintText
+	*/
+	virtual const FText GetHintText() const = 0;
+
+	/**
+	* Returns the virtual keyboard type.
+	*
+	* @return  VirtualKeyboardType
+	*/
+	virtual EKeyboardType GetVirtualKeyboardType() const = 0;
+
+	/**
+	* Returns whether the entry is multi-line
+	*
+	* @return Whether the entry is multi-line
+	*/
+	virtual bool IsMultilineEntry() const = 0;
+};
diff --git a/Engine/Source/Runtime/Slate/Public/Widgets/Input/SEditableText.h b/Engine/Source/Runtime/Slate/Public/Widgets/Input/SEditableText.h
--- a/Engine/Source/Runtime/Slate/Public/Widgets/Input/SEditableText.h
+++ b/Engine/Source/Runtime/Slate/Public/Widgets/Input/SEditableText.h
@@ -4,12 +4,12 @@
 
 #include "ITextEditorWidget.h"
 #include "ITextInputMethodSystem.h"
-
+#include "IVirtualKeyboardEntry.h"
 
 /**
  * Editable text widget
  */
-class SLATE_API SEditableText : public SLeafWidget, public ITextEditorWidget
+class SLATE_API SEditableText : public SLeafWidget, public ITextEditorWidget, public IVirtualKeyboardEntry
 {
 
 public:
@@ -32,6 +32,7 @@ public:
 		, _ClearKeyboardFocusOnCommit(true)
 		, _MinDesiredWidth(0.0f)
 		, _SelectAllTextOnCommit( false )
+		, _VirtualKeyboardType(EKeyboardType::Keyboard_Default)
 		{}
 
 		/** Sets the text content for this editable text widget */
@@ -102,6 +103,9 @@ public:
 		/** Menu extender for the right-click context menu */
 		SLATE_EVENT( FMenuExtensionDelegate, ContextMenuExtender )
 
+		/** The type of virtual keyboard to use on mobile devices */
+		SLATE_ATTRIBUTE( EKeyboardType, VirtualKeyboardType)
+
 	SLATE_END_ARGS()
 
 	/** Constructor */
@@ -115,16 +119,6 @@ public:
 	 */
 	void Construct( const FArguments& InArgs );
 
-	/**
-	 * Returns the text.
-	 *
-	 * @return  Text
-	 */
-	const FText& GetText() const
-	{
-		return Text.Get();
-	}
-
 	/**
 	 * Sets the text currently being edited 
 	 *
@@ -280,7 +274,7 @@ protected:
 		{
 		}
 	};
-
+	
 public:
 	// BEGIN ITextEditorWidget interface
 	virtual void StartChangingText() override;
@@ -321,6 +315,31 @@ public:
 	virtual void LoadText() override;
 	// END ITextEditorWidget interface
 
+public:
+	// BEGIN IVirtualKeyboardEntry interface
+	virtual void SetTextFromVirtualKeyboard(const FText& InNewText) override;
+
+	virtual const FText& GetText() const override
+	{
+		return Text.Get();
+	}
+
+	virtual const FText GetHintText() const override
+	{
+		return HintText.Get();
+	}
+
+	virtual EKeyboardType GetVirtualKeyboardType() const override
+	{
+		return VirtualKeyboardType.Get();
+	}
+
+	virtual bool IsMultilineEntry() const override
+	{
+		return false;
+	}
+	// END IVirtualKeyboardEntry interface
+
 protected:
 	// BEGIN SWidget interface
 	virtual void Tick( const FGeometry& AllottedGeometry, const double InCurrentTime, const float InDeltaTime ) override;
@@ -461,7 +480,7 @@ private:
 	/** The text content for this editable text widget */
 	TAttribute< FText > Text;
 
-	/** The text content fot watermark/hinting what text to type here */
+	/** The text content for watermark/hinting what text to type here */
 	TAttribute< FText > HintText;
 
 	/** The font used to draw the text */
@@ -583,4 +602,10 @@ private:
 
 	/** Notification interface object for text input method systems. */
 	TSharedPtr<ITextInputMethodChangeNotifier> TextInputMethodChangeNotifier;
+
+	/** The type of virtual keyboard to use for editing this text on mobile */
+	TAttribute<EKeyboardType> VirtualKeyboardType;
+
+	/** Whether the text has been changed by a virtual keyboard */
+	bool bTextChangedByVirtualKeyboard;
 };
diff --git a/Engine/Source/Runtime/Slate/Public/Widgets/Input/SSpinBox.h b/Engine/Source/Runtime/Slate/Public/Widgets/Input/SSpinBox.h
--- a/Engine/Source/Runtime/Slate/Public/Widgets/Input/SSpinBox.h
+++ b/Engine/Source/Runtime/Slate/Public/Widgets/Input/SSpinBox.h
@@ -157,16 +157,17 @@ public:
 				.ClearKeyboardFocusOnCommit( InArgs._ClearKeyboardFocusOnCommit )
 				.SelectAllTextOnCommit( InArgs._SelectAllTextOnCommit )
 				.MinDesiredWidth( this, &SSpinBox<NumericType>::GetTextMinDesiredWidth )
-			]
+				.VirtualKeyboardType(EKeyboardType::Keyboard_Number)
+			]			
 
-			+ SHorizontalBox::Slot()
+			+SHorizontalBox::Slot()
 			.AutoWidth()
 			.HAlign(HAlign_Fill) 
 			.VAlign(VAlign_Center) 
 			[
 				SNew(SImage)
 				.Image( &InArgs._Style->ArrowsImage )
-				.ColorAndOpacity(FSlateColor::UseForeground())
+				.ColorAndOpacity( FSlateColor::UseForeground())
 			]
 		];
 	}
diff --git a/Engine/Source/Runtime/Slate/Public/Widgets/Input/SVirtualKeyboardEntry.h b/Engine/Source/Runtime/Slate/Public/Widgets/Input/SVirtualKeyboardEntry.h
--- a/Engine/Source/Runtime/Slate/Public/Widgets/Input/SVirtualKeyboardEntry.h
+++ b/Engine/Source/Runtime/Slate/Public/Widgets/Input/SVirtualKeyboardEntry.h
@@ -2,17 +2,9 @@
 
 #pragma once
 
-// @todo - hook up keyboard types
-enum EKeyboardType
-{
-	Keyboard_Default,
-	Keyboard_Number,
-	Keyboard_Web,
-	Keyboard_Email,
-	Keyboard_Password,
-};
+#include "IVirtualKeyboardEntry.h"
 
-class SLATE_API SVirtualKeyboardEntry : public SLeafWidget
+class SLATE_API SVirtualKeyboardEntry : public SLeafWidget, public IVirtualKeyboardEntry
 {
 
 public:
@@ -71,65 +63,46 @@ public:
 	 */
 	void Construct( const FArguments& InArgs );
 
+	/** Restores the text to the original state */
+	void RestoreOriginalText();
+
+	/** @return whether the current text varies from the original */
+	bool HasTextChangedFromOriginal() const;
+
 	/**
-	 * Returns the text.
+	 * Sets the font used to draw the text
 	 *
-	 * @return  Text
+	 * @param  InNewFont	The new font to use
 	 */
-	const FText& GetText() const
+	void SetFont( const TAttribute< FSlateFontInfo >& InNewFont );
+
+	/** @return Whether the user can edit the text. */
+	bool GetIsReadOnly() const;
+
+public:
+	// BEGIN IVirtualKeyboardEntry Interface
+	virtual void SetTextFromVirtualKeyboard(const FText& InNewText) override;
+
+	virtual const FText& GetText() const override
 	{
 		return Text.Get();
 	}
 
-	/**
-	 * Sets the text currently being edited 
-	 *
-	 * @param  InNewText  The new text
-	 */
-	void SetText( const TAttribute< FText >& InNewText );
-
-	/**
-	 * Returns the hint text.
-	 *
-	 * @return  HintText
-	 */
-	const FText& GetHintText() const
+	virtual const FText GetHintText() const override
 	{
 		return HintText.Get();
 	}
 
-	/**
-	 * Returns the keyboard type.
-	 *
-	 * @return  KeyboardType
-	 */
-	EKeyboardType GetKeyboardType() const
+	virtual EKeyboardType GetVirtualKeyboardType() const override
 	{
 		return KeyboardType.Get();
 	}
 
-	/**
-	 * Restores the text to the original state
-	 */
-	void RestoreOriginalText();
-
-	/**
-	 *	Returns whether the current text varies from the original
-	 */
-	bool HasTextChangedFromOriginal() const;
-
-	/**
-	 * Sets the font used to draw the text
-	 *
-	 * @param  InNewFont	The new font to use
-	 */
-	void SetFont( const TAttribute< FSlateFontInfo >& InNewFont );
-
-	/**
-	 * Can the user edit the text?
-	 */
-	bool GetIsReadOnly() const;
-
+	virtual bool IsMultilineEntry() const override
+	{
+		return false;
+	}
+	// END IVirtualKeyboardEntry Interface
 
 protected:
 
@@ -144,9 +117,7 @@ protected:
 
 private:
 
-	/**
-	 * @return  the string that needs to be rendered
-	 */
+	/** @return The string that needs to be rendered */
 	FString GetStringToRender() const;
 
 private:
diff --git a/Engine/Source/Runtime/Slate/Public/Widgets/Text/SMultiLineEditableText.h b/Engine/Source/Runtime/Slate/Public/Widgets/Text/SMultiLineEditableText.h
--- a/Engine/Source/Runtime/Slate/Public/Widgets/Text/SMultiLineEditableText.h
+++ b/Engine/Source/Runtime/Slate/Public/Widgets/Text/SMultiLineEditableText.h
@@ -5,12 +5,13 @@
 #if WITH_FANCY_TEXT
 
 #include "ITextEditorWidget.h"
+#include "IVirtualKeyboardEntry.h"
 #include "ITextInputMethodSystem.h"
 #include "ITextLayoutMarshaller.h"
 #include "SScrollBar.h"
 
 /** An editable text widget that supports multiple lines and soft word-wrapping. */
-class SLATE_API SMultiLineEditableText : public SWidget, public ITextEditorWidget
+class SLATE_API SMultiLineEditableText : public SWidget, public ITextEditorWidget, public IVirtualKeyboardEntry
 {
 public:
 	/** Called when the cursor is moved within the text area */
@@ -89,14 +90,6 @@ public:
 
 	void Construct( const FArguments& InArgs );
 
-	/**
-	 * Gets the text assigned to this text block
-	 */
-	const FText& GetText() const
-	{
-		return BoundText.Get();
-	}
-
 	/**
 	 * Sets the text for this text block
 	 */
@@ -397,6 +390,30 @@ private:
 	virtual void LoadText() override;
 	// END ITextEditorWidget interface
 
+public:
+	// BEGIN IVirtualKeyboardEntry interface
+	virtual void SetTextFromVirtualKeyboard(const FText& InNewText) override;
+
+	virtual const FText& GetText() const override
+	{
+		return BoundText.Get();
+	}
+
+	virtual const FText GetHintText() const override
+	{
+		return FText();
+	}
+
+	virtual EKeyboardType GetVirtualKeyboardType() const override
+	{
+		return EKeyboardType::Keyboard_Default;
+	}
+
+	virtual bool IsMultilineEntry() const override
+	{
+		return true;
+	}
+	// END IVirtualKeyboardEntry interface
 
 private:
 	// BEGIN SWidget interface
@@ -623,6 +640,9 @@ private:
 
 	/** Notification interface object for text input method systems. */
 	TSharedPtr<ITextInputMethodChangeNotifier> TextInputMethodChangeNotifier;
+
+	/** Whether the text has been changed by a virtual keyboard */
+	bool bTextChangedByVirtualKeyboard;
 };
 
 
diff --git a/Engine/Source/Runtime/SlateCore/Public/Input/Events.h b/Engine/Source/Runtime/SlateCore/Public/Input/Events.h
--- a/Engine/Source/Runtime/SlateCore/Public/Input/Events.h
+++ b/Engine/Source/Runtime/SlateCore/Public/Input/Events.h
@@ -504,7 +504,7 @@ public:
 		: FInputEvent(InModifierKeys, false)
 		, ScreenSpacePosition(InScreenSpacePosition)
 		, LastScreenSpacePosition(InLastScreenSpacePosition)
-		, CursorDelta(LastScreenSpacePosition - ScreenSpacePosition)
+		, CursorDelta(InScreenSpacePosition - InLastScreenSpacePosition)
 		, PressedButtons(bPressLeftMouseButton ? FTouchKeySet::StandardSet : FTouchKeySet::EmptySet)
 		, EffectingButton(EKeys::LeftMouseButton)
 		, UserIndex(InUserIndex)
