Proper fix for tooltip jitter on Linux.

- We cannot speculatively set cached screen position in MoveWindowTo(), since the OS may disallow the move and/or modify the location.
- Instead, the code in FSlateApplication needs to be fixed not to assume GetPositionInScreen() to be accurate immediately after the move (TTP #347070).
- This hack work arounds the problem by limiting speculative update of the cached position to tooltips only.

diff --git a/Engine/Source/Runtime/Slate/Private/Framework/Application/SlateApplication.cpp b/Engine/Source/Runtime/Slate/Private/Framework/Application/SlateApplication.cpp
--- a/Engine/Source/Runtime/Slate/Private/Framework/Application/SlateApplication.cpp
+++ b/Engine/Source/Runtime/Slate/Private/Framework/Application/SlateApplication.cpp
@@ -2348,6 +2348,10 @@ void FSlateApplication::UpdateToolTip( bool AllowSpawningOfNewToolTips )
 			WindowLocation = CalculatePopupWindowPosition( Anchor, PinnedToolTipWindow->GetDesiredSize() );
 
 			// Update the tool tip window positioning
+			// SetCachedScreenPosition is a hack (issue tracked as TTP #347070) which is needed because code in TickWindowAndChildren()/DrawPrepass()
+			// assumes GetPositionInScreen() to correspond to the new window location in the same tick. This is true on Windows, but other
+			// OSes (Linux in particular) may not update cached screen position until next time events are polled.
+			PinnedToolTipWindow->SetCachedScreenPosition( WindowLocation );
 			PinnedToolTipWindow->MoveWindowTo( WindowLocation );
 		}
 	}
diff --git a/Engine/Source/Runtime/SlateCore/Private/Widgets/SWindow.cpp b/Engine/Source/Runtime/SlateCore/Private/Widgets/SWindow.cpp
--- a/Engine/Source/Runtime/SlateCore/Private/Widgets/SWindow.cpp
+++ b/Engine/Source/Runtime/SlateCore/Private/Widgets/SWindow.cpp
@@ -697,8 +697,6 @@ void SWindow::MoveWindowTo( FVector2D NewPosition )
 	if (NativeWindow.IsValid())
 	{
 		NativeWindow->MoveWindowTo( FMath::TruncToInt(NewPosition.X), FMath::TruncToInt(NewPosition.Y) );
-		// cache position immediately since the "windows moved" event may be delayed until next frame.
-		SetCachedScreenPosition(NewPosition);
 	}
 	else
 	{
