Work around for virtual joystick scaling issues.

#TTP 346047  - Thumbsticks are very small

diff --git a/Engine/Source/Runtime/Slate/Private/Widgets/Input/SVirtualJoystick.cpp b/Engine/Source/Runtime/Slate/Private/Widgets/Input/SVirtualJoystick.cpp
--- a/Engine/Source/Runtime/Slate/Private/Widgets/Input/SVirtualJoystick.cpp
+++ b/Engine/Source/Runtime/Slate/Private/Widgets/Input/SVirtualJoystick.cpp
@@ -9,7 +9,9 @@ const float OPACITY_LERP_RATE = 3.f;
 static FORCEINLINE float GetScaleFactor(const FGeometry& Geometry)
 {
 	const float DesiredWidth = 1024.0f;
-	return Geometry.GetDrawSize().GetMax() / DesiredWidth;
+
+	float UndoDPIScaling = 1.0f / Geometry.Scale;
+	return (Geometry.GetDrawSize().GetMax() / DesiredWidth) * UndoDPIScaling;
 }
 
 FORCEINLINE float SVirtualJoystick::GetBaseOpacity()
