Fixed some layout shuffling when using wrapping with a multi-line editable text

TTP# 347063 - EDITOR: BUG: Multiline editor tries to be larger than it needs to be at first

The scrollbars were always initially made visible, which in the case of wrapping, meant they could appear for a frame and then vanish again once the text had been wrapped. They're now hidden until we have valid wrapping information for the text layout.

This also addresses an issue where the property editor could incorrectly clip off some properties when a multi-line editable text changed size. This was due to the fact that you're required to notify the property editor that is has to refresh when a widget has changed size, and SPropertyEditorText now does this.

diff --git a/Engine/Source/Editor/PropertyEditor/Private/Presentation/PropertyEditor/PropertyEditor.cpp b/Engine/Source/Editor/PropertyEditor/Private/Presentation/PropertyEditor/PropertyEditor.cpp
--- a/Engine/Source/Editor/PropertyEditor/Private/Presentation/PropertyEditor/PropertyEditor.cpp
+++ b/Engine/Source/Editor/PropertyEditor/Private/Presentation/PropertyEditor/PropertyEditor.cpp
@@ -412,6 +412,11 @@ bool FPropertyEditor::IsPropertyEditingEnabled() const
 		(EditConditionProperty == NULL || IsEditConditionMet( EditConditionProperty, PropertyEditConditions ));
 }
 
+void FPropertyEditor::RequestRefresh()
+{
+	PropertyUtilities->RequestRefresh();
+}
+
 bool FPropertyEditor::HasEditCondition() const 
 { 
 	return EditConditionProperty != NULL; 
diff --git a/Engine/Source/Editor/PropertyEditor/Private/Presentation/PropertyEditor/PropertyEditor.h b/Engine/Source/Editor/PropertyEditor/Private/Presentation/PropertyEditor/PropertyEditor.h
--- a/Engine/Source/Editor/PropertyEditor/Private/Presentation/PropertyEditor/PropertyEditor.h
+++ b/Engine/Source/Editor/PropertyEditor/Private/Presentation/PropertyEditor/PropertyEditor.h
@@ -33,6 +33,8 @@ public:
 
 	bool IsPropertyEditingEnabled() const;
 
+	void RequestRefresh();
+
 	bool SupportsEditConditionToggle() const;
 
 	/**	@return Whether the property is editconst */
diff --git a/Engine/Source/Editor/PropertyEditor/Private/UserInterface/PropertyEditor/SPropertyEditorText.cpp b/Engine/Source/Editor/PropertyEditor/Private/UserInterface/PropertyEditor/SPropertyEditorText.cpp
--- a/Engine/Source/Editor/PropertyEditor/Private/UserInterface/PropertyEditor/SPropertyEditorText.cpp
+++ b/Engine/Source/Editor/PropertyEditor/Private/UserInterface/PropertyEditor/SPropertyEditorText.cpp
@@ -105,6 +105,16 @@ FReply SPropertyEditorText::OnKeyboardFocusReceived( const FGeometry& MyGeometry
 	return FReply::Handled().SetKeyboardFocus( PrimaryWidget.ToSharedRef(), InKeyboardFocusEvent.GetCause() );
 }
 
+void SPropertyEditorText::Tick( const FGeometry& AllottedGeometry, const double InCurrentTime, const float InDeltaTime )
+{
+	const float CurrentHeight = AllottedGeometry.GetLocalSize().Y;
+	if (bIsMultiLine && PreviousHeight.IsSet() && PreviousHeight.GetValue() != CurrentHeight)
+	{
+		PropertyEditor->RequestRefresh();
+	}
+	PreviousHeight = CurrentHeight;
+}
+
 bool SPropertyEditorText::CanEdit() const
 {
 	return PropertyEditor.IsValid() ? !PropertyEditor->IsEditConst() : true;
diff --git a/Engine/Source/Editor/PropertyEditor/Private/UserInterface/PropertyEditor/SPropertyEditorText.h b/Engine/Source/Editor/PropertyEditor/Private/UserInterface/PropertyEditor/SPropertyEditorText.h
--- a/Engine/Source/Editor/PropertyEditor/Private/UserInterface/PropertyEditor/SPropertyEditorText.h
+++ b/Engine/Source/Editor/PropertyEditor/Private/UserInterface/PropertyEditor/SPropertyEditorText.h
@@ -23,6 +23,8 @@ public:
 
 	FReply OnKeyboardFocusReceived( const FGeometry& MyGeometry, const FKeyboardFocusEvent& InKeyboardFocusEvent ) override;
 
+	virtual void Tick( const FGeometry& AllottedGeometry, const double InCurrentTime, const float InDeltaTime ) override;
+
 private:
 
 	virtual void OnTextCommitted( const FText& NewText, ETextCommit::Type CommitInfo );
@@ -39,6 +41,8 @@ private:
 
 	TSharedPtr< class SWidget > PrimaryWidget;
 
+	TOptional< float > PreviousHeight;
+
 	/** Cached flag as we would like multi-line text widgets to be slightly larger */
 	bool bIsMultiLine;
 };
\ No newline at end of file
diff --git a/Engine/Source/Runtime/Slate/Private/Widgets/Text/SMultiLineEditableText.cpp b/Engine/Source/Runtime/Slate/Private/Widgets/Text/SMultiLineEditableText.cpp
--- a/Engine/Source/Runtime/Slate/Private/Widgets/Text/SMultiLineEditableText.cpp
+++ b/Engine/Source/Runtime/Slate/Private/Widgets/Text/SMultiLineEditableText.cpp
@@ -238,12 +238,14 @@ void SMultiLineEditableText::Construct( const FArguments& InArgs )
 	HScrollBar = InArgs._HScrollBar;
 	if (HScrollBar.IsValid())
 	{
+		HScrollBar->SetUserVisibility(EVisibility::Collapsed);
 		HScrollBar->SetOnUserScrolled(FOnUserScrolled::CreateSP(this, &SMultiLineEditableText::OnHScrollBarMoved));
 	}
 
 	VScrollBar = InArgs._VScrollBar;
 	if (VScrollBar.IsValid())
 	{
+		VScrollBar->SetUserVisibility(EVisibility::Collapsed);
 		VScrollBar->SetOnUserScrolled(FOnUserScrolled::CreateSP(this, &SMultiLineEditableText::OnVScrollBarMoved));
 	}
 
@@ -2081,6 +2083,10 @@ void SMultiLineEditableText::Tick( const FGeometry& AllottedGeometry, const doub
 	const float FontMaxCharHeight = FTextEditHelper::GetFontHeight(TextStyle.Font);
 	const float CaretWidth = FTextEditHelper::CalculateCaretWidth(FontMaxCharHeight);
 
+	// If we're auto-wrapping, we need to hide the scrollbars until the first valid auto-wrap has been performed
+	// If we don't do this, then we can get some nasty layout shuffling as the scrollbars appear for one frame and then vanish again
+	const EVisibility ScrollBarVisiblityOverride = (AutoWrapText.Get() && CachedSize != AllottedGeometry.Size) ? EVisibility::Collapsed : EVisibility::Visible;
+
 	// Try and make sure that the line containing the cursor is in view
 	if (PositionToScrollIntoView.IsSet())
 	{
@@ -2133,6 +2139,7 @@ void SMultiLineEditableText::Tick( const FGeometry& AllottedGeometry, const doub
 		if (HScrollBar.IsValid())
 		{
 			HScrollBar->SetState(ViewOffset, ViewFraction);
+			HScrollBar->SetUserVisibility(ScrollBarVisiblityOverride);
 			if (!HScrollBar->IsNeeded())
 			{
 				// We cannot scroll, so ensure that there is no offset
@@ -2155,6 +2162,7 @@ void SMultiLineEditableText::Tick( const FGeometry& AllottedGeometry, const doub
 		if (VScrollBar.IsValid())
 		{
 			VScrollBar->SetState(ViewOffset, ViewFraction);
+			VScrollBar->SetUserVisibility(ScrollBarVisiblityOverride);
 			if (!VScrollBar->IsNeeded())
 			{
 				// We cannot scroll, so ensure that there is no offset
