Fix incorrect calculation of SlatePrepass and fixed the name of the stat's display name

diff --git a/Engine/Source/Runtime/Slate/Private/Framework/Application/SlateApplication.cpp b/Engine/Source/Runtime/Slate/Private/Framework/Application/SlateApplication.cpp
--- a/Engine/Source/Runtime/Slate/Private/Framework/Application/SlateApplication.cpp
+++ b/Engine/Source/Runtime/Slate/Private/Framework/Application/SlateApplication.cpp
@@ -670,7 +670,7 @@ static void DrawWindowPrepass( TSharedRef<SWindow> WindowToDraw )
 
 void FSlateApplication::DrawPrepass( TSharedPtr<SWindow> DrawOnlyThisWindow )
 {
-	SCOPE_CYCLE_COUNTER(STAT_SlateDrawWindowTime);
+
 
 	TSharedPtr<SWindow> ActiveModalWindow = GetActiveModalWindow();
 
diff --git a/Engine/Source/Runtime/Slate/Public/Slate.h b/Engine/Source/Runtime/Slate/Public/Slate.h
--- a/Engine/Source/Runtime/Slate/Public/Slate.h
+++ b/Engine/Source/Runtime/Slate/Public/Slate.h
@@ -18,7 +18,7 @@
 
 /** Stat tracking for Slate */
 DECLARE_CYCLE_STAT_EXTERN( TEXT("Total Slate Tick Time"), STAT_SlateTickTime, STATGROUP_Slate , );
-DECLARE_CYCLE_STAT_EXTERN( TEXT("CacheDesiredSize"), STAT_SlateCacheDesiredSize, STATGROUP_Slate , );
+DECLARE_CYCLE_STAT_EXTERN( TEXT("SlatePrepass"), STAT_SlateCacheDesiredSize, STATGROUP_Slate , );
 DECLARE_CYCLE_STAT_EXTERN( TEXT("STextBlock"), STAT_SlateCacheDesiredSize_STextBlock, STATGROUP_Slate , );
 DECLARE_CYCLE_STAT_EXTERN( TEXT("TickWidgets"), STAT_SlateTickWidgets, STATGROUP_Slate , );
 DECLARE_CYCLE_STAT_EXTERN( TEXT("DrawWindows"), STAT_SlateDrawWindowTime, STATGROUP_Slate , );
