Added Slate.SkipSecondPrepass cvar that allows us to inspect the overhead of the second prepass phase dynamically.

diff --git a/Engine/Source/Runtime/Slate/Private/Framework/Application/SlateApplication.cpp b/Engine/Source/Runtime/Slate/Private/Framework/Application/SlateApplication.cpp
--- a/Engine/Source/Runtime/Slate/Private/Framework/Application/SlateApplication.cpp
+++ b/Engine/Source/Runtime/Slate/Private/Framework/Application/SlateApplication.cpp
@@ -68,6 +68,11 @@ TAutoConsoleVariable<int32> TargetFrameRateForResponsiveness(
 	35,	// Frames per second
 	TEXT( "Minimum sustained average frame rate required before we consider the editor to be \"responsive\" for a smooth UI experience" ) );
 
+TAutoConsoleVariable<int32> SkipSecondPrepass(
+	TEXT("Slate.SkipSecondPrepass"),
+	0,
+	TEXT("Whether to skip the second Slate PrePass call (the one right before rendering)."));
+
 //////////////////////////////////////////////////////////////////////////
 bool FSlateApplication::MouseCaptorHelper::HasCapture() const
 {
@@ -670,8 +675,6 @@ static void DrawWindowPrepass( TSharedRef<SWindow> WindowToDraw )
 
 void FSlateApplication::DrawPrepass( TSharedPtr<SWindow> DrawOnlyThisWindow )
 {
-
-
 	TSharedPtr<SWindow> ActiveModalWindow = GetActiveModalWindow();
 
 	if (ActiveModalWindow.IsValid())
@@ -725,7 +728,10 @@ void FSlateApplication::PrivateDrawWindows( TSharedPtr<SWindow> DrawOnlyThisWind
 
 	FWidgetPath FocusPath = FocusedWidgetPath.ToWidgetPath();
 
-	DrawPrepass( DrawOnlyThisWindow );
+	if ( !SkipSecondPrepass.GetValueOnGameThread() )
+	{
+		DrawPrepass( DrawOnlyThisWindow );
+	}
 
 	FDrawWindowArgs DrawWindowArgs( Renderer->GetDrawBuffer(), FocusPath, WidgetsUnderCursor );
 
