Fixed not being able to click on widgets inside a scissor rect

diff --git a/Engine/Source/Runtime/Slate/Public/Widgets/Layout/SScissorRectBox.h b/Engine/Source/Runtime/Slate/Public/Widgets/Layout/SScissorRectBox.h
--- a/Engine/Source/Runtime/Slate/Public/Widgets/Layout/SScissorRectBox.h
+++ b/Engine/Source/Runtime/Slate/Public/Widgets/Layout/SScissorRectBox.h
@@ -25,7 +25,9 @@ public:
 	};
 
 	SLATE_BEGIN_ARGS(SScissorRectBox)
-	{}
+	{
+		_Visibility = EVisibility::SelfHitTestInvisible;
+	}
 		SLATE_DEFAULT_SLOT(FArguments, Content)
 	SLATE_END_ARGS()
 
