Fixed a 1 frame delay in text block size when scaling

We don't know the scale when calling ComputeDesiredSize, so we were trying to cache and use the previous frames scale instead, however that was producing some noticable scaling artefacts on text blocks.

We've "fixed" this by always having ComputeDesiredSize return the text size with a text scale of 1. This means that we might over estimate the amount of space required for the text (as it doesn't scale linearly like the Slate text layout assumes widgets do), but it fixes the size changing after 1 frame (this is how the text block sized used to be calculated before we had text layouts).

diff --git a/Engine/Source/Runtime/Slate/Private/Framework/Text/TextLayout.cpp b/Engine/Source/Runtime/Slate/Private/Framework/Text/TextLayout.cpp
--- a/Engine/Source/Runtime/Slate/Private/Framework/Text/TextLayout.cpp
+++ b/Engine/Source/Runtime/Slate/Private/Framework/Text/TextLayout.cpp
@@ -615,6 +615,11 @@ void FTextLayout::DirtyLayout()
 	ClearWrappingCache();
 }
 
+bool FTextLayout::IsLayoutDirty() const
+{
+	return !!(DirtyFlags & EDirtyState::Layout);
+}
+
 void FTextLayout::ClearRunRenderers()
 {
 	for (int32 Index = 0; Index < LineModels.Num(); Index++)
diff --git a/Engine/Source/Runtime/Slate/Private/Widgets/Text/SRichTextBlock.cpp b/Engine/Source/Runtime/Slate/Private/Widgets/Text/SRichTextBlock.cpp
--- a/Engine/Source/Runtime/Slate/Private/Widgets/Text/SRichTextBlock.cpp
+++ b/Engine/Source/Runtime/Slate/Private/Widgets/Text/SRichTextBlock.cpp
@@ -22,8 +22,6 @@ void SRichTextBlock::Construct( const FArguments& InArgs )
 	LineHeightPercentage = InArgs._LineHeightPercentage;
 	Justification = InArgs._Justification;
 
-	CachedScale = 1.0f;
-
 	{
 		TSharedPtr<IRichTextMarkupParser> Parser = InArgs._Parser;
 		if ( !Parser.IsValid() )
@@ -45,9 +43,6 @@ int32 SRichTextBlock::OnPaint( const FPaintArgs& Args, const FGeometry& Allotted
 {
 	SCOPE_CYCLE_COUNTER( STAT_SlateOnPaint_SRichTextBlock );
 
-	// todo: jdale - The scale needs to be passed to ComputeDesiredSize
-	CachedScale = AllottedGeometry.Scale;
-
 	// OnPaint will also update the text layout cache if required
 	LayerId = TextLayoutCache->OnPaint(
 		FTextBlockLayout::FWidgetArgs(BoundText, HighlightText, WrapTextAt, AutoWrapText, Margin, LineHeightPercentage, Justification), 
@@ -60,12 +55,12 @@ int32 SRichTextBlock::OnPaint( const FPaintArgs& Args, const FGeometry& Allotted
 FVector2D SRichTextBlock::ComputeDesiredSize() const
 {
 	// todo: jdale - The scale needs to be passed to ComputeDesiredSize
-	const float Scale = CachedScale;
+	//const float Scale = CachedScale;
 
 	// ComputeDesiredSize will also update the text layout cache if required
 	const FVector2D TextSize = TextLayoutCache->ComputeDesiredSize(
 		FTextBlockLayout::FWidgetArgs(BoundText, HighlightText, WrapTextAt, AutoWrapText, Margin, LineHeightPercentage, Justification), 
-		Scale, TextStyle
+		/*Scale, */TextStyle
 		);
 
 	return TextSize;
diff --git a/Engine/Source/Runtime/Slate/Private/Widgets/Text/STextBlock.cpp b/Engine/Source/Runtime/Slate/Private/Widgets/Text/STextBlock.cpp
--- a/Engine/Source/Runtime/Slate/Private/Widgets/Text/STextBlock.cpp
+++ b/Engine/Source/Runtime/Slate/Private/Widgets/Text/STextBlock.cpp
@@ -16,8 +16,6 @@ void STextBlock::Construct( const FArguments& InArgs )
 	Justification = InArgs._Justification;
 	MinDesiredWidth = InArgs._MinDesiredWidth;
 
-	CachedScale = 1.0f;
-
 	Font = InArgs._Font;
 	ColorAndOpacity = InArgs._ColorAndOpacity;
 	ShadowOffset = InArgs._ShadowOffset;
@@ -101,9 +99,6 @@ int32 STextBlock::OnPaint( const FPaintArgs& Args, const FGeometry& AllottedGeom
 	SCOPE_CYCLE_COUNTER( STAT_SlateOnPaint_STextBlock );
 #endif
 
-	// todo: jdale - The scale needs to be passed to ComputeDesiredSize
-	CachedScale = AllottedGeometry.Scale;
-
 #if WITH_FANCY_TEXT
 
 	// OnPaint will also update the text layout cache if required
@@ -182,12 +177,12 @@ FVector2D STextBlock::ComputeDesiredSize() const
 #if WITH_FANCY_TEXT
 
 	// todo: jdale - The scale needs to be passed to ComputeDesiredSize
-	const float Scale = CachedScale;
+	//const float Scale = CachedScale;
 
 	// ComputeDesiredSize will also update the text layout cache if required
 	const FVector2D TextSize = TextLayoutCache->ComputeDesiredSize(
 		FTextBlockLayout::FWidgetArgs(BoundText, HighlightText, WrapTextAt, AutoWrapText, Margin, LineHeightPercentage, Justification), 
-		Scale, GetComputedTextStyle()
+		/*Scale, */GetComputedTextStyle()
 		);
 
 #else//WITH_FANCY_TEXT
diff --git a/Engine/Source/Runtime/Slate/Private/Widgets/Text/TextBlockLayout.cpp b/Engine/Source/Runtime/Slate/Private/Widgets/Text/TextBlockLayout.cpp
--- a/Engine/Source/Runtime/Slate/Private/Widgets/Text/TextBlockLayout.cpp
+++ b/Engine/Source/Runtime/Slate/Private/Widgets/Text/TextBlockLayout.cpp
@@ -18,13 +18,13 @@ FTextBlockLayout::FTextBlockLayout(FTextBlockStyle InDefaultTextStyle, TSharedRe
 	, Marshaller(InMarshaller)
 	, TextHighlighter(FSlateTextHighlightRunRenderer::Create())
 	, CachedSize(ForceInitToZero)
+	, CachedDesiredSize(ForceInitToZero)
 {
 	TextLayout->SetLineBreakIterator(InLineBreakPolicy);
 }
 
-FVector2D FTextBlockLayout::ComputeDesiredSize(const FWidgetArgs& InWidgetArgs, const float InScale, const FTextBlockStyle& InTextStyle)
+FVector2D FTextBlockLayout::ComputeDesiredSize(const FWidgetArgs& InWidgetArgs/*, const float InScale*/, const FTextBlockStyle& InTextStyle)
 {
-	TextLayout->SetScale(InScale);
 	TextLayout->SetWrappingWidth(CalculateWrappingWidth(InWidgetArgs));
 	TextLayout->SetMargin(InWidgetArgs.Margin.Get());
 	TextLayout->SetJustification(InWidgetArgs.Justification.Get());
@@ -74,16 +74,31 @@ FVector2D FTextBlockLayout::ComputeDesiredSize(const FWidgetArgs& InWidgetArgs,
 		}
 	}
 
-	TextLayout->UpdateIfNeeded();
+	// We need to update our cached desired size if the text layout has become dirty
+	// todo: jdale - This is a hack until we can perform accurate measuring in ComputeDesiredSize
+	if(TextLayout->IsLayoutDirty())
+	{
+		// The desired size must always have a scale of 1, OnPaint will make sure the scale is set correctly for painting
+		TextLayout->SetScale(1.0f);
+		TextLayout->UpdateIfNeeded();
 
-	return TextLayout->GetSize();
+		CachedDesiredSize = TextLayout->GetSize();
+	}
+	else
+	{
+		// This logic may look odd, but IsLayoutDirty() only checks that we've made a change that might affect the layout (which in turn might affect the 
+		// desired size), however there's also highlight changes, which don't affect the size, but still require a call to UpdateIfNeeded() to be applied
+		TextLayout->UpdateIfNeeded();
+	}
+
+	return CachedDesiredSize;
 }
 
 int32 FTextBlockLayout::OnPaint(const FWidgetArgs& InWidgetArgs, const FPaintArgs& InPaintArgs, const FGeometry& InAllottedGeometry, const FSlateRect& InClippingRect, FSlateWindowElementList& OutDrawElements, int32 LayerId, const FWidgetStyle& InWidgetStyle, bool bParentEnabled)
 {
 	CachedSize = InAllottedGeometry.Size;
 
-	TextLayout->SetWrappingWidth(CalculateWrappingWidth(InWidgetArgs));
+	TextLayout->SetScale(InAllottedGeometry.Scale);
 	TextLayout->SetVisibleRegion(InAllottedGeometry.Size, FVector2D::ZeroVector);
 
 	TextLayout->UpdateIfNeeded();
diff --git a/Engine/Source/Runtime/Slate/Public/Framework/Text/TextLayout.h b/Engine/Source/Runtime/Slate/Public/Framework/Text/TextLayout.h
--- a/Engine/Source/Runtime/Slate/Public/Framework/Text/TextLayout.h
+++ b/Engine/Source/Runtime/Slate/Public/Framework/Text/TextLayout.h
@@ -345,6 +345,8 @@ public:
 
 	void DirtyLayout();
 
+	bool IsLayoutDirty() const;
+
 	int32 GetLineViewIndexForTextLocation(const TArray< FTextLayout::FLineView >& LineViews, const FTextLocation& Location, const bool bPerformInclusiveBoundsCheck);
 
 	/**
diff --git a/Engine/Source/Runtime/Slate/Public/Widgets/Text/SRichTextBlock.h b/Engine/Source/Runtime/Slate/Public/Widgets/Text/SRichTextBlock.h
--- a/Engine/Source/Runtime/Slate/Public/Widgets/Text/SRichTextBlock.h
+++ b/Engine/Source/Runtime/Slate/Public/Widgets/Text/SRichTextBlock.h
@@ -161,9 +161,6 @@ private:
 
 	/** How the text should be aligned with the margin. */
 	TAttribute< float > LineHeightPercentage;
-
-	/** todo: jdale - The scale needs to be passed to ComputeDesiredSize */
-	mutable float CachedScale;
 };
 
 #endif //WITH_FANCY_TEXT
\ No newline at end of file
diff --git a/Engine/Source/Runtime/Slate/Public/Widgets/Text/STextBlock.h b/Engine/Source/Runtime/Slate/Public/Widgets/Text/STextBlock.h
--- a/Engine/Source/Runtime/Slate/Public/Widgets/Text/STextBlock.h
+++ b/Engine/Source/Runtime/Slate/Public/Widgets/Text/STextBlock.h
@@ -247,7 +247,4 @@ private:
 
 	/** The delegate to execute when this text is double clicked */
 	FOnClicked OnDoubleClicked;
-
-	/** todo: jdale - The scale needs to be passed to ComputeDesiredSize */
-	mutable float CachedScale;
 };
diff --git a/Engine/Source/Runtime/Slate/Public/Widgets/Text/TextBlockLayout.h b/Engine/Source/Runtime/Slate/Public/Widgets/Text/TextBlockLayout.h
--- a/Engine/Source/Runtime/Slate/Public/Widgets/Text/TextBlockLayout.h
+++ b/Engine/Source/Runtime/Slate/Public/Widgets/Text/TextBlockLayout.h
@@ -47,7 +47,7 @@ public:
 	/**
 	 * Get the computed desired size for this layout, updating the internal cache as required
 	 */
-	FVector2D ComputeDesiredSize(const FWidgetArgs& InWidgetArgs, const float InScale, const FTextBlockStyle& InTextStyle);
+	FVector2D ComputeDesiredSize(const FWidgetArgs& InWidgetArgs/*, const float InScale*/, const FTextBlockStyle& InTextStyle);
 
 	/**
 	 * Paint this layout, updating the internal cache as required
@@ -96,6 +96,9 @@ private:
 	/** The last known size of the layout from the previous OnPaint, used to guess at an auto-wrapping width in ComputeDesiredSize */
 	FVector2D CachedSize;
 
+	/** The cached desired size - this should be updated using a text layout scale of 1 whenever the text layout becomes dirty */
+	FVector2D CachedDesiredSize;
+
 	/** The state of the text the last time it was updated (used to allow updates when the text is changed) */
 	FTextSnapshot TextLastUpdate;
 
