Fixed a crash when the TSF IME API tried to get the selected text from SOutputLog

FOutputLogTextLayoutMarshaller::GetText wasn't implemented, so it was always returning an empty string. I've implemented this and verified that it doesn't cause any performance issues for the output log.

I had to change SMultiLineEditableText::LoadText to only do that when it was bound to a delegate, as that was causing a hitch while the output log was updated when you gave it focus.

I also fixed giving the output log focus sometimes jumping to the top of the log (SMultiLineEditableText::OnKeyboardFocusReceived needed the same fix as SMultiLineEditableText::OnKeyboardFocusLost to prevent the cursor jumping into focus).

diff --git a/Engine/Source/Developer/OutputLog/Private/SOutputLog.cpp b/Engine/Source/Developer/OutputLog/Private/SOutputLog.cpp
--- a/Engine/Source/Developer/OutputLog/Private/SOutputLog.cpp
+++ b/Engine/Source/Developer/OutputLog/Private/SOutputLog.cpp
@@ -591,6 +591,7 @@ void FOutputLogTextLayoutMarshaller::SetText(const FString& SourceString, FTextL
 
 void FOutputLogTextLayoutMarshaller::GetText(FString& TargetString, const FTextLayout& SourceTextLayout)
 {
+	SourceTextLayout.GetAsText(TargetString);
 }
 
 bool FOutputLogTextLayoutMarshaller::AppendMessage(const TCHAR* InText, const ELogVerbosity::Type InVerbosity, const FName& InCategory)
diff --git a/Engine/Source/Runtime/Slate/Private/Widgets/Text/SMultiLineEditableText.cpp b/Engine/Source/Runtime/Slate/Private/Widgets/Text/SMultiLineEditableText.cpp
--- a/Engine/Source/Runtime/Slate/Private/Widgets/Text/SMultiLineEditableText.cpp
+++ b/Engine/Source/Runtime/Slate/Private/Widgets/Text/SMultiLineEditableText.cpp
@@ -467,6 +467,13 @@ FReply SMultiLineEditableText::OnKeyboardFocusReceived( const FGeometry& MyGeome
 
 
 		UpdateCursorHighlight();
+
+		// UpdateCursorHighlight always tries to scroll to the cursor, but we don't want that to happen when we 
+		// gain focus since it can cause the scroll position to jump unexpectedly
+		// If we gained focus via a mouse click that moved the cursor, then MoveCursor will already take care
+		// of making sure that gets scrolled into view
+		PositionToScrollIntoView = TOptional<FScrollInfo>();
+
 		return SWidget::OnKeyboardFocusReceived( MyGeometry, InKeyboardFocusEvent );
 	}
 
@@ -2052,8 +2059,13 @@ void SMultiLineEditableText::OnWindowClosed(const TSharedRef<SWindow>&)
 
 void SMultiLineEditableText::LoadText()
 {
-	SetText(BoundText);
-	TextLayout->UpdateIfNeeded();
+	// We only need to do this if we're bound to a delegate, otherwise the text layout will already be up-to-date
+	// either from Construct, or a call to SetText
+	if (BoundText.IsBound())
+	{
+		SetText(BoundText);
+		TextLayout->UpdateIfNeeded();
+	}
 }
 
 void SMultiLineEditableText::Tick( const FGeometry& AllottedGeometry, const double InCurrentTime, const float InDeltaTime )
