Use current cursor location for mouse rather than cached version.

diff --git a/Engine/Source/Runtime/Slate/Private/Framework/Application/SlateApplication.cpp b/Engine/Source/Runtime/Slate/Private/Framework/Application/SlateApplication.cpp
--- a/Engine/Source/Runtime/Slate/Private/Framework/Application/SlateApplication.cpp
+++ b/Engine/Source/Runtime/Slate/Private/Framework/Application/SlateApplication.cpp
@@ -2530,12 +2530,12 @@ void FSlateApplication::SynthesizeMouseMove()
 {
 	for (auto IndexLocationPairIterator = PointerIndexLastPositionMap.CreateIterator(); IndexLocationPairIterator; ++IndexLocationPairIterator)
 	{
-		// TODO: Use correct user index instead of 0.
+		// @TODO umg: Use correct user index instead of 0.
 		FPointerEvent MouseEvent
 		(
 			0,
 			IndexLocationPairIterator.Key(),
-			IndexLocationPairIterator.Value(),
+			(IndexLocationPairIterator.Key() == CursorPointerIndex) ? FSlateApplication::GetCursorPos() : IndexLocationPairIterator.Value(),
 			IndexLocationPairIterator.Value(), 
 			false
 		);
