Fix the Window menu on Mac OS X not being properly populated.

diff --git a/Engine/Source/Runtime/Core/Private/Mac/MacPlatformMisc.cpp b/Engine/Source/Runtime/Core/Private/Mac/MacPlatformMisc.cpp
--- a/Engine/Source/Runtime/Core/Private/Mac/MacPlatformMisc.cpp
+++ b/Engine/Source/Runtime/Core/Private/Mac/MacPlatformMisc.cpp
@@ -324,6 +324,7 @@ void FMacPlatformMisc::UpdateWindowMenu()
 	[WindowMenu addItem:CloseItem];
 	[WindowMenu addItem:[NSMenuItem separatorItem]];
 	[WindowMenu addItem:BringAllToFrontItem];
+	[WindowMenu addItem:[NSMenuItem separatorItem]];
 }
 
 void FMacPlatformMisc::ActivateApplication()
diff --git a/Engine/Source/Runtime/Slate/Private/Framework/MultiBox/Mac/MacMenu.cpp b/Engine/Source/Runtime/Slate/Private/Framework/MultiBox/Mac/MacMenu.cpp
--- a/Engine/Source/Runtime/Slate/Private/Framework/MultiBox/Mac/MacMenu.cpp
+++ b/Engine/Source/Runtime/Slate/Private/Framework/MultiBox/Mac/MacMenu.cpp
@@ -51,6 +51,7 @@ void FSlateMacMenu::UpdateWithMultiBox(const TSharedRef< FMultiBox >& MultiBox)
 	MainThreadCall(^{
 		int32 NumItems = [[NSApp mainMenu] numberOfItems];
 		FText HelpTitle = NSLOCTEXT("MainMenu", "HelpMenu", "Help");
+		FText WindowLabel = NSLOCTEXT("MainMenu", "WindowMenu", "Window");
 		NSMenuItem* HelpMenu = [[[NSApp mainMenu] itemWithTitle:HelpTitle.ToString().GetNSString()] retain];
 		for (int32 Index = NumItems - 1; Index > 0; Index--)
 		{
@@ -70,6 +71,29 @@ void FSlateMacMenu::UpdateWithMultiBox(const TSharedRef< FMultiBox >& MultiBox)
 			[MenuItem setTitle:Title];
 			[[NSApp mainMenu] addItem:MenuItem];
 			[MenuItem setSubmenu:Menu];
+			
+			const bool bIsWindowMenu = (WindowLabel.ToString().Compare(FString(Title)) == 0);
+			if (bIsWindowMenu)
+			{
+				[NSApp setWindowsMenu:nil];
+				
+				[Menu removeAllItems];
+				
+				NSMenuItem* MinimizeItem = [[[NSMenuItem alloc] initWithTitle:@"Minimize" action:@selector(miniaturize:) keyEquivalent:@"m"] autorelease];
+				NSMenuItem* ZoomItem = [[[NSMenuItem alloc] initWithTitle:@"Zoom" action:@selector(performZoom:) keyEquivalent:@""] autorelease];
+				NSMenuItem* CloseItem = [[[NSMenuItem alloc] initWithTitle:@"Close" action:@selector(performClose:) keyEquivalent:@"w"] autorelease];
+				NSMenuItem* BringAllToFrontItem = [[[NSMenuItem alloc] initWithTitle:@"Bring All to Front" action:@selector(arrangeInFront:) keyEquivalent:@""] autorelease];
+				
+				[Menu addItem:MinimizeItem];
+				[Menu addItem:ZoomItem];
+				[Menu addItem:CloseItem];
+				[Menu addItem:[NSMenuItem separatorItem]];
+				[Menu addItem:BringAllToFrontItem];
+				[Menu addItem:[NSMenuItem separatorItem]];
+				
+				[NSApp setWindowsMenu:Menu];
+				[Menu addItem:[NSMenuItem separatorItem]];
+			}
 		}
 
 		if ([[NSApp mainMenu] itemWithTitle:HelpTitle.ToString().GetNSString()] == nil && HelpMenu)
@@ -90,21 +114,18 @@ void FSlateMacMenu::UpdateMenu(FMacMenu* Menu)
 		FText WindowLabel = NSLOCTEXT("MainMenu", "WindowMenu", "Window");
 		const bool bIsWindowMenu = (WindowLabel.ToString().Compare(FString([Menu title])) == 0);
 
-		if ([[Menu itemArray] count] == 0)
+		int32 ItemIndexOffset = 0;
+		if (bIsWindowMenu)
 		{
-			if (bIsWindowMenu)
+			int32 SeparatorIndex = 0;
+			for (NSMenuItem* Item in [Menu itemArray])
 			{
-				NSMenuItem* MinimizeItem = [[[NSMenuItem alloc] initWithTitle:@"Minimize" action:@selector(miniaturize:) keyEquivalent:@"m"] autorelease];
-				NSMenuItem* ZoomItem = [[[NSMenuItem alloc] initWithTitle:@"Zoom" action:@selector(performZoom:) keyEquivalent:@""] autorelease];
-				NSMenuItem* CloseItem = [[[NSMenuItem alloc] initWithTitle:@"Close" action:@selector(performClose:) keyEquivalent:@"w"] autorelease];
-				NSMenuItem* BringAllToFrontItem = [[[NSMenuItem alloc] initWithTitle:@"Bring All to Front" action:@selector(arrangeInFront:) keyEquivalent:@""] autorelease];
-
-				[Menu addItem:MinimizeItem];
-				[Menu addItem:ZoomItem];
-				[Menu addItem:CloseItem];
-				[Menu addItem:[NSMenuItem separatorItem]];
-				[Menu addItem:BringAllToFrontItem];
-				[Menu addItem:[NSMenuItem separatorItem]];
+				SeparatorIndex += [Item isSeparatorItem] ? 1 : 0;
+				ItemIndexOffset++;
+				if (SeparatorIndex == 3)
+				{
+					break;
+				}
 			}
 		}
 
@@ -124,7 +145,7 @@ void FSlateMacMenu::UpdateMenu(FMacMenu* Menu)
 		int32 ItemIndexAdjust = 0;
 		for (int32 Index = 0; Index < MenuBlocks.Num(); Index++)
 		{
-			const int32 ItemIndex = (bIsWindowMenu ? Index + 6 : Index) - ItemIndexAdjust;
+			const int32 ItemIndex = (bIsWindowMenu ? Index + ItemIndexOffset : Index) - ItemIndexAdjust;
 			NSMenuItem* MenuItem = [Menu numberOfItems] > ItemIndex ? [Menu itemAtIndex:ItemIndex] : nil;
 
 			EMultiBlockType::Type Type = MenuBlocks[Index]->GetType();
