Replacing SVirtualKeyboard::SetText

diff --git a/Engine/Source/Runtime/Slate/Private/Widgets/Input/SVirtualKeyboardEntry.cpp b/Engine/Source/Runtime/Slate/Private/Widgets/Input/SVirtualKeyboardEntry.cpp
--- a/Engine/Source/Runtime/Slate/Private/Widgets/Input/SVirtualKeyboardEntry.cpp
+++ b/Engine/Source/Runtime/Slate/Private/Widgets/Input/SVirtualKeyboardEntry.cpp
@@ -31,11 +31,21 @@ void SVirtualKeyboardEntry::Construct( const FArguments& InArgs )
 	KeyboardType = InArgs._KeyboardType;
 }
 
-/**
- * Sets the text string currently being edited 
- *
- * @param  InNewText  The new text string
- */
+void SVirtualKeyboardEntry::SetText(const TAttribute< FText >& InNewText)
+{
+	EditedText = InNewText.Get();
+
+	// Don't set text if the text attribute has a 'getter' binding on it, otherwise we'd blow away
+	// that binding.  If there is a getter binding, then we'll assume it will provide us with
+	// updated text after we've fired our 'text changed' callbacks
+	if (!Text.IsBound())
+	{
+		Text.Set(EditedText);
+	}
+
+	bNeedsUpdate = true;
+}
+
 void SVirtualKeyboardEntry::SetTextFromVirtualKeyboard( const FText& InNewText )
 {
 	// Only set the text if the text attribute doesn't have a getter binding (otherwise it would be blown away).
diff --git a/Engine/Source/Runtime/Slate/Public/Widgets/Input/SVirtualKeyboardEntry.h b/Engine/Source/Runtime/Slate/Public/Widgets/Input/SVirtualKeyboardEntry.h
--- a/Engine/Source/Runtime/Slate/Public/Widgets/Input/SVirtualKeyboardEntry.h
+++ b/Engine/Source/Runtime/Slate/Public/Widgets/Input/SVirtualKeyboardEntry.h
@@ -69,6 +69,13 @@ public:
 	/** @return whether the current text varies from the original */
 	bool HasTextChangedFromOriginal() const;
 
+	/**
+	* Sets the text currently being edited
+	*
+	* @param  InNewText  The new text
+	*/
+	void SetText(const TAttribute< FText >& InNewText);
+
 	/**
 	 * Sets the font used to draw the text
 	 *
