List and tile views that have no scroll bar now support overscroll.

Also added an option to control whether overscroll is allowed on STableViewBase derivatives;
This partly addresses TTP#346682 - Need to support overscroll when items are not enough to allow scrolling (consistency with ListViews)

diff --git a/Engine/Source/Editor/GameProjectGeneration/Private/SProjectBrowser.cpp b/Engine/Source/Editor/GameProjectGeneration/Private/SProjectBrowser.cpp
--- a/Engine/Source/Editor/GameProjectGeneration/Private/SProjectBrowser.cpp
+++ b/Engine/Source/Editor/GameProjectGeneration/Private/SProjectBrowser.cpp
@@ -352,6 +352,7 @@ void SProjectBrowser::ConstructCategory( const TSharedRef<SVerticalBox>& InCateg
 		.ListItemsSource(&Category->FilteredProjectItemsSource)
 		.SelectionMode(ESelectionMode::Single)
 		.ClearSelectionOnClick(false)
+		.AllowOverscroll(EAllowOverscroll::No)
 		.OnGenerateTile(this, &SProjectBrowser::MakeProjectViewWidget)
 		.OnContextMenuOpening(this, &SProjectBrowser::OnGetContextMenuContent)
 		.OnMouseButtonDoubleClick(this, &SProjectBrowser::HandleProjectItemDoubleClick)
diff --git a/Engine/Source/Runtime/Slate/Private/Widgets/Views/STableViewBase.cpp b/Engine/Source/Runtime/Slate/Private/Widgets/Views/STableViewBase.cpp
--- a/Engine/Source/Runtime/Slate/Private/Widgets/Views/STableViewBase.cpp
+++ b/Engine/Source/Runtime/Slate/Private/Widgets/Views/STableViewBase.cpp
@@ -167,7 +167,7 @@ void STableViewBase::TickInertialScroll( const FGeometry& AllottedGeometry, cons
 		{
 			if ( CanUseInertialScroll(ScrollVelocity) )
 			{
-				this->ScrollBy(AllottedGeometry, ScrollVelocity * InDeltaTime, EAllowOverscroll::Yes);
+				this->ScrollBy(AllottedGeometry, ScrollVelocity * InDeltaTime, AllowOverscroll);
 			}
 			else
 			{
@@ -185,7 +185,7 @@ void STableViewBase::Tick( const FGeometry& AllottedGeometry, const double InCur
 	{
 		TickInertialScroll(AllottedGeometry, InCurrentTime, InDeltaTime);
 
-		if ( !IsRightClickScrolling() )
+		if ( !IsRightClickScrolling() && AllowOverscroll == EAllowOverscroll::Yes)
 		{
 			// If we are currently in overscroll, the list will need refreshing.
 			// Do this before UpdateOverscroll, as that could cause GetOverscroll() to be 0
@@ -225,9 +225,12 @@ void STableViewBase::Tick( const FGeometry& AllottedGeometry, const double InCur
 			
 			ScrollOffset = FMath::Max(0.0, ScrollOffset);
 			ItemsPanel->SmoothScrollOffset( FMath::Fractional(ScrollOffset / GetNumItemsWide()) );
-			const float OverscrollAmount = Overscroll.GetOverscroll();
-			ItemsPanel->SetOverscrollAmount( OverscrollAmount );
-			
+
+			if (AllowOverscroll == EAllowOverscroll::Yes)
+			{
+				const float OverscrollAmount = Overscroll.GetOverscroll();
+				ItemsPanel->SetOverscrollAmount( OverscrollAmount );
+			}
 
 			UpdateSelectionSet();
 
@@ -242,7 +245,7 @@ void STableViewBase::Tick( const FGeometry& AllottedGeometry, const double InCur
 
 			NotifyItemScrolledIntoView();
 
-			bWasAtEndOfList = FMath::IsNearlyZero(ScrollBar->DistanceFromBottom()) ? true : false;
+			bWasAtEndOfList = ScrollBar->DistanceFromBottom() < KINDA_SMALL_NUMBER ? true : false;
 
 			bItemsNeedRefresh = false;
 			ItemsPanel->SetRefreshPending(false);
@@ -358,7 +361,7 @@ FReply STableViewBase::OnMouseMove( const FGeometry& MyGeometry, const FPointerE
 		{
 			TickScrollDelta -= ScrollByAmount;
 
-			const float AmountScrolled = this->ScrollBy( MyGeometry, -ScrollByAmount, EAllowOverscroll::Yes );
+			const float AmountScrolled = this->ScrollBy( MyGeometry, -ScrollByAmount, AllowOverscroll );
 
 			FReply Reply = FReply::Handled();
 
@@ -455,7 +458,7 @@ FReply STableViewBase::OnTouchMoved( const FGeometry& MyGeometry, const FPointer
 		AmountScrolledWhileRightMouseDown += FMath::Abs( ScrollByAmount );
 		TickScrollDelta += ScrollByAmount;
 
-		const float AmountScrolled = this->ScrollBy( MyGeometry, ScrollByAmount, EAllowOverscroll::Yes );
+		const float AmountScrolled = this->ScrollBy( MyGeometry, ScrollByAmount, AllowOverscroll );
 
 		if (AmountScrolledWhileRightMouseDown > SlatePanTriggerDistance)
 		{
@@ -503,7 +506,8 @@ TSharedPtr<SHeaderRow> STableViewBase::GetHeaderRow() const
 
 bool STableViewBase::IsRightClickScrolling() const
 {
-	return AmountScrolledWhileRightMouseDown >= SlatePanTriggerDistance && this->ScrollBar->IsNeeded();
+	return AmountScrolledWhileRightMouseDown >= SlatePanTriggerDistance &&
+		(this->ScrollBar->IsNeeded() || AllowOverscroll == EAllowOverscroll::Yes);
 }
 
 bool STableViewBase::IsUserScrolling() const
@@ -557,16 +561,17 @@ STableViewBase::STableViewBase( ETableViewMode::Type InTableViewMode )
 	, SoftwareCursorPosition( ForceInitToZero )
 	, bShowSoftwareCursor( false )
 	, Overscroll()
+	, AllowOverscroll(EAllowOverscroll::Yes)
 	, bItemsNeedRefresh( true )	
 {
 }
 
-float STableViewBase::ScrollBy( const FGeometry& MyGeometry, float ScrollByAmountInSlateUnits, EAllowOverscroll AllowOverscroll )
+float STableViewBase::ScrollBy( const FGeometry& MyGeometry, float ScrollByAmountInSlateUnits, EAllowOverscroll InAllowOverscroll )
 {
 	const int32 NumItemsBeingObserved = GetNumItemsBeingObserved();
 	const float FractionalScrollOffsetInItems = (ScrollOffset + GetScrollRateInItems() * ScrollByAmountInSlateUnits) / NumItemsBeingObserved;
 	const double ClampedScrollOffsetInItems = FMath::Clamp<double>( FractionalScrollOffsetInItems*NumItemsBeingObserved, -10.0f, NumItemsBeingObserved+10.0f ) * NumItemsBeingObserved;
-	if (AllowOverscroll == EAllowOverscroll::Yes)
+	if (InAllowOverscroll == EAllowOverscroll::Yes)
 	{
 		Overscroll.ScrollBy( ClampedScrollOffsetInItems - ScrollByAmountInSlateUnits );
 	}
diff --git a/Engine/Source/Runtime/Slate/Public/Widgets/Views/SListView.h b/Engine/Source/Runtime/Slate/Public/Widgets/Views/SListView.h
--- a/Engine/Source/Runtime/Slate/Public/Widgets/Views/SListView.h
+++ b/Engine/Source/Runtime/Slate/Public/Widgets/Views/SListView.h
@@ -64,6 +64,7 @@ public:
 		, _SelectionMode(ESelectionMode::Multi)
 		, _ClearSelectionOnClick(true)
 		, _ExternalScrollbar()
+		, _AllowOverscroll(EAllowOverscroll::Yes)
 	{ }
 
 		SLATE_EVENT( FOnGenerateRow, OnGenerateRow )
@@ -91,6 +92,8 @@ public:
 		SLATE_ARGUMENT( TSharedPtr<SScrollBar>, ExternalScrollbar )
 
 		SLATE_ATTRIBUTE( EVisibility, ScrollbarVisibility)
+		
+		SLATE_ARGUMENT( EAllowOverscroll, AllowOverscroll );
 
 		SLATE_END_ARGS()
 
@@ -113,6 +116,8 @@ public:
 
 		this->bClearSelectionOnClick = InArgs._ClearSelectionOnClick;
 
+		this->AllowOverscroll = InArgs._AllowOverscroll;
+
 		// Check for any parameters that the coder forgot to specify.
 		FString ErrorString;
 		{
@@ -1139,18 +1144,14 @@ protected:
 		}
 	}
 
-	virtual float ScrollBy( const FGeometry& MyGeometry, float ScrollByAmountInSlateUnits, EAllowOverscroll AllowOverscroll ) override
+	virtual float ScrollBy( const FGeometry& MyGeometry, float ScrollByAmountInSlateUnits, EAllowOverscroll InAllowOverscroll ) override
 	{
 		float AbsScrollByAmount = FMath::Abs( ScrollByAmountInSlateUnits );
 		int32 StartingItemIndex = (int32)ScrollOffset;
 		double NewScrollOffset = ScrollOffset;
 
 		const bool bWholeListVisible = ScrollOffset == 0 && bWasAtEndOfList;
-		if (bWholeListVisible)
-		{
-			return 0;
-		}
-		else if ( AllowOverscroll == EAllowOverscroll::Yes && Overscroll.ShouldApplyOverscroll( ScrollOffset == 0, bWasAtEndOfList, ScrollByAmountInSlateUnits ) )
+		if ( InAllowOverscroll == EAllowOverscroll::Yes && Overscroll.ShouldApplyOverscroll( ScrollOffset == 0, bWasAtEndOfList, ScrollByAmountInSlateUnits ) )
 		{
 			const float UnclampedScrollDelta = FMath::Sign(ScrollByAmountInSlateUnits) * AbsScrollByAmount;				
 			const float ActuallyScrolledBy = Overscroll.ScrollBy( UnclampedScrollDelta );
@@ -1160,7 +1161,7 @@ protected:
 			}
 			return ActuallyScrolledBy;
 		}
-		else
+		else if (!bWholeListVisible)
 		{
 			// We know how far we want to scroll in SlateUnits, but we store scroll offset in "number of widgets".
 			// Challenge: each widget can be a different height.
@@ -1266,6 +1267,8 @@ protected:
 
 			return ScrollTo( NewScrollOffset );
 		}
+
+		return 0;
 	}
 
 protected:
diff --git a/Engine/Source/Runtime/Slate/Public/Widgets/Views/STableViewBase.h b/Engine/Source/Runtime/Slate/Public/Widgets/Views/STableViewBase.h
--- a/Engine/Source/Runtime/Slate/Public/Widgets/Views/STableViewBase.h
+++ b/Engine/Source/Runtime/Slate/Public/Widgets/Views/STableViewBase.h
@@ -106,7 +106,7 @@ protected:
 	 *
 	 * @return The amount actually scrolled in items
 	 */
-	virtual float ScrollBy( const FGeometry& MyGeometry, float ScrollByAmount, EAllowOverscroll AllowOverscroll );
+	virtual float ScrollBy( const FGeometry& MyGeometry, float ScrollByAmount, EAllowOverscroll InAllowOverscroll );
 
 	/**
 	 * Scroll the view to an offset
@@ -293,6 +293,9 @@ protected:
 		float OverscrollAmount;
 	} Overscroll;
 
+	/** Whether to permit overscroll on this list view */
+	EAllowOverscroll AllowOverscroll;
+
 private:
 	
 	/** When true, a refresh should occur the next tick */
diff --git a/Engine/Source/Runtime/Slate/Public/Widgets/Views/STileView.h b/Engine/Source/Runtime/Slate/Public/Widgets/Views/STileView.h
--- a/Engine/Source/Runtime/Slate/Public/Widgets/Views/STileView.h
+++ b/Engine/Source/Runtime/Slate/Public/Widgets/Views/STileView.h
@@ -32,6 +32,7 @@ public:
 		, _ClearSelectionOnClick(true)
 		, _ExternalScrollbar()
 		, _ScrollbarVisibility(EVisibility::Visible)
+		, _AllowOverscroll(EAllowOverscroll::Yes)
 		{}
 
 		SLATE_EVENT( FOnGenerateRow, OnGenerateTile )
@@ -58,6 +59,8 @@ public:
 
 		SLATE_ATTRIBUTE(EVisibility, ScrollbarVisibility)
 
+		SLATE_ARGUMENT( EAllowOverscroll, AllowOverscroll );
+
 	SLATE_END_ARGS()
 
 	/**
@@ -78,6 +81,8 @@ public:
 
 		this->bClearSelectionOnClick = InArgs._ClearSelectionOnClick;
 
+		this->AllowOverscroll = InArgs._AllowOverscroll;
+
 		// Check for any parameters that the coder forgot to specify.
 		FString ErrorString;
 		{
@@ -266,18 +271,14 @@ public:
 
 protected:
 
-	virtual float ScrollBy( const FGeometry& MyGeometry, float ScrollByAmountInSlateUnits, EAllowOverscroll AllowOverscroll ) override
+	virtual float ScrollBy( const FGeometry& MyGeometry, float ScrollByAmountInSlateUnits, EAllowOverscroll InAllowOverscroll ) override
 	{
 		// Working around a CLANG bug, where all the base class
 		// members require an explicit namespace resolution.
 		typedef STableViewBase S;
 
 		const bool bWholeListVisible = S::ScrollOffset == 0 && S::bWasAtEndOfList;
-		if (bWholeListVisible)
-		{
-			return 0;
-		}
-		else if ( AllowOverscroll == EAllowOverscroll::Yes && S::Overscroll.ShouldApplyOverscroll( S::ScrollOffset == 0, S::bWasAtEndOfList, ScrollByAmountInSlateUnits ) )
+		if ( InAllowOverscroll == EAllowOverscroll::Yes && S::Overscroll.ShouldApplyOverscroll( S::ScrollOffset == 0, S::bWasAtEndOfList, ScrollByAmountInSlateUnits ) )
 		{
 			const float UnclampedScrollDelta = ScrollByAmountInSlateUnits / GetNumItemsWide();
 			const float ActuallyScrolledBy = S::Overscroll.ScrollBy( UnclampedScrollDelta );
@@ -287,13 +288,15 @@ protected:
 			}
 			return ActuallyScrolledBy;
 		}
-		else
+		else if (!bWholeListVisible)
 		{
 			const float ItemHeight = this->GetItemHeight();
 			const double NewScrollOffset = this->ScrollOffset + ( ( ScrollByAmountInSlateUnits * GetNumItemsWide() ) / ItemHeight );
 
 			return this->ScrollTo( NewScrollOffset );
 		}
+
+		return 0;
 	}
 
 	virtual int32 GetNumItemsWide() const override
