Fixed side effects of recent Mac-specific drag and drop changes (temp solution, cross-platform version will be added soon)

diff --git a/Engine/Source/Runtime/Core/Private/Mac/CocoaTextView.cpp b/Engine/Source/Runtime/Core/Private/Mac/CocoaTextView.cpp
--- a/Engine/Source/Runtime/Core/Private/Mac/CocoaTextView.cpp
+++ b/Engine/Source/Runtime/Core/Private/Mac/CocoaTextView.cpp
@@ -64,6 +64,8 @@
 	{
 		[CocoaWindow mouseDown:theEvent];
 	}
+
+	[NSApp preventWindowOrdering];
 }
 
 - (void)mouseDragged:(NSEvent *)theEvent
diff --git a/Engine/Source/Runtime/Slate/Private/Framework/Application/SlateApplication.cpp b/Engine/Source/Runtime/Slate/Private/Framework/Application/SlateApplication.cpp
--- a/Engine/Source/Runtime/Slate/Private/Framework/Application/SlateApplication.cpp
+++ b/Engine/Source/Runtime/Slate/Private/Framework/Application/SlateApplication.cpp
@@ -3289,6 +3289,9 @@ bool FSlateApplication::ProcessMouseButtonDownEvent( const TSharedPtr< FGenericW
 			FWidgetPath WidgetsUnderCursor = LocateWindowUnderMouse( MouseEvent.GetScreenSpacePosition(), GetInteractiveTopLevelWindows() );
 			MouseEvent.SetEventPath(WidgetsUnderCursor);
 
+#if PLATFORM_MAC
+			NSWindow* ActiveWindow = [NSApp keyWindow];
+#endif
 			PopupSupport.SendNotifications( WidgetsUnderCursor );
 
 			// Switch worlds widgets in the current path
@@ -3350,11 +3353,11 @@ bool FSlateApplication::ProcessMouseButtonDownEvent( const TSharedPtr< FGenericW
 				}
 
 #if PLATFORM_MAC
-				if (WidgetsUnderCursor.TopLevelWindow.IsValid() && !DragDetector.DetectDragForWidget.IsValid() && MenuStack.GetNumStackLevels() == 0)
+				if (WidgetsUnderCursor.TopLevelWindow.IsValid() && !DragDetector.DetectDragForWidget.IsValid() && ActiveWindow == [NSApp keyWindow])
 				{
 					MouseCaptorHelper Captor = MouseCaptor;
-					WidgetsUnderCursor.TopLevelWindow->BringToFront(true);
 					FPlatformMisc::ActivateApplication();
+					WidgetsUnderCursor.TopLevelWindow->BringToFront(true);
 					MouseCaptor = Captor;
 				}
 #endif
@@ -3494,6 +3497,13 @@ bool FSlateApplication::ProcessMouseButtonUpEvent( FPointerEvent& MouseEvent )
 				Reply = MouseCaptorWidget.Widget->OnMouseButtonUp( MouseCaptorWidget.Geometry, MouseEvent ).SetHandler( MouseCaptorWidget.Widget );
 			}
 			ProcessReply(MouseCaptorPath, Reply, &MouseCaptorPath, &MouseEvent);
+#if PLATFORM_MAC
+			TSharedPtr<SWindow> ActiveWindow = GetActiveTopLevelWindow();
+			if (MouseCaptorPath.TopLevelWindow.IsValid() && ActiveWindow != MouseCaptorPath.TopLevelWindow)
+			{
+				MouseCaptorPath.TopLevelWindow->BringToFront(true);
+			}
+#endif
 			LOG_EVENT( EEventLog::MouseButtonUp, Reply );
 		}
 	}
