FSlateApplication - Only synthesize events for the cursor and use the constructor which creates mouse events and not touch events.
+ Fixes TTP #347390

diff --git a/Engine/Source/Runtime/Slate/Private/Framework/Application/SlateApplication.cpp b/Engine/Source/Runtime/Slate/Private/Framework/Application/SlateApplication.cpp
--- a/Engine/Source/Runtime/Slate/Private/Framework/Application/SlateApplication.cpp
+++ b/Engine/Source/Runtime/Slate/Private/Framework/Application/SlateApplication.cpp
@@ -2528,19 +2528,17 @@ bool FSlateApplication::IsWindowInDestroyQueue(TSharedRef<SWindow> Window) const
 
 void FSlateApplication::SynthesizeMouseMove()
 {
-	for (auto IndexLocationPairIterator = PointerIndexLastPositionMap.CreateIterator(); IndexLocationPairIterator; ++IndexLocationPairIterator)
-	{
-		// @TODO umg: Use correct user index instead of 0.
-		FPointerEvent MouseEvent
-		(
-			0,
-			IndexLocationPairIterator.Key(),
-			(IndexLocationPairIterator.Key() == CursorPointerIndex) ? FSlateApplication::GetCursorPos() : IndexLocationPairIterator.Value(),
-			IndexLocationPairIterator.Value(), 
-			false
+	FPointerEvent MouseEvent(
+		CursorPointerIndex,
+		GetCursorPos(),
+		GetLastCursorPos(),
+		PressedMouseButtons,
+		EKeys::Invalid,
+		0,
+		PlatformApplication->GetModifierKeys()
 		);
-		ProcessMouseMoveEvent(MouseEvent);
-	}
+
+	ProcessMouseMoveEvent(MouseEvent);
 }
 
 void FSlateApplication::OnLogSlateEvent(EEventLog::Type Elovent, const FString& AdditionalContent)
