Refactored Mac menu code to work better with game thread and fix crashes caused by performing game thread calls in menuNeedsUpdate:

diff --git a/Engine/Source/Programs/CrashReporter/CrashReportClient/Private/Mac/CrashReportClientMainMac.cpp b/Engine/Source/Programs/CrashReporter/CrashReportClient/Private/Mac/CrashReportClientMainMac.cpp
--- a/Engine/Source/Programs/CrashReporter/CrashReportClient/Private/Mac/CrashReportClientMainMac.cpp
+++ b/Engine/Source/Programs/CrashReporter/CrashReportClient/Private/Mac/CrashReportClientMainMac.cpp
@@ -48,6 +48,7 @@ static FString GSavedCommandLine;
 - (IBAction)requestQuit:(id)Sender
 {
 	GIsRequestingExit = true;
+	FPlatformMisc::UpdateCachedMacMenuState = nullptr;
 }
 
 - (void) runGameThread:(id)Arg
diff --git a/Engine/Source/Programs/SlateViewer/Private/Mac/SlateViewerMainMac.cpp b/Engine/Source/Programs/SlateViewer/Private/Mac/SlateViewerMainMac.cpp
--- a/Engine/Source/Programs/SlateViewer/Private/Mac/SlateViewerMainMac.cpp
+++ b/Engine/Source/Programs/SlateViewer/Private/Mac/SlateViewerMainMac.cpp
@@ -35,6 +35,7 @@ static FString GSavedCommandLine;
 	if(!GIsRequestingExit)
 	{
 		GIsRequestingExit = true;
+		FPlatformMisc::UpdateCachedMacMenuState = nullptr;
 		return NSTerminateCancel;
 	}
 	return NSTerminateNow;
diff --git a/Engine/Source/Programs/SymbolDebugger/Private/Mac/SymbolDebuggerMainMac.cpp b/Engine/Source/Programs/SymbolDebugger/Private/Mac/SymbolDebuggerMainMac.cpp
--- a/Engine/Source/Programs/SymbolDebugger/Private/Mac/SymbolDebuggerMainMac.cpp
+++ b/Engine/Source/Programs/SymbolDebugger/Private/Mac/SymbolDebuggerMainMac.cpp
@@ -34,6 +34,7 @@ static FString GSavedCommandLine;
 	if(!GIsRequestingExit)
 	{
 		GIsRequestingExit = true;
+		FPlatformMisc::UpdateCachedMacMenuState = nullptr;
 		return NSTerminateCancel;
 	}
 	return NSTerminateNow;
diff --git a/Engine/Source/Programs/UnrealFrontend/Private/Mac/MacUnrealFrontendMain.cpp b/Engine/Source/Programs/UnrealFrontend/Private/Mac/MacUnrealFrontendMain.cpp
--- a/Engine/Source/Programs/UnrealFrontend/Private/Mac/MacUnrealFrontendMain.cpp
+++ b/Engine/Source/Programs/UnrealFrontend/Private/Mac/MacUnrealFrontendMain.cpp
@@ -27,6 +27,7 @@ static FString GSavedCommandLine;
 - (IBAction)requestQuit:(id)Sender
 {
 	GIsRequestingExit = true;
+	FPlatformMisc::UpdateCachedMacMenuState = nullptr;
 }
 
 - (void) runGameThread:(id)Arg
diff --git a/Engine/Source/Runtime/Core/Private/Mac/MacApplication.cpp b/Engine/Source/Runtime/Core/Private/Mac/MacApplication.cpp
--- a/Engine/Source/Runtime/Core/Private/Mac/MacApplication.cpp
+++ b/Engine/Source/Runtime/Core/Private/Mac/MacApplication.cpp
@@ -637,6 +637,7 @@ void FMacApplication::ProcessNSEvent(NSEvent* const Event, TSharedPtr< FMacWindo
 					MessageHandler->OnCursorSet();
 				}
 			}
+			FPlatformMisc::bChachedMacMenuStateNeedsUpdate = true;
 			break;
 		}
 
@@ -783,6 +784,7 @@ void FMacApplication::ProcessNSEvent(NSEvent* const Event, TSharedPtr< FMacWindo
 			{
 				ResendEvent(Event);
 			}
+			FPlatformMisc::bChachedMacMenuStateNeedsUpdate = true;
 			break;
 		}
 	}
diff --git a/Engine/Source/Runtime/Core/Private/Mac/MacPlatformMisc.cpp b/Engine/Source/Runtime/Core/Private/Mac/MacPlatformMisc.cpp
--- a/Engine/Source/Runtime/Core/Private/Mac/MacPlatformMisc.cpp
+++ b/Engine/Source/Runtime/Core/Private/Mac/MacPlatformMisc.cpp
@@ -200,6 +200,9 @@ struct MacApplicationInfo
 };
 static MacApplicationInfo GMacAppInfo;
 
+UpdateCachedMacMenuStateProc FMacPlatformMisc::UpdateCachedMacMenuState = nullptr;
+bool FMacPlatformMisc::bChachedMacMenuStateNeedsUpdate = true;
+
 void FMacPlatformMisc::PlatformPreInit()
 {
 	// Increase the maximum number of simultaneously open files
@@ -459,6 +462,12 @@ void FMacPlatformMisc::PumpMessages( bool bFromMainLoop )
 	if( bFromMainLoop )
 	{
 		ProcessGameThreadEvents();
+
+		if (UpdateCachedMacMenuState && bChachedMacMenuStateNeedsUpdate && IsInGameThread())
+		{
+			UpdateCachedMacMenuState();
+			bChachedMacMenuStateNeedsUpdate = false;
+		}
 	}
 }
 
diff --git a/Engine/Source/Runtime/Core/Public/Mac/MacPlatformMisc.h b/Engine/Source/Runtime/Core/Public/Mac/MacPlatformMisc.h
--- a/Engine/Source/Runtime/Core/Public/Mac/MacPlatformMisc.h
+++ b/Engine/Source/Runtime/Core/Public/Mac/MacPlatformMisc.h
@@ -6,6 +6,8 @@
 
 #pragma once
 
+typedef void (*UpdateCachedMacMenuStateProc)(void);
+
 /**
 * Mac implementation of the misc OS functions
 **/
@@ -135,6 +137,8 @@ struct CORE_API FMacPlatformMisc : public FGenericPlatformMisc
 
 	static void UpdateWindowMenu();
 
+	static UpdateCachedMacMenuStateProc UpdateCachedMacMenuState;
+
 	static void ActivateApplication();
 
 	/**
@@ -146,6 +150,8 @@ struct CORE_API FMacPlatformMisc : public FGenericPlatformMisc
 	 * Returns whether the Mac OS X version is 10.9.x or not.
 	 */
 	static bool IsRunningOnMavericks();
+
+	static bool bChachedMacMenuStateNeedsUpdate;
 };
 
 #ifdef __OBJC__
diff --git a/Engine/Source/Runtime/Launch/Private/Mac/LaunchMac.cpp b/Engine/Source/Runtime/Launch/Private/Mac/LaunchMac.cpp
--- a/Engine/Source/Runtime/Launch/Private/Mac/LaunchMac.cpp
+++ b/Engine/Source/Runtime/Launch/Private/Mac/LaunchMac.cpp
@@ -103,6 +103,7 @@ void EngineCrashHandler(const FGenericCrashContext & GenericContext)
 	GameThreadCall(^{
 		if (GEngine)
 		{
+			FPlatformMisc::UpdateCachedMacMenuState = nullptr;
 			if (GIsEditor)
 			{
 				if (IsRunningCommandlet())
diff --git a/Engine/Source/Runtime/Slate/Private/Framework/MultiBox/Mac/MacMenu.cpp b/Engine/Source/Runtime/Slate/Private/Framework/MultiBox/Mac/MacMenu.cpp
--- a/Engine/Source/Runtime/Slate/Private/Framework/MultiBox/Mac/MacMenu.cpp
+++ b/Engine/Source/Runtime/Slate/Private/Framework/MultiBox/Mac/MacMenu.cpp
@@ -4,6 +4,23 @@
 #include "MacMenu.h"
 #include "CocoaThread.h"
 
+struct FMacMenuItemState
+{
+	TSharedPtr<const FMenuEntryBlock> Block;
+	EMultiBlockType::Type Type;
+	NSString* Title;
+	NSString* KeyEquivalent;
+	uint32 KeyModifiers;
+	NSImage* Icon;
+	bool IsSubMenu;
+	bool IsEnabled;
+	uint32 State;
+
+	FMacMenuItemState() : Type(EMultiBlockType::None), Title(nil), KeyEquivalent(nil), KeyModifiers(0), Icon(nil), IsSubMenu(false), IsEnabled(false), State(0) {}
+};
+
+static TMap<FMacMenu*, TArray<FMacMenuItemState>> GCachedMenuState;
+
 @interface FMacMenuItem : NSMenuItem
 @property (assign) TSharedPtr<const FMenuEntryBlock> MenuEntryBlock;
 - (void)performAction;
@@ -11,7 +28,7 @@
 
 @implementation FMacMenuItem
 
-- (id)initWithMenuEntryBlock:(TSharedRef< const FMenuEntryBlock >&)Block
+- (id)initWithMenuEntryBlock:(TSharedPtr< const FMenuEntryBlock >)Block
 {
 	self = [super initWithTitle:@"" action:nil keyEquivalent:@""];
 	self.MenuEntryBlock = Block;
@@ -31,11 +48,12 @@
 
 @implementation FMacMenu
 
-- (id)initWithMenuEntryBlock:(TSharedRef< const FMenuEntryBlock >&)Block
+- (id)initWithMenuEntryBlock:(TSharedPtr< const FMenuEntryBlock >)Block
 {
 	self = [super initWithTitle:@""];
 	[self setDelegate:self];
 	self.MenuEntryBlock = Block;
+	GCachedMenuState.Add(self, TArray<FMacMenuItemState>());
 	return self;
 }
 
@@ -44,73 +62,92 @@
 	FSlateMacMenu::UpdateMenu(self);
 }
 
+- (void)menuWillOpen:(NSMenu*)Menu
+{
+	FPlatformMisc::bChachedMacMenuStateNeedsUpdate = true;
+}
+
 @end
 
+static TSharedPtr<FMultiBox> GCurrentMultiBox = nullptr;
+
 void FSlateMacMenu::UpdateWithMultiBox(const TSharedRef< FMultiBox >& MultiBox)
 {
 	MainThreadCall(^{
-		int32 NumItems = [[NSApp mainMenu] numberOfItems];
-		FText HelpTitle = NSLOCTEXT("MainMenu", "HelpMenu", "Help");
-		FText WindowLabel = NSLOCTEXT("MainMenu", "WindowMenu", "Window");
-		NSMenuItem* HelpMenu = [[[NSApp mainMenu] itemWithTitle:HelpTitle.ToString().GetNSString()] retain];
-		for (int32 Index = NumItems - 1; Index > 0; Index--)
+		if (MultiBox != GCurrentMultiBox)
 		{
-			[[NSApp mainMenu] removeItemAtIndex:Index];
-		}
+			if (!FPlatformMisc::UpdateCachedMacMenuState)
+			{
+				FPlatformMisc::UpdateCachedMacMenuState = UpdateCachedState;
+			}
 
-		const TArray<TSharedRef<const FMultiBlock> >& MenuBlocks = MultiBox->GetBlocks();
+			int32 NumItems = [[NSApp mainMenu] numberOfItems];
+			FText HelpTitle = NSLOCTEXT("MainMenu", "HelpMenu", "Help");
+			FText WindowLabel = NSLOCTEXT("MainMenu", "WindowMenu", "Window");
+			NSMenuItem* HelpMenu = [[[NSApp mainMenu] itemWithTitle:HelpTitle.ToString().GetNSString()] retain];
+			for (int32 Index = NumItems - 1; Index > 0; Index--)
+			{
+				GCachedMenuState.Remove((FMacMenu*)[[NSApp mainMenu] itemAtIndex:Index]);
+				[[NSApp mainMenu] removeItemAtIndex:Index];
+			}
 
-		for (int32 Index = 0; Index < MenuBlocks.Num(); Index++)
-		{
-			TSharedRef<const FMenuEntryBlock> Block = StaticCastSharedRef<const FMenuEntryBlock>(MenuBlocks[Index]);
-			FMacMenu* Menu = [[FMacMenu alloc] initWithMenuEntryBlock:Block];
-			NSString* Title = FSlateMacMenu::GetMenuItemTitle(Block);
-			[Menu setTitle:Title];
-
-			NSMenuItem* MenuItem = [[NSMenuItem new] autorelease];
-			[MenuItem setTitle:Title];
-			[[NSApp mainMenu] addItem:MenuItem];
-			[MenuItem setSubmenu:Menu];
-			
-			const bool bIsWindowMenu = (WindowLabel.ToString().Compare(FString(Title)) == 0);
-			if (bIsWindowMenu)
+			const TArray<TSharedRef<const FMultiBlock> >& MenuBlocks = MultiBox->GetBlocks();
+
+			for (int32 Index = 0; Index < MenuBlocks.Num(); Index++)
+			{
+				TSharedRef<const FMenuEntryBlock> Block = StaticCastSharedRef<const FMenuEntryBlock>(MenuBlocks[Index]);
+				FMacMenu* Menu = [[FMacMenu alloc] initWithMenuEntryBlock:Block];
+				NSString* Title = FSlateMacMenu::GetMenuItemTitle(Block);
+				[Menu setTitle:Title];
+
+				NSMenuItem* MenuItem = [[NSMenuItem new] autorelease];
+				[MenuItem setTitle:Title];
+				[[NSApp mainMenu] addItem:MenuItem];
+				[MenuItem setSubmenu:Menu];
+
+				const bool bIsWindowMenu = (WindowLabel.ToString().Compare(FString(Title)) == 0);
+				if (bIsWindowMenu)
+				{
+					[NSApp setWindowsMenu:nil];
+
+					[Menu removeAllItems];
+
+					NSMenuItem* MinimizeItem = [[[NSMenuItem alloc] initWithTitle:@"Minimize" action:@selector(miniaturize:) keyEquivalent:@"m"] autorelease];
+					NSMenuItem* ZoomItem = [[[NSMenuItem alloc] initWithTitle:@"Zoom" action:@selector(performZoom:) keyEquivalent:@""] autorelease];
+					NSMenuItem* CloseItem = [[[NSMenuItem alloc] initWithTitle:@"Close" action:@selector(performClose:) keyEquivalent:@"w"] autorelease];
+					NSMenuItem* BringAllToFrontItem = [[[NSMenuItem alloc] initWithTitle:@"Bring All to Front" action:@selector(arrangeInFront:) keyEquivalent:@""] autorelease];
+
+					[Menu addItem:MinimizeItem];
+					[Menu addItem:ZoomItem];
+					[Menu addItem:CloseItem];
+					[Menu addItem:[NSMenuItem separatorItem]];
+					[Menu addItem:BringAllToFrontItem];
+					[Menu addItem:[NSMenuItem separatorItem]];
+
+					[NSApp setWindowsMenu:Menu];
+					[Menu addItem:[NSMenuItem separatorItem]];
+				}
+			}
+
+			if ([[NSApp mainMenu] itemWithTitle:HelpTitle.ToString().GetNSString()] == nil && HelpMenu)
 			{
-				[NSApp setWindowsMenu:nil];
-				
-				[Menu removeAllItems];
-				
-				NSMenuItem* MinimizeItem = [[[NSMenuItem alloc] initWithTitle:@"Minimize" action:@selector(miniaturize:) keyEquivalent:@"m"] autorelease];
-				NSMenuItem* ZoomItem = [[[NSMenuItem alloc] initWithTitle:@"Zoom" action:@selector(performZoom:) keyEquivalent:@""] autorelease];
-				NSMenuItem* CloseItem = [[[NSMenuItem alloc] initWithTitle:@"Close" action:@selector(performClose:) keyEquivalent:@"w"] autorelease];
-				NSMenuItem* BringAllToFrontItem = [[[NSMenuItem alloc] initWithTitle:@"Bring All to Front" action:@selector(arrangeInFront:) keyEquivalent:@""] autorelease];
-				
-				[Menu addItem:MinimizeItem];
-				[Menu addItem:ZoomItem];
-				[Menu addItem:CloseItem];
-				[Menu addItem:[NSMenuItem separatorItem]];
-				[Menu addItem:BringAllToFrontItem];
-				[Menu addItem:[NSMenuItem separatorItem]];
-				
-				[NSApp setWindowsMenu:Menu];
-				[Menu addItem:[NSMenuItem separatorItem]];
+				[[NSApp mainMenu] addItem:HelpMenu];
 			}
-		}
 
-		if ([[NSApp mainMenu] itemWithTitle:HelpTitle.ToString().GetNSString()] == nil && HelpMenu)
-		{
-			[[NSApp mainMenu] addItem:HelpMenu];
-		}
+			if (HelpMenu)
+			{
+				[HelpMenu release];
+			}
 
-		if (HelpMenu)
-		{
-			[HelpMenu release];
+			GCurrentMultiBox = MultiBox;
+			FPlatformMisc::bChachedMacMenuStateNeedsUpdate = true;
 		}
 	});
 }
 
 void FSlateMacMenu::UpdateMenu(FMacMenu* Menu)
 {
-	GameThreadCall(^{
+	MainThreadCall(^{
 		FText WindowLabel = NSLOCTEXT("MainMenu", "WindowMenu", "Window");
 		const bool bIsWindowMenu = (WindowLabel.ToString().Compare(FString([Menu title])) == 0);
 
@@ -129,41 +166,28 @@ void FSlateMacMenu::UpdateMenu(FMacMenu* Menu)
 			}
 		}
 
-		if (!Menu.MultiBox.IsValid())
-		{
-			const bool bShouldCloseWindowAfterMenuSelection = true;
-			FMenuBuilder MenuBuilder(bShouldCloseWindowAfterMenuSelection, Menu.MenuEntryBlock->GetActionList(), Menu.MenuEntryBlock->Extender);
-			{
-				// Have the menu fill its contents
-				Menu.MenuEntryBlock->EntryBuilder.ExecuteIfBound(MenuBuilder);
-			}
-			TSharedRef<SWidget> Widget = MenuBuilder.MakeWidget();
-			Menu.MultiBox = TSharedPtr<const FMultiBox>(StaticCastSharedRef<SMultiBoxWidget>(Widget)->GetMultiBox());
-		}
-
-		const TArray<TSharedRef<const FMultiBlock>>& MenuBlocks = Menu.MultiBox->GetBlocks();
+		TArray<FMacMenuItemState>& MenuState = GCachedMenuState[Menu];
 		int32 ItemIndexAdjust = 0;
-		for (int32 Index = 0; Index < MenuBlocks.Num(); Index++)
+		for (int32 Index = 0; Index < MenuState.Num(); Index++)
 		{
+			FMacMenuItemState& MenuItemState = MenuState[Index];
 			const int32 ItemIndex = (bIsWindowMenu ? Index + ItemIndexOffset : Index) - ItemIndexAdjust;
 			NSMenuItem* MenuItem = [Menu numberOfItems] > ItemIndex ? [Menu itemAtIndex:ItemIndex] : nil;
 
-			EMultiBlockType::Type Type = MenuBlocks[Index]->GetType();
-			if (Type == EMultiBlockType::MenuEntry)
+			if (MenuItemState.Type == EMultiBlockType::MenuEntry)
 			{
-				TSharedRef<const FMenuEntryBlock> Block = StaticCastSharedRef<const FMenuEntryBlock>(MenuBlocks[Index]);
-				if (MenuItem && (![MenuItem isKindOfClass:[FMacMenuItem class]] || (Block->bIsSubMenu && [MenuItem submenu] == nil) || (!Block->bIsSubMenu && [MenuItem submenu] != nil)))
+				if (MenuItem && (![MenuItem isKindOfClass:[FMacMenuItem class]] || (MenuItemState.IsSubMenu && [MenuItem submenu] == nil) || (!MenuItemState.IsSubMenu && [MenuItem submenu] != nil)))
 				{
 					[Menu removeItem:MenuItem];
 					MenuItem = nil;
 				}
 				if (!MenuItem)
 				{
-					MenuItem = [[[FMacMenuItem alloc] initWithMenuEntryBlock:Block] autorelease];
+					MenuItem = [[[FMacMenuItem alloc] initWithMenuEntryBlock:MenuItemState.Block] autorelease];
 
-					if (Block->bIsSubMenu)
+					if (MenuItemState.IsSubMenu)
 					{
-						FMacMenu* SubMenu = [[FMacMenu alloc] initWithMenuEntryBlock:Block];
+						FMacMenu* SubMenu = [[FMacMenu alloc] initWithMenuEntryBlock:MenuItemState.Block];
 						[MenuItem setSubmenu:SubMenu];
 					}
 
@@ -177,20 +201,17 @@ void FSlateMacMenu::UpdateMenu(FMacMenu* Menu)
 					}
 				}
 
-				[MenuItem setTitle:FSlateMacMenu::GetMenuItemTitle(Block)];
+				[MenuItem setTitle:MenuItemState.Title];
 
-				uint32 KeyModifiers;
-				NSString* KeyEquivalent = FSlateMacMenu::GetMenuItemKeyEquivalent(Block, &KeyModifiers);
-				[MenuItem setKeyEquivalent:KeyEquivalent];
-				[MenuItem setKeyEquivalentModifierMask:KeyModifiers];
+				[MenuItem setKeyEquivalent:MenuItemState.KeyEquivalent];
+				[MenuItem setKeyEquivalentModifierMask:MenuItemState.KeyModifiers];
 
 				if (bIsWindowMenu)
 				{
-					NSImage* MenuImage = FSlateMacMenu::GetMenuItemIcon(Block);
+					NSImage* MenuImage = MenuItemState.Icon;
 					if(MenuImage)
 					{
 						[MenuItem setImage:MenuImage];
-						[MenuImage release];
 					}
 				}
 				else
@@ -198,18 +219,18 @@ void FSlateMacMenu::UpdateMenu(FMacMenu* Menu)
 					[MenuItem setImage:nil];
 				}
 
-				if (FSlateMacMenu::IsMenuItemEnabled(Block))
+				if (MenuItemState.IsEnabled)
 				{
 					[MenuItem setTarget:MenuItem];
 					[MenuItem setAction:@selector(performAction)];
 				}
 
-				if (!Block->bIsSubMenu)
+				if (!MenuItemState.IsSubMenu)
 				{
-					[MenuItem setState:FSlateMacMenu::GetMenuItemState(Block)];
+					[MenuItem setState:MenuItemState.State];
 				}
 			}
-			else if (Type == EMultiBlockType::MenuSeparator)
+			else if (MenuItemState.Type == EMultiBlockType::MenuSeparator)
 			{
 				if (MenuItem && ![MenuItem isSeparatorItem])
 				{
@@ -233,7 +254,54 @@ void FSlateMacMenu::UpdateMenu(FMacMenu* Menu)
 				ItemIndexAdjust++;
 			}
 		}
-	}, InGameRunLoopMode(@[ UE4NilEventMode, UE4ShowEventMode, UE4ResizeEventMode, UE4FullscreenEventMode, UE4CloseEventMode, UE4IMEEventMode ]));
+	});
+}
+
+void FSlateMacMenu::UpdateCachedState()
+{
+	for (TMap<FMacMenu*, TArray<FMacMenuItemState>>::TIterator It(GCachedMenuState); It; ++It)
+	{
+		FMacMenu* Menu = It.Key();
+		TArray<FMacMenuItemState>& MenuState = It.Value();
+
+		if (!Menu.MultiBox.IsValid())
+		{
+			const bool bShouldCloseWindowAfterMenuSelection = true;
+			FMenuBuilder MenuBuilder(bShouldCloseWindowAfterMenuSelection, Menu.MenuEntryBlock->GetActionList(), Menu.MenuEntryBlock->Extender);
+			{
+				// Have the menu fill its contents
+				Menu.MenuEntryBlock->EntryBuilder.ExecuteIfBound(MenuBuilder);
+			}
+			TSharedRef<SWidget> Widget = MenuBuilder.MakeWidget();
+			Menu.MultiBox = TSharedPtr<const FMultiBox>(StaticCastSharedRef<SMultiBoxWidget>(Widget)->GetMultiBox());
+		}
+
+		const TArray<TSharedRef<const FMultiBlock>>& MenuBlocks = Menu.MultiBox->GetBlocks();
+		for (int32 Index = MenuState.Num(); MenuBlocks.Num() > MenuState.Num(); Index++)
+		{
+			MenuState.Add(FMacMenuItemState());
+		}
+		for (int32 Index = 0; Index < MenuBlocks.Num(); Index++)
+		{
+			FMacMenuItemState& ItemState = MenuState[Index];
+			ItemState.Type = MenuBlocks[Index]->GetType();
+
+			if (ItemState.Type == EMultiBlockType::MenuEntry)
+			{
+				TSharedRef<const FMenuEntryBlock> Block = StaticCastSharedRef<const FMenuEntryBlock>(MenuBlocks[Index]);
+				ItemState.Block = Block;
+				ItemState.Title = FSlateMacMenu::GetMenuItemTitle(Block);
+				ItemState.KeyEquivalent = FSlateMacMenu::GetMenuItemKeyEquivalent(Block, &ItemState.KeyModifiers);
+				if (!ItemState.Icon)
+				{
+					ItemState.Icon = FSlateMacMenu::GetMenuItemIcon(Block);
+				}
+				ItemState.IsSubMenu = Block->bIsSubMenu;
+				ItemState.IsEnabled = FSlateMacMenu::IsMenuItemEnabled(Block);
+				ItemState.State = ItemState.IsSubMenu ? 0 : FSlateMacMenu::GetMenuItemState(Block);
+			}
+		}
+	}
 }
 
 void FSlateMacMenu::ExecuteMenuItemAction(const TSharedRef< const class FMenuEntryBlock >& Block)
diff --git a/Engine/Source/Runtime/Slate/Private/Framework/MultiBox/Mac/MacMenu.h b/Engine/Source/Runtime/Slate/Private/Framework/MultiBox/Mac/MacMenu.h
--- a/Engine/Source/Runtime/Slate/Private/Framework/MultiBox/Mac/MacMenu.h
+++ b/Engine/Source/Runtime/Slate/Private/Framework/MultiBox/Mac/MacMenu.h
@@ -17,6 +17,7 @@ public:
 
 	static void UpdateWithMultiBox(const TSharedRef<FMultiBox>& MultiBox);
 	static void UpdateMenu(FMacMenu* Menu);
+	static void UpdateCachedState();
 	static void ExecuteMenuItemAction(const TSharedRef<const FMenuEntryBlock>& Block);
 
 private:
