UMG/Slate - Moving Scale Box into Slate!

diff --git a/Engine/Source/Runtime/Slate/Private/Widgets/Layout/SScaleBox.cpp b/Engine/Source/Runtime/Slate/Private/Widgets/Layout/SScaleBox.cpp
new file mode 100644
--- /dev/null
+++ b/Engine/Source/Runtime/Slate/Private/Widgets/Layout/SScaleBox.cpp
@@ -0,0 +1,115 @@
+// Copyright 1998-2014 Epic Games, Inc. All Rights Reserved.
+
+#include "SlatePrivatePCH.h"
+
+#include "LayoutUtils.h"
+#include "SScaleBox.h"
+
+
+/* SScaleBox interface
+ *****************************************************************************/
+
+void SScaleBox::Construct( const SScaleBox::FArguments& InArgs )
+{
+	ChildSlot
+	[
+		InArgs._Content.Widget
+	];
+}
+
+/* SWidget overrides
+ *****************************************************************************/
+
+void SScaleBox::OnArrangeChildren( const FGeometry& AllottedGeometry, FArrangedChildren& ArrangedChildren ) const
+{
+	const EVisibility ChildVisibility = ChildSlot.GetWidget()->GetVisibility();
+	if ( ArrangedChildren.Accepts(ChildVisibility) )
+	{
+		FVector2D DesiredSize = ChildSlot.GetWidget()->GetDesiredSize();
+
+		float FinalScale = 1;
+
+		EStretch::Type CurrentStretch = Stretch.Get();
+		EStretchDirection::Type CurrentStretchDirection = StretchDirection.Get();
+
+		if ( DesiredSize.X != 0 && DesiredSize.Y != 0 )
+		{
+			switch ( CurrentStretch )
+			{
+			case EStretch::None:
+				break;
+			case EStretch::Fill:
+				DesiredSize = AllottedGeometry.Size;
+				break;
+			case EStretch::ScaleToFit:
+				FinalScale = FMath::Min(AllottedGeometry.Size.X / DesiredSize.X, AllottedGeometry.Size.Y / DesiredSize.Y);
+				break;
+			case EStretch::ScaleToFill:
+				FinalScale = FMath::Max(AllottedGeometry.Size.X / DesiredSize.X, AllottedGeometry.Size.Y / DesiredSize.Y);
+				break;
+			}
+
+			switch ( CurrentStretchDirection )
+			{
+			case EStretchDirection::DownOnly:
+				FinalScale = FMath::Min(FinalScale, 1.0f);
+				break;
+			case EStretchDirection::UpOnly:
+				FinalScale = FMath::Max(FinalScale, 1.0f);
+				break;
+			}
+		}
+
+		FVector2D FinalOffset(0, 0);
+
+		if ( CurrentStretch != EStretch::Fill )
+		{
+			const FMargin SlotPadding(ChildSlot.SlotPadding.Get());
+			AlignmentArrangeResult XResult = AlignChild<Orient_Horizontal>(AllottedGeometry.Size.X, ChildSlot, SlotPadding, FinalScale, false);
+			AlignmentArrangeResult YResult = AlignChild<Orient_Vertical>(AllottedGeometry.Size.Y, ChildSlot, SlotPadding, FinalScale, false);
+
+			FinalOffset = FVector2D(XResult.Offset, YResult.Offset) * ( 1.0f / FinalScale );
+		}
+
+		ArrangedChildren.AddWidget(ChildVisibility, AllottedGeometry.MakeChild(
+			ChildSlot.GetWidget(),
+			FinalOffset,
+			DesiredSize,
+			FinalScale
+		) );
+	}
+}
+
+
+int32 SScaleBox::OnPaint( const FPaintArgs& Args, const FGeometry& AllottedGeometry, const FSlateRect& MyClippingRect, FSlateWindowElementList& OutDrawElements, int32 LayerId, const FWidgetStyle& InWidgetStyle, bool bParentEnabled ) const
+{
+	return SCompoundWidget::OnPaint(Args, AllottedGeometry, MyClippingRect, OutDrawElements, LayerId, InWidgetStyle, bParentEnabled);
+}
+
+void SScaleBox::SetContent(TSharedRef<SWidget> InContent)
+{
+	ChildSlot
+	[
+		InContent
+	];
+}
+
+void SScaleBox::SetHAlign(EHorizontalAlignment HAlign)
+{
+	ChildSlot.HAlignment = HAlign;
+}
+
+void SScaleBox::SetVAlign(EVerticalAlignment VAlign)
+{
+	ChildSlot.VAlignment = VAlign;
+}
+
+void SScaleBox::SetStretchDirection(EStretchDirection::Type InStretchDirection)
+{
+	StretchDirection = InStretchDirection;
+}
+
+void SScaleBox::SetStretch(EStretch::Type InStretch)
+{
+	Stretch = InStretch;
+}
diff --git a/Engine/Source/Runtime/Slate/Public/Widgets/Layout/SScaleBox.h b/Engine/Source/Runtime/Slate/Public/Widgets/Layout/SScaleBox.h
new file mode 100644
--- /dev/null
+++ b/Engine/Source/Runtime/Slate/Public/Widgets/Layout/SScaleBox.h
@@ -0,0 +1,97 @@
+// Copyright 1998-2014 Epic Games, Inc. All Rights Reserved.
+
+#pragma once
+
+#include "SScaleBox.generated.h"
+
+UENUM()
+namespace EStretchDirection
+{
+	enum Type
+	{
+		/** Will scale the content up or down */
+		Both,
+		/** Will only make the content smaller, will never scale it larger than the content's desired size. */
+		DownOnly,
+		/** Will only make the content larger, will never scale it smaller than the content's desired size. */
+		UpOnly
+	};
+}
+
+UENUM()
+namespace EStretch
+{
+	enum Type
+	{
+		/** Does not scale the content. */
+		None,
+		/** Scales the content non-uniformly filling the entire space of the area */
+		Fill,
+		/** Scales the content uniformly (preserving aspect ratio) until it can no longer scale the content without clipping it. */
+		ScaleToFit,
+		/** Scales the content uniformly (preserving aspect ratio), until all sides meet or exceed the size of the area.  Will result in clipping longer sides. */
+		ScaleToFill
+	};
+}
+
+/**
+ * Allows you to place content with a desired size and have it scale to meet the constraints placed on this box's alloted area.  If
+ * you needed to have a background image scale to fill an area but not become distorted with different aspect ratios, or if you need
+ * to auto fit some text to an area, this is the control for you.
+ */
+class SLATE_API SScaleBox : public SCompoundWidget
+{
+public:
+	SLATE_BEGIN_ARGS(SScaleBox)
+	: _Content()
+	, _HAlign(HAlign_Center)
+	, _VAlign(VAlign_Center)
+	, _StretchDirection(EStretchDirection::Both)
+	, _Stretch(EStretch::None)
+	{}
+		/** Slot for this designers content (optional) */
+		SLATE_DEFAULT_SLOT(FArguments, Content)
+
+		/** The horizontal alignment of the content */
+		SLATE_ARGUMENT(EHorizontalAlignment, HAlign)
+
+		/** The vertical alignment of the content */
+		SLATE_ARGUMENT(EVerticalAlignment, VAlign)
+		
+		/** Controls in what direction content can be scaled */
+		SLATE_ATTRIBUTE(EStretchDirection::Type, StretchDirection)
+		
+		/** The stretching rule to apply when content is stretched */
+		SLATE_ATTRIBUTE(EStretch::Type, Stretch)
+
+	SLATE_END_ARGS()
+
+	void Construct(const FArguments& InArgs);
+	
+	// SWidget interface
+	virtual void OnArrangeChildren(const FGeometry& AllottedGeometry, FArrangedChildren& ArrangedChildren) const override;
+	virtual int32 OnPaint(const FPaintArgs& Args, const FGeometry& AllottedGeometry, const FSlateRect& MyClippingRect, FSlateWindowElementList& OutDrawElements, int32 LayerId, const FWidgetStyle& InWidgetStyle, bool bParentEnabled) const override;
+	// End SWidget of interface
+
+	/** See Content slot */
+	void SetContent(TSharedRef<SWidget> InContent);
+
+	/** See HAlign argument */
+	void SetHAlign(EHorizontalAlignment HAlign);
+
+	/** See VAlign argument */
+	void SetVAlign(EVerticalAlignment VAlign);
+
+	/** See StretchDirection argument */
+	void SetStretchDirection(EStretchDirection::Type InStretchDirection);
+
+	/** See Stretch argument */
+	void SetStretch(EStretch::Type InStretch);
+	
+private:
+	/** The allowed direction of stretching of the content */
+	TAttribute<EStretchDirection::Type> StretchDirection;
+
+	/** The method of scaling that is applied to the content. */
+	TAttribute<EStretch::Type> Stretch;
+};
diff --git a/Engine/Source/Runtime/UMG/Private/Slate/SScaleBox.cpp b/Engine/Source/Runtime/UMG/Private/Slate/SScaleBox.cpp
deleted file mode 100644
diff --git a/Engine/Source/Runtime/UMG/Public/Slate/SScaleBox.h b/Engine/Source/Runtime/UMG/Public/Slate/SScaleBox.h
deleted file mode 100644
