Fix memory corruption bug in the new Mac menu state caching code, adding a submenu was modifying the state map which was obliterating the state array. The map now holds a shared pointer to the array rather than the array itself.

diff --git a/Engine/Source/Runtime/Slate/Private/Framework/MultiBox/Mac/MacMenu.cpp b/Engine/Source/Runtime/Slate/Private/Framework/MultiBox/Mac/MacMenu.cpp
--- a/Engine/Source/Runtime/Slate/Private/Framework/MultiBox/Mac/MacMenu.cpp
+++ b/Engine/Source/Runtime/Slate/Private/Framework/MultiBox/Mac/MacMenu.cpp
@@ -19,7 +19,7 @@ struct FMacMenuItemState
 	FMacMenuItemState() : Type(EMultiBlockType::None), Title(nil), KeyEquivalent(nil), KeyModifiers(0), Icon(nil), IsSubMenu(false), IsEnabled(false), State(0) {}
 };
 
-static TMap<FMacMenu*, TArray<FMacMenuItemState>> GCachedMenuState;
+static TMap<FMacMenu*, TSharedPtr<TArray<FMacMenuItemState>>> GCachedMenuState;
 
 @interface FMacMenuItem : NSMenuItem
 @property (assign) TSharedPtr<const FMenuEntryBlock> MenuEntryBlock;
@@ -53,7 +53,7 @@ static TMap<FMacMenu*, TArray<FMacMenuItemState>> GCachedMenuState;
 	self = [super initWithTitle:@""];
 	[self setDelegate:self];
 	self.MenuEntryBlock = Block;
-	GCachedMenuState.Add(self, TArray<FMacMenuItemState>());
+	GCachedMenuState.Add(self, TSharedPtr<TArray<FMacMenuItemState>>(new TArray<FMacMenuItemState>()));
 	return self;
 }
 
@@ -166,11 +166,11 @@ void FSlateMacMenu::UpdateMenu(FMacMenu* Menu)
 			}
 		}
 
-		TArray<FMacMenuItemState>& MenuState = GCachedMenuState[Menu];
+		TSharedPtr<TArray<FMacMenuItemState>> MenuState = GCachedMenuState[Menu];
 		int32 ItemIndexAdjust = 0;
-		for (int32 Index = 0; Index < MenuState.Num(); Index++)
+		for (int32 Index = 0; Index < MenuState->Num(); Index++)
 		{
-			FMacMenuItemState& MenuItemState = MenuState[Index];
+			FMacMenuItemState& MenuItemState = (*MenuState)[Index];
 			const int32 ItemIndex = (bIsWindowMenu ? Index + ItemIndexOffset : Index) - ItemIndexAdjust;
 			NSMenuItem* MenuItem = [Menu numberOfItems] > ItemIndex ? [Menu itemAtIndex:ItemIndex] : nil;
 
@@ -259,10 +259,10 @@ void FSlateMacMenu::UpdateMenu(FMacMenu* Menu)
 
 void FSlateMacMenu::UpdateCachedState()
 {
-	for (TMap<FMacMenu*, TArray<FMacMenuItemState>>::TIterator It(GCachedMenuState); It; ++It)
+	for (TMap<FMacMenu*, TSharedPtr<TArray<FMacMenuItemState>>>::TIterator It(GCachedMenuState); It; ++It)
 	{
 		FMacMenu* Menu = It.Key();
-		TArray<FMacMenuItemState>& MenuState = It.Value();
+		TSharedPtr<TArray<FMacMenuItemState>> MenuState = It.Value();
 
 		if (!Menu.MultiBox.IsValid())
 		{
@@ -277,13 +277,13 @@ void FSlateMacMenu::UpdateCachedState()
 		}
 
 		const TArray<TSharedRef<const FMultiBlock>>& MenuBlocks = Menu.MultiBox->GetBlocks();
-		for (int32 Index = MenuState.Num(); MenuBlocks.Num() > MenuState.Num(); Index++)
+		for (int32 Index = MenuState->Num(); MenuBlocks.Num() > MenuState->Num(); Index++)
 		{
-			MenuState.Add(FMacMenuItemState());
+			MenuState->Add(FMacMenuItemState());
 		}
 		for (int32 Index = 0; Index < MenuBlocks.Num(); Index++)
 		{
-			FMacMenuItemState& ItemState = MenuState[Index];
+			FMacMenuItemState& ItemState = (*MenuState)[Index];
 			ItemState.Type = MenuBlocks[Index]->GetType();
 
 			if (ItemState.Type == EMultiBlockType::MenuEntry)
