Limit Mac-specific window activation handling for drag and drop to left mouse button presses.

diff --git a/Engine/Source/Runtime/Slate/Private/Framework/Application/SlateApplication.cpp b/Engine/Source/Runtime/Slate/Private/Framework/Application/SlateApplication.cpp
--- a/Engine/Source/Runtime/Slate/Private/Framework/Application/SlateApplication.cpp
+++ b/Engine/Source/Runtime/Slate/Private/Framework/Application/SlateApplication.cpp
@@ -3332,7 +3332,7 @@ bool FSlateApplication::ProcessMouseButtonDownEvent( const TSharedPtr< FGenericW
 			// If none of the widgets requested keyboard focus to be set (or set the keyboard focus explicitly), set it to the leaf-most widget under the mouse.
 			// On Mac we prevent the OS from activating the window on mouse down, so we have full control and can activate only if there's nothing draggable under the mouse cursor.
 			const bool bFocusChangedByEventHandler = PreviouslyFocusedWidget != GetKeyboardFocusedWidget();
-			if ((!Reply.GetFocusRecepient().IsValid() || (PLATFORM_MAC && !DragDetector.DetectDragForWidget.IsValid())) && !bFocusChangedByEventHandler)
+			if ((!Reply.GetFocusRecepient().IsValid() || (PLATFORM_MAC && MouseEvent.GetEffectingButton() == EKeys::LeftMouseButton && !DragDetector.DetectDragForWidget.IsValid())) && !bFocusChangedByEventHandler)
 			{
 				// The event handler for OnMouseButton down may have altered the widget hierarchy.
 				// Refresh the previously-cached widget path.
@@ -3351,7 +3351,7 @@ bool FSlateApplication::ProcessMouseButtonDownEvent( const TSharedPtr< FGenericW
 				}
 
 #if PLATFORM_MAC
-				if (WidgetsUnderCursor.TopLevelWindow.IsValid() && !DragDetector.DetectDragForWidget.IsValid() && ActiveWindow == [NSApp keyWindow])
+				if (MouseEvent.GetEffectingButton() == EKeys::LeftMouseButton && WidgetsUnderCursor.TopLevelWindow.IsValid() && !DragDetector.DetectDragForWidget.IsValid() && ActiveWindow == [NSApp keyWindow])
 				{
 					MouseCaptorHelper Captor = MouseCaptor;
 					FPlatformMisc::ActivateApplication();
@@ -3497,7 +3497,7 @@ bool FSlateApplication::ProcessMouseButtonUpEvent( FPointerEvent& MouseEvent )
 			ProcessReply(MouseCaptorPath, Reply, &MouseCaptorPath, &MouseEvent);
 #if PLATFORM_MAC
 			TSharedPtr<SWindow> ActiveWindow = GetActiveTopLevelWindow();
-			if (MouseCaptorPath.TopLevelWindow.IsValid() && ActiveWindow != MouseCaptorPath.TopLevelWindow)
+			if (MouseEvent.GetEffectingButton() == EKeys::LeftMouseButton && MouseCaptorPath.TopLevelWindow.IsValid() && ActiveWindow != MouseCaptorPath.TopLevelWindow)
 			{
 				MouseCaptorPath.TopLevelWindow->BringToFront(true);
 			}
