Make sure ObjC objects in Mac menu state cache are not released by autorelease pool too early

diff --git a/Engine/Source/Runtime/Slate/Private/Framework/MultiBox/Mac/MacMenu.cpp b/Engine/Source/Runtime/Slate/Private/Framework/MultiBox/Mac/MacMenu.cpp
--- a/Engine/Source/Runtime/Slate/Private/Framework/MultiBox/Mac/MacMenu.cpp
+++ b/Engine/Source/Runtime/Slate/Private/Framework/MultiBox/Mac/MacMenu.cpp
@@ -17,6 +17,12 @@ struct FMacMenuItemState
 	uint32 State;
 
 	FMacMenuItemState() : Type(EMultiBlockType::None), Title(nil), KeyEquivalent(nil), KeyModifiers(0), Icon(nil), IsSubMenu(false), IsEnabled(false), State(0) {}
+	~FMacMenuItemState()
+	{
+		if (Title) [Title release];
+		if (KeyEquivalent) [KeyEquivalent release];
+		if (Icon) [Icon release];
+	}
 };
 
 static TMap<FMacMenu*, TSharedPtr<TArray<FMacMenuItemState>>> GCachedMenuState;
@@ -290,11 +296,11 @@ void FSlateMacMenu::UpdateCachedState()
 			{
 				TSharedRef<const FMenuEntryBlock> Block = StaticCastSharedRef<const FMenuEntryBlock>(MenuBlocks[Index]);
 				ItemState.Block = Block;
-				ItemState.Title = FSlateMacMenu::GetMenuItemTitle(Block);
-				ItemState.KeyEquivalent = FSlateMacMenu::GetMenuItemKeyEquivalent(Block, &ItemState.KeyModifiers);
+				ItemState.Title = [FSlateMacMenu::GetMenuItemTitle(Block) retain];
+				ItemState.KeyEquivalent = [FSlateMacMenu::GetMenuItemKeyEquivalent(Block, &ItemState.KeyModifiers) retain];
 				if (!ItemState.Icon)
 				{
-					ItemState.Icon = FSlateMacMenu::GetMenuItemIcon(Block);
+					ItemState.Icon = [FSlateMacMenu::GetMenuItemIcon(Block) retain];
 				}
 				ItemState.IsSubMenu = Block->bIsSubMenu;
 				ItemState.IsEnabled = FSlateMacMenu::IsMenuItemEnabled(Block);
