Fixed SynthesizeMouseMove() to properly signal to Slate that it is a synthetic event. However, we should be smarter about sending out synthesized mouse moves when the cursor is over a non-slate window.

Resolves TTP 347873: UE4: EDITOR: SLATE: Illegal call to StaticFindObject() while serializing object data
Resolved TTP 347767: Regression: EDITOR: Tooltips appear even if editor does not have focus and a window is layered on top of it

diff --git a/Engine/Source/Runtime/Slate/Private/Framework/Application/SlateApplication.cpp b/Engine/Source/Runtime/Slate/Private/Framework/Application/SlateApplication.cpp
--- a/Engine/Source/Runtime/Slate/Private/Framework/Application/SlateApplication.cpp
+++ b/Engine/Source/Runtime/Slate/Private/Framework/Application/SlateApplication.cpp
@@ -2528,7 +2528,17 @@ bool FSlateApplication::IsWindowInDestroyQueue(TSharedRef<SWindow> Window) const
 
 void FSlateApplication::SynthesizeMouseMove()
 {
-	FPointerEvent MouseEvent(
+	// Synthetic mouse events accomplish two goals:
+	// 1) The UI can change even if the mosue doesn't move.
+	//    Synthesizing a mouse move sends out events.
+	//    In this case, the current and previous position will be the same.
+	//
+	// 2) The mouse moves, but the OS decided not to send us an event.
+	//    e.g. Mouse moved outside of our window.
+	//    In this case, the previous and current positions differ.
+
+	FPointerEvent MouseEvent
+	(
 		CursorPointerIndex,
 		GetCursorPos(),
 		GetLastCursorPos(),
@@ -2536,9 +2546,9 @@ void FSlateApplication::SynthesizeMouseMove()
 		EKeys::Invalid,
 		0,
 		PlatformApplication->GetModifierKeys()
-		);
+	);
 
-	ProcessMouseMoveEvent(MouseEvent);
+	ProcessMouseMoveEvent(MouseEvent, true);
 }
 
 void FSlateApplication::OnLogSlateEvent(EEventLog::Type Elovent, const FString& AdditionalContent)
@@ -3683,9 +3693,9 @@ bool FSlateApplication::OnRawMouseMove( const int32 X, const int32 Y )
 	return true;
 }
 
-bool FSlateApplication::ProcessMouseMoveEvent( FPointerEvent& MouseEvent )
+bool FSlateApplication::ProcessMouseMoveEvent( FPointerEvent& MouseEvent, bool bIsSynthetic )
 {
-	if ( !MouseEvent.GetCursorDelta().IsZero() )
+	if ( !bIsSynthetic )
 	{
 		// Detecting a mouse move of zero delta is our way of filtering out synthesized move events
 		const bool AllowSpawningOfToolTips = true;
diff --git a/Engine/Source/Runtime/Slate/Public/Framework/Application/SlateApplication.h b/Engine/Source/Runtime/Slate/Public/Framework/Application/SlateApplication.h
--- a/Engine/Source/Runtime/Slate/Public/Framework/Application/SlateApplication.h
+++ b/Engine/Source/Runtime/Slate/Public/Framework/Application/SlateApplication.h
@@ -673,10 +673,11 @@ public:
 	 * Called by the native application in response to a mouse move. Routs the event to Slate Widgets.
 	 *
 	 * @param  InMouseEvent  Mouse event
+	 * @param  bIsSynthetic  True when the even is synthesized by slate.
 	 *
 	 * @return  Was this event handled by the Slate application?
 	 */
-	bool ProcessMouseMoveEvent( FPointerEvent& MouseEvent );
+	bool ProcessMouseMoveEvent( FPointerEvent& MouseEvent, bool bIsSynthetic = false );
 
 	/**
 	 * Called by the native application in response to a mouse button press. Routs the event to Slate Widgets.
