Fixed a potential crash when calling SMultiLineEditableText::SetEditableText

TTP# 347830 - EDITOR: CRASH: Possible to get a crash in a multi line text field after resetting it to default

It was possible for the cursor to be left past the end of the newly set text, resulting in a crash the next time the cursor was updated. We now ensure that the cursor is always in a valid position after calling SetEditableText.

diff --git a/Engine/Source/Runtime/Slate/Private/Widgets/Text/SMultiLineEditableText.cpp b/Engine/Source/Runtime/Slate/Private/Widgets/Text/SMultiLineEditableText.cpp
--- a/Engine/Source/Runtime/Slate/Private/Widgets/Text/SMultiLineEditableText.cpp
+++ b/Engine/Source/Runtime/Slate/Private/Widgets/Text/SMultiLineEditableText.cpp
@@ -406,6 +406,22 @@ bool SMultiLineEditableText::SetEditableText(const FText& TextToSet, const bool
 			TextLayout->AddLine(LineText, Runs);
 		}
 
+		{
+			const FTextLocation OldCursorPos = CursorInfo.GetCursorInteractionLocation();
+
+			// Make sure the cursor is still at a valid location
+			if(OldCursorPos.GetLineIndex() >= Lines.Num() || OldCursorPos.GetOffset() > Lines[OldCursorPos.GetLineIndex()].Text->Len())
+			{
+				const int32 LastLineIndex = Lines.Num() - 1;
+				const FTextLocation NewCursorPosition = FTextLocation(LastLineIndex, Lines[LastLineIndex].Text->Len());
+
+				CursorInfo.SetCursorLocationAndCalculateAlignment(TextLayout, NewCursorPosition);
+				OnCursorMoved.ExecuteIfBound(CursorInfo.GetCursorInteractionLocation());
+				UpdatePreferredCursorScreenOffsetInLine();
+				UpdateCursorHighlight();
+			}
+		}
+
 		return true;
 	}
 
