SlateApplication - When resetting to default input state, use ProcessReply to reset mouse capture.

diff --git a/Engine/Source/Runtime/Slate/Private/Framework/Application/SlateApplication.cpp b/Engine/Source/Runtime/Slate/Private/Framework/Application/SlateApplication.cpp
--- a/Engine/Source/Runtime/Slate/Private/Framework/Application/SlateApplication.cpp
+++ b/Engine/Source/Runtime/Slate/Private/Framework/Application/SlateApplication.cpp
@@ -1459,9 +1459,9 @@ void FSlateApplication::SetKeyboardFocus( const TSharedPtr< SWidget >& OptionalW
 
 void FSlateApplication::ResetToDefaultInputSettings()
 {
-	if (MouseCaptor.HasCapture())
+	for (auto MouseCaptorPath : MouseCaptor.ToWidgetPaths())
 	{
-		MouseCaptor.InvalidateCaptureForAllPointers();
+		ProcessReply(MouseCaptorPath, FReply::Handled().ReleaseMouseCapture(), NULL, NULL);
 	}
 
 	for (int32 UserIndex = 0; UserIndex < ARRAY_COUNT(JoystickCaptorWeakPaths); UserIndex++)
