TArray::FindMatch deprecation as it does the same as TArray::IndexOfByPredicate.

diff --git a/Engine/Source/Developer/Profiler/Private/Widgets/SEventGraph.cpp b/Engine/Source/Developer/Profiler/Private/Widgets/SEventGraph.cpp
--- a/Engine/Source/Developer/Profiler/Private/Widgets/SEventGraph.cpp
+++ b/Engine/Source/Developer/Profiler/Private/Widgets/SEventGraph.cpp
@@ -2594,7 +2594,7 @@ void SEventGraph::SetHierarchicalSelectedEvents( const TArray<FEventGraphSampleP
 			: EventPtr( InEventPtr )
 		{}
 
-		bool Matches( const FEventGraphSamplePtr& Other ) const
+		bool operator()(const FEventGraphSamplePtr& Other) const
 		{
 			return EventPtr->_StatName == Other->_StatName;
 		}
@@ -2608,7 +2608,7 @@ void SEventGraph::SetHierarchicalSelectedEvents( const TArray<FEventGraphSampleP
 	{
 		for( int32 Nx = 0; Nx < HierarchicalSelectedEvents.Num(); ++Nx )
 		{
-			const int32 Index = Events_FlatCoalesced.FindMatch( FCoalescedEventMatcher( HierarchicalSelectedEvents[Nx] ) );
+			const int32 Index = Events_FlatCoalesced.IndexOfByPredicate(FCoalescedEventMatcher(HierarchicalSelectedEvents[Nx]));
 			if( Index != INDEX_NONE )
 			{
 				SelectedEvents.AddUnique( Events_FlatCoalesced[Index] );
diff --git a/Engine/Source/Editor/DetailCustomizations/Private/DialogueStructsCustomizations.cpp b/Engine/Source/Editor/DetailCustomizations/Private/DialogueStructsCustomizations.cpp
--- a/Engine/Source/Editor/DetailCustomizations/Private/DialogueStructsCustomizations.cpp
+++ b/Engine/Source/Editor/DetailCustomizations/Private/DialogueStructsCustomizations.cpp
@@ -205,7 +205,7 @@ void SSpeakerDropDown::Tick( const FGeometry& AllottedGeometry, const double InC
 struct FSpeakerMatcher
 {
 	const UDialogueVoice* Speaker;
-	bool Matches(const FDialogueContextMapping& Mapping) const
+	bool operator()(const FDialogueContextMapping& Mapping) const
 	{
 		return Mapping.Context.Speaker == Speaker;
 	}
@@ -954,7 +954,7 @@ bool SValidatedDialogueContextHeaderWidget::IsSpeakerValid() const
 	if( DialogueWave )
 	{
 		FSpeakerMatcher SpeakerMatcher = { Speaker };
-		if( DialogueWave->ContextMappings.FindMatch( SpeakerMatcher ) != INDEX_NONE )
+		if (DialogueWave->ContextMappings.IndexOfByPredicate(SpeakerMatcher) != INDEX_NONE)
 		{
 			bSpeakerIsValid = true;
 		}
diff --git a/Engine/Source/Editor/UnrealEd/Private/EditorServer.cpp b/Engine/Source/Editor/UnrealEd/Private/EditorServer.cpp
--- a/Engine/Source/Editor/UnrealEd/Private/EditorServer.cpp
+++ b/Engine/Source/Editor/UnrealEd/Private/EditorServer.cpp
@@ -4237,16 +4237,16 @@ void UEditorEngine::MoveViewportCamerasToActor(const TArray<AActor*> &Actors, bo
 								: ComponentToMatch( InComponentToMatch )
 							{}
 
-							bool Matches( const UClass* ComponentClass ) const
+							bool operator()(const UClass* ComponentClass) const
 							{
-								return ComponentToMatch->IsA( ComponentClass );
+								return ComponentToMatch->IsA(ComponentClass);
 							}
 
 							UPrimitiveComponent* ComponentToMatch;
 						};
 
 						// Some components can have huge bounds but are not visible.  Ignore these components unless it is the only component on the actor 
-						const bool bIgnore = Components.Num() > 1 && PrimitiveComponentTypesToIgnore.FindMatch( ComponentTypeMatcher( PrimitiveComponent ) ) != INDEX_NONE;
+						const bool bIgnore = Components.Num() > 1 && PrimitiveComponentTypesToIgnore.IndexOfByPredicate(ComponentTypeMatcher(PrimitiveComponent)) != INDEX_NONE;
 
 						if( !bIgnore )
 						{
diff --git a/Engine/Source/Editor/WorldBrowser/Private/Tiles/WorldTileCollectionModel.cpp b/Engine/Source/Editor/WorldBrowser/Private/Tiles/WorldTileCollectionModel.cpp
--- a/Engine/Source/Editor/WorldBrowser/Private/Tiles/WorldTileCollectionModel.cpp
+++ b/Engine/Source/Editor/WorldBrowser/Private/Tiles/WorldTileCollectionModel.cpp
@@ -231,7 +231,7 @@ TSharedPtr<FLevelDragDropOp> FWorldTileCollectionModel::CreateDragDropOp() const
 			else
 			{
 				//
-				int32 TileStreamingIdx = GetWorld()->WorldComposition->TilesStreaming.FindMatch(
+				int32 TileStreamingIdx = GetWorld()->WorldComposition->TilesStreaming.IndexOfByPredicate(
 					ULevelStreaming::FPackageNameMatcher(LevelModel->GetLongPackageName())
 					);
 
diff --git a/Engine/Source/Programs/BuildPatchTool/Private/BuildPatchToolMain.cpp b/Engine/Source/Programs/BuildPatchTool/Private/BuildPatchToolMain.cpp
--- a/Engine/Source/Programs/BuildPatchTool/Private/BuildPatchToolMain.cpp
+++ b/Engine/Source/Programs/BuildPatchTool/Private/BuildPatchToolMain.cpp
@@ -87,9 +87,9 @@ struct FCommandLineMatcher
 
 	FCommandLineMatcher(){}
 
-	FORCEINLINE bool Matches( const FString& ToMatch ) const
+	FORCEINLINE bool operator()(const FString& ToMatch) const
 	{
-		if( ToMatch.StartsWith( Command, ESearchCase::CaseSensitive ) )
+		if (ToMatch.StartsWith(Command, ESearchCase::CaseSensitive))
 		{
 			return true;
 		}
@@ -179,53 +179,53 @@ int32 BuildPatchToolMain( const TCHAR* CommandLine )
 		bSuccess = true;
 
 		Matcher.Command = TEXT("compactify");
-		bCompactify = Switches.FindMatch(Matcher) != INDEX_NONE;
+		bCompactify = Switches.IndexOfByPredicate(Matcher) != INDEX_NONE;
 		bPatchGeneration = !bCompactify;
 
 		Matcher.Command = TEXT("preview");
-		bPreview = bCompactify && Switches.FindMatch(Matcher) != INDEX_NONE;
+		bPreview = bCompactify && Switches.IndexOfByPredicate(Matcher) != INDEX_NONE;
 		
 		Matcher.Command = TEXT( "BuildRoot" );
-		BuildRootIdx = Switches.FindMatch( Matcher );
+		BuildRootIdx = Switches.IndexOfByPredicate(Matcher);
 
 		Matcher.Command = TEXT( "CloudDir" );
-		CloudDirIdx = Switches.FindMatch( Matcher );
+		CloudDirIdx = Switches.IndexOfByPredicate(Matcher);
 
 		Matcher.Command = TEXT( "AppID" );
-		AppIDIdx = Switches.FindMatch( Matcher );
+		AppIDIdx = Switches.IndexOfByPredicate(Matcher);
 
 		Matcher.Command = TEXT( "AppName" );
-		AppNameIdx = Switches.FindMatch( Matcher );
+		AppNameIdx = Switches.IndexOfByPredicate(Matcher);
 
 		Matcher.Command = TEXT( "BuildVersion" );
-		BuildVersionIdx = Switches.FindMatch( Matcher );
+		BuildVersionIdx = Switches.IndexOfByPredicate(Matcher);
 
 		Matcher.Command = TEXT( "AppLaunch" );
-		AppLaunchIdx = Switches.FindMatch( Matcher );
+		AppLaunchIdx = Switches.IndexOfByPredicate(Matcher);
 
 		Matcher.Command = TEXT( "AppArgs" );
-		AppArgsIdx = Switches.FindMatch( Matcher );
+		AppArgsIdx = Switches.IndexOfByPredicate(Matcher);
 
 		Matcher.Command = TEXT( "FileIgnoreList" );
-		FileIgnoreListIdx = Switches.FindMatch( Matcher );
+		FileIgnoreListIdx = Switches.IndexOfByPredicate(Matcher);
 
 		Matcher.Command = TEXT( "PrereqName" );
-		PrereqNameIdx = Switches.FindMatch( Matcher );
+		PrereqNameIdx = Switches.IndexOfByPredicate(Matcher);
 
 		Matcher.Command = TEXT( "PrereqPath" );
-		PrereqPathIdx = Switches.FindMatch( Matcher );
+		PrereqPathIdx = Switches.IndexOfByPredicate(Matcher);
 
 		Matcher.Command = TEXT( "PrereqArgs" );
-		PrereqArgsIdx = Switches.FindMatch( Matcher );
+		PrereqArgsIdx = Switches.IndexOfByPredicate(Matcher);
 
 		Matcher.Command = TEXT("ManifestsList");
-		ManifestsListIdx = Switches.FindMatch(Matcher);
+		ManifestsListIdx = Switches.IndexOfByPredicate(Matcher);
 
 		Matcher.Command = TEXT("ManifestsFile");
-		ManifestsFileIdx = Switches.FindMatch(Matcher);
+		ManifestsFileIdx = Switches.IndexOfByPredicate(Matcher);
 
 		Matcher.Command = TEXT("DataAgeThreshold");
-		DataAgeThresholdIdx = Switches.FindMatch(Matcher);
+		DataAgeThresholdIdx = Switches.IndexOfByPredicate(Matcher);
 
 		// Check required param indexes
 		bSuccess = bSuccess && CloudDirIdx != INDEX_NONE;
diff --git a/Engine/Source/Runtime/Core/Public/Containers/Array.h b/Engine/Source/Runtime/Core/Public/Containers/Array.h
--- a/Engine/Source/Runtime/Core/Public/Containers/Array.h
+++ b/Engine/Source/Runtime/Core/Public/Containers/Array.h
@@ -882,6 +882,7 @@ public:
 	 * @returns Index of the found element. INDEX_NONE otherwise.
 	 */
 	template<typename MatchFunctorType>
+	DEPRECATED(4.6, "This function is deprecated as it does the same as IndexOfByPredicate(). Please use IndexOfByPredicate() instead.")
 	int32 FindMatch(const MatchFunctorType& Matcher) const
 	{
 		const ElementType* const RESTRICT DataEnd = GetData() + ArrayNum;
diff --git a/Engine/Source/Runtime/Engine/Classes/Engine/LevelStreaming.h b/Engine/Source/Runtime/Engine/Classes/Engine/LevelStreaming.h
--- a/Engine/Source/Runtime/Engine/Classes/Engine/LevelStreaming.h
+++ b/Engine/Source/Runtime/Engine/Classes/Engine/LevelStreaming.h
@@ -229,7 +229,7 @@ class ULevelStreaming : public UObject
 		{
 		}
 
-		bool Matches( const ULevelStreaming* Candidate ) const
+		bool operator()(const ULevelStreaming* Candidate) const
 		{
 			return Candidate->GetWorldAssetPackageFName() == PackageName;
 		}
diff --git a/Engine/Source/Runtime/Engine/Private/LevelStreaming.cpp b/Engine/Source/Runtime/Engine/Private/LevelStreaming.cpp
--- a/Engine/Source/Runtime/Engine/Private/LevelStreaming.cpp
+++ b/Engine/Source/Runtime/Engine/Private/LevelStreaming.cpp
@@ -653,7 +653,7 @@ ULevelStreaming* ULevelStreaming::CreateInstance(FString InstanceUniqueName)
 		FName	InstanceUniquePackageName = FName(*(InstancePackagePath + InstanceShortPackageName));
 
 		// check if instance name is unique among existing streaming level objects
-		const bool bUniqueName = (InWorld->StreamingLevels.FindMatch(ULevelStreaming::FPackageNameMatcher(InstanceUniquePackageName)) == INDEX_NONE);
+		const bool bUniqueName = (InWorld->StreamingLevels.IndexOfByPredicate(ULevelStreaming::FPackageNameMatcher(InstanceUniquePackageName)) == INDEX_NONE);
 				
 		if (bUniqueName)
 		{
diff --git a/Engine/Source/Runtime/Engine/Private/PlayerMuteList.cpp b/Engine/Source/Runtime/Engine/Private/PlayerMuteList.cpp
--- a/Engine/Source/Runtime/Engine/Private/PlayerMuteList.cpp
+++ b/Engine/Source/Runtime/Engine/Private/PlayerMuteList.cpp
@@ -6,7 +6,7 @@
 static inline void AddIdToMuteList(TArray< TSharedRef<class FUniqueNetId> >& MuteList, const TSharedPtr<FUniqueNetId>& UniqueIdToAdd)
 {
 	FUniqueNetIdMatcher UniqueIdToAddMatch(*UniqueIdToAdd);
-	if (MuteList.FindMatch(UniqueIdToAddMatch) == INDEX_NONE)
+	if (MuteList.IndexOfByPredicate(UniqueIdToAddMatch) == INDEX_NONE)
 	{
 		MuteList.Add(UniqueIdToAdd.ToSharedRef());
 	}
@@ -15,7 +15,7 @@ static inline void AddIdToMuteList(TArray< TSharedRef<class FUniqueNetId> >& Mut
 static inline void RemoveIdFromMuteList(TArray< TSharedRef<class FUniqueNetId> >& MuteList, const TSharedPtr<FUniqueNetId>& UniqueIdToRemove)
 {
 	FUniqueNetIdMatcher UniqueIdToRemoveMatch(*UniqueIdToRemove);
-	int32 RemoveIndex = MuteList.FindMatch(UniqueIdToRemoveMatch);
+	int32 RemoveIndex = MuteList.IndexOfByPredicate(UniqueIdToRemoveMatch);
 	if (RemoveIndex != INDEX_NONE)
 	{
 		MuteList.RemoveAtSwap(RemoveIndex);
@@ -66,16 +66,16 @@ void FPlayerMuteList::ServerUnmutePlayer(APlayerController* OwningPC, const FUni
 		FUniqueNetIdMatcher OwningPlayerIdMatch(*OwningPC->PlayerState->UniqueId);
 
 		// Make sure this player isn't muted for gameplay reasons
-		if (GameplayVoiceMuteList.FindMatch(PlayerIdToUnmuteMatch) == INDEX_NONE &&
+		if (GameplayVoiceMuteList.IndexOfByPredicate(PlayerIdToUnmuteMatch) == INDEX_NONE &&
 			// And make sure they didn't mute us
-			OtherPC->MuteList.VoiceMuteList.FindMatch(OwningPlayerIdMatch) == INDEX_NONE)
+			OtherPC->MuteList.VoiceMuteList.IndexOfByPredicate(OwningPlayerIdMatch) == INDEX_NONE)
 		{
 			OwningPC->ClientUnmutePlayer(UnmuteId);
 		}
 
 		// If the other player doesn't have this player muted
-		if (OtherPC->MuteList.VoiceMuteList.FindMatch(OwningPlayerIdMatch) == INDEX_NONE &&
-			OtherPC->MuteList.GameplayVoiceMuteList.FindMatch(OwningPlayerIdMatch) == INDEX_NONE)
+		if (OtherPC->MuteList.VoiceMuteList.IndexOfByPredicate(OwningPlayerIdMatch) == INDEX_NONE &&
+			OtherPC->MuteList.GameplayVoiceMuteList.IndexOfByPredicate(OwningPlayerIdMatch) == INDEX_NONE)
 		{
 			// Remove them from the packet filter list
 			RemoveIdFromMuteList(VoicePacketFilter, PlayerIdToUnmute);
@@ -162,9 +162,9 @@ void FPlayerMuteList::GameplayUnmutePlayer(APlayerController* OwningPC, const FU
 		FUniqueNetIdMatcher OwningPlayerIdMatch(*OwningPC->PlayerState->UniqueId);
 
 		// Make sure this player isn't explicitly muted
-		if (VoiceMuteList.FindMatch(PlayerIdToUnmuteMatch) == INDEX_NONE &&
+		if (VoiceMuteList.IndexOfByPredicate(PlayerIdToUnmuteMatch) == INDEX_NONE &&
 			// And make sure they didn't mute us
-			OtherPC->MuteList.VoiceMuteList.FindMatch(OwningPlayerIdMatch) == INDEX_NONE)
+			OtherPC->MuteList.VoiceMuteList.IndexOfByPredicate(OwningPlayerIdMatch) == INDEX_NONE)
 		{
 			RemoveIdFromMuteList(VoicePacketFilter, PlayerIdToUnmute);
 
@@ -177,7 +177,7 @@ void FPlayerMuteList::GameplayUnmutePlayer(APlayerController* OwningPC, const FU
 bool FPlayerMuteList::IsPlayerMuted(const FUniqueNetId& PlayerId)
 {
 	FUniqueNetIdMatcher PlayerIdMatch(PlayerId);
-	return VoicePacketFilter.FindMatch(PlayerIdMatch) != INDEX_NONE;
+	return VoicePacketFilter.IndexOfByPredicate(PlayerIdMatch) != INDEX_NONE;
 }
 
 FString DumpMutelistState(UWorld* World)
diff --git a/Engine/Source/Runtime/Online/OnlineSubsystem/Private/OnlineSubsystem.cpp b/Engine/Source/Runtime/Online/OnlineSubsystem/Private/OnlineSubsystem.cpp
--- a/Engine/Source/Runtime/Online/OnlineSubsystem/Private/OnlineSubsystem.cpp
+++ b/Engine/Source/Runtime/Online/OnlineSubsystem/Private/OnlineSubsystem.cpp
@@ -81,7 +81,7 @@ bool IsPlayerInSessionImpl(IOnlineSession* SessionInt, FName SessionName, const
 
 		FUniqueNetIdMatcher PlayerMatch(UniqueId);
 		if (bIsSessionOwner || 
-			Session->RegisteredPlayers.FindMatch(PlayerMatch) != INDEX_NONE)
+			Session->RegisteredPlayers.IndexOfByPredicate(PlayerMatch) != INDEX_NONE)
 		{
 			bFound = true;
 		}
diff --git a/Engine/Source/Runtime/Online/OnlineSubsystem/Public/OnlineSubsystemTypes.h b/Engine/Source/Runtime/Online/OnlineSubsystem/Public/OnlineSubsystemTypes.h
--- a/Engine/Source/Runtime/Online/OnlineSubsystem/Public/OnlineSubsystemTypes.h
+++ b/Engine/Source/Runtime/Online/OnlineSubsystem/Public/OnlineSubsystemTypes.h
@@ -752,7 +752,7 @@ public:
 };
 
 /**
- * TArray helper for FindMatch() function
+ * TArray helper for IndexOfByPredicate() function
  */
 struct FUniqueNetIdMatcher
 {
@@ -771,27 +771,27 @@ public:
 	 *
 	 * @return true if they are an exact match, false otherwise
 	 */
- 	bool Matches(const FUniqueNetId& Candidate) const
- 	{
- 		return UniqueIdTarget == Candidate;
- 	}
+	bool operator()(const FUniqueNetId& Candidate) const
+	{
+		return UniqueIdTarget == Candidate;
+	}
  
 	/**
 	 * Match a given unique Id against the one stored in this struct
 	 *
 	 * @return true if they are an exact match, false otherwise
 	 */
- 	bool Matches(const TSharedPtr<FUniqueNetId>& Candidate) const
- 	{
- 		return UniqueIdTarget == *Candidate;
- 	}
+	bool operator()(const TSharedPtr<FUniqueNetId>& Candidate) const
+	{
+		return UniqueIdTarget == *Candidate;
+	}
 
 	/**
 	 * Match a given unique Id against the one stored in this struct
 	 *
 	 * @return true if they are an exact match, false otherwise
 	 */
-	bool Matches(const TSharedRef<FUniqueNetId>& Candidate) const
+	bool operator()(const TSharedRef<FUniqueNetId>& Candidate) const
 	{
 		return UniqueIdTarget == *Candidate;
 	}
diff --git a/Engine/Source/Runtime/Online/OnlineSubsystemNull/Private/OnlineSessionInterfaceNull.cpp b/Engine/Source/Runtime/Online/OnlineSubsystemNull/Private/OnlineSessionInterfaceNull.cpp
--- a/Engine/Source/Runtime/Online/OnlineSubsystemNull/Private/OnlineSessionInterfaceNull.cpp
+++ b/Engine/Source/Runtime/Online/OnlineSubsystemNull/Private/OnlineSessionInterfaceNull.cpp
@@ -839,7 +839,7 @@ bool FOnlineSessionNull::RegisterPlayers(FName SessionName, const TArray< TShare
 			const TSharedRef<FUniqueNetId>& PlayerId = Players[PlayerIdx];
 
 			FUniqueNetIdMatcher PlayerMatch(*PlayerId);
-			if (Session->RegisteredPlayers.FindMatch(PlayerMatch) == INDEX_NONE)
+			if (Session->RegisteredPlayers.IndexOfByPredicate(PlayerMatch) == INDEX_NONE)
 			{
 				Session->RegisteredPlayers.Add(PlayerId);
 				RegisterVoice(*PlayerId);
@@ -889,7 +889,7 @@ bool FOnlineSessionNull::UnregisterPlayers(FName SessionName, const TArray< TSha
 			const TSharedRef<FUniqueNetId>& PlayerId = Players[PlayerIdx];
 
 			FUniqueNetIdMatcher PlayerMatch(*PlayerId);
-			int32 RegistrantIndex = Session->RegisteredPlayers.FindMatch(PlayerMatch);
+			int32 RegistrantIndex = Session->RegisteredPlayers.IndexOfByPredicate(PlayerMatch);
 			if (RegistrantIndex != INDEX_NONE)
 			{
 				Session->RegisteredPlayers.RemoveAtSwap(RegistrantIndex);
diff --git a/Engine/Source/Runtime/Online/OnlineSubsystemSteam/Private/OnlineSessionInterfaceSteam.cpp b/Engine/Source/Runtime/Online/OnlineSubsystemSteam/Private/OnlineSessionInterfaceSteam.cpp
--- a/Engine/Source/Runtime/Online/OnlineSubsystemSteam/Private/OnlineSessionInterfaceSteam.cpp
+++ b/Engine/Source/Runtime/Online/OnlineSubsystemSteam/Private/OnlineSessionInterfaceSteam.cpp
@@ -1297,7 +1297,7 @@ bool FOnlineSessionSteam::RegisterPlayers(FName SessionName, const TArray< TShar
 				const FUniqueNetIdSteam& SteamId = (const FUniqueNetIdSteam&)*PlayerId;
 
 				FUniqueNetIdMatcher PlayerMatch(SteamId);
-				if (Session->RegisteredPlayers.FindMatch(PlayerMatch) == INDEX_NONE)
+				if (Session->RegisteredPlayers.IndexOfByPredicate(PlayerMatch) == INDEX_NONE)
 				{
 					Session->RegisteredPlayers.Add(PlayerId);
 
@@ -1373,7 +1373,7 @@ bool FOnlineSessionSteam::UnregisterPlayers(FName SessionName, const TArray< TSh
 				const TSharedRef<FUniqueNetId>& PlayerId = Players[PlayerIdx];
 
 				FUniqueNetIdMatcher PlayerMatch(*PlayerId);
-				int32 RegistrantIndex = Session->RegisteredPlayers.FindMatch(PlayerMatch);
+				int32 RegistrantIndex = Session->RegisteredPlayers.IndexOfByPredicate(PlayerMatch);
 				if (RegistrantIndex != INDEX_NONE)
 				{
 					Session->RegisteredPlayers.RemoveAtSwap(RegistrantIndex);
diff --git a/Engine/Source/Runtime/Online/OnlineSubsystemSteam/Private/VoiceInterfaceSteam.cpp b/Engine/Source/Runtime/Online/OnlineSubsystemSteam/Private/VoiceInterfaceSteam.cpp
--- a/Engine/Source/Runtime/Online/OnlineSubsystemSteam/Private/VoiceInterfaceSteam.cpp
+++ b/Engine/Source/Runtime/Online/OnlineSubsystemSteam/Private/VoiceInterfaceSteam.cpp
@@ -551,7 +551,7 @@ TSharedPtr<FVoicePacket> FOnlineVoiceSteam::SerializeRemotePacket(FArchive& Ar)
 		if (!SteamSubsystem->IsDedicated())
 		{
 			FUniqueNetIdMatcher PlayerMatch(*NewPacket->GetSender());
-			if (MuteList.FindMatch(PlayerMatch) == INDEX_NONE)
+			if (MuteList.IndexOfByPredicate(PlayerMatch) == INDEX_NONE)
 			{
 				VoiceData.RemotePackets.Add(NewPacket);
 			}
diff --git a/Engine/Source/Runtime/Online/OnlineSubsystemUtils/Private/PartyBeaconHost.cpp b/Engine/Source/Runtime/Online/OnlineSubsystemUtils/Private/PartyBeaconHost.cpp
--- a/Engine/Source/Runtime/Online/OnlineSubsystemUtils/Private/PartyBeaconHost.cpp
+++ b/Engine/Source/Runtime/Online/OnlineSubsystemUtils/Private/PartyBeaconHost.cpp
@@ -124,7 +124,7 @@ void APartyBeaconHost::Tick(float DeltaTime)
 								bIsSessionOwner)
 							{
 								FUniqueNetIdMatcher PlayerMatch(*PlayerEntry.UniqueId);
-								int32 FoundIdx = State->PlayersPendingJoin.FindMatch(PlayerMatch);
+								int32 FoundIdx = State->PlayersPendingJoin.IndexOfByPredicate(PlayerMatch);
 								if (FoundIdx != INDEX_NONE)
 								{
 									UE_LOG(LogBeacon, Display, TEXT("Beacon (%s): pending player %s found in session (%s)."),
@@ -145,7 +145,7 @@ void APartyBeaconHost::Tick(float DeltaTime)
 
 								// if the player is pending it's initial join then check against TravelSessionTimeoutSecs instead
 								FUniqueNetIdMatcher PlayerMatch(*PlayerEntry.UniqueId);
-								int32 FoundIdx = State->PlayersPendingJoin.FindMatch(PlayerMatch);
+								int32 FoundIdx = State->PlayersPendingJoin.IndexOfByPredicate(PlayerMatch);
 								const bool bIsPlayerPendingJoin = FoundIdx != INDEX_NONE;
 								// if the timeout has been exceeded then add to list of players 
 								// that need to be logged out from the beacon
diff --git a/Engine/Source/Runtime/Online/OnlineSubsystemUtils/Private/VoiceInterfaceImpl.cpp b/Engine/Source/Runtime/Online/OnlineSubsystemUtils/Private/VoiceInterfaceImpl.cpp
--- a/Engine/Source/Runtime/Online/OnlineSubsystemUtils/Private/VoiceInterfaceImpl.cpp
+++ b/Engine/Source/Runtime/Online/OnlineSubsystemUtils/Private/VoiceInterfaceImpl.cpp
@@ -550,7 +550,7 @@ TSharedPtr<FVoicePacket> FOnlineVoiceImpl::SerializeRemotePacket(FArchive& Ar)
 		if (!OnlineSubsystem->IsDedicated())
 		{
 			FUniqueNetIdMatcher PlayerMatch(*NewPacket->GetSender());
-			if (MuteList.FindMatch(PlayerMatch) == INDEX_NONE)
+			if (MuteList.IndexOfByPredicate(PlayerMatch) == INDEX_NONE)
 			{
 				VoiceData.RemotePackets.Add(NewPacket);
 			}
diff --git a/Engine/Source/Runtime/Slate/Private/Framework/Application/SlateApplication.cpp b/Engine/Source/Runtime/Slate/Private/Framework/Application/SlateApplication.cpp
--- a/Engine/Source/Runtime/Slate/Private/Framework/Application/SlateApplication.cpp
+++ b/Engine/Source/Runtime/Slate/Private/Framework/Application/SlateApplication.cpp
@@ -290,7 +290,7 @@ void FPopupSupport::SendNotifications( const FWidgetPath& WidgetsUnderCursor )
 		: WidgetToMatch( InWidgetToMatch )
 		{}
 
-		bool Matches( const FArrangedWidget& Candidate ) const
+		bool operator()(const FArrangedWidget& Candidate) const
 		{
 			return WidgetToMatch == Candidate.Widget;
 		}
@@ -307,7 +307,7 @@ void FPopupSupport::SendNotifications( const FWidgetPath& WidgetsUnderCursor )
 		{
 			// Did we click outside the region in this subscription? If so send the notification.
 			FArrangedWidgetMatcher Matcher(Subscriber.DetectClicksOutsideMe.Pin().ToSharedRef());
-			const bool bClickedOutsideOfWidget = WidgetsUnderCursor.Widgets.GetInternalArray().FindMatch( Matcher ) == INDEX_NONE;
+			const bool bClickedOutsideOfWidget = WidgetsUnderCursor.Widgets.GetInternalArray().IndexOfByPredicate(Matcher) == INDEX_NONE;
 			if ( bClickedOutsideOfWidget )
 			{
 				Subscriber.Notification.ExecuteIfBound();
diff --git a/Engine/Source/Runtime/Slate/Private/Framework/Docking/SDockingTabStack.cpp b/Engine/Source/Runtime/Slate/Private/Framework/Docking/SDockingTabStack.cpp
--- a/Engine/Source/Runtime/Slate/Private/Framework/Docking/SDockingTabStack.cpp
+++ b/Engine/Source/Runtime/Slate/Private/Framework/Docking/SDockingTabStack.cpp
@@ -290,9 +290,9 @@ int32 SDockingTabStack::GetNumTabs() const
 	return TabWell->GetNumTabs();
 }
 
-bool SDockingTabStack::HasTab( const struct FTabMatcher& TabMatcher ) const
+bool SDockingTabStack::HasTab(const struct FTabMatcher& TabMatcher) const
 {
-	return Tabs.FindMatch( TabMatcher ) != INDEX_NONE;
+	return Tabs.IndexOfByPredicate(TabMatcher) != INDEX_NONE;
 }
 
 FGeometry SDockingTabStack::GetTabStackGeometry() const
@@ -896,7 +896,7 @@ FSlateColor SDockingTabStack::GetUnhideTabWellButtonOpacity() const
 
 int32 SDockingTabStack::OpenPersistentTab( const FTabId& TabId, int32 OpenLocationAmongActiveTabs )
 {
-	const int32 ExistingClosedTabIndex = Tabs.FindMatch( FTabMatcher(TabId, ETabState::ClosedTab) );
+	const int32 ExistingClosedTabIndex = Tabs.IndexOfByPredicate(FTabMatcher(TabId, ETabState::ClosedTab));
 
 	if (OpenLocationAmongActiveTabs == INDEX_NONE)
 	{						
@@ -969,7 +969,7 @@ int32 SDockingTabStack::OpenPersistentTab( const FTabId& TabId, int32 OpenLocati
 
 int32 SDockingTabStack::ClosePersistentTab( const FTabId& TabId )
 {
-	const int32 TabIndex = Tabs.FindMatch( FTabMatcher(TabId, ETabState::OpenedTab) );
+	const int32 TabIndex = Tabs.IndexOfByPredicate(FTabMatcher(TabId, ETabState::OpenedTab));
 	if (TabIndex != INDEX_NONE)
 	{
 		Tabs[TabIndex].TabState = ETabState::ClosedTab;
@@ -979,7 +979,7 @@ int32 SDockingTabStack::ClosePersistentTab( const FTabId& TabId )
 
 void SDockingTabStack::RemovePersistentTab( const FTabId& TabId )
 {
-	const int32 TabIndex = Tabs.FindMatch( FTabMatcher(TabId) );
+	const int32 TabIndex = Tabs.IndexOfByPredicate(FTabMatcher(TabId));
 	Tabs.RemoveAtSwap(TabIndex);
 }
 
diff --git a/Engine/Source/Runtime/Slate/Private/Framework/Docking/SDockingTabStack.h b/Engine/Source/Runtime/Slate/Private/Framework/Docking/SDockingTabStack.h
--- a/Engine/Source/Runtime/Slate/Private/Framework/Docking/SDockingTabStack.h
+++ b/Engine/Source/Runtime/Slate/Private/Framework/Docking/SDockingTabStack.h
@@ -247,13 +247,13 @@ struct FTabMatcher
 	{
 	}
 
-	bool Matches( const FTabManager::FTab& Candidate ) const
+	bool operator()(const FTabManager::FTab& Candidate) const
 	{
 		return
-			( (Candidate.TabState & RequiredTabState) != 0 ) &&
-			( Candidate.TabId.TabType == TabIdToMatch.TabType ) &&
+			((Candidate.TabState & RequiredTabState) != 0) &&
+			(Candidate.TabId.TabType == TabIdToMatch.TabType) &&
 			// INDEX_NONE is treated as a wildcard
-			( (TreatIndexNoneAsWildcard && TabIdToMatch.InstanceId == INDEX_NONE) || TabIdToMatch.InstanceId == Candidate.TabId.InstanceId );;
+			((TreatIndexNoneAsWildcard && TabIdToMatch.InstanceId == INDEX_NONE) || TabIdToMatch.InstanceId == Candidate.TabId.InstanceId);;
 	}
 
 	FTabId TabIdToMatch;
diff --git a/Engine/Source/Runtime/Slate/Private/Framework/Docking/TabManager.cpp b/Engine/Source/Runtime/Slate/Private/Framework/Docking/TabManager.cpp
--- a/Engine/Source/Runtime/Slate/Private/Framework/Docking/TabManager.cpp
+++ b/Engine/Source/Runtime/Slate/Private/Framework/Docking/TabManager.cpp
@@ -1267,7 +1267,7 @@ bool FTabManager::HasAnyMatchingTabs( const TSharedRef<FTabManager::FLayoutNode>
 
 	if ( AsStack.IsValid() )
 	{
-		return INDEX_NONE != AsStack->Tabs.FindMatch(Matcher);
+		return INDEX_NONE != AsStack->Tabs.IndexOfByPredicate(Matcher);
 	}
 	else
 	{
@@ -1291,7 +1291,7 @@ bool FTabManager::HasOpenTabs( const TSharedRef<FTabManager::FLayoutNode>& SomeN
 	{
 		const FTabManager* TabManager;
 
-		bool Matches( const FTab& Candidate ) const
+		bool operator()(const FTab& Candidate) const
 		{
 			return TabManager->IsValidTabForSpawning(Candidate) && Candidate.TabState == ETabState::OpenedTab;
 		}
@@ -1309,7 +1309,7 @@ bool FTabManager::HasValidTabs( const TSharedRef<FTabManager::FLayoutNode>& Some
 	{
 		const FTabManager* TabManager;
 
-		bool Matches( const FTab& Candidate ) const
+		bool operator()(const FTab& Candidate) const
 		{
 			return TabManager->IsValidTabForSpawning(Candidate);
 		}
@@ -1407,7 +1407,7 @@ TSharedPtr<FTabManager::FStack> FTabManager::FindTabUnderNode( const FTabMatcher
 
 	if (NodeAsStack.IsValid())
 	{
-		const int32 TabIndex = NodeAsStack->Tabs.FindMatch( Matcher );
+		const int32 TabIndex = NodeAsStack->Tabs.IndexOfByPredicate(Matcher);
 		if (TabIndex != INDEX_NONE)
 		{
 			return NodeAsStack;
@@ -1473,7 +1473,7 @@ void FTabManager::RemoveTabFromCollapsedAreas( const FTabMatcher& Matcher )
 
 			if (StackWithMatchingTab.IsValid())
 			{
-				const int32 TabIndex = StackWithMatchingTab->Tabs.FindMatch( Matcher );
+				const int32 TabIndex = StackWithMatchingTab->Tabs.IndexOfByPredicate(Matcher);
 				if ( ensure(TabIndex != INDEX_NONE) )
 				{
 					StackWithMatchingTab->Tabs.RemoveAt(TabIndex);
@@ -1616,7 +1616,7 @@ bool FGlobalTabmanager::CanCloseManager( const TSet< TSharedRef<SDockTab> >& Tab
 
 void FGlobalTabmanager::DrawAttentionToTabManager( const TSharedRef<FTabManager>& ChildManager )
 {
-	const int32 MajorTabIndex = SubTabManagers.FindMatch( FindByManager(ChildManager) );
+	const int32 MajorTabIndex = SubTabManagers.IndexOfByPredicate(FindByManager(ChildManager));
 	if (MajorTabIndex != INDEX_NONE)
 	{
 		this->DrawAttention( SubTabManagers[MajorTabIndex].MajorTab.Pin().ToSharedRef() );
@@ -1646,7 +1646,7 @@ void FGlobalTabmanager::UpdateMainMenu(const TSharedRef<SDockTab>& ForTab, bool
 	TSharedPtr<FTabManager> Tabmanager = ForTab->GetTabManager();
 	if(Tabmanager == AsShared())
 	{
-		const int32 TabIndex = SubTabManagers.FindMatch( FindByTab(ForTab) );
+		const int32 TabIndex = SubTabManagers.IndexOfByPredicate(FindByTab(ForTab));
 		if (TabIndex != INDEX_NONE)
 		{
 			Tabmanager = SubTabManagers[TabIndex].TabManager.Pin();
@@ -1703,7 +1703,7 @@ void FGlobalTabmanager::OnTabForegrounded( const TSharedPtr<SDockTab>& NewForegr
 	if (NewForegroundTab.IsValid())
 	{
 		// Show any child windows associated with the Major Tab that got foregrounded.
-		const int32 ForegroundedTabIndex = SubTabManagers.FindMatch( FindByTab(NewForegroundTab.ToSharedRef()) );
+		const int32 ForegroundedTabIndex = SubTabManagers.IndexOfByPredicate(FindByTab(NewForegroundTab.ToSharedRef()));
 		if (ForegroundedTabIndex != INDEX_NONE)
 		{
 			TSharedPtr<FTabManager> ForegroundTabManager = SubTabManagers[ForegroundedTabIndex].TabManager.Pin();
@@ -1714,7 +1714,7 @@ void FGlobalTabmanager::OnTabForegrounded( const TSharedPtr<SDockTab>& NewForegr
 	if ( BackgroundedTab.IsValid() )
 	{
 		// Hide any child windows associated with the Major Tab that got backgrounded.
-		const int32 BackgroundedTabIndex = SubTabManagers.FindMatch( FindByTab(BackgroundedTab.ToSharedRef()) );
+		const int32 BackgroundedTabIndex = SubTabManagers.IndexOfByPredicate(FindByTab(BackgroundedTab.ToSharedRef()));
 		if (BackgroundedTabIndex != INDEX_NONE)
 		{
 			TSharedPtr<FTabManager> BackgroundedTabManager = SubTabManagers[BackgroundedTabIndex].TabManager.Pin();
@@ -1727,7 +1727,7 @@ void FGlobalTabmanager::OnTabRelocated( const TSharedRef<SDockTab>& RelocatedTab
 {
 	if (NewOwnerWindow.IsValid())
 	{
-		const int32 RelocatedManagerIndex = SubTabManagers.FindMatch( FindByTab(RelocatedTab) );
+		const int32 RelocatedManagerIndex = SubTabManagers.IndexOfByPredicate(FindByTab(RelocatedTab));
 		if (RelocatedManagerIndex != INDEX_NONE)
 		{
 			const TSharedRef<FTabManager>& RelocatedManager = SubTabManagers[RelocatedManagerIndex].TabManager.Pin().ToSharedRef();
@@ -1769,7 +1769,7 @@ void FGlobalTabmanager::OnTabClosing( const TSharedRef<SDockTab>& TabBeingClosed
 {
 	// Is this a major tab that contained a Sub TabManager?
 	// If so, need to properly close the sub tab manager
-	const int32 TabManagerBeingClosedIndex = SubTabManagers.FindMatch( FindByTab(TabBeingClosed) );
+	const int32 TabManagerBeingClosedIndex = SubTabManagers.IndexOfByPredicate(FindByTab(TabBeingClosed));
 	if (TabManagerBeingClosedIndex != INDEX_NONE)
 	{
 		const TSharedRef<FTabManager>& TabManagerBeingClosed = SubTabManagers[TabManagerBeingClosedIndex].TabManager.Pin().ToSharedRef();
diff --git a/Engine/Source/Runtime/Slate/Public/Framework/Docking/TabManager.h b/Engine/Source/Runtime/Slate/Public/Framework/Docking/TabManager.h
--- a/Engine/Source/Runtime/Slate/Public/Framework/Docking/TabManager.h
+++ b/Engine/Source/Runtime/Slate/Public/Framework/Docking/TabManager.h
@@ -820,7 +820,7 @@ private:
 		{
 		}
 
-		bool Matches( const FGlobalTabmanager::FSubTabManager& TabManagerPair ) const
+		bool operator()(const FGlobalTabmanager::FSubTabManager& TabManagerPair) const
 		{
 			return TabManagerPair.TabManager.IsValid() && TabManagerPair.MajorTab.IsValid() && TabManagerPair.MajorTab.Pin() == TabToFind;
 		}
@@ -835,7 +835,7 @@ private:
 		{
 		}
 
-		bool Matches( const FGlobalTabmanager::FSubTabManager& TabManagerPair ) const
+		bool operator()(const FGlobalTabmanager::FSubTabManager& TabManagerPair) const
 		{
 			return TabManagerPair.TabManager.IsValid() && TabManagerPair.MajorTab.IsValid() && TabManagerPair.TabManager.Pin() == ManagerToFind;
 		}
