FSlateMacMenu::UpdateCachedState() now checks with the global tab manager if there's any active tab instead of depending on not being called when there's not. This should solve any remaining issues with executing it while the app is quitting.

diff --git a/Engine/Source/Programs/CrashReporter/CrashReportClient/Private/Mac/CrashReportClientMainMac.cpp b/Engine/Source/Programs/CrashReporter/CrashReportClient/Private/Mac/CrashReportClientMainMac.cpp
--- a/Engine/Source/Programs/CrashReporter/CrashReportClient/Private/Mac/CrashReportClientMainMac.cpp
+++ b/Engine/Source/Programs/CrashReporter/CrashReportClient/Private/Mac/CrashReportClientMainMac.cpp
@@ -48,7 +48,6 @@ static FString GSavedCommandLine;
 - (IBAction)requestQuit:(id)Sender
 {
 	GIsRequestingExit = true;
-	FPlatformMisc::UpdateCachedMacMenuState = nullptr;
 }
 
 - (void) runGameThread:(id)Arg
diff --git a/Engine/Source/Programs/SlateViewer/Private/Mac/SlateViewerMainMac.cpp b/Engine/Source/Programs/SlateViewer/Private/Mac/SlateViewerMainMac.cpp
--- a/Engine/Source/Programs/SlateViewer/Private/Mac/SlateViewerMainMac.cpp
+++ b/Engine/Source/Programs/SlateViewer/Private/Mac/SlateViewerMainMac.cpp
@@ -35,7 +35,6 @@ static FString GSavedCommandLine;
 	if(!GIsRequestingExit || ([NSThread gameThread] && [NSThread gameThread] != [NSThread mainThread]))
 	{
 		GIsRequestingExit = true;
-		FPlatformMisc::UpdateCachedMacMenuState = nullptr;
 		return NSTerminateLater;
 	}
 	else
diff --git a/Engine/Source/Programs/SymbolDebugger/Private/Mac/SymbolDebuggerMainMac.cpp b/Engine/Source/Programs/SymbolDebugger/Private/Mac/SymbolDebuggerMainMac.cpp
--- a/Engine/Source/Programs/SymbolDebugger/Private/Mac/SymbolDebuggerMainMac.cpp
+++ b/Engine/Source/Programs/SymbolDebugger/Private/Mac/SymbolDebuggerMainMac.cpp
@@ -34,7 +34,6 @@ static FString GSavedCommandLine;
 	if(!GIsRequestingExit || ([NSThread gameThread] && [NSThread gameThread] != [NSThread mainThread]))
 	{
 		GIsRequestingExit = true;
-		FPlatformMisc::UpdateCachedMacMenuState = nullptr;
 		return NSTerminateLater;
 	}
 	else
diff --git a/Engine/Source/Programs/UnrealFrontend/Private/Mac/MacUnrealFrontendMain.cpp b/Engine/Source/Programs/UnrealFrontend/Private/Mac/MacUnrealFrontendMain.cpp
--- a/Engine/Source/Programs/UnrealFrontend/Private/Mac/MacUnrealFrontendMain.cpp
+++ b/Engine/Source/Programs/UnrealFrontend/Private/Mac/MacUnrealFrontendMain.cpp
@@ -27,7 +27,6 @@ static FString GSavedCommandLine;
 - (IBAction)requestQuit:(id)Sender
 {
 	GIsRequestingExit = true;
-	FPlatformMisc::UpdateCachedMacMenuState = nullptr;
 }
 
 - (void) runGameThread:(id)Arg
diff --git a/Engine/Source/Runtime/Core/Private/Mac/CocoaWindow.cpp b/Engine/Source/Runtime/Core/Private/Mac/CocoaWindow.cpp
--- a/Engine/Source/Runtime/Core/Private/Mac/CocoaWindow.cpp
+++ b/Engine/Source/Runtime/Core/Private/Mac/CocoaWindow.cpp
@@ -519,12 +519,6 @@ NSString* NSPerformDragOperation = @"NSPerformDragOperation";
 	}
 	self.bForwardEvents = false;
 	[self setDelegate:nil];
-
-	if ([[NSApp windows] count] == 1)
-	{
-		// It's the last window. App will quit, so stop updating menu state cache.
-		GameThreadCall(^{ FPlatformMisc::UpdateCachedMacMenuState = nullptr; }, @[ NSDefaultRunLoopMode ], false);
-	}
 }
 
 - (void)mouseDown:(NSEvent*)Event
diff --git a/Engine/Source/Runtime/Launch/Private/Mac/LaunchMac.cpp b/Engine/Source/Runtime/Launch/Private/Mac/LaunchMac.cpp
--- a/Engine/Source/Runtime/Launch/Private/Mac/LaunchMac.cpp
+++ b/Engine/Source/Runtime/Launch/Private/Mac/LaunchMac.cpp
@@ -103,7 +103,6 @@ void EngineCrashHandler(const FGenericCrashContext & GenericContext)
 	GameThreadCall(^{
 		if (GEngine)
 		{
-			FPlatformMisc::UpdateCachedMacMenuState = nullptr;
 			if (GIsEditor)
 			{
 				if (IsRunningCommandlet())
diff --git a/Engine/Source/Runtime/Slate/Private/Framework/MultiBox/Mac/MacMenu.cpp b/Engine/Source/Runtime/Slate/Private/Framework/MultiBox/Mac/MacMenu.cpp
--- a/Engine/Source/Runtime/Slate/Private/Framework/MultiBox/Mac/MacMenu.cpp
+++ b/Engine/Source/Runtime/Slate/Private/Framework/MultiBox/Mac/MacMenu.cpp
@@ -265,46 +265,50 @@ void FSlateMacMenu::UpdateMenu(FMacMenu* Menu)
 
 void FSlateMacMenu::UpdateCachedState()
 {
-	for (TMap<FMacMenu*, TSharedPtr<TArray<FMacMenuItemState>>>::TIterator It(GCachedMenuState); It; ++It)
+	TSharedPtr<SDockTab> ActiveTab = FGlobalTabmanager::Get()->GetActiveTab();
+	if (ActiveTab.IsValid())
 	{
-		FMacMenu* Menu = It.Key();
-		TSharedPtr<TArray<FMacMenuItemState>> MenuState = It.Value();
-
-		if (!Menu.MultiBox.IsValid())
+		for (TMap<FMacMenu*, TSharedPtr<TArray<FMacMenuItemState>>>::TIterator It(GCachedMenuState); It; ++It)
 		{
-			const bool bShouldCloseWindowAfterMenuSelection = true;
-			FMenuBuilder MenuBuilder(bShouldCloseWindowAfterMenuSelection, Menu.MenuEntryBlock->GetActionList(), Menu.MenuEntryBlock->Extender);
+			FMacMenu* Menu = It.Key();
+			TSharedPtr<TArray<FMacMenuItemState>> MenuState = It.Value();
+
+			if (!Menu.MultiBox.IsValid())
 			{
-				// Have the menu fill its contents
-				Menu.MenuEntryBlock->EntryBuilder.ExecuteIfBound(MenuBuilder);
+				const bool bShouldCloseWindowAfterMenuSelection = true;
+				FMenuBuilder MenuBuilder(bShouldCloseWindowAfterMenuSelection, Menu.MenuEntryBlock->GetActionList(), Menu.MenuEntryBlock->Extender);
+				{
+					// Have the menu fill its contents
+					Menu.MenuEntryBlock->EntryBuilder.ExecuteIfBound(MenuBuilder);
+				}
+				TSharedRef<SWidget> Widget = MenuBuilder.MakeWidget();
+				Menu.MultiBox = TSharedPtr<const FMultiBox>(StaticCastSharedRef<SMultiBoxWidget>(Widget)->GetMultiBox());
 			}
-			TSharedRef<SWidget> Widget = MenuBuilder.MakeWidget();
-			Menu.MultiBox = TSharedPtr<const FMultiBox>(StaticCastSharedRef<SMultiBoxWidget>(Widget)->GetMultiBox());
-		}
-
-		const TArray<TSharedRef<const FMultiBlock>>& MenuBlocks = Menu.MultiBox->GetBlocks();
-		for (int32 Index = MenuState->Num(); MenuBlocks.Num() > MenuState->Num(); Index++)
-		{
-			MenuState->Add(FMacMenuItemState());
-		}
-		for (int32 Index = 0; Index < MenuBlocks.Num(); Index++)
-		{
-			FMacMenuItemState& ItemState = (*MenuState)[Index];
-			ItemState.Type = MenuBlocks[Index]->GetType();
 
-			if (ItemState.Type == EMultiBlockType::MenuEntry)
+			const TArray<TSharedRef<const FMultiBlock>>& MenuBlocks = Menu.MultiBox->GetBlocks();
+			for (int32 Index = MenuState->Num(); MenuBlocks.Num() > MenuState->Num(); Index++)
 			{
-				TSharedRef<const FMenuEntryBlock> Block = StaticCastSharedRef<const FMenuEntryBlock>(MenuBlocks[Index]);
-				ItemState.Block = Block;
-				ItemState.Title = [FSlateMacMenu::GetMenuItemTitle(Block) retain];
-				ItemState.KeyEquivalent = [FSlateMacMenu::GetMenuItemKeyEquivalent(Block, &ItemState.KeyModifiers) retain];
-				if (!ItemState.Icon)
+				MenuState->Add(FMacMenuItemState());
+			}
+			for (int32 Index = 0; Index < MenuBlocks.Num(); Index++)
+			{
+				FMacMenuItemState& ItemState = (*MenuState)[Index];
+				ItemState.Type = MenuBlocks[Index]->GetType();
+
+				if (ItemState.Type == EMultiBlockType::MenuEntry)
 				{
-					ItemState.Icon = [FSlateMacMenu::GetMenuItemIcon(Block) retain];
+					TSharedRef<const FMenuEntryBlock> Block = StaticCastSharedRef<const FMenuEntryBlock>(MenuBlocks[Index]);
+					ItemState.Block = Block;
+					ItemState.Title = [FSlateMacMenu::GetMenuItemTitle(Block) retain];
+					ItemState.KeyEquivalent = [FSlateMacMenu::GetMenuItemKeyEquivalent(Block, &ItemState.KeyModifiers) retain];
+					if (!ItemState.Icon)
+					{
+						ItemState.Icon = [FSlateMacMenu::GetMenuItemIcon(Block) retain];
+					}
+					ItemState.IsSubMenu = Block->bIsSubMenu;
+					ItemState.IsEnabled = FSlateMacMenu::IsMenuItemEnabled(Block);
+					ItemState.State = ItemState.IsSubMenu ? 0 : FSlateMacMenu::GetMenuItemState(Block);
 				}
-				ItemState.IsSubMenu = Block->bIsSubMenu;
-				ItemState.IsEnabled = FSlateMacMenu::IsMenuItemEnabled(Block);
-				ItemState.State = ItemState.IsSubMenu ? 0 : FSlateMacMenu::GetMenuItemState(Block);
 			}
 		}
 	}
