Fixed Slate synthesizing pointer moves even on platforms that do not support a cursor.

diff --git a/Engine/Source/Runtime/Slate/Private/Framework/Application/SlateApplication.cpp b/Engine/Source/Runtime/Slate/Private/Framework/Application/SlateApplication.cpp
--- a/Engine/Source/Runtime/Slate/Private/Framework/Application/SlateApplication.cpp
+++ b/Engine/Source/Runtime/Slate/Private/Framework/Application/SlateApplication.cpp
@@ -2528,27 +2528,30 @@ bool FSlateApplication::IsWindowInDestroyQueue(TSharedRef<SWindow> Window) const
 
 void FSlateApplication::SynthesizeMouseMove()
 {
-	// Synthetic mouse events accomplish two goals:
-	// 1) The UI can change even if the mosue doesn't move.
-	//    Synthesizing a mouse move sends out events.
-	//    In this case, the current and previous position will be the same.
-	//
-	// 2) The mouse moves, but the OS decided not to send us an event.
-	//    e.g. Mouse moved outside of our window.
-	//    In this case, the previous and current positions differ.
-
-	FPointerEvent MouseEvent
-	(
-		CursorPointerIndex,
-		GetCursorPos(),
-		GetLastCursorPos(),
-		PressedMouseButtons,
-		EKeys::Invalid,
-		0,
-		PlatformApplication->GetModifierKeys()
-	);
+	if (PlatformApplication->Cursor.IsValid())
+	{
+		// Synthetic mouse events accomplish two goals:
+		// 1) The UI can change even if the mosue doesn't move.
+		//    Synthesizing a mouse move sends out events.
+		//    In this case, the current and previous position will be the same.
+		//
+		// 2) The mouse moves, but the OS decided not to send us an event.
+		//    e.g. Mouse moved outside of our window.
+		//    In this case, the previous and current positions differ.
+
+		FPointerEvent MouseEvent
+		(
+			CursorPointerIndex,
+			GetCursorPos(),
+			GetLastCursorPos(),
+			PressedMouseButtons,
+			EKeys::Invalid,
+			0,
+			PlatformApplication->GetModifierKeys()
+		);
 
-	ProcessMouseMoveEvent(MouseEvent, true);
+		ProcessMouseMoveEvent(MouseEvent, true);
+	}
 }
 
 void FSlateApplication::OnLogSlateEvent(EEventLog::Type Elovent, const FString& AdditionalContent)
