Fixed list views scrolling backwards when touch-scrolling

This regression was caused by a change to the touch FPointerEvent constructor in CL#2303366, which corrects the sign of the cursor delta. List view touch scrolling had previously (perhaps unknowingly) worked around this issue by negating the scroll direction. Touch scrolling and mouse scrolling now work in the same way.

This addresses TTP#348351 - Slate lists scroll backwards when using touch.

diff --git a/Engine/Source/Runtime/Slate/Private/Widgets/Views/STableViewBase.cpp b/Engine/Source/Runtime/Slate/Private/Widgets/Views/STableViewBase.cpp
--- a/Engine/Source/Runtime/Slate/Private/Widgets/Views/STableViewBase.cpp
+++ b/Engine/Source/Runtime/Slate/Private/Widgets/Views/STableViewBase.cpp
@@ -456,9 +456,9 @@ FReply STableViewBase::OnTouchMoved( const FGeometry& MyGeometry, const FPointer
 	{
 		const float ScrollByAmount = InTouchEvent.GetCursorDelta().Y / MyGeometry.Scale;
 		AmountScrolledWhileRightMouseDown += FMath::Abs( ScrollByAmount );
-		TickScrollDelta += ScrollByAmount;
+		TickScrollDelta -= ScrollByAmount;
 
-		const float AmountScrolled = this->ScrollBy( MyGeometry, ScrollByAmount, AllowOverscroll );
+		const float AmountScrolled = this->ScrollBy( MyGeometry, -ScrollByAmount, AllowOverscroll );
 
 		if (AmountScrolledWhileRightMouseDown > SlatePanTriggerDistance)
 		{
