Fix MAC: CRASH: Closing level blueprint after PIE without changing Slate - when updating the multi-box we must reset the cached state, not just remove the same entries, because not all of them will exist. This ensures we don't keep around menus pointing to invalid FMenuExtenders.

diff --git a/Engine/Source/Runtime/Slate/Private/Framework/MultiBox/Mac/MacMenu.cpp b/Engine/Source/Runtime/Slate/Private/Framework/MultiBox/Mac/MacMenu.cpp
--- a/Engine/Source/Runtime/Slate/Private/Framework/MultiBox/Mac/MacMenu.cpp
+++ b/Engine/Source/Runtime/Slate/Private/Framework/MultiBox/Mac/MacMenu.cpp
@@ -93,9 +93,9 @@ void FSlateMacMenu::UpdateWithMultiBox(const TSharedRef< FMultiBox >& MultiBox)
 			NSMenuItem* HelpMenu = [[[NSApp mainMenu] itemWithTitle:HelpTitle.ToString().GetNSString()] retain];
 			for (int32 Index = NumItems - 1; Index > 0; Index--)
 			{
-				GCachedMenuState.Remove((FMacMenu*)[[NSApp mainMenu] itemAtIndex:Index]);
 				[[NSApp mainMenu] removeItemAtIndex:Index];
 			}
+			GCachedMenuState.Reset();
 
 			const TArray<TSharedRef<const FMultiBlock> >& MenuBlocks = MultiBox->GetBlocks();
 
@@ -138,6 +138,7 @@ void FSlateMacMenu::UpdateWithMultiBox(const TSharedRef< FMultiBox >& MultiBox)
 			if ([[NSApp mainMenu] itemWithTitle:HelpTitle.ToString().GetNSString()] == nil && HelpMenu)
 			{
 				[[NSApp mainMenu] addItem:HelpMenu];
+				GCachedMenuState.Add((FMacMenu*)HelpMenu, TSharedPtr<TArray<FMacMenuItemState>>(new TArray<FMacMenuItemState>()));
 			}
 
 			if (HelpMenu)
