Don't call BringToFront on mouse up if the window is already active

diff --git a/Engine/Source/Runtime/Slate/Private/Framework/Application/SlateApplication.cpp b/Engine/Source/Runtime/Slate/Private/Framework/Application/SlateApplication.cpp
--- a/Engine/Source/Runtime/Slate/Private/Framework/Application/SlateApplication.cpp
+++ b/Engine/Source/Runtime/Slate/Private/Framework/Application/SlateApplication.cpp
@@ -3492,6 +3492,9 @@ bool FSlateApplication::ProcessMouseButtonUpEvent( FPointerEvent& MouseEvent )
 		FWidgetPath MouseCaptorPath = MouseCaptor.ToWidgetPath(MouseEvent.GetPointerIndex());
 		if ( ensureMsg(MouseCaptorPath.Widgets.Num() > 0, TEXT("A window had a widget with mouse capture. That entire window has been dismissed before the mouse up could be processed.")) )
 		{
+#if PLATFORM_MAC
+			NSWindow* ActiveNativeWindow = [NSApp keyWindow];
+#endif
 			FArrangedWidget& MouseCaptorWidget = MouseCaptorPath.Widgets.Last();
 			MouseEvent.SetEventPath(MouseCaptorPath);
 
@@ -3509,8 +3512,9 @@ bool FSlateApplication::ProcessMouseButtonUpEvent( FPointerEvent& MouseEvent )
 			}
 			ProcessReply(MouseCaptorPath, Reply, &MouseCaptorPath, &MouseEvent);
 #if PLATFORM_MAC
+			// Activate a window under the mouse if it's inactive and mouse up didn't bring any window to front
 			TSharedPtr<SWindow> ActiveWindow = GetActiveTopLevelWindow();
-			if (MouseEvent.GetEffectingButton() == EKeys::LeftMouseButton && MouseCaptorPath.TopLevelWindow.IsValid() && ActiveWindow != MouseCaptorPath.TopLevelWindow)
+			if (MouseEvent.GetEffectingButton() == EKeys::LeftMouseButton && MouseCaptorPath.TopLevelWindow.IsValid() && ActiveWindow != MouseCaptorPath.TopLevelWindow && ActiveNativeWindow == [NSApp keyWindow])
 			{
 				MouseCaptorPath.TopLevelWindow->BringToFront(true);
 			}
