Slate - The Menu Anchor now allows content to appear on the left.  Customized the enum display names to be less wordy for UMG.

diff --git a/Engine/Source/Runtime/Slate/Private/Widgets/Input/SMenuAnchor.cpp b/Engine/Source/Runtime/Slate/Private/Widgets/Input/SMenuAnchor.cpp
--- a/Engine/Source/Runtime/Slate/Private/Widgets/Input/SMenuAnchor.cpp
+++ b/Engine/Source/Runtime/Slate/Private/Widgets/Input/SMenuAnchor.cpp
@@ -125,19 +125,30 @@ int32 SMenuAnchor::OnPaint( const FPaintArgs& Args, const FGeometry& AllottedGeo
 	return LayerId;
 }
 
+static FVector2D GetMenuOffsetForPlacement(const FGeometry& AllottedGeometry, EMenuPlacement PlacementMode, FVector2D PopupSizeLocalSpace)
+{
+	switch ( PlacementMode )
+	{
+	case MenuPlacement_AboveAnchor:
+		return FVector2D(0.0f, -AllottedGeometry.GetLocalSize().Y - PopupSizeLocalSpace.Y);
+	case MenuPlacement_ComboBoxRight:
+		return FVector2D(AllottedGeometry.GetLocalSize().X - PopupSizeLocalSpace.X, 0.0f);
+	case MenuPlacement_MenuLeft:
+		return FVector2D(-AllottedGeometry.GetLocalSize().X - PopupSizeLocalSpace.X, 0.0f);
+	default:
+		return FVector2D::ZeroVector;
+	}
+}
+
 FGeometry SMenuAnchor::ComputeMenuPlacement( const FGeometry& AllottedGeometry, const FVector2D& PopupDesiredSize, EMenuPlacement PlacementMode )
 {
 	// Compute the popup size, offset, and anchor rect  in local space
 	const FVector2D PopupSizeLocalSpace = (PlacementMode == MenuPlacement_ComboBox || PlacementMode == MenuPlacement_ComboBoxRight)
 		? FVector2D( FMath::Max( AllottedGeometry.Size.X, PopupDesiredSize.X ), PopupDesiredSize.Y )
 		: PopupDesiredSize;
-	const FVector2D OffsetLocalSpace = (PlacementMode == MenuPlacement_AboveAnchor)
-		? FVector2D( 0.0f, -AllottedGeometry.GetLocalSize().Y - PopupSizeLocalSpace.Y )
-		: (PlacementMode == MenuPlacement_ComboBoxRight) 
-			? FVector2D(AllottedGeometry.GetLocalSize().X - PopupSizeLocalSpace.X, 0.0f) 
-			: FVector2D::ZeroVector;
+	const FVector2D OffsetLocalSpace = GetMenuOffsetForPlacement(AllottedGeometry, PlacementMode, PopupSizeLocalSpace);
 	const FSlateRect AnchorLocalSpace = FSlateRect::FromPointAndExtent(OffsetLocalSpace, AllottedGeometry.Size);
-	const EOrientation Orientation = (PlacementMode == MenuPlacement_MenuRight) ? Orient_Horizontal : Orient_Vertical;
+	const EOrientation Orientation = (PlacementMode == MenuPlacement_MenuRight || PlacementMode == MenuPlacement_MenuLeft) ? Orient_Horizontal : Orient_Vertical;
 	
 	// ask the application to compute the proper desktop offset for the anchor. This requires the offsets to be in desktop space.
 	const FVector2D NewPositionDesktopSpace = FSlateApplication::Get().CalculatePopupWindowPosition( 
diff --git a/Engine/Source/Runtime/SlateCore/Public/Types/SlateEnums.h b/Engine/Source/Runtime/SlateCore/Public/Types/SlateEnums.h
--- a/Engine/Source/Runtime/SlateCore/Public/Types/SlateEnums.h
+++ b/Engine/Source/Runtime/SlateCore/Public/Types/SlateEnums.h
@@ -115,7 +115,7 @@ UENUM(BlueprintType)
 enum EMenuPlacement
 {
 	/** Place the menu immediately below the anchor */
-	MenuPlacement_BelowAnchor UMETA(DisplayName="Below Anchor"),
+	MenuPlacement_BelowAnchor UMETA(DisplayName="Below"),
 
 	/** Place the menu immediately below the anchor and match is width to the anchor's content */
 	MenuPlacement_ComboBox UMETA(DisplayName="Combo Box"),
@@ -124,10 +124,13 @@ enum EMenuPlacement
 	MenuPlacement_ComboBoxRight UMETA(DisplayName="Combo Box Right"),
 
 	/** Place the menu to the right of the anchor */
-	MenuPlacement_MenuRight UMETA(DisplayName="Menu Right"),
+	MenuPlacement_MenuRight UMETA(DisplayName="Right"),
 
 	/** Place the menu immediately above the anchor, not transition effect */
-	MenuPlacement_AboveAnchor UMETA(DisplayName="Above Anchor"),
+	MenuPlacement_AboveAnchor UMETA(DisplayName="Above"),
+
+	/** Place the menu to the left of the anchor */
+	MenuPlacement_MenuLeft UMETA(DisplayName="Left"),
 };
 
 
