Added an SWebBrowser widget

Web Browser widget displays a web page and allows for some optional basic controls.
Added initialisation and update of Web Browser module to Slate Application so that each web browser window doesn't have to manage this individually.
Added example usage of the widget to STestSuite so that it can be used via the SlateViewer app.
Fixed error in name of CEF3Utils module.

diff --git a/Engine/Source/Developer/CEF3Utils/Private/CEF3Utils.cpp b/Engine/Source/Developer/CEF3Utils/Private/CEF3Utils.cpp
--- a/Engine/Source/Developer/CEF3Utils/Private/CEF3Utils.cpp
+++ b/Engine/Source/Developer/CEF3Utils/Private/CEF3Utils.cpp
@@ -4,7 +4,7 @@
 
 DEFINE_LOG_CATEGORY(LogCEF3Utils);
 
-IMPLEMENT_MODULE(FDefaultModuleImpl, CEF3);
+IMPLEMENT_MODULE(FDefaultModuleImpl, CEF3Utils);
 
 #if WITH_CEF3
 namespace CEF3Utils
diff --git a/Engine/Source/Runtime/AppFramework/Private/Widgets/Testing/STestSuite.cpp b/Engine/Source/Runtime/AppFramework/Private/Widgets/Testing/STestSuite.cpp
--- a/Engine/Source/Runtime/AppFramework/Private/Widgets/Testing/STestSuite.cpp
+++ b/Engine/Source/Runtime/AppFramework/Private/Widgets/Testing/STestSuite.cpp
@@ -12,6 +12,7 @@
 #include "TransformCalculus3D.h"
 #include "SlateRenderTransform.h"
 #include "SlateLayoutTransform.h"
+#include "SWebBrowser.h"
 
 #define LOCTEXT_NAMESPACE "STestSuite"
 
@@ -4626,6 +4627,19 @@ TSharedRef<SDockTab> SpawnTab(const FSpawnTabArgs& Args, FName TabIdentifier)
 				]
 			];
 	}
+	else if (TabIdentifier == FName(TEXT("WebBrowserTab")))
+	{
+		const FString DocPath(FPaths::ConvertRelativePathToFull(FPaths::Combine(*FPaths::EngineDir(), TEXT("Documentation/HTML/INT/Resources/index.html"))));
+		return SNew(SDockTab)
+			.Label(LOCTEXT("WebBrowserTab", "Web Browser"))
+			.ToolTipText(LOCTEXT("WebBrowserTabToolTip", "Switches to the Web Browser to test its features."))
+			[
+				SNew(SWebBrowser)
+				.ParentWindow(Args.GetOwnerWindow())
+				//.InitialURL(DocPath)
+				//.InitialURL(TEXT("www.youtube.com"))
+			];
+	}
 #endif //WITH_FANCY_TEXT
 	else if (TabIdentifier == FName("LayoutRoundingTab"))
 	{
@@ -4811,6 +4825,7 @@ TSharedRef<SDockTab> SpawnTestSuite1( const FSpawnTabArgs& Args )
 				->AddTab("RichTextTab", ETabState::OpenedTab)
 				->AddTab("MultiLineEditTab", ETabState::OpenedTab)
 				->AddTab("RichEditableTextTab", ETabState::OpenedTab)
+				->AddTab("WebBrowserTab", ETabState::OpenedTab)
 #endif //WITH_FANCY_TEXT
 				)
 		)
@@ -4863,6 +4878,10 @@ TSharedRef<SDockTab> SpawnTestSuite1( const FSpawnTabArgs& Args )
 		.SetDisplayName( NSLOCTEXT("TestSuite1", "RichEditableTextTab", "Rich Editable Text Test") )
 		.SetGroup(TestSuiteMenu::SuiteTabs);
 
+	TestSuite1TabManager->RegisterTabSpawner("WebBrowserTab", FOnSpawnTab::CreateStatic(&SpawnTab, FName("WebBrowserTab")))
+		.SetDisplayName(NSLOCTEXT("TestSuite1", "WebBrowserTab", "Web Browser test"))
+		.SetGroup(TestSuiteMenu::SuiteTabs);
+
 	TestSuite1TabManager->RegisterTabSpawner( "LayoutRoundingTab", FOnSpawnTab::CreateStatic( &SpawnTab, FName("LayoutRoundingTab") ) )
 		.SetDisplayName( NSLOCTEXT("TestSuite1", "LayoutRoundingTab", "Layout Rounding") )
 		.SetGroup(TestSuiteMenu::SuiteTabs);
diff --git a/Engine/Source/Runtime/Slate/Private/Framework/Application/SlateApplication.cpp b/Engine/Source/Runtime/Slate/Private/Framework/Application/SlateApplication.cpp
--- a/Engine/Source/Runtime/Slate/Private/Framework/Application/SlateApplication.cpp
+++ b/Engine/Source/Runtime/Slate/Private/Framework/Application/SlateApplication.cpp
@@ -5,6 +5,9 @@
 #include "SWindowTitleBar.h"
 #include "HittestGrid.h"
 
+#include "WebBrowserModule.h"
+#include "IWebBrowserSingleton.h"
+
 
 DECLARE_CYCLE_STAT( TEXT("Message Tick Time"), STAT_SlateMessageTick, STATGROUP_Slate );
 DECLARE_CYCLE_STAT( TEXT("Update Tooltip Time"), STAT_SlateUpdateTooltip, STATGROUP_Slate );
@@ -444,6 +447,9 @@ void FSlateApplication::InitializeRenderer( TSharedRef<FSlateRenderer> InRendere
 {
 	Renderer = InRenderer;
 	Renderer->Initialize();
+
+	// Pass renderer on to the Web Browser module
+	IWebBrowserModule::Get().GetSingleton()->SetSlateRenderer(InRenderer);
 }
 
 void FSlateApplication::InitializeSound( const TSharedRef<ISlateSoundDevice>& InSlateSoundDevice )
@@ -921,6 +927,9 @@ void FSlateApplication::Tick()
 	// (so they know their size)
 	FSlateNotificationManager::Get().Tick();
 
+	// Calling this after Tick gives web browser a chance to render pages to texture after size changes
+	IWebBrowserModule::Get().GetSingleton()->PumpMessages();
+
 	// Draw all windows
 	DrawWindows();
 }
diff --git a/Engine/Source/Runtime/Slate/Private/Widgets/SWebBrowser.cpp b/Engine/Source/Runtime/Slate/Private/Widgets/SWebBrowser.cpp
new file mode 100644
--- /dev/null
+++ b/Engine/Source/Runtime/Slate/Private/Widgets/SWebBrowser.cpp
@@ -0,0 +1,161 @@
+// Copyright 1998-2014 Epic Games, Inc. All Rights Reserved.
+
+#include "SlatePrivatePCH.h"
+#include "SWebBrowser.h"
+
+#include "WebBrowserModule.h"
+#include "IWebBrowserSingleton.h"
+#include "IWebBrowserWindow.h"
+#include "WebBrowserViewport.h"
+
+#define LOCTEXT_NAMESPACE "WebBrowser"
+
+SWebBrowser::SWebBrowser()
+{
+}
+
+void SWebBrowser::Construct(const FArguments& InArgs)
+{
+	void* OSWindowHandle = nullptr;
+	if (InArgs._ParentWindow.IsValid())
+	{
+		TSharedPtr<FGenericWindow> NativeWindow = InArgs._ParentWindow->GetNativeWindow();
+		OSWindowHandle = NativeWindow->GetOSWindowHandle();
+	}
+	BrowserWindow = IWebBrowserModule::Get().GetSingleton()->CreateBrowserWindow(OSWindowHandle, InArgs._InitialURL, InArgs._ViewportSize.Get().X, InArgs._ViewportSize.Get().Y);
+
+	TSharedPtr<SViewport> ViewportWidget;
+
+	ChildSlot
+	[
+		SNew(SVerticalBox)
+		+SVerticalBox::Slot()
+		.AutoHeight()
+		.Padding(0, 5)
+		[
+			SNew(SHorizontalBox)
+			.Visibility(InArgs._ShowControls ? EVisibility::Visible : EVisibility::Collapsed)
+			+SHorizontalBox::Slot()
+			.AutoWidth()
+			[
+				SNew(SButton)
+				.Text(LOCTEXT("Back","Back"))
+				.IsEnabled(this, &SWebBrowser::CanGoBack)
+				.OnClicked(this, &SWebBrowser::OnBackClicked)
+			]
+			+SHorizontalBox::Slot()
+			.AutoWidth()
+			[
+				SNew(SButton)
+				.Text(LOCTEXT("Forward", "Forward"))
+				.IsEnabled(this, &SWebBrowser::CanGoForward)
+				.OnClicked(this, &SWebBrowser::OnForwardClicked)
+			]
+			+SHorizontalBox::Slot()
+			.AutoWidth()
+			[
+				SNew(SButton)
+				.Text(this, &SWebBrowser::GetReloadButtonText)
+				.OnClicked(this, &SWebBrowser::OnReloadClicked)
+			]
+			+SHorizontalBox::Slot()
+			.FillWidth(1.0f)
+			.VAlign(VAlign_Center)
+			.HAlign(HAlign_Right)
+			.Padding(5)
+			[
+				SNew(STextBlock)
+				.Text(this, &SWebBrowser::GetTitleText)
+				.Justification(ETextJustify::Right)
+			]
+		]
+		+SVerticalBox::Slot()
+		[
+			SAssignNew(ViewportWidget, SViewport)
+			.ViewportSize(InArgs._ViewportSize)
+			.EnableGammaCorrection(false)
+		]
+	];
+
+	if (BrowserWindow.IsValid())
+	{
+		BrowserViewport = MakeShareable(new FWebBrowserViewport(BrowserWindow, ViewportWidget));
+		ViewportWidget->SetViewportInterface(BrowserViewport.ToSharedRef());
+	}
+}
+
+FText SWebBrowser::GetTitleText() const
+{
+	if (BrowserWindow.IsValid())
+	{
+		return FText::FromString(BrowserWindow->GetTitle());
+	}
+	return LOCTEXT("InvalidWindow", "Browser Window is not valid/supported");
+}
+
+bool SWebBrowser::CanGoBack() const
+{
+	if (BrowserWindow.IsValid())
+	{
+		return BrowserWindow->CanGoBack();
+	}
+	return false;
+}
+
+FReply SWebBrowser::OnBackClicked()
+{
+	if (BrowserWindow.IsValid())
+	{
+		BrowserWindow->GoBack();
+	}
+	return FReply::Handled();
+}
+
+bool SWebBrowser::CanGoForward() const
+{
+	if (BrowserWindow.IsValid())
+	{
+		return BrowserWindow->CanGoForward();
+	}
+	return false;
+}
+
+FReply SWebBrowser::OnForwardClicked()
+{
+	if (BrowserWindow.IsValid())
+	{
+		BrowserWindow->GoForward();
+	}
+	return FReply::Handled();
+}
+
+FText SWebBrowser::GetReloadButtonText() const
+{
+	static FText Reload = LOCTEXT("Reload", "Reload");
+	static FText Stop = LOCTEXT("Stop", "Stop");
+
+	if (BrowserWindow.IsValid())
+	{
+		if (BrowserWindow->IsLoading())
+		{
+			return Stop;
+		}
+	}
+	return Reload;
+}
+
+FReply SWebBrowser::OnReloadClicked()
+{
+	if (BrowserWindow.IsValid())
+	{
+		if (BrowserWindow->IsLoading())
+		{
+			BrowserWindow->StopLoad();
+		}
+		else
+		{
+			BrowserWindow->Reload();
+		}
+	}
+	return FReply::Handled();
+}
diff --git a/Engine/Source/Runtime/Slate/Public/Widgets/SWebBrowser.h b/Engine/Source/Runtime/Slate/Public/Widgets/SWebBrowser.h
new file mode 100644
--- /dev/null
+++ b/Engine/Source/Runtime/Slate/Public/Widgets/SWebBrowser.h
@@ -0,0 +1,84 @@
+// Copyright 1998-2014 Epic Games, Inc. All Rights Reserved.
+
+#pragma once
+
+class IWebBrowserWindow;
+class FWebBrowserViewport;
+
+class SLATE_API SWebBrowser : public SCompoundWidget
+{
+public:
+	SLATE_BEGIN_ARGS(SWebBrowser)
+		: _InitialURL(TEXT("www.google.com"))
+		, _ShowControls(true)
+		, _ViewportSize(FVector2D(320, 240))
+		{}
+
+		/** A reference to the parent window */
+		SLATE_ARGUMENT(TSharedPtr<SWindow>, ParentWindow)
+
+		/** URL that the browser will initially navigate to */
+		SLATE_ARGUMENT(FString, InitialURL)
+
+		/** Whether to show standard controls like Back, Forward, Reload etc. */
+		SLATE_ARGUMENT(bool, ShowControls)
+
+		/** Desired size of the web browser viewport */
+		SLATE_ATTRIBUTE(FVector2D, ViewportSize);
+	SLATE_END_ARGS()
+
+	/**
+	 * Constructor
+	 */
+	SWebBrowser();
+
+	/**
+	 * Construct the widget.
+	 *
+	 * @param InArgs  Declaration from which to construct the widget.
+	 */
+	void Construct(const FArguments& InArgs);
+
+	/**
+	 * Get the current title of the web page
+	 */
+	FText GetTitleText() const;
+
+private:
+	/**
+	 * Returns true if the browser can navigate backwards.
+	 */
+	bool CanGoBack() const;
+
+	/**
+	 * Navigate backwards.
+	 */
+	FReply OnBackClicked();
+
+	/**
+	 * Returns true if the browser can navigate forwards.
+	 */
+	bool CanGoForward() const;
+
+	/**
+	 * Navigate forwards.
+	 */
+	FReply OnForwardClicked();
+
+	/**
+	 * Get text for reload button depending on status
+	 */
+	FText GetReloadButtonText() const;
+
+	/**
+	 * Reload or stop loading
+	 */
+	FReply OnReloadClicked();
+
+private:
+	/** Interface for dealing with a web browser window */
+	TSharedPtr<IWebBrowserWindow>	BrowserWindow;
+
+	/** Viewport interface for rendering the web page */
+	TSharedPtr<FWebBrowserViewport>	BrowserViewport;
+};
diff --git a/Engine/Source/Runtime/Slate/Slate.Build.cs b/Engine/Source/Runtime/Slate/Slate.Build.cs
--- a/Engine/Source/Runtime/Slate/Slate.Build.cs
+++ b/Engine/Source/Runtime/Slate/Slate.Build.cs
@@ -18,6 +18,12 @@ public class Slate : ModuleRules
 			}
 		);
 
+		PrivateDependencyModuleNames.AddRange(
+			new string[] {
+				"WebBrowser"
+			}
+		);
+
         if (!UEBuildConfiguration.bBuildDedicatedServer && UEBuildConfiguration.bCompileFreeType)
         {
 			AddThirdPartyPrivateStaticDependencies(Target, "FreeType2");
