Slate - We no longer create the menus when destroying the menu.

diff --git a/Engine/Source/Runtime/Slate/Private/Widgets/Input/SMenuAnchor.cpp b/Engine/Source/Runtime/Slate/Private/Widgets/Input/SMenuAnchor.cpp
--- a/Engine/Source/Runtime/Slate/Private/Widgets/Input/SMenuAnchor.cpp
+++ b/Engine/Source/Runtime/Slate/Private/Widgets/Input/SMenuAnchor.cpp
@@ -178,14 +178,13 @@ void SMenuAnchor::RequestClosePopupWindow( const TSharedRef<SWindow>& PopupWindo
 
 void SMenuAnchor::OnClickedOutsidePopup()
 {
-		bDismissedThisTick = true;
-		if (ensure(Method == UseCurrentWindow))
-		{
-			FSlateApplication::Get().GetPopupSupport().UnregisterClickNotification( FOnClickedOutside::CreateSP(this, &SMenuAnchor::OnClickedOutsidePopup) );
-			SetIsOpen(false);	
-		}
+	bDismissedThisTick = true;
+	if (ensure(Method == UseCurrentWindow))
+	{
+		FSlateApplication::Get().GetPopupSupport().UnregisterClickNotification( FOnClickedOutside::CreateSP(this, &SMenuAnchor::OnClickedOutsidePopup) );
+		SetIsOpen(false);	
 	}
-
+}
 
 void SMenuAnchor::SetContent(TSharedRef<SWidget> InContent)
 {
@@ -211,82 +210,85 @@ void SMenuAnchor::SetIsOpen( bool InIsOpen, const bool bFocusMenu )
 	// Prevent redundant opens/closes
 	if ( IsOpen() != InIsOpen )
 	{
-		TSharedPtr< SWidget > MenuContentPtr = OnGetMenuContent.IsBound() ? OnGetMenuContent.Execute() : MenuContent;
-
-		if (InIsOpen && MenuContentPtr.IsValid() )
+		if ( InIsOpen )
 		{
-			// OPEN POPUP
-			if (OnMenuOpenChanged.IsBound())
-			{
-				OnMenuOpenChanged.Execute(true);
-			}
-			
-			// This can be called at any time so we use the push menu override that explicitly allows us to specify our parent
-
-			// Figure out where the menu anchor is on the screen, so we can set the initial position of our pop-up window
-			SlatePrepass();
-
-			// NOTE: Careful, GeneratePathToWidget can be reentrant in that it can call visibility delegates and such
-			FWidgetPath MyWidgetPath;
-			FSlateApplication::Get().GeneratePathToWidgetChecked( AsShared(), MyWidgetPath );
-			const FGeometry& MyGeometry = MyWidgetPath.Widgets.Last().Geometry;
-
-			// @todo Slate: This code does not properly propagate the layout scale of the widget we are creating the popup for.
-			// The popup instead has a scale of one, but a computed size as if the contents were scaled. This code should ideally
-			// wrap the contents with an SFxWidget that applies the necessary layout scale. This is a very rare case right now.
-			
-			// Figure out how big the content widget is so we can set the window's initial size properly
-			TSharedRef< SWidget > MenuContentRef = MenuContentPtr.ToSharedRef();
-			MenuContentRef->SlatePrepass();
-
-			// Combo-boxes never size down smaller than the widget that spawned them, but all
-			// other pop-up menus are currently auto-sized
-			const FVector2D DesiredContentSize = MenuContentRef->GetDesiredSize();  // @todo: This is ignoring any window border size!
-			const EMenuPlacement PlacementMode = Placement.Get();
-
-			const FVector2D NewPosition = MyGeometry.AbsolutePosition;
-			FVector2D NewWindowSize = DesiredContentSize;
-			const FVector2D SummonLocationSize = MyGeometry.Size;
-
-			FPopupTransitionEffect TransitionEffect( FPopupTransitionEffect::None );
-			if( PlacementMode == MenuPlacement_ComboBox || PlacementMode == MenuPlacement_ComboBoxRight )
-			{
-				TransitionEffect = FPopupTransitionEffect( FPopupTransitionEffect::ComboButton );
-				NewWindowSize = FVector2D( FMath::Max(MyGeometry.Size.X, DesiredContentSize.X), DesiredContentSize.Y );
-			}
-			else if( PlacementMode == MenuPlacement_BelowAnchor )
-			{
-				TransitionEffect = FPopupTransitionEffect( FPopupTransitionEffect::TopMenu );
-			}
-			else if( PlacementMode == MenuPlacement_MenuRight )
-			{
-				TransitionEffect = FPopupTransitionEffect( FPopupTransitionEffect::SubMenu );
-			}
+			TSharedPtr< SWidget > MenuContentPtr = OnGetMenuContent.IsBound() ? OnGetMenuContent.Execute() : MenuContent;
 
-			if (Method == CreateNewWindow)
+			if ( MenuContentPtr.IsValid() )
 			{
-				// Open the pop-up
-				TSharedRef<SWindow> PopupWindow = FSlateApplication::Get().PushMenu( AsShared(), MenuContentRef, NewPosition, TransitionEffect, bFocusMenu, false, NewWindowSize, SummonLocationSize );
-				PopupWindow->SetRequestDestroyWindowOverride( FRequestDestroyWindowOverride::CreateSP(this, &SMenuAnchor::RequestClosePopupWindow) );
-				PopupWindowPtr = PopupWindow;
-			}
-			else
-			{
-				// We are re-using the current window instead of creating a new one.
-				// The popup will be presented via an overlay service.
-				ensure(Method==UseCurrentWindow);
-				Children[1]
-				[
-					MenuContentRef
-				];
+				// OPEN POPUP
+				if ( OnMenuOpenChanged.IsBound() )
+				{
+					OnMenuOpenChanged.Execute(true);
+				}
 
-				// We want to support dismissing the popup widget when the user clicks outside it.
-				FSlateApplication::Get( ).GetPopupSupport( ).RegisterClickNotification( MenuContentRef, FOnClickedOutside::CreateSP( this, &SMenuAnchor::OnClickedOutsidePopup ) );
-			}
+				// This can be called at any time so we use the push menu override that explicitly allows us to specify our parent
 
-			if(bFocusMenu)
-			{
-				FSlateApplication::Get().SetKeyboardFocus(MenuContentRef, EKeyboardFocusCause::SetDirectly);
+				// Figure out where the menu anchor is on the screen, so we can set the initial position of our pop-up window
+				SlatePrepass();
+
+				// NOTE: Careful, GeneratePathToWidget can be reentrant in that it can call visibility delegates and such
+				FWidgetPath MyWidgetPath;
+				FSlateApplication::Get().GeneratePathToWidgetChecked(AsShared(), MyWidgetPath);
+				const FGeometry& MyGeometry = MyWidgetPath.Widgets.Last().Geometry;
+
+				// @todo Slate: This code does not properly propagate the layout scale of the widget we are creating the popup for.
+				// The popup instead has a scale of one, but a computed size as if the contents were scaled. This code should ideally
+				// wrap the contents with an SFxWidget that applies the necessary layout scale. This is a very rare case right now.
+
+				// Figure out how big the content widget is so we can set the window's initial size properly
+				TSharedRef< SWidget > MenuContentRef = MenuContentPtr.ToSharedRef();
+				MenuContentRef->SlatePrepass();
+
+				// Combo-boxes never size down smaller than the widget that spawned them, but all
+				// other pop-up menus are currently auto-sized
+				const FVector2D DesiredContentSize = MenuContentRef->GetDesiredSize();  // @todo: This is ignoring any window border size!
+				const EMenuPlacement PlacementMode = Placement.Get();
+
+				const FVector2D NewPosition = MyGeometry.AbsolutePosition;
+				FVector2D NewWindowSize = DesiredContentSize;
+				const FVector2D SummonLocationSize = MyGeometry.Size;
+
+				FPopupTransitionEffect TransitionEffect(FPopupTransitionEffect::None);
+				if ( PlacementMode == MenuPlacement_ComboBox || PlacementMode == MenuPlacement_ComboBoxRight )
+				{
+					TransitionEffect = FPopupTransitionEffect(FPopupTransitionEffect::ComboButton);
+					NewWindowSize = FVector2D(FMath::Max(MyGeometry.Size.X, DesiredContentSize.X), DesiredContentSize.Y);
+				}
+				else if ( PlacementMode == MenuPlacement_BelowAnchor )
+				{
+					TransitionEffect = FPopupTransitionEffect(FPopupTransitionEffect::TopMenu);
+				}
+				else if ( PlacementMode == MenuPlacement_MenuRight )
+				{
+					TransitionEffect = FPopupTransitionEffect(FPopupTransitionEffect::SubMenu);
+				}
+
+				if ( Method == CreateNewWindow )
+				{
+					// Open the pop-up
+					TSharedRef<SWindow> PopupWindow = FSlateApplication::Get().PushMenu(AsShared(), MenuContentRef, NewPosition, TransitionEffect, bFocusMenu, false, NewWindowSize, SummonLocationSize);
+					PopupWindow->SetRequestDestroyWindowOverride(FRequestDestroyWindowOverride::CreateSP(this, &SMenuAnchor::RequestClosePopupWindow));
+					PopupWindowPtr = PopupWindow;
+				}
+				else
+				{
+					// We are re-using the current window instead of creating a new one.
+					// The popup will be presented via an overlay service.
+					ensure(Method == UseCurrentWindow);
+					Children[1]
+						[
+							MenuContentRef
+						];
+
+					// We want to support dismissing the popup widget when the user clicks outside it.
+					FSlateApplication::Get().GetPopupSupport().RegisterClickNotification(MenuContentRef, FOnClickedOutside::CreateSP(this, &SMenuAnchor::OnClickedOutsidePopup));
+				}
+
+				if ( bFocusMenu )
+				{
+					FSlateApplication::Get().SetKeyboardFocus(MenuContentRef, EKeyboardFocusCause::SetDirectly);
+				}
 			}
 		}
 		else
