Updated Mac menu code to not depend on active tab being always correctly set and instead check if there's any visible, regular window open when deciding whether or not menus need update. This solves an issue with menu being empty in some cases.

diff --git a/Engine/Source/Runtime/Core/Public/Mac/MacApplication.h b/Engine/Source/Runtime/Core/Public/Mac/MacApplication.h
--- a/Engine/Source/Runtime/Core/Public/Mac/MacApplication.h
+++ b/Engine/Source/Runtime/Core/Public/Mac/MacApplication.h
@@ -112,6 +112,8 @@ public:
 	static TSharedPtr< FMacWindow > FindMacWindowByNSWindow( FCocoaWindow* const WindowHandle );
 	static void ProcessEvent(FMacEvent const* const Event);
 
+	const TArray<TSharedRef<FMacWindow>>& GetAllWindows() const { return Windows; }
+
 private:
 	enum FMacApplicationEventTypes
 	{
diff --git a/Engine/Source/Runtime/Slate/Private/Framework/MultiBox/Mac/MacMenu.cpp b/Engine/Source/Runtime/Slate/Private/Framework/MultiBox/Mac/MacMenu.cpp
--- a/Engine/Source/Runtime/Slate/Private/Framework/MultiBox/Mac/MacMenu.cpp
+++ b/Engine/Source/Runtime/Slate/Private/Framework/MultiBox/Mac/MacMenu.cpp
@@ -3,6 +3,7 @@
 #include "SlatePrivatePCH.h"
 #include "MacMenu.h"
 #include "CocoaThread.h"
+#include "MacApplication.h"
 
 struct FMacMenuItemState
 {
@@ -266,8 +267,21 @@ void FSlateMacMenu::UpdateMenu(FMacMenu* Menu)
 
 void FSlateMacMenu::UpdateCachedState()
 {
-	TSharedPtr<SDockTab> ActiveTab = FGlobalTabmanager::Get()->GetActiveTab();
-	if (ActiveTab.IsValid())
+	bool bShouldUpdate = false;
+
+	// @todo: Ideally this would ask global tab manager if there's any active tab, but that cannot be done reliably at the moment
+	// so instead we assume that as long as there's any visible, regular window open, we do have some menu to show/update.
+	const TArray<TSharedRef<FMacWindow>>&AllWindows = MacApplication->GetAllWindows();
+	for (auto Window : AllWindows)
+	{
+		if (Window->IsRegularWindow() && Window->IsVisible())
+		{
+			bShouldUpdate = true;
+			break;
+		}
+	}
+
+	if (bShouldUpdate)
 	{
 		for (TMap<FMacMenu*, TSharedPtr<TArray<FMacMenuItemState>>>::TIterator It(GCachedMenuState); It; ++It)
 		{
