Slate - The Menu anchor now handles the case where it has no content, but still wants to display a popped up child.

diff --git a/Engine/Source/Runtime/Slate/Private/Widgets/Input/SMenuAnchor.cpp b/Engine/Source/Runtime/Slate/Private/Widgets/Input/SMenuAnchor.cpp
--- a/Engine/Source/Runtime/Slate/Private/Widgets/Input/SMenuAnchor.cpp
+++ b/Engine/Source/Runtime/Slate/Private/Widgets/Input/SMenuAnchor.cpp
@@ -109,17 +109,40 @@ int32 SMenuAnchor::OnPaint( const FPaintArgs& Args, const FGeometry& AllottedGeo
 	// There may be zero elements in this array if our child collapsed/hidden
 	if ( ArrangedChildren.Num() > 0 )
 	{
-		FArrangedWidget& TheChild = ArrangedChildren[0];
+		const FArrangedWidget& FirstChild = ArrangedChildren[0];
 
-		const FSlateRect ChildClippingRect = AllottedGeometry.GetClippingRect().IntersectionWith( MyClippingRect );
+		// In the case where the user doesn't provide content to the menu anchor, the null widget
+		// wont appear in the visible set of arranged children, so only immediately paint the first child,
+		// if it's visible and matches the first slot content.
+		const bool bHasArrangedAnchorContent = FirstChild.Widget == Children[0].GetWidget();
+		if ( bHasArrangedAnchorContent )
+		{
+			const FSlateRect ChildClippingRect = AllottedGeometry.GetClippingRect().IntersectionWith(MyClippingRect);
+			LayerId = FirstChild.Widget->Paint(Args.WithNewParent(this), FirstChild.Geometry, ChildClippingRect, OutDrawElements, LayerId + 1, InWidgetStyle, ShouldBeEnabled(bParentEnabled));
+		}
 
-		LayerId = TheChild.Widget->Paint( Args.WithNewParent( this ), TheChild.Geometry, ChildClippingRect, OutDrawElements, LayerId + 1, InWidgetStyle, ShouldBeEnabled( bParentEnabled ) );
+		const bool bIsOpen = IsOpen();
 
-		if ( IsOpen( ) && ArrangedChildren.Num( ) > 1 && !ChildClippingRect.IsEmpty() )
+		if ( bIsOpen )
 		{
-			OutDrawElements.QueueDeferredPainting(
-				FSlateWindowElementList::FDeferredPaint( ArrangedChildren[1].Widget, Args, ArrangedChildren[1].Geometry, MyClippingRect, InWidgetStyle, bParentEnabled ) );
-		}
+			// In the case where the anchor content is present and visible, it's the 1 index child, in the case
+			// where the anchor content is invisible, it's the 0 index child.
+			FArrangedWidget* PopupChild = nullptr;
+			if ( bHasArrangedAnchorContent && ArrangedChildren.Num() > 1 )
+			{
+				PopupChild = &ArrangedChildren[1];
+			}
+			else if ( !bHasArrangedAnchorContent && ArrangedChildren.Num() == 1 )
+			{
+				PopupChild = &ArrangedChildren[0];
+			}
+
+			if ( PopupChild != nullptr )
+			{
+				 OutDrawElements.QueueDeferredPainting(
+					FSlateWindowElementList::FDeferredPaint(PopupChild->Widget, Args, PopupChild->Geometry, MyClippingRect, InWidgetStyle, bParentEnabled));
+			}
+		} 
 	}
 
 	return LayerId;
