Fixed FSlateTextRun::Measure not taking the shadow offset into account when working out the size

diff --git a/Engine/Source/Runtime/Slate/Private/Framework/Text/SlateTextRun.cpp b/Engine/Source/Runtime/Slate/Private/Framework/Text/SlateTextRun.cpp
--- a/Engine/Source/Runtime/Slate/Private/Framework/Text/SlateTextRun.cpp
+++ b/Engine/Source/Runtime/Slate/Private/Framework/Text/SlateTextRun.cpp
@@ -42,11 +42,11 @@ FVector2D FSlateTextRun::Measure( int32 BeginIndex, int32 EndIndex, float Scale
 {
 	if ( EndIndex - BeginIndex == 0 )
 	{
-		return FVector2D( 0, GetMaxHeight( Scale ) );
+		return FVector2D( 0, GetMaxHeight( Scale ) ) + Style.ShadowOffset;
 	}
 
 	const TSharedRef< FSlateFontMeasure > FontMeasure = FSlateApplication::Get().GetRenderer()->GetFontMeasureService();
-	return FontMeasure->Measure( **Text, BeginIndex, EndIndex, Style.Font, true, Scale );
+	return FontMeasure->Measure( **Text, BeginIndex, EndIndex, Style.Font, true, Scale ) + Style.ShadowOffset;
 }
 
 int8 FSlateTextRun::GetKerning(int32 CurrentIndex, float Scale) const
