Update Mac menu in UpdateWithMultiBox() even if multibox didn't change in case some menu items were added or removed

diff --git a/Engine/Source/Runtime/Slate/Private/Framework/MultiBox/Mac/MacMenu.cpp b/Engine/Source/Runtime/Slate/Private/Framework/MultiBox/Mac/MacMenu.cpp
--- a/Engine/Source/Runtime/Slate/Private/Framework/MultiBox/Mac/MacMenu.cpp
+++ b/Engine/Source/Runtime/Slate/Private/Framework/MultiBox/Mac/MacMenu.cpp
@@ -76,80 +76,74 @@ static TMap<FMacMenu*, TSharedPtr<TArray<FMacMenuItemState>>> GCachedMenuState;
 
 @end
 
-static TSharedPtr<FMultiBox> GCurrentMultiBox = nullptr;
-
 void FSlateMacMenu::UpdateWithMultiBox(const TSharedRef< FMultiBox >& MultiBox)
 {
 	MainThreadCall(^{
-		if (MultiBox != GCurrentMultiBox)
+		if (!FPlatformMisc::UpdateCachedMacMenuState)
 		{
-			if (!FPlatformMisc::UpdateCachedMacMenuState)
-			{
-				FPlatformMisc::UpdateCachedMacMenuState = UpdateCachedState;
-			}
+			FPlatformMisc::UpdateCachedMacMenuState = UpdateCachedState;
+		}
 
-			int32 NumItems = [[NSApp mainMenu] numberOfItems];
-			FText HelpTitle = NSLOCTEXT("MainMenu", "HelpMenu", "Help");
-			FText WindowLabel = NSLOCTEXT("MainMenu", "WindowMenu", "Window");
-			NSMenuItem* HelpMenu = [[[NSApp mainMenu] itemWithTitle:HelpTitle.ToString().GetNSString()] retain];
-			for (int32 Index = NumItems - 1; Index > 0; Index--)
-			{
-				[[NSApp mainMenu] removeItemAtIndex:Index];
-			}
-			GCachedMenuState.Reset();
+		int32 NumItems = [[NSApp mainMenu] numberOfItems];
+		FText HelpTitle = NSLOCTEXT("MainMenu", "HelpMenu", "Help");
+		FText WindowLabel = NSLOCTEXT("MainMenu", "WindowMenu", "Window");
+		NSMenuItem* HelpMenu = [[[NSApp mainMenu] itemWithTitle:HelpTitle.ToString().GetNSString()] retain];
+		for (int32 Index = NumItems - 1; Index > 0; Index--)
+		{
+			[[NSApp mainMenu] removeItemAtIndex:Index];
+		}
+		GCachedMenuState.Reset();
 
-			const TArray<TSharedRef<const FMultiBlock> >& MenuBlocks = MultiBox->GetBlocks();
+		const TArray<TSharedRef<const FMultiBlock> >& MenuBlocks = MultiBox->GetBlocks();
 
-			for (int32 Index = 0; Index < MenuBlocks.Num(); Index++)
+		for (int32 Index = 0; Index < MenuBlocks.Num(); Index++)
+		{
+			TSharedRef<const FMenuEntryBlock> Block = StaticCastSharedRef<const FMenuEntryBlock>(MenuBlocks[Index]);
+			FMacMenu* Menu = [[FMacMenu alloc] initWithMenuEntryBlock:Block];
+			NSString* Title = FSlateMacMenu::GetMenuItemTitle(Block);
+			[Menu setTitle:Title];
+
+			NSMenuItem* MenuItem = [[NSMenuItem new] autorelease];
+			[MenuItem setTitle:Title];
+			[[NSApp mainMenu] addItem:MenuItem];
+			[MenuItem setSubmenu:Menu];
+
+			const bool bIsWindowMenu = (WindowLabel.ToString().Compare(FString(Title)) == 0);
+			if (bIsWindowMenu)
 			{
-				TSharedRef<const FMenuEntryBlock> Block = StaticCastSharedRef<const FMenuEntryBlock>(MenuBlocks[Index]);
-				FMacMenu* Menu = [[FMacMenu alloc] initWithMenuEntryBlock:Block];
-				NSString* Title = FSlateMacMenu::GetMenuItemTitle(Block);
-				[Menu setTitle:Title];
-
-				NSMenuItem* MenuItem = [[NSMenuItem new] autorelease];
-				[MenuItem setTitle:Title];
-				[[NSApp mainMenu] addItem:MenuItem];
-				[MenuItem setSubmenu:Menu];
-
-				const bool bIsWindowMenu = (WindowLabel.ToString().Compare(FString(Title)) == 0);
-				if (bIsWindowMenu)
-				{
-					[NSApp setWindowsMenu:nil];
+				[NSApp setWindowsMenu:nil];
 
-					[Menu removeAllItems];
+				[Menu removeAllItems];
 
-					NSMenuItem* MinimizeItem = [[[NSMenuItem alloc] initWithTitle:@"Minimize" action:@selector(miniaturize:) keyEquivalent:@"m"] autorelease];
-					NSMenuItem* ZoomItem = [[[NSMenuItem alloc] initWithTitle:@"Zoom" action:@selector(performZoom:) keyEquivalent:@""] autorelease];
-					NSMenuItem* CloseItem = [[[NSMenuItem alloc] initWithTitle:@"Close" action:@selector(performClose:) keyEquivalent:@"w"] autorelease];
-					NSMenuItem* BringAllToFrontItem = [[[NSMenuItem alloc] initWithTitle:@"Bring All to Front" action:@selector(arrangeInFront:) keyEquivalent:@""] autorelease];
+				NSMenuItem* MinimizeItem = [[[NSMenuItem alloc] initWithTitle:@"Minimize" action:@selector(miniaturize:) keyEquivalent:@"m"] autorelease];
+				NSMenuItem* ZoomItem = [[[NSMenuItem alloc] initWithTitle:@"Zoom" action:@selector(performZoom:) keyEquivalent:@""] autorelease];
+				NSMenuItem* CloseItem = [[[NSMenuItem alloc] initWithTitle:@"Close" action:@selector(performClose:) keyEquivalent:@"w"] autorelease];
+				NSMenuItem* BringAllToFrontItem = [[[NSMenuItem alloc] initWithTitle:@"Bring All to Front" action:@selector(arrangeInFront:) keyEquivalent:@""] autorelease];
 
-					[Menu addItem:MinimizeItem];
-					[Menu addItem:ZoomItem];
-					[Menu addItem:CloseItem];
-					[Menu addItem:[NSMenuItem separatorItem]];
-					[Menu addItem:BringAllToFrontItem];
-					[Menu addItem:[NSMenuItem separatorItem]];
+				[Menu addItem:MinimizeItem];
+				[Menu addItem:ZoomItem];
+				[Menu addItem:CloseItem];
+				[Menu addItem:[NSMenuItem separatorItem]];
+				[Menu addItem:BringAllToFrontItem];
+				[Menu addItem:[NSMenuItem separatorItem]];
 
-					[NSApp setWindowsMenu:Menu];
-					[Menu addItem:[NSMenuItem separatorItem]];
-				}
-			}
-
-			if ([[NSApp mainMenu] itemWithTitle:HelpTitle.ToString().GetNSString()] == nil && HelpMenu)
-			{
-				[[NSApp mainMenu] addItem:HelpMenu];
-				GCachedMenuState.Add((FMacMenu*)HelpMenu, TSharedPtr<TArray<FMacMenuItemState>>(new TArray<FMacMenuItemState>()));
+				[NSApp setWindowsMenu:Menu];
+				[Menu addItem:[NSMenuItem separatorItem]];
 			}
+		}
 
-			if (HelpMenu)
-			{
-				[HelpMenu release];
-			}
+		if ([[NSApp mainMenu] itemWithTitle:HelpTitle.ToString().GetNSString()] == nil && HelpMenu)
+		{
+			[[NSApp mainMenu] addItem:HelpMenu];
+			GCachedMenuState.Add((FMacMenu*)HelpMenu, TSharedPtr<TArray<FMacMenuItemState>>(new TArray<FMacMenuItemState>()));
+		}
 
-			GCurrentMultiBox = MultiBox;
-			FPlatformMisc::bChachedMacMenuStateNeedsUpdate = true;
+		if (HelpMenu)
+		{
+			[HelpMenu release];
 		}
+
+		FPlatformMisc::bChachedMacMenuStateNeedsUpdate = true;
 	});
 }
 
