- Fixed compile error with LOCTEXT_NAMESPACE not being undef'd

diff --git a/Engine/Source/Runtime/Slate/Private/Widgets/SWebBrowser.cpp b/Engine/Source/Runtime/Slate/Private/Widgets/SWebBrowser.cpp
--- a/Engine/Source/Runtime/Slate/Private/Widgets/SWebBrowser.cpp
+++ b/Engine/Source/Runtime/Slate/Private/Widgets/SWebBrowser.cpp
@@ -165,3 +165,5 @@ FReply SWebBrowser::OnReloadClicked()
 	}
 	return FReply::Handled();
 }
+
+#undef LOCTEXT_NAMESPACE
