Fixed some clipping issues with text shadows

FSlateTextRun and FSlateHyperlinkRun now support negative shadow offsets.

This also fixes the maximum size of a text run being reported incorrectly when it was using a shadow offset.

diff --git a/Engine/Source/Runtime/Slate/Private/Framework/Text/SlateHyperlinkRun.cpp b/Engine/Source/Runtime/Slate/Private/Framework/Text/SlateHyperlinkRun.cpp
--- a/Engine/Source/Runtime/Slate/Private/Framework/Text/SlateHyperlinkRun.cpp
+++ b/Engine/Source/Runtime/Slate/Private/Framework/Text/SlateHyperlinkRun.cpp
@@ -29,24 +29,26 @@ void FSlateHyperlinkRun::SetTextRange( const FTextRange& Value )
 int16 FSlateHyperlinkRun::GetBaseLine( float Scale ) const 
 {
 	const TSharedRef< FSlateFontMeasure > FontMeasure = FSlateApplication::Get().GetRenderer()->GetFontMeasureService();
-	return FontMeasure->GetBaseline( Style.TextStyle.Font, Scale );
+	return FontMeasure->GetBaseline( Style.TextStyle.Font, Scale ) - FMath::Max(0.0f, Style.TextStyle.ShadowOffset.Y);
 }
 
 int16 FSlateHyperlinkRun::GetMaxHeight( float Scale ) const 
 {
 	const TSharedRef< FSlateFontMeasure > FontMeasure = FSlateApplication::Get().GetRenderer()->GetFontMeasureService();
-	return FontMeasure->GetMaxCharacterHeight( Style.TextStyle.Font, Scale );
+	return FontMeasure->GetMaxCharacterHeight( Style.TextStyle.Font, Scale ) + FMath::Abs(Style.TextStyle.ShadowOffset.Y);
 }
 
 FVector2D FSlateHyperlinkRun::Measure( int32 StartIndex, int32 EndIndex, float Scale ) const 
 {
+	const FVector2D ShadowOffsetToApply((EndIndex == Range.EndIndex) ? FMath::Abs(Style.TextStyle.ShadowOffset.X) : 0.0f, FMath::Abs(Style.TextStyle.ShadowOffset.Y));
+
 	if ( EndIndex - StartIndex == 0 )
 	{
-		return FVector2D( 0, GetMaxHeight( Scale ) );
+		return FVector2D( ShadowOffsetToApply.X, GetMaxHeight( Scale ) );
 	}
 
 	const TSharedRef< FSlateFontMeasure > FontMeasure = FSlateApplication::Get().GetRenderer()->GetFontMeasureService();
-	return FontMeasure->Measure( **Text, StartIndex, EndIndex, Style.TextStyle.Font, true, Scale );
+	return FontMeasure->Measure( **Text, StartIndex, EndIndex, Style.TextStyle.Font, true, Scale ) + ShadowOffsetToApply;
 }
 
 int8 FSlateHyperlinkRun::GetKerning( int32 CurrentIndex, float Scale ) const 
diff --git a/Engine/Source/Runtime/Slate/Private/Framework/Text/SlateTextRun.cpp b/Engine/Source/Runtime/Slate/Private/Framework/Text/SlateTextRun.cpp
--- a/Engine/Source/Runtime/Slate/Private/Framework/Text/SlateTextRun.cpp
+++ b/Engine/Source/Runtime/Slate/Private/Framework/Text/SlateTextRun.cpp
@@ -29,24 +29,26 @@ void FSlateTextRun::SetTextRange( const FTextRange& Value )
 int16 FSlateTextRun::GetBaseLine( float Scale ) const 
 {
 	const TSharedRef< FSlateFontMeasure > FontMeasure = FSlateApplication::Get().GetRenderer()->GetFontMeasureService();
-	return FontMeasure->GetBaseline( Style.Font, Scale );
+	return FontMeasure->GetBaseline( Style.Font, Scale ) - FMath::Max(0.0f, Style.ShadowOffset.Y);
 }
 
 int16 FSlateTextRun::GetMaxHeight( float Scale ) const 
 {
 	const TSharedRef< FSlateFontMeasure > FontMeasure = FSlateApplication::Get().GetRenderer()->GetFontMeasureService();
-	return FontMeasure->GetMaxCharacterHeight( Style.Font, Scale );
+	return FontMeasure->GetMaxCharacterHeight( Style.Font, Scale ) + FMath::Abs(Style.ShadowOffset.Y);
 }
 
 FVector2D FSlateTextRun::Measure( int32 BeginIndex, int32 EndIndex, float Scale ) const 
 {
+	const FVector2D ShadowOffsetToApply((EndIndex == Range.EndIndex) ? FMath::Abs(Style.ShadowOffset.X) : 0.0f, FMath::Abs(Style.ShadowOffset.Y));
+
 	if ( EndIndex - BeginIndex == 0 )
 	{
-		return FVector2D( 0, GetMaxHeight( Scale ) ) + Style.ShadowOffset;
+		return FVector2D( ShadowOffsetToApply.X, GetMaxHeight( Scale ) );
 	}
 
 	const TSharedRef< FSlateFontMeasure > FontMeasure = FSlateApplication::Get().GetRenderer()->GetFontMeasureService();
-	return FontMeasure->Measure( **Text, BeginIndex, EndIndex, Style.Font, true, Scale ) + Style.ShadowOffset;
+	return FontMeasure->Measure( **Text, BeginIndex, EndIndex, Style.Font, true, Scale ) + ShadowOffsetToApply;
 }
 
 int8 FSlateTextRun::GetKerning(int32 CurrentIndex, float Scale) const
@@ -78,13 +80,23 @@ int32 FSlateTextRun::OnPaint( const FPaintArgs& Args, const FTextLayout::FLineVi
 	// The block size and offset values are pre-scaled, so we need to account for that when converting the block offsets into paint geometry
 	const float InverseScale = Inverse(AllottedGeometry.Scale);
 
+	// A negative shadow offset should be applied as a positive offset to the text to avoid clipping issues
+	const FVector2D DrawShadowOffset(
+		(Style.ShadowOffset.X > 0.0f) ? Style.ShadowOffset.X : 0.0f, 
+		(Style.ShadowOffset.Y > 0.0f) ? Style.ShadowOffset.Y : 0.0f
+		);
+	const FVector2D DrawTextOffset(
+		(Style.ShadowOffset.X < 0.0f) ? -Style.ShadowOffset.X : 0.0f, 
+		(Style.ShadowOffset.Y < 0.0f) ? -Style.ShadowOffset.Y : 0.0f
+		);
+
 	// Draw the optional shadow
 	if ( ShouldDropShadow )
 	{
 		FSlateDrawElement::MakeText(
 			OutDrawElements,
 			++LayerId,
-			AllottedGeometry.ToPaintGeometry(TransformVector(InverseScale, Block->GetSize()), FSlateLayoutTransform(TransformPoint(InverseScale, Block->GetLocationOffset() + Style.ShadowOffset))),
+			AllottedGeometry.ToPaintGeometry(TransformVector(InverseScale, Block->GetSize()), FSlateLayoutTransform(TransformPoint(InverseScale, Block->GetLocationOffset() + DrawShadowOffset))),
 			Text.Get(),
 			BlockRange.BeginIndex,
 			BlockRange.EndIndex,
@@ -99,7 +111,7 @@ int32 FSlateTextRun::OnPaint( const FPaintArgs& Args, const FTextLayout::FLineVi
 	FSlateDrawElement::MakeText(
 		OutDrawElements,
 		++LayerId,
-		AllottedGeometry.ToPaintGeometry(TransformVector(InverseScale, Block->GetSize()), FSlateLayoutTransform(TransformPoint(InverseScale, Block->GetLocationOffset()))),
+		AllottedGeometry.ToPaintGeometry(TransformVector(InverseScale, Block->GetSize()), FSlateLayoutTransform(TransformPoint(InverseScale, Block->GetLocationOffset() + DrawTextOffset))),
 		Text.Get(),
 		BlockRange.BeginIndex,
 		BlockRange.EndIndex,
