Force update Mac menu bar in FTabManager::OnTabRelocated()

diff --git a/Engine/Source/Runtime/Slate/Private/Framework/Docking/TabManager.cpp b/Engine/Source/Runtime/Slate/Private/Framework/Docking/TabManager.cpp
--- a/Engine/Source/Runtime/Slate/Private/Framework/Docking/TabManager.cpp
+++ b/Engine/Source/Runtime/Slate/Private/Framework/Docking/TabManager.cpp
@@ -1334,6 +1334,8 @@ void FTabManager::OnTabRelocated( const TSharedRef<SDockTab>& RelocatedTab, cons
 		DockAreas[DockAreaIndex].Pin()->OnTabFoundNewHome( RelocatedTab, NewOwnerWindow.ToSharedRef() );
 	}
 
+	FGlobalTabmanager::Get()->UpdateMainMenu(RelocatedTab, true);
+
 	UpdateStats();
 }
 
