Making sure that Web Browser does not affect monolithic game builds

Ensured that the Developer level Web Browser module is only referenced in editor and program builds. This means that you can create an SWebBrowser wherever you like but it won't have any functionality in game.

diff --git a/Engine/Source/Runtime/Slate/Private/Framework/Application/SlateApplication.cpp b/Engine/Source/Runtime/Slate/Private/Framework/Application/SlateApplication.cpp
--- a/Engine/Source/Runtime/Slate/Private/Framework/Application/SlateApplication.cpp
+++ b/Engine/Source/Runtime/Slate/Private/Framework/Application/SlateApplication.cpp
@@ -5,8 +5,10 @@
 #include "SWindowTitleBar.h"
 #include "HittestGrid.h"
 
+#if WITH_EDITOR || IS_PROGRAM
 #include "WebBrowserModule.h"
 #include "IWebBrowserSingleton.h"
+#endif
 
 
 DECLARE_CYCLE_STAT( TEXT("Message Tick Time"), STAT_SlateMessageTick, STATGROUP_Slate );
@@ -449,8 +451,10 @@ void FSlateApplication::InitializeRenderer( TSharedRef<FSlateRenderer> InRendere
 	Renderer = InRenderer;
 	Renderer->Initialize();
 
+#if WITH_EDITOR || IS_PROGRAM
 	// Pass renderer on to the Web Browser module
 	IWebBrowserModule::Get().GetSingleton()->SetSlateRenderer(InRenderer);
+#endif
 }
 
 void FSlateApplication::InitializeSound( const TSharedRef<ISlateSoundDevice>& InSlateSoundDevice )
@@ -928,8 +932,10 @@ void FSlateApplication::Tick()
 	// (so they know their size)
 	FSlateNotificationManager::Get().Tick();
 
+#if WITH_EDITOR || IS_PROGRAM
 	// Calling this after Tick gives web browser a chance to render pages to texture after size changes
 	IWebBrowserModule::Get().GetSingleton()->PumpMessages();
+#endif
 
 	// Draw all windows
 	DrawWindows();
diff --git a/Engine/Source/Runtime/Slate/Private/Widgets/SWebBrowser.cpp b/Engine/Source/Runtime/Slate/Private/Widgets/SWebBrowser.cpp
--- a/Engine/Source/Runtime/Slate/Private/Widgets/SWebBrowser.cpp
+++ b/Engine/Source/Runtime/Slate/Private/Widgets/SWebBrowser.cpp
@@ -3,10 +3,12 @@
 #include "SlatePrivatePCH.h"
 #include "SWebBrowser.h"
 
+#if WITH_EDITOR || IS_PROGRAM
 #include "WebBrowserModule.h"
 #include "IWebBrowserSingleton.h"
 #include "IWebBrowserWindow.h"
 #include "WebBrowserViewport.h"
+#endif
 
 #define LOCTEXT_NAMESPACE "WebBrowser"
 
@@ -22,11 +24,14 @@ void SWebBrowser::Construct(const FArguments& InArgs)
 		TSharedPtr<FGenericWindow> NativeWindow = InArgs._ParentWindow->GetNativeWindow();
 		OSWindowHandle = NativeWindow->GetOSWindowHandle();
 	}
+
+#if WITH_EDITOR || IS_PROGRAM
 	BrowserWindow = IWebBrowserModule::Get().GetSingleton()->CreateBrowserWindow(OSWindowHandle,
-																				InArgs._InitialURL,
-																				InArgs._ViewportSize.Get().X,
-																				InArgs._ViewportSize.Get().Y,
-																				InArgs._SupportsTransparency);
+	                                                                             InArgs._InitialURL,
+	                                                                             InArgs._ViewportSize.Get().X,
+	                                                                             InArgs._ViewportSize.Get().Y,
+	                                                                             InArgs._SupportsTransparency);
+#endif
 
 	TSharedPtr<SViewport> ViewportWidget;
 
@@ -83,55 +88,67 @@ void SWebBrowser::Construct(const FArguments& InArgs)
 		]
 	];
 
+#if WITH_EDITOR || IS_PROGRAM
 	if (BrowserWindow.IsValid())
 	{
 		BrowserViewport = MakeShareable(new FWebBrowserViewport(BrowserWindow, ViewportWidget));
 		ViewportWidget->SetViewportInterface(BrowserViewport.ToSharedRef());
 	}
+#endif
 }
 
 FText SWebBrowser::GetTitleText() const
 {
+#if WITH_EDITOR || IS_PROGRAM
 	if (BrowserWindow.IsValid())
 	{
 		return FText::FromString(BrowserWindow->GetTitle());
 	}
+#endif
 	return LOCTEXT("InvalidWindow", "Browser Window is not valid/supported");
 }
 
 bool SWebBrowser::CanGoBack() const
 {
+#if WITH_EDITOR || IS_PROGRAM
 	if (BrowserWindow.IsValid())
 	{
 		return BrowserWindow->CanGoBack();
 	}
+#endif
 	return false;
 }
 
 FReply SWebBrowser::OnBackClicked()
 {
+#if WITH_EDITOR || IS_PROGRAM
 	if (BrowserWindow.IsValid())
 	{
 		BrowserWindow->GoBack();
 	}
+#endif
 	return FReply::Handled();
 }
 
 bool SWebBrowser::CanGoForward() const
 {
+#if WITH_EDITOR || IS_PROGRAM
 	if (BrowserWindow.IsValid())
 	{
 		return BrowserWindow->CanGoForward();
 	}
+#endif
 	return false;
 }
 
 FReply SWebBrowser::OnForwardClicked()
 {
+#if WITH_EDITOR || IS_PROGRAM
 	if (BrowserWindow.IsValid())
 	{
 		BrowserWindow->GoForward();
 	}
+#endif
 	return FReply::Handled();
 }
 
@@ -140,6 +157,7 @@ FText SWebBrowser::GetReloadButtonText() const
 	static FText Reload = LOCTEXT("Reload", "Reload");
 	static FText Stop = LOCTEXT("Stop", "Stop");
 
+#if WITH_EDITOR || IS_PROGRAM
 	if (BrowserWindow.IsValid())
 	{
 		if (BrowserWindow->IsLoading())
@@ -147,11 +165,13 @@ FText SWebBrowser::GetReloadButtonText() const
 			return Stop;
 		}
 	}
+#endif
 	return Reload;
 }
 
 FReply SWebBrowser::OnReloadClicked()
 {
+#if WITH_EDITOR || IS_PROGRAM
 	if (BrowserWindow.IsValid())
 	{
 		if (BrowserWindow->IsLoading())
@@ -163,6 +183,7 @@ FReply SWebBrowser::OnReloadClicked()
 			BrowserWindow->Reload();
 		}
 	}
+#endif
 	return FReply::Handled();
 }
 
