Fix mac menu specific crash not clearing focus before menu action is executed, leaving widgets in dangerous states.

diff --git a/Engine/Source/Runtime/Slate/Private/Framework/MultiBox/Mac/MacMenu.cpp b/Engine/Source/Runtime/Slate/Private/Framework/MultiBox/Mac/MacMenu.cpp
--- a/Engine/Source/Runtime/Slate/Private/Framework/MultiBox/Mac/MacMenu.cpp
+++ b/Engine/Source/Runtime/Slate/Private/Framework/MultiBox/Mac/MacMenu.cpp
@@ -72,6 +72,11 @@ static TMap<FMacMenu*, TSharedPtr<TArray<FMacMenuItemState>>> GCachedMenuState;
 - (void)menuWillOpen:(NSMenu*)Menu
 {
 	FPlatformMisc::bChachedMacMenuStateNeedsUpdate = true;
+	
+	GameThreadCall(^{
+		FSlateApplication::Get().ClearKeyboardFocus( EKeyboardFocusCause::WindowActivate );
+	}, @[ NSDefaultRunLoopMode ], false);
+
 }
 
 @end
