Slate - Now initializing the scale box parameters in the construct function.

diff --git a/Engine/Source/Runtime/Slate/Private/Widgets/Layout/SScaleBox.cpp b/Engine/Source/Runtime/Slate/Private/Widgets/Layout/SScaleBox.cpp
--- a/Engine/Source/Runtime/Slate/Private/Widgets/Layout/SScaleBox.cpp
+++ b/Engine/Source/Runtime/Slate/Private/Widgets/Layout/SScaleBox.cpp
@@ -11,7 +11,12 @@
 
 void SScaleBox::Construct( const SScaleBox::FArguments& InArgs )
 {
+	Stretch = InArgs._Stretch;
+	StretchDirection = InArgs._StretchDirection;
+
 	ChildSlot
+	.HAlign(InArgs._HAlign)
+	.VAlign(InArgs._VAlign)
 	[
 		InArgs._Content.Widget
 	];
