Slate synthetic mouse moves no longer get routed in cases when high-precision mouse moves are occurring.

TTP 349481 SynthesizeMouseInput causing extra mouse input

diff --git a/Engine/Source/Runtime/Slate/Private/Framework/Application/SlateApplication.cpp b/Engine/Source/Runtime/Slate/Private/Framework/Application/SlateApplication.cpp
--- a/Engine/Source/Runtime/Slate/Private/Framework/Application/SlateApplication.cpp
+++ b/Engine/Source/Runtime/Slate/Private/Framework/Application/SlateApplication.cpp
@@ -3846,23 +3846,26 @@ bool FSlateApplication::ProcessMouseMoveEvent( FPointerEvent& MouseEvent, bool b
 
 	if (MouseCaptorPath.IsValid())
 	{
-		FArrangedWidget& MouseCaptorWidget = MouseCaptorPath.Widgets.Last();
-		MouseEvent.SetEventPath(MouseCaptorPath);
+		if ( !bIsSynthetic )
+		{
+			FArrangedWidget& MouseCaptorWidget = MouseCaptorPath.Widgets.Last();
+			MouseEvent.SetEventPath(MouseCaptorPath);
 
-		// Switch worlds widgets in the current path
-		FScopedSwitchWorldHack SwitchWorld( MouseCaptorPath );
+			// Switch worlds widgets in the current path
+			FScopedSwitchWorldHack SwitchWorld(MouseCaptorPath);
 
-		FReply Reply = FReply::Unhandled();
-		if (MouseEvent.IsTouchEvent())
-		{
-			Reply = MouseCaptorWidget.Widget->OnTouchMoved( MouseCaptorWidget.Geometry, MouseEvent ).SetHandler( MouseCaptorWidget.Widget );
-		}
-		if (!MouseEvent.IsTouchEvent() || (!Reply.IsEventHandled() && bTouchFallbackToMouse))
-		{
-			Reply = MouseCaptorWidget.Widget->OnMouseMove( MouseCaptorWidget.Geometry, MouseEvent ).SetHandler( MouseCaptorWidget.Widget );
+			FReply Reply = FReply::Unhandled();
+			if (MouseEvent.IsTouchEvent())
+			{
+				Reply = MouseCaptorWidget.Widget->OnTouchMoved(MouseCaptorWidget.Geometry, MouseEvent).SetHandler(MouseCaptorWidget.Widget);
+			}
+			if (!MouseEvent.IsTouchEvent() || (!Reply.IsEventHandled() && bTouchFallbackToMouse))
+			{
+				Reply = MouseCaptorWidget.Widget->OnMouseMove(MouseCaptorWidget.Geometry, MouseEvent).SetHandler(MouseCaptorWidget.Widget);
+			}
+			ProcessReply(MouseCaptorPath, Reply, &MouseCaptorPath, &MouseEvent);
+			bHandled = Reply.IsEventHandled();
 		}
-		ProcessReply(MouseCaptorPath, Reply, &MouseCaptorPath, &MouseEvent);
-		bHandled = Reply.IsEventHandled();
 	}
 	else
 	{	
