Removed a bunch of stuff from Slate standard include, created SlateBasics.h

diff --git a/Engine/Source/Runtime/Slate/Public/Framework/Notifications/NotificationManager.h b/Engine/Source/Runtime/Slate/Public/Framework/Notifications/NotificationManager.h
--- a/Engine/Source/Runtime/Slate/Public/Framework/Notifications/NotificationManager.h
+++ b/Engine/Source/Runtime/Slate/Public/Framework/Notifications/NotificationManager.h
@@ -4,6 +4,9 @@
 
 #include "LockFreeList.h"
 
+class SNotificationItem;
+struct FNotificationInfo;
+
 /**
  * A class which manages a group of notification windows                 
  */
diff --git a/Engine/Source/Runtime/Slate/Public/Framework/SlateDelegates.h b/Engine/Source/Runtime/Slate/Public/Framework/SlateDelegates.h
--- a/Engine/Source/Runtime/Slate/Public/Framework/SlateDelegates.h
+++ b/Engine/Source/Runtime/Slate/Public/Framework/SlateDelegates.h
@@ -207,3 +207,7 @@ public:
 
 /** Notification for when a keyboard event occurs */
 DECLARE_DELEGATE_RetVal_OneParam( FReply, FOnKeyboardEvent, const FKeyboardEvent& )
+
+/** A Delegate for passing along a string of a source code location to access */
+DECLARE_DELEGATE_RetVal_ThreeParams(bool, FAccessSourceCode, const FString& /*FileName*/, int32 /*InLineNumber*/, int32 /*InColumnNumber*/);
+
diff --git a/Engine/Source/Runtime/Slate/Public/Slate.h b/Engine/Source/Runtime/Slate/Public/Slate.h
--- a/Engine/Source/Runtime/Slate/Public/Slate.h
+++ b/Engine/Source/Runtime/Slate/Public/Slate.h
@@ -2,170 +2,6 @@
 
 #pragma once
 
-#include "Core.h"
-#include "CoreUObject.h"
-#include "Json.h"
-#include "SlateCore.h"
-#include "SlateClasses.h"
-
-
-// Enables or disables more slate stats.  These stats are usually per widget and can be gathered hundreds or thousands of times per frame
-// enabling these constantly causes a bit of a performance drop due to stat gathering overhead so these should only be enabled for profiling
-#define SLATE_HD_STATS 0
-
-#define BEGIN_SLATE_FUNCTION_BUILD_OPTIMIZATION PRAGMA_DISABLE_OPTIMIZATION
-#define END_SLATE_FUNCTION_BUILD_OPTIMIZATION   PRAGMA_ENABLE_OPTIMIZATION
-
-/** Stat tracking for Slate */
-DECLARE_CYCLE_STAT_EXTERN( TEXT("Total Slate Tick Time"), STAT_SlateTickTime, STATGROUP_Slate , );
-DECLARE_CYCLE_STAT_EXTERN( TEXT("SlatePrepass"), STAT_SlateCacheDesiredSize, STATGROUP_Slate , );
-DECLARE_CYCLE_STAT_EXTERN( TEXT("STextBlock"), STAT_SlateCacheDesiredSize_STextBlock, STATGROUP_Slate , );
-DECLARE_CYCLE_STAT_EXTERN( TEXT("TickWidgets"), STAT_SlateTickWidgets, STATGROUP_Slate , );
-DECLARE_CYCLE_STAT_EXTERN( TEXT("DrawWindows"), STAT_SlateDrawWindowTime, STATGROUP_Slate , );
-DECLARE_CYCLE_STAT_EXTERN( TEXT("OnPaint SViewport"), STAT_SlateOnPaint_SViewport, STATGROUP_Slate , );
-DECLARE_CYCLE_STAT_EXTERN( TEXT("OnPaint SImage"), STAT_SlateOnPaint_SImage, STATGROUP_Slate , );
-DECLARE_CYCLE_STAT_EXTERN( TEXT("OnPaint SBox"), STAT_SlateOnPaint_SBox, STATGROUP_Slate , );
-DECLARE_CYCLE_STAT_EXTERN( TEXT("OnPaint SCompoundWidget"), STAT_SlateOnPaint_SCompoundWidget, STATGROUP_Slate , );
-DECLARE_CYCLE_STAT_EXTERN( TEXT("OnPaint SBorder"), STAT_SlateOnPaint_SBorder, STATGROUP_Slate , );
-DECLARE_CYCLE_STAT_EXTERN( TEXT("OnPaint SPanel"), STAT_SlateOnPaint_SPanel, STATGROUP_Slate , );
-DECLARE_CYCLE_STAT_EXTERN( TEXT("OnPaint SOverlay"), STAT_SlateOnPaint_SOverlay, STATGROUP_Slate , );
-DECLARE_CYCLE_STAT_EXTERN( TEXT("OnPaint SEditableText"), STAT_SlateOnPaint_SEditableText, STATGROUP_Slate , );
-DECLARE_CYCLE_STAT_EXTERN( TEXT("OnPaint STextBlock"), STAT_SlateOnPaint_STextBlock, STATGROUP_Slate , );
-DECLARE_CYCLE_STAT_EXTERN( TEXT("OnPaint SRichTextBlock"), STAT_SlateOnPaint_SRichTextBlock, STATGROUP_Slate , );
-DECLARE_CYCLE_STAT_EXTERN( TEXT("OnPaint SSlider"), STAT_SlateOnPaint_SSlider, STATGROUP_Slate , );
-DECLARE_CYCLE_STAT_EXTERN( TEXT("OnPaint SGraphPanel"), STAT_SlateOnPaint_SGraphPanel, STATGROUP_Slate , SLATE_API);
-DECLARE_CYCLE_STAT_EXTERN( TEXT("Slate Misc Time"), STAT_SlateMiscTime, STATGROUP_Slate , );
-
-
-#if PLATFORM_IOS
-	#include "IOS/IOSPlatformTextField.h"
-#elif PLATFORM_ANDROID
-	#include "Android/AndroidPlatformTextField.h"
-#else
-	#include "GenericPlatformTextField.h"
-#endif
-
-#include "InputCore.h"
-
-#include "SlateDelegates.h"
-
-// Application
-#include "MenuStack.h"
-#include "IPlatformTextField.h"
-#include "IWidgetReflector.h"
-#include "SlateApplication.h"
-#include "SlateIcon.h"
-
-// Commands
-#include "InputGesture.h"
-#include "UIAction.h"
-#include "UICommandInfo.h"
-#include "InputBindingManager.h"
-#include "Commands.h"
-#include "UICommandList.h"
-
-// Legacy
-#include "SMissingWidget.h"
-#include "SlateScrollHelper.h"
-#include "SlateStyleRegistry.h"
-#include "UICommandDragDropOp.h"
-#include "InertialScrollManager.h"
-#include "GenericCommands.h"
-#include "SWeakWidget.h"
-#include "TextLayoutEngine.h"
-#include "SPanel.h"
-#include "SCompoundWidget.h"
-#include "SFxWidget.h"
-#include "SBorder.h"
-#include "SSeparator.h"
-#include "SSpacer.h"
-#include "SWrapBox.h"
-#include "SImage.h"
-#include "SSpinningImage.h"
-#include "SProgressBar.h"
-#include "SCanvas.h"
-#include "STextBlock.h"
-#include "TextDecorators.h"
-#include "SRichTextBlock.h"
-#include "SBox.h"
-#include "SHeader.h"
-#include "SGridPanel.h"
-#include "SUniformGridPanel.h"
-#include "SMenuAnchor.h"
-#include "MultiBoxDefs.h"
-////
-#include "MultiBox.h"
-////
-#include "MultiBoxBuilder.h"
-#include "MultiBoxExtender.h"
-#include "SMultiLineEditableText.h"
-#include "SMultiLineEditableTextBox.h"
-#include "SEditableText.h"
-#include "SEditableTextBox.h"
-#include "SSearchBox.h"
-#include "SButton.h"
-#include "SToolTip.h"
-#include "SScrollBar.h"
-#include "SScrollBorder.h"
-#include "SErrorText.h"
-#include "SErrorHint.h"
-#include "SPopUpErrorText.h"
-#include "SSplitter.h"
-#include "TableViewTypeTraits.h"
-#include "SExpanderArrow.h"
-#include "ITypedTableView.h"
-#include "STableViewBase.h"
-#include "SHeaderRow.h"
-#include "STableRow.h"
-#include "SListView.h"
-#include "STileView.h"
-#include "STreeView.h"
-#include "SScrollBox.h"
-#include "SViewport.h"
-#include "SColorBlock.h"
-#include "SCheckBox.h"
-#include "SSpinBox.h"
-#include "SSlider.h"
-#include "SVolumeControl.h"
-#include "SColorSpectrum.h"
-#include "SColorWheel.h"
-///
-#include "SToolBarButtonBlock.h"
-#include "SToolBarComboButtonBlock.h"
-///
-#include "SComboButton.h"
-#include "SComboBox.h"
-#include "SHyperlink.h"
-#include "SRichTextHyperlink.h"
-#include "SThrobber.h"
-#include "STextEntryPopup.h"
-#include "STextComboPopup.h"
-#include "SExpandableButton.h"
-#include "SExpandableArea.h"
-#include "SNotificationList.h"
-#include "SWidgetSwitcher.h"
-#include "SSuggestionTextBox.h"
-#include "SBreadcrumbTrail.h"
-#include "STextComboBox.h"
-#include "SNumericEntryBox.h"
-#include "SEditableComboBox.h"
-#include "NotificationManager.h"
-#include "SDPIScaler.h"
-#include "SInlineEditableTextBlock.h"
-#include "SVirtualKeyboardEntry.h"
-#include "ScrollyZoomy.h"
-#include "SSafeZone.h"
-#include "MarqueeRect.h"
-#include "SRotatorInputBox.h"
-#include "SVectorInputBox.h"
-#include "SVirtualJoystick.h"
-
-// Docking Framework
-#include "WorkspaceItem.h"
-#include "TabManager.h"
-#include "SDockTab.h"
-#include "LayoutService.h"
-
-// Old docking
-#include "SDockableTab.h"
-#include "SDockTabStack.h"
+EMIT_DEPRECATED_WARNING_MESSAGE("including Slate.h is deprecated. Please include SlateBasics.h and then individual widget headers instead.");
+#include "SlateBasics.h"
+#include "SlateExtras.h"
\ No newline at end of file
diff --git a/Engine/Source/Runtime/Slate/Public/SlateBasics.h b/Engine/Source/Runtime/Slate/Public/SlateBasics.h
new file mode 100644
--- /dev/null
+++ b/Engine/Source/Runtime/Slate/Public/SlateBasics.h
@@ -0,0 +1,152 @@
+// Copyright 1998-2014 Epic Games, Inc. All Rights Reserved.
+
+#pragma once
+
+#include "Core.h"
+#include "CoreUObject.h"
+#include "Json.h"
+#include "SlateCore.h"
+#include "SlateClasses.h"
+
+
+// Enables or disables more slate stats.  These stats are usually per widget and can be gathered hundreds or thousands of times per frame
+// enabling these constantly causes a bit of a performance drop due to stat gathering overhead so these should only be enabled for profiling
+#define SLATE_HD_STATS 0
+
+#define BEGIN_SLATE_FUNCTION_BUILD_OPTIMIZATION PRAGMA_DISABLE_OPTIMIZATION
+#define END_SLATE_FUNCTION_BUILD_OPTIMIZATION   PRAGMA_ENABLE_OPTIMIZATION
+
+#include "IPlatformTextField.h"
+#if PLATFORM_IOS
+	#include "IOS/IOSPlatformTextField.h"
+#elif PLATFORM_ANDROID
+	#include "Android/AndroidPlatformTextField.h"
+#else
+	#include "GenericPlatformTextField.h"
+#endif
+
+#include "InputCore.h"
+
+
+#include "SlateDelegates.h"
+#include "SlateFwd.h"
+
+// Application
+#include "SlateApplication.h"
+#include "SlateIcon.h"
+
+// Commands
+#include "InputGesture.h"
+#include "UIAction.h"
+#include "UICommandInfo.h"
+#include "InputBindingManager.h"
+#include "Commands.h"
+#include "UICommandList.h"
+
+// Legacy
+class SMissingWidget;
+class FSlateStyleRegistry;
+class FUICommandDragDropOp;
+class FInertialScrollManager;
+class FGenericCommands;
+#include "SWeakWidget.h"
+#include "TextLayoutEngine.h"
+#include "SPanel.h"
+#include "SCompoundWidget.h"
+#include "SFxWidget.h"
+#include "SBorder.h"
+#include "SSeparator.h"
+#include "SSpacer.h"
+#include "SWrapBox.h"
+#include "SImage.h"
+#include "SSpinningImage.h"
+#include "SProgressBar.h"
+#include "SCanvas.h"
+#include "STextBlock.h"
+#include "TextDecorators.h"
+#include "SRichTextBlock.h"
+#include "SBox.h"
+#include "SHeader.h"
+#include "SGridPanel.h"
+#include "SUniformGridPanel.h"
+#include "SMenuAnchor.h"
+#include "MultiBoxDefs.h"
+////
+#include "MultiBox.h"
+////
+#include "MultiBoxBuilder.h"
+#include "MultiBoxExtender.h"
+#include "SMultiLineEditableText.h"
+#include "SMultiLineEditableTextBox.h"
+#include "SEditableText.h"
+#include "SEditableTextBox.h"
+class SSearchBox;
+#include "SButton.h"
+#include "SToolTip.h"
+#include "SScrollBar.h"
+#include "SScrollBorder.h"
+#include "SErrorText.h"
+#include "SErrorHint.h"
+#include "SPopUpErrorText.h"
+#include "SSplitter.h"
+#include "TableViewTypeTraits.h"
+#include "SExpanderArrow.h"
+#include "ITypedTableView.h"
+#include "STableViewBase.h"
+#include "SHeaderRow.h"
+#include "STableRow.h"
+#include "SListView.h"
+#include "STileView.h"
+#include "STreeView.h"
+#include "SScrollBox.h"
+#include "SViewport.h"
+#include "SColorBlock.h"
+#include "SCheckBox.h"
+#include "SSpinBox.h"
+#include "SSlider.h"
+class SVolumeControl;
+class SColorSpectrum;
+class SColorWheel;
+
+class FToolBarButtonBlock;
+class SToolBarButtonBlock;
+class FToolBarComboButtonBlock;
+class SToolBarComboButtonBlock;
+///
+#include "SComboButton.h"
+#include "SComboBox.h"
+class SHyperlink;
+class SRichTextHyperlink;
+class SThrobber;
+class SCircularThrobber;
+class STextEntryPopup;
+class STextComboPopup;
+class SExpandableButton;
+class SExpandableArea;
+class SNotificationItem;
+class SNotificationList;
+class SWidgetSwitcher;
+class SSuggestionTextBox;
+template <typename ItemType> class SBreadcrumbTrail;
+class STextComboBox;
+template<typename NumericType> class SNumericEntryBox;
+template<typename OptionType> class SEditableComboBox;
+class FSlateNotificationManager;
+class SDPIScaler;
+class SInlineEditableTextBlock;
+class SVirtualKeyboardEntry;
+class SSafeZone;
+struct FMarqueeRect;
+class SRotatorInputBox;
+class SVectorInputBox;
+class SVirtualJoystick;
+class SSearchBox;
+
+// Docking Framework
+#include "WorkspaceItem.h"
+#include "TabManager.h"
+#include "LayoutService.h"
+
+enum ETabActivationCause : uint8;
+enum ETabRole : uint8;
+class SDockTab;
diff --git a/Engine/Source/Runtime/Slate/Public/SlateExtras.h b/Engine/Source/Runtime/Slate/Public/SlateExtras.h
new file mode 100644
--- /dev/null
+++ b/Engine/Source/Runtime/Slate/Public/SlateExtras.h
@@ -0,0 +1,59 @@
+// Copyright 1998-2014 Epic Games, Inc. All Rights Reserved.
+
+#pragma once
+
+/**
+ * This file exists for backward compatibility. Slate.h has been split from one monolithic header
+ * into SlateBasics.h and SlateExtras.h. Over time more files will be shifted from SlateBasic.h to
+ * SlateExtras.h so older code using Slate.h will not be affected.
+ * 
+ * However, no one should include SlateExtras.h. It exists purely to assist migration of legacy projects.
+ */
+
+#include "SlateBasics.h"
+
+// Legacy
+#include "SMissingWidget.h"
+#include "SlateScrollHelper.h"
+#include "SlateStyleRegistry.h"
+#include "UICommandDragDropOp.h"
+#include "InertialScrollManager.h"
+#include "GenericCommands.h"
+
+#include "SSearchBox.h"
+#include "SVolumeControl.h"
+#include "SColorSpectrum.h"
+#include "SColorWheel.h"
+///
+#include "SToolBarButtonBlock.h"
+#include "SToolBarComboButtonBlock.h"
+///
+#include "SHyperlink.h"
+#include "SRichTextHyperlink.h"
+#include "SThrobber.h"
+#include "STextEntryPopup.h"
+#include "STextComboPopup.h"
+#include "SExpandableButton.h"
+#include "SExpandableArea.h"
+#include "SNotificationList.h"
+#include "SWidgetSwitcher.h"
+#include "SSuggestionTextBox.h"
+#include "SBreadcrumbTrail.h"
+#include "STextComboBox.h"
+#include "SNumericEntryBox.h"
+#include "SEditableComboBox.h"
+#include "NotificationManager.h"
+#include "SDPIScaler.h"
+#include "SInlineEditableTextBlock.h"
+#include "SVirtualKeyboardEntry.h"
+#include "ScrollyZoomy.h"
+#include "SSafeZone.h"
+#include "MarqueeRect.h"
+#include "SRotatorInputBox.h"
+#include "SVectorInputBox.h"
+#include "SVirtualJoystick.h"
+
+// Docking Framework
+#include "SDockTab.h"
+#include "SDockableTab.h"
+#include "SDockTabStack.h"
diff --git a/Engine/Source/Runtime/Slate/Public/Widgets/Docking/SDockTab.h b/Engine/Source/Runtime/Slate/Public/Widgets/Docking/SDockTab.h
--- a/Engine/Source/Runtime/Slate/Public/Widgets/Docking/SDockTab.h
+++ b/Engine/Source/Runtime/Slate/Public/Widgets/Docking/SDockTab.h
@@ -7,27 +7,21 @@ class SDockingTabWell;
 class SDockingTabStack;
 
 /** How will this tab be used. */
-namespace ETabRole
-{
-	enum Type
-	{
-		MajorTab,
-		PanelTab,
-		NomadTab,
-		DocumentTab,
-		NumRoles
-	};
-}
-
-namespace ETabActivationCause
+enum ETabRole : uint8
 {
+	MajorTab,
+	PanelTab,
+	NomadTab,
+	DocumentTab,
+	NumRoles
+};
+
 	/** The cause of a tab activation */
-	enum Type
-	{
-		UserClickedOnTab,
-		SetDirectly
-	};
-}
+enum ETabActivationCause : uint8
+{
+	UserClickedOnTab,
+	SetDirectly
+};
 
 /**
  * A tab widget that also holds on to some content that should be shown when this tab is selected.
@@ -42,7 +36,7 @@ public:
 	DECLARE_DELEGATE_OneParam(FOnTabClosedCallback, TSharedRef<SDockTab>);
 	
 	/** Invoked when a tab is activated */
-	DECLARE_DELEGATE_TwoParams(FOnTabActivatedCallback, TSharedRef<SDockTab>, ETabActivationCause::Type);
+	DECLARE_DELEGATE_TwoParams(FOnTabActivatedCallback, TSharedRef<SDockTab>, ETabActivationCause);
 
 	/** Invoked when this tab should save some information about its content. */
 	DECLARE_DELEGATE(FOnPersistVisualState);
@@ -70,7 +64,7 @@ public:
 		SLATE_NAMED_SLOT( FArguments, TabWellContentLeft )
 		SLATE_NAMED_SLOT( FArguments, TabWellContentRight )
 		SLATE_ATTRIBUTE( FMargin, ContentPadding )
-		SLATE_ARGUMENT( ETabRole::Type, TabRole )
+		SLATE_ARGUMENT( ETabRole, TabRole )
 		SLATE_ATTRIBUTE( FText, Label )
 		SLATE_ATTRIBUTE( const FSlateBrush*, Icon )
 		SLATE_EVENT( FOnTabClosedCallback, OnTabClosed )
@@ -115,7 +109,7 @@ public:
 	bool IsForeground() const;
 
 	/** Is this an MajorTab? A tool panel tab? */
-	ETabRole::Type GetTabRole() const;
+	ETabRole GetTabRole() const;
 
 	/** Returns true if the tab role is NomadTab, but is being visualized as a major tab. */
 	bool IsNomadTabWithMajorTabStyle() const;
@@ -195,7 +189,7 @@ public:
 	 * Make this tab active in its tabwell 
 	 * @param	InActivationMethod	How this tab was activated.
 	 */
-	void ActivateInParent(ETabActivationCause::Type InActivationCause);
+	void ActivateInParent(ETabActivationCause InActivationCause);
 
 	/** Set the tab manager that is controlling this tab */
 	void SetTabManager( const TSharedPtr<FTabManager>& InTabManager );
@@ -295,7 +289,7 @@ protected:
 	FTabId LayoutIdentifier;
 
 	/** Is this an MajorTab? A tool panel tab? */
-	ETabRole::Type TabRole;
+	ETabRole TabRole;
 
 	/** The tab's parent tab well. Null if it is a floating tab. */
 	TWeakPtr<SDockingTabWell> ParentPtr;
diff --git a/Engine/Source/Runtime/Slate/Public/Widgets/Docking/SDockTabStack.h b/Engine/Source/Runtime/Slate/Public/Widgets/Docking/SDockTabStack.h
--- a/Engine/Source/Runtime/Slate/Public/Widgets/Docking/SDockTabStack.h
+++ b/Engine/Source/Runtime/Slate/Public/Widgets/Docking/SDockTabStack.h
@@ -2,6 +2,7 @@
 
 #pragma once
 
+class SDockableTab;
 
 /**
  * A node in the Docking/Tabbing hierarchy.
diff --git a/Engine/Source/Runtime/Slate/Public/Widgets/Docking/SDockableTab.h b/Engine/Source/Runtime/Slate/Public/Widgets/Docking/SDockableTab.h
--- a/Engine/Source/Runtime/Slate/Public/Widgets/Docking/SDockableTab.h
+++ b/Engine/Source/Runtime/Slate/Public/Widgets/Docking/SDockableTab.h
@@ -1,7 +1,7 @@
 // Copyright 1998-2014 Epic Games, Inc. All Rights Reserved.
 
 #pragma once
-
+#include "SDockTab.h"
 
 class SDockTabStack;
 
@@ -69,7 +69,7 @@ public:
 	 * Make this tab active in its tabwell 
 	 * @param	InActivationMethod	How this tab was activated.
 	 */
-	void ActivateInParent(ETabActivationCause::Type InActivationCause);
+	void ActivateInParent(ETabActivationCause InActivationCause);
 
 protected:
 
@@ -77,7 +77,7 @@ protected:
 	FName LayoutIdentifier;
 
 	/** Is this an MajorTab? A tool panel tab? */
-	ETabRole::Type TabRole;
+	ETabRole TabRole;
 
 	/** The width that this tab will overlap with side-by-side tabs. */
 	float OverlapWidth;
diff --git a/Engine/Source/Runtime/Slate/Public/Widgets/Input/SEditableText.h b/Engine/Source/Runtime/Slate/Public/Widgets/Input/SEditableText.h
--- a/Engine/Source/Runtime/Slate/Public/Widgets/Input/SEditableText.h
+++ b/Engine/Source/Runtime/Slate/Public/Widgets/Input/SEditableText.h
@@ -5,6 +5,7 @@
 #include "ITextEditorWidget.h"
 #include "ITextInputMethodSystem.h"
 #include "IVirtualKeyboardEntry.h"
+#include "SlateScrollHelper.h"
 
 /**
  * Editable text widget
diff --git a/Engine/Source/Runtime/Slate/Public/Widgets/Layout/SSafeZone.h b/Engine/Source/Runtime/Slate/Public/Widgets/Layout/SSafeZone.h
--- a/Engine/Source/Runtime/Slate/Public/Widgets/Layout/SSafeZone.h
+++ b/Engine/Source/Runtime/Slate/Public/Widgets/Layout/SSafeZone.h
@@ -47,31 +47,9 @@ class SLATE_API SSafeZone : public SBox
 
 	SLATE_END_ARGS()
 
-	void Construct( const FArguments& InArgs )
-	{
-		SBox::Construct( SBox::FArguments()
-			.HAlign(InArgs._HAlign)
-			.VAlign(InArgs._VAlign)
-			.Padding( this, &SSafeZone::GetSafeZonePadding )
-			[
-				InArgs._Content.Widget
-			]
-		);
+	void Construct( const FArguments& InArgs );
 
-		IsTitleSafe = InArgs._IsTitleSafe;
-		Padding = InArgs._Padding;
-	}
-
-	FMargin GetSafeZonePadding() const
-	{
-		// @todo: should we have a function that just returns the safe area size?
-		FDisplayMetrics Metrics;
-		FSlateApplication::Get().GetDisplayMetrics(Metrics);
-
-		// return either the TitleSafe or the ActionSafe size, added to the user padding
-		return Padding.Get() + (IsTitleSafe ? FMargin(Metrics.TitleSafePaddingSize.X, Metrics.TitleSafePaddingSize.Y) : 
-			FMargin(Metrics.ActionSafePaddingSize.X, Metrics.ActionSafePaddingSize.Y));
-	}
+	FMargin GetSafeZonePadding() const;
 
 private:
 	/** Cached values from the args */
diff --git a/Engine/Source/Runtime/Slate/Public/Widgets/Layout/SScrollBox.h b/Engine/Source/Runtime/Slate/Public/Widgets/Layout/SScrollBox.h
--- a/Engine/Source/Runtime/Slate/Public/Widgets/Layout/SScrollBox.h
+++ b/Engine/Source/Runtime/Slate/Public/Widgets/Layout/SScrollBox.h
@@ -2,6 +2,8 @@
 
 #pragma once
 
+#include "InertialScrollManager.h"
+
 /** SScrollBox can scroll through an arbitrary number of widgets. */
 class SLATE_API SScrollBox : public SCompoundWidget
 {
diff --git a/Engine/Source/Runtime/Slate/Public/Widgets/Notifications/SNotificationList.h b/Engine/Source/Runtime/Slate/Public/Widgets/Notifications/SNotificationList.h
--- a/Engine/Source/Runtime/Slate/Public/Widgets/Notifications/SNotificationList.h
+++ b/Engine/Source/Runtime/Slate/Public/Widgets/Notifications/SNotificationList.h
@@ -2,6 +2,8 @@
 
 #pragma once
 
+#include "SlateFwd.h"
+#include "SCheckBox.h"
 
 /**
  * Interface class for an item in the event message list.
diff --git a/Engine/Source/Runtime/Slate/Public/Widgets/Views/STableViewBase.h b/Engine/Source/Runtime/Slate/Public/Widgets/Views/STableViewBase.h
--- a/Engine/Source/Runtime/Slate/Public/Widgets/Views/STableViewBase.h
+++ b/Engine/Source/Runtime/Slate/Public/Widgets/Views/STableViewBase.h
@@ -3,7 +3,7 @@
 #pragma once
 
 #include "IScrollableWidget.h"
-
+#include "InertialScrollManager.h"
 
 class SListPanel;
 class SHeaderRow;
diff --git a/Engine/Source/Runtime/SlateCore/Private/Widgets/SCompoundWidget.cpp b/Engine/Source/Runtime/SlateCore/Private/Widgets/SCompoundWidget.cpp
--- a/Engine/Source/Runtime/SlateCore/Private/Widgets/SCompoundWidget.cpp
+++ b/Engine/Source/Runtime/SlateCore/Private/Widgets/SCompoundWidget.cpp
@@ -2,6 +2,7 @@
 
 #include "SlateCorePrivatePCH.h"
 
+DECLARE_CYCLE_STAT( TEXT("OnPaint SCompoundWidget"), STAT_SlateOnPaint_SCompoundWidget, STATGROUP_Slate );
 
 int32 SCompoundWidget::OnPaint( const FPaintArgs& Args, const FGeometry& AllottedGeometry, const FSlateRect& MyClippingRect, FSlateWindowElementList& OutDrawElements, int32 LayerId, const FWidgetStyle& InWidgetStyle, bool bParentEnabled ) const
 {
diff --git a/Engine/Source/Runtime/SlateCore/Private/Widgets/SOverlay.cpp b/Engine/Source/Runtime/SlateCore/Private/Widgets/SOverlay.cpp
--- a/Engine/Source/Runtime/SlateCore/Private/Widgets/SOverlay.cpp
+++ b/Engine/Source/Runtime/SlateCore/Private/Widgets/SOverlay.cpp
@@ -4,6 +4,8 @@
 
 #include "LayoutUtils.h"
 
+DECLARE_CYCLE_STAT( TEXT("OnPaint SOverlay"), STAT_SlateOnPaint_SOverlay, STATGROUP_Slate );
+
 SOverlay::SOverlay()
 : Children()
 {
diff --git a/Engine/Source/Runtime/SlateCore/Private/Widgets/SPanel.cpp b/Engine/Source/Runtime/SlateCore/Private/Widgets/SPanel.cpp
--- a/Engine/Source/Runtime/SlateCore/Private/Widgets/SPanel.cpp
+++ b/Engine/Source/Runtime/SlateCore/Private/Widgets/SPanel.cpp
@@ -2,6 +2,7 @@
 
 #include "SlateCorePrivatePCH.h"
 
+DECLARE_CYCLE_STAT( TEXT("OnPaint SPanel"), STAT_SlateOnPaint_SPanel, STATGROUP_Slate );
 
 /* SWidget overrides
  *****************************************************************************/
