Slate - SUniformGridPanel now allows setting the min desired width/height of items with functions post construct.

diff --git a/Engine/Source/Runtime/Slate/Private/Widgets/Layout/SUniformGridPanel.cpp b/Engine/Source/Runtime/Slate/Private/Widgets/Layout/SUniformGridPanel.cpp
--- a/Engine/Source/Runtime/Slate/Private/Widgets/Layout/SUniformGridPanel.cpp
+++ b/Engine/Source/Runtime/Slate/Private/Widgets/Layout/SUniformGridPanel.cpp
@@ -96,6 +96,16 @@ void SUniformGridPanel::SetSlotPadding(TAttribute<FMargin> InSlotPadding)
 	SlotPadding = InSlotPadding;
 }
 
+void SUniformGridPanel::SetMinDesiredSlotWidth(TAttribute<float> InMinDesiredSlotWidth)
+{
+	MinDesiredSlotWidth = InMinDesiredSlotWidth;
+}
+
+void SUniformGridPanel::SetMinDesiredSlotHeight(TAttribute<float> InMinDesiredSlotHeight)
+{
+	MinDesiredSlotHeight = InMinDesiredSlotHeight;
+}
+
 SUniformGridPanel::FSlot& SUniformGridPanel::AddSlot( int32 Column, int32 Row )
 {
 	FSlot& NewSlot = *(new FSlot( Column, Row ));
diff --git a/Engine/Source/Runtime/Slate/Public/Widgets/Layout/SUniformGridPanel.h b/Engine/Source/Runtime/Slate/Public/Widgets/Layout/SUniformGridPanel.h
--- a/Engine/Source/Runtime/Slate/Public/Widgets/Layout/SUniformGridPanel.h
+++ b/Engine/Source/Runtime/Slate/Public/Widgets/Layout/SUniformGridPanel.h
@@ -64,6 +64,12 @@ public:
 	/** See SlotPadding attribute */
 	void SetSlotPadding(TAttribute<FMargin> InSlotPadding);
 
+	/** See MinDesiredSlotWidth attribute */
+	void SetMinDesiredSlotWidth(TAttribute<float> InMinDesiredSlotWidth);
+
+	/** See MinDesiredSlotHeight attribute */
+	void SetMinDesiredSlotHeight(TAttribute<float> InMinDesiredSlotHeight);
+
 	/**
 	 * Dynamically add a new slot to the UI at specified Column and Row.
 	 *
