Added OnFinishedPointerInput(), which gets called after event processing and before GEngine->Tick(). SViewport and FSceneViewport now uses this function instead of ProcessInput() to accumulate high-precision mouse events.

Addresses TTP#349213 MouseX and MouseY are a frame behind all other input

diff --git a/Engine/Source/Runtime/Engine/Private/Slate/SceneViewport.cpp b/Engine/Source/Runtime/Engine/Private/Slate/SceneViewport.cpp
--- a/Engine/Source/Runtime/Engine/Private/Slate/SceneViewport.cpp
+++ b/Engine/Source/Runtime/Engine/Private/Slate/SceneViewport.cpp
@@ -162,6 +162,11 @@ void FSceneViewport::SetMouse( int32 X, int32 Y )
 	CachedMousePos = FIntPoint(X, Y);
 }
 
+void FSceneViewport::ProcessInput( float DeltaTime )
+{
+	// Required 
+}
+
 void FSceneViewport::UpdateCachedMousePos( const FGeometry& InGeometry, const FPointerEvent& InMouseEvent )
 {
 	CachedMousePos = InGeometry.AbsoluteToLocal( InMouseEvent.GetScreenSpacePosition() ).IntPoint();
@@ -216,7 +221,7 @@ void FSceneViewport::ApplyModifierKeys( const FModifierKeysState& InKeysState )
 	}
 }
 
-void FSceneViewport::ProcessInput( float DeltaTime )
+void FSceneViewport::ProcessAccumulatedPointerInput()
 {
 	if( !ViewportClient )
 	{
@@ -230,6 +235,7 @@ void FSceneViewport::ProcessInput( float DeltaTime )
 
 	if (NumMouseSamplesX > 0 || NumMouseSamplesY > 0)
 	{
+		const float DeltaTime = FApp::GetDeltaTime();
 		ViewportClient->InputAxis( this, 0, EKeys::MouseX, MouseDelta.X, DeltaTime, NumMouseSamplesX );
 		ViewportClient->InputAxis( this, 0, EKeys::MouseY, MouseDelta.Y, DeltaTime, NumMouseSamplesY );
 	}
@@ -718,6 +724,11 @@ FReply FSceneViewport::OnMotionDetected( const FGeometry& MyGeometry, const FMot
 	return CurrentReplyState;
 }
 
+void FSceneViewport::OnFinishedPointerInput()
+{
+	ProcessAccumulatedPointerInput();
+}
+
 FReply FSceneViewport::OnKeyDown( const FGeometry& InGeometry, const FKeyboardEvent& InKeyboardEvent )
 {
 	// Start a new reply state
diff --git a/Engine/Source/Runtime/Engine/Public/Slate/SceneViewport.h b/Engine/Source/Runtime/Engine/Public/Slate/SceneViewport.h
--- a/Engine/Source/Runtime/Engine/Public/Slate/SceneViewport.h
+++ b/Engine/Source/Runtime/Engine/Public/Slate/SceneViewport.h
@@ -82,7 +82,7 @@ public:
 	 * Additional input processing that happens every frame                   
 	 */
 	virtual void ProcessInput( float DeltaTime ) override;
-
+	
 	virtual FVector2D VirtualDesktopPixelToViewport(FIntPoint VirtualDesktopPointPx) const override;
 	virtual FIntPoint ViewportToVirtualDesktopPixel(FVector2D ViewportCoordinate) const override;
 
@@ -203,6 +203,7 @@ public:
 	virtual FReply OnTouchEnded( const FGeometry& MyGeometry, const FPointerEvent& InTouchEvent ) override;
 	virtual FReply OnTouchGesture( const FGeometry& MyGeometry, const FPointerEvent& InGestureEvent ) override;
 	virtual FReply OnMotionDetected( const FGeometry& MyGeometry, const FMotionEvent& InMotionEvent ) override;
+	virtual void OnFinishedPointerInput() override;
 	virtual FReply OnKeyDown( const FGeometry& InGeometry, const FKeyboardEvent& InKeyboardEvent ) override;
 	virtual FReply OnKeyUp( const FGeometry& InGeometry, const FKeyboardEvent& InKeyboardEvent ) override;
 	virtual FReply OnKeyChar( const FGeometry& InGeometry, const FCharacterEvent& InCharacterEvent ) override;
@@ -237,6 +238,10 @@ private:
 	 */
 	virtual void ResizeViewport( uint32 NewSizeX,uint32 NewSizeY,EWindowMode::Type NewWindowMode,int32 InPosX, int32 InPosY );
 
+	/**
+	 * Called from slate when input is finished for this frame, and we should process any accumulated mouse data.
+	 */
+	void ProcessAccumulatedPointerInput();
 
 	/**
 	 * Updates the cached mouse position from a mouse event
diff --git a/Engine/Source/Runtime/Launch/Private/LaunchEngineLoop.cpp b/Engine/Source/Runtime/Launch/Private/LaunchEngineLoop.cpp
--- a/Engine/Source/Runtime/Launch/Private/LaunchEngineLoop.cpp
+++ b/Engine/Source/Runtime/Launch/Private/LaunchEngineLoop.cpp
@@ -2184,6 +2184,9 @@ void FEngineLoop::Tick()
 			FSlateApplication::Get().PollGameDeviceState();
 		}
 
+		// Gives widgets a chance to process any accumulated input
+		FSlateApplication::Get().FinishedInputThisFrame();
+
 		GEngine->Tick( FApp::GetDeltaTime(), bIdleMode );
 
 		if (GShaderCompilingManager)
diff --git a/Engine/Source/Runtime/Slate/Private/Framework/Application/SlateApplication.cpp b/Engine/Source/Runtime/Slate/Private/Framework/Application/SlateApplication.cpp
--- a/Engine/Source/Runtime/Slate/Private/Framework/Application/SlateApplication.cpp
+++ b/Engine/Source/Runtime/Slate/Private/Framework/Application/SlateApplication.cpp
@@ -804,6 +804,54 @@ void FSlateApplication::PollGameDeviceState()
 	}
 }
 
+void FSlateApplication::FinishedInputThisFrame()
+{
+	const float DeltaTime = GetDeltaTime();
+
+	// All the input events have been processed.
+
+	// Any widgets that may have received pointer input events
+	// are given a chance to process accumulated values.
+	if (MouseCaptor.HasCapture())
+	{
+		TArray<TSharedRef<SWidget>> Captors = MouseCaptor.ToSharedWidgets();
+		for (const auto & Captor : Captors )
+		{
+			Captor->OnFinishedPointerInput();
+		}
+	}
+	else
+	{
+		for ( const TWeakPtr<SWidget>& WidgetPtr : WidgetsUnderCursorLastEvent.Widgets )
+		{
+			const TSharedPtr<SWidget>& Widget = WidgetPtr.Pin();
+			if (Widget.IsValid())
+			{
+				Widget->OnFinishedPointerInput();
+			}
+			else
+			{
+				break;
+			}
+		}
+	}
+
+	// Any widgets that may have recieved key events
+	// are given a chance to process accumulated values.
+	for( const TWeakPtr<SWidget>& WidgetPtr : FocusedWidgetPath.Widgets )
+	{
+		const TSharedPtr<SWidget>& Widget = WidgetPtr.Pin();
+		if (Widget.IsValid())
+		{
+			Widget->OnFinishedKeyInput();
+		}
+		else
+		{
+			break;
+		}
+	}
+}
+
 /**
  * Ticks this application
  */
@@ -3082,13 +3130,10 @@ bool FSlateApplication::ProcessKeyDownEvent( FKeyboardEvent& InKeyboardEvent )
 				const bool bIsWidgetReflectorPicking = WidgetReflector.IsValid() && WidgetReflector->IsInPickingMode();
 				if ( bIsWidgetReflectorPicking )
 				{
-					if ( WidgetReflector.IsValid() )
-					{
-						WidgetReflector->OnWidgetPicked();
-						Reply = FReply::Handled();
+					WidgetReflector->OnWidgetPicked();
+					Reply = FReply::Handled();
 
-						return Reply.IsEventHandled();
-					}
+					return Reply.IsEventHandled();
 				}
 			}
 
diff --git a/Engine/Source/Runtime/Slate/Private/Widgets/SViewport.cpp b/Engine/Source/Runtime/Slate/Private/Widgets/SViewport.cpp
--- a/Engine/Source/Runtime/Slate/Private/Widgets/SViewport.cpp
+++ b/Engine/Source/Runtime/Slate/Private/Widgets/SViewport.cpp
@@ -272,3 +272,12 @@ FReply SViewport::OnMotionDetected( const FGeometry& MyGeometry, const FMotionEv
 {
 	return ViewportInterface.IsValid() ? ViewportInterface.Pin()->OnMotionDetected(MyGeometry, MotionEvent) : FReply::Unhandled();
 }
+
+void SViewport::OnFinishedPointerInput()
+{
+	TSharedPtr<ISlateViewport> PinnedInterface = ViewportInterface.Pin();
+	if (PinnedInterface.IsValid())
+	{
+		PinnedInterface->OnFinishedPointerInput();
+	}
+}
diff --git a/Engine/Source/Runtime/Slate/Public/Framework/Application/SlateApplication.h b/Engine/Source/Runtime/Slate/Public/Framework/Application/SlateApplication.h
--- a/Engine/Source/Runtime/Slate/Public/Framework/Application/SlateApplication.h
+++ b/Engine/Source/Runtime/Slate/Public/Framework/Application/SlateApplication.h
@@ -210,6 +210,11 @@ public:
 	 */
 	void PollGameDeviceState();
 
+	/**
+	 * Occurs before Tick(), after all pointer and keyboard input has been processed.
+	 */
+	void FinishedInputThisFrame();
+
 	/**
 	 * Ticks this application
 	 */
diff --git a/Engine/Source/Runtime/Slate/Public/Widgets/SViewport.h b/Engine/Source/Runtime/Slate/Public/Widgets/SViewport.h
--- a/Engine/Source/Runtime/Slate/Public/Widgets/SViewport.h
+++ b/Engine/Source/Runtime/Slate/Public/Widgets/SViewport.h
@@ -157,6 +157,7 @@ public:
 	virtual FReply OnControllerButtonReleased( const FGeometry& MyGeometry, const FControllerEvent& ControllerEvent ) override;
 	virtual FReply OnControllerAnalogValueChanged( const FGeometry& MyGeometry, const FControllerEvent& ControllerEvent ) override;
 	virtual FReply OnMotionDetected( const FGeometry& MyGeometry, const FMotionEvent& MotionEvent ) override;
+	virtual void OnFinishedPointerInput() override;
 
 private:
 
diff --git a/Engine/Source/Runtime/SlateCore/Private/Widgets/SWidget.cpp b/Engine/Source/Runtime/SlateCore/Private/Widgets/SWidget.cpp
--- a/Engine/Source/Runtime/SlateCore/Private/Widgets/SWidget.cpp
+++ b/Engine/Source/Runtime/SlateCore/Private/Widgets/SWidget.cpp
@@ -236,6 +236,16 @@ FReply SWidget::OnMotionDetected( const FGeometry& MyGeometry, const FMotionEven
 	return FReply::Unhandled();
 }
 
+void SWidget::OnFinishedPointerInput()
+{
+
+}
+
+void SWidget::OnFinishedKeyInput()
+{
+
+}
+
 
 EWindowZone::Type SWidget::GetWindowZoneOverride() const
 {
diff --git a/Engine/Source/Runtime/SlateCore/Public/Rendering/RenderingCommon.h b/Engine/Source/Runtime/SlateCore/Public/Rendering/RenderingCommon.h
--- a/Engine/Source/Runtime/SlateCore/Public/Rendering/RenderingCommon.h
+++ b/Engine/Source/Runtime/SlateCore/Public/Rendering/RenderingCommon.h
@@ -472,6 +472,13 @@ public:
 		return FReply::Unhandled();
 	}
 
+	/**
+	 * Called after all input for this frame is processed.
+	 */
+	virtual void OnFinishedPointerInput()
+	{
+	}
+
 	/**
 	 * Called when the viewport loses keyboard focus.  
 	 *
diff --git a/Engine/Source/Runtime/SlateCore/Public/Widgets/SWidget.h b/Engine/Source/Runtime/SlateCore/Public/Widgets/SWidget.h
--- a/Engine/Source/Runtime/SlateCore/Public/Widgets/SWidget.h
+++ b/Engine/Source/Runtime/SlateCore/Public/Widgets/SWidget.h
@@ -403,6 +403,20 @@ public:
 	 */
 	virtual FReply OnMotionDetected( const FGeometry& MyGeometry, const FMotionEvent& InMotionEvent );
 
+	/**
+	 * All the pointer (mouse, touch, stylus, etc.) events from this frame have been routed.
+	 * This is a widget's chance to act on any accumulated data.
+	 *
+	 * 
+	 */
+	virtual void OnFinishedPointerInput();
+
+	/**
+	 * All the key (keyboard, gamepay, joystick, etc.) input from this frame has been routed.
+	 * This is a widget's chance to act on any accumulated data.
+	 */
+	virtual void OnFinishedKeyInput();
+
 	/**
 	 * Called when the mouse is moved over the widget's window, to determine if we should report whether
 	 * OS-specific features should be active at this location (such as a title bar grip, system menu, etc.)
