Slate input event routing now a bit more systematic. ProcessReply(), SetHandler(), and Logging() now happen automatically when using FEventRouter. Logging is still incomplete.

diff --git a/Engine/Source/Developer/SlateReflector/Private/Widgets/SWidgetReflector.cpp b/Engine/Source/Developer/SlateReflector/Private/Widgets/SWidgetReflector.cpp
--- a/Engine/Source/Developer/SlateReflector/Private/Widgets/SWidgetReflector.cpp
+++ b/Engine/Source/Developer/SlateReflector/Private/Widgets/SWidgetReflector.cpp
@@ -5,130 +5,167 @@
 
 #define LOCTEXT_NAMESPACE "SWidgetReflector"
 
+static const int32 MaxLoggedEvents = 100;
+
+#define WITH_EVENT_LOGGING 0
+
+struct FLoggedEvent
+{
+	FLoggedEvent( const FInputEvent& InEvent, const FReply& InReply )
+	: Event( InEvent )
+	, Handler( InReply.GetHandler() )
+	, EventText( InEvent.ToText() )
+	, HandlerText( InReply.GetHandler().IsValid() ? FText::FromString(InReply.GetHandler()->ToString()) : LOCTEXT("NullHandler", "null") )
+	{		
+	}
+
+	FText ToText( )
+	{
+		return FText::Format( NSLOCTEXT("","","{0}  |  {1}"), EventText, HandlerText );
+	}
+	
+	FInputEvent Event;
+	TWeakPtr<SWidget> Handler;
+	FText EventText;
+	FText HandlerText;
+};
+
 
 /* SWidgetReflector interface
  *****************************************************************************/
 
 void SWidgetReflector::Construct( const FArguments& InArgs )
 {
+	LoggedEvents.Reserve(MaxLoggedEvents);
+
 	bShowFocus = false;
 	bIsPicking = false;
 
 	this->ChildSlot
 	[
 		SNew(SBorder)
-			.BorderImage(FCoreStyle::Get().GetBrush("ToolPanel.GroupBorder"))
-			[
-				SNew(SVerticalBox)
+		.BorderImage(FCoreStyle::Get().GetBrush("ToolPanel.GroupBorder"))
+		[
+			SNew(SVerticalBox)
 				
-				+ SVerticalBox::Slot()
-					.AutoHeight()
-					[
-						SNew(SHorizontalBox)
+			+ SVerticalBox::Slot()
+			.AutoHeight()
+			[
+				SNew(SHorizontalBox)
 					
-						+ SHorizontalBox::Slot()
-							.AutoWidth()
-							[
-								SNew(STextBlock)
-									.Text(LOCTEXT("AppScale", "Application Scale: ").ToString())
-							]
-
-						+ SHorizontalBox::Slot()
-						.MaxWidth(250)
+				+ SHorizontalBox::Slot()
+				.AutoWidth()
+				[
+					SNew(STextBlock)
+						.Text(LOCTEXT("AppScale", "Application Scale: ").ToString())
+				]
+
+				+ SHorizontalBox::Slot()
+				.MaxWidth(250)
+				[
+					SNew(SSpinBox<float>)
+					.Value(this, &SWidgetReflector::HandleAppScaleSliderValue)
+					.MinValue(0.1f)
+					.MaxValue(3.0f)
+					.Delta(0.01f)
+					.OnValueChanged(this, &SWidgetReflector::HandleAppScaleSliderChanged)
+				]
+			]
+
+			+ SVerticalBox::Slot()
+			.AutoHeight()
+			[
+				SNew(SHorizontalBox)
+
+				+ SHorizontalBox::Slot()
+				.AutoWidth()
+				.Padding(5.0f)
+				[
+					// Check box that controls LIVE MODE
+					SNew(SCheckBox)
+						.IsChecked(this, &SWidgetReflector::HandleFocusCheckBoxIsChecked)
+						.OnCheckStateChanged(this, &SWidgetReflector::HandleFocusCheckBoxCheckedStateChanged)
 						[
-							SNew(SSpinBox<float>)
-								.Value(this, &SWidgetReflector::HandleAppScaleSliderValue)
-								.MinValue(0.1f)
-								.MaxValue(3.0f)
-								.Delta(0.01f)
-								.OnValueChanged(this, &SWidgetReflector::HandleAppScaleSliderChanged)
+							SNew(STextBlock)
+								.Text(LOCTEXT("ShowFocus", "Show Focus").ToString())
 						]
-					]
+				]
+
+				+ SHorizontalBox::Slot()
+				.AutoWidth()
+				.Padding(5.0f)
+				[
+					// Check box that controls PICKING A WIDGET TO INSPECT
+					SNew(SButton)
+						.OnClicked(this, &SWidgetReflector::HandlePickButtonClicked)
+						.ButtonColorAndOpacity(this, &SWidgetReflector::HandlePickButtonColorAndOpacity)
+						[
+							SNew(STextBlock)
+								.Text(this, &SWidgetReflector::HandlePickButtonText)
+						]
+				]
+			]
 
-				+ SVerticalBox::Slot()
-					.AutoHeight()
+			+ SVerticalBox::Slot()
+			.FillHeight(1.0f)
+			[
+				// The tree view that shows all the info that we capture.
+				SAssignNew(ReflectorTree, SReflectorTree)
+				.ItemHeight(24.0f)
+				.TreeItemsSource(&ReflectorTreeRoot)
+				.OnGenerateRow(this, &SWidgetReflector::HandleReflectorTreeGenerateRow)
+				.OnGetChildren(this, &SWidgetReflector::HandleReflectorTreeGetChildren)
+				.OnSelectionChanged(this, &SWidgetReflector::HandleReflectorTreeSelectionChanged)
+				.HeaderRow
+				(
+					SNew(SHeaderRow)
+
+					+ SHeaderRow::Column("WidgetName")
+					.DefaultLabel(LOCTEXT("WidgetName", "Widget Name").ToString())
+					.FillWidth(0.65f)
+
+					+ SHeaderRow::Column("ForegroundColor")
+					.FixedWidth(24.0f)
+					.VAlignHeader(VAlign_Center)
+					.HeaderContent()
 					[
-						SNew(SHorizontalBox)
-
-						+ SHorizontalBox::Slot()
-							.AutoWidth()
-							.Padding(5.0f)
-							[
-								// Check box that controls LIVE MODE
-								SNew(SCheckBox)
-									.IsChecked(this, &SWidgetReflector::HandleFocusCheckBoxIsChecked)
-									.OnCheckStateChanged(this, &SWidgetReflector::HandleFocusCheckBoxCheckedStateChanged)
-									[
-										SNew(STextBlock)
-											.Text(LOCTEXT("ShowFocus", "Show Focus").ToString())
-									]
-							]
-
-						+ SHorizontalBox::Slot()
-							.AutoWidth()
-							.Padding(5.0f)
-							[
-								// Check box that controls PICKING A WIDGET TO INSPECT
-								SNew(SButton)
-									.OnClicked(this, &SWidgetReflector::HandlePickButtonClicked)
-									.ButtonColorAndOpacity(this, &SWidgetReflector::HandlePickButtonColorAndOpacity)
-									[
-										SNew(STextBlock)
-											.Text(this, &SWidgetReflector::HandlePickButtonText)
-									]
-							]
+						SNew(STextBlock)
+						.Text(LOCTEXT("ForegroundColor", "FG").ToString())
+						.ToolTipText(LOCTEXT("ForegroundColorToolTip", "Foreground Color").ToString())
 					]
 
-				+ SVerticalBox::Slot()
-					.FillHeight(1.0f)
-					[
-						// The tree view that shows all the info that we capture.
-						SAssignNew(ReflectorTree, SReflectorTree)
-							.ItemHeight(24.0f)
-							.TreeItemsSource(&ReflectorTreeRoot)
-							.OnGenerateRow(this, &SWidgetReflector::HandleReflectorTreeGenerateRow)
-							.OnGetChildren(this, &SWidgetReflector::HandleReflectorTreeGetChildren)
-							.OnSelectionChanged(this, &SWidgetReflector::HandleReflectorTreeSelectionChanged)
-							.HeaderRow
-							(
-								SNew(SHeaderRow)
-
-								+ SHeaderRow::Column("WidgetName")
-									.DefaultLabel(LOCTEXT("WidgetName", "Widget Name").ToString())
-									.FillWidth(0.65f)
-
-								+ SHeaderRow::Column("ForegroundColor")
-									.FixedWidth(24.0f)
-									.VAlignHeader(VAlign_Center)
-									.HeaderContent()
-									[
-										SNew(STextBlock)
-											.Text(LOCTEXT("ForegroundColor", "FG").ToString())
-											.ToolTipText(LOCTEXT("ForegroundColorToolTip", "Foreground Color").ToString())
-									]
-
-								+ SHeaderRow::Column("Visibility")
-									.DefaultLabel(LOCTEXT("Visibility", "Visibility" ).ToString())
-									.FixedWidth(125.0f)
-
-								+ SHeaderRow::Column("WidgetInfo")
-									.DefaultLabel(LOCTEXT("WidgetInfo", "Widget Info" ).ToString())
-									.FillWidth(0.25f)
-
-								+ SHeaderRow::Column("Address")
-									.DefaultLabel( LOCTEXT("Address", "Address") )
-									.FillWidth( 0.10f )
-							)
-					]
+					+ SHeaderRow::Column("Visibility")
+					.DefaultLabel(LOCTEXT("Visibility", "Visibility" ).ToString())
+					.FixedWidth(125.0f)
 
-				+ SVerticalBox::Slot()
-					.AutoHeight()
-					[
-						// Frame rate counter
-						SNew(STextBlock)
-							.Text(this, &SWidgetReflector::HandleFrameRateText)
-					]
+					+ SHeaderRow::Column("WidgetInfo")
+					.DefaultLabel(LOCTEXT("WidgetInfo", "Widget Info" ).ToString())
+					.FillWidth(0.25f)
+
+					+ SHeaderRow::Column("Address")
+					.DefaultLabel( LOCTEXT("Address", "Address") )
+					.FillWidth( 0.10f )
+				)
+			]
+
+			#if WITH_EVENT_LOGGING
+			+ SVerticalBox::Slot()
+			.FillHeight(1.0f)
+			[
+				SAssignNew(EventListView, SListView<TSharedRef<FLoggedEvent>>)
+				.ListItemsSource( &LoggedEvents )
+				.OnGenerateRow(this, &SWidgetReflector::GenerateEventLogRow)
 			]
+			#endif //WITH_EVENT_LOGGING
+			
+			+ SVerticalBox::Slot()
+			.AutoHeight()
+			[
+				// Frame rate counter
+				SNew(STextBlock)
+					.Text(this, &SWidgetReflector::HandleFrameRateText)
+			]
+		]
 	];
 }
 
@@ -139,7 +176,20 @@ void SWidgetReflector::Construct( const FArguments& InArgs )
 void SWidgetReflector::Tick( const FGeometry& AllottedGeometry, const double InCurrentTime, const float InDeltaTime )
 {
 	SCompoundWidget::Tick(AllottedGeometry, InCurrentTime, InDeltaTime);
-	ReflectorTree->RequestTreeRefresh();
+}
+
+void SWidgetReflector::OnEventProcessed( const FInputEvent& Event, const FReply& InReply )
+{
+	#if WITH_EVENT_LOGGING
+		if (LoggedEvents.Num() >= MaxLoggedEvents)
+		{
+			LoggedEvents.Empty();
+		}
+
+		LoggedEvents.Add( MakeShareable( new FLoggedEvent( Event, InReply ) ) );
+		EventListView->RequestListRefresh();
+		EventListView->RequestScrollIntoView( LoggedEvents.Last() );
+	#endif //WITH_EVENT_LOGGING
 }
 
 
@@ -331,5 +381,14 @@ TSharedRef<ITableRow> SWidgetReflector::HandleReflectorTreeGenerateRow( TSharedP
 		.SourceCodeAccessor(SourceAccessDelegate);
 }
 
+TSharedRef<ITableRow> SWidgetReflector::GenerateEventLogRow( TSharedRef<FLoggedEvent> InLoggedEvent, const TSharedRef<STableViewBase>& OwnerTable )
+{
+	return SNew(STableRow<TSharedRef<FLoggedEvent>>, OwnerTable)
+	[
+		SNew(STextBlock)
+		.Text( InLoggedEvent->ToText() )
+	];
+}
+
 
 #undef LOCTEXT_NAMESPACE
diff --git a/Engine/Source/Developer/SlateReflector/Private/Widgets/SWidgetReflector.h b/Engine/Source/Developer/SlateReflector/Private/Widgets/SWidgetReflector.h
--- a/Engine/Source/Developer/SlateReflector/Private/Widgets/SWidgetReflector.h
+++ b/Engine/Source/Developer/SlateReflector/Private/Widgets/SWidgetReflector.h
@@ -4,6 +4,8 @@
 
 #include "IWidgetReflector.h"
 
+struct FLoggedEvent;
+
 /**
  * Abstract base class for widget reflectors.
  */
@@ -36,6 +38,8 @@ public:
 
 	// IWidgetReflector interface
 
+	virtual void OnEventProcessed( const FInputEvent& Event, const FReply& InReply ) override;
+
 	virtual bool IsInPickingMode( ) const override
 	{
 		return bIsPicking;
@@ -169,8 +173,13 @@ private:
 		SelectedNodes = ReflectorTree->GetSelectedItems();
 	}
 
+	TSharedRef<ITableRow> GenerateEventLogRow( TSharedRef<FLoggedEvent> InReflectorNode, const TSharedRef<STableViewBase>& OwnerTable );
+
 private:
 
+	TArray< TSharedRef<FLoggedEvent> > LoggedEvents;
+	TSharedPtr< SListView< TSharedRef< FLoggedEvent > > > EventListView;
+
 	TSharedPtr<SReflectorTree> ReflectorTree;
 
 	TArray<TSharedPtr<FReflectorNode>> SelectedNodes;
diff --git a/Engine/Source/Runtime/Slate/Private/Framework/Application/SlateApplication.cpp b/Engine/Source/Runtime/Slate/Private/Framework/Application/SlateApplication.cpp
--- a/Engine/Source/Runtime/Slate/Private/Framework/Application/SlateApplication.cpp
+++ b/Engine/Source/Runtime/Slate/Private/Framework/Application/SlateApplication.cpp
@@ -9,6 +9,108 @@
 #include "GenericCommands.h"
 #include "NotificationManager.h"
 
+
+class FEventRouter
+{
+
+// @todo slate : handle preview events.
+
+// @todo slate : Widget Reflector should log: (1) Every process reply (2) Every time the event is handled and by who.
+
+// @todo slate : Remove remaining [&]-style mass captures.
+
+// @todo slate : Eliminate all ad-hoc uses of SetEventPath()
+
+// @todo slate : Remove CALL_WIDGET_FUNCTION
+
+public:
+	static void LogEvent( FSlateApplication* ThisApplication, const FInputEvent& Event, const FReply& Reply )
+	{
+		TSharedPtr<IWidgetReflector> Reflector = ThisApplication->WidgetReflectorPtr.Pin();
+		if (Reflector.IsValid() && Reply.IsEventHandled())
+		{
+			Reflector->OnEventProcessed( Event, Reply );
+		}
+	}
+
+	/**
+	 * Bubble an event along a focus path (as opposed to PointerPath)
+	 *
+	 * Focus paths are determined by the Keyboard Focus or Controller Focus.
+	 * Focus paths change when the user navigates focus (e.g. Tab or
+	 * Shift Tab or clicks on a focusable widget.)
+	 */
+	template< typename FuncType, typename EventType >
+	static FReply BubbleAlongFocusPath( FSlateApplication* ThisApplication, const FWidgetPath& RoutingPath, EventType KeyEventCopy, const FuncType& Lambda )
+	{
+		FReply Reply = FReply::Unhandled();
+
+		KeyEventCopy.SetEventPath(RoutingPath);
+
+		for ( int32 WidgetIndex = RoutingPath.Widgets.Num() - 1; !Reply.IsEventHandled() && WidgetIndex >= 0; --WidgetIndex )
+		{
+			const FArrangedWidget& ArrangedWidget = RoutingPath.Widgets[WidgetIndex];
+			Reply = Lambda( ArrangedWidget, KeyEventCopy ).SetHandler( ArrangedWidget.Widget );
+			ThisApplication->ProcessReply( RoutingPath, Reply, NULL, NULL );
+		}
+
+		LogEvent(ThisApplication, KeyEventCopy, Reply);
+		
+		return Reply;
+	}
+
+	/**
+	 * Bubble an event along a PointerPath (as opposed to FocusPath)
+	 *
+	 * The pointer path is determined on the fly: whatever is under
+	 * the cursor will get a shot at the event. Widgets can request
+	 * pointer capture. Wile pointer capture is active, only the
+	 * widget in question will receive pointer events associated
+	 * with that pointing device until the pointer capture is
+	 * released.
+	 */
+	template< typename FuncType >
+	static FReply BubbleAlongPointerPath( FSlateApplication* ThisApplication, const FWidgetPath& RoutingPath, FPointerEvent PointerEventCopy, const FuncType& Lambda )
+	{
+		FReply Reply = FReply::Unhandled();
+
+		PointerEventCopy.SetEventPath( RoutingPath );
+
+		for ( int32 WidgetIndex = RoutingPath.Widgets.Num() - 1; !Reply.IsEventHandled() && WidgetIndex >= 0; --WidgetIndex )
+		{
+			const FArrangedWidget& ArrangedWidget = RoutingPath.Widgets[WidgetIndex];
+			Reply = Lambda( ArrangedWidget, PointerEventCopy ).SetHandler( ArrangedWidget.Widget );
+			const FWidgetPath* WidgetsUnderCursorPtr = &RoutingPath;
+			ThisApplication->ProcessReply( RoutingPath, Reply, WidgetsUnderCursorPtr, &PointerEventCopy );
+		}
+
+		LogEvent(ThisApplication, PointerEventCopy, Reply);
+
+		return Reply;
+	}
+
+	/** Just like BubbleAlongPointerPath, but only the TargetWidget in the path gets the event. */
+	template< typename FuncType >
+	static FReply SendToLeafmost( FSlateApplication* ThisApplication, const FWidgetPath& RoutingPath, const TSharedPtr<SWidget> TargetWidget, FPointerEvent PointerEventCopy, const FuncType& Lambda )
+	{
+		PointerEventCopy.SetEventPath( RoutingPath );
+
+		const FArrangedWidget& ArrangedWidget = TargetWidget.IsValid()
+			? RoutingPath.FindArrangedWidget( TargetWidget.ToSharedRef() )
+			: RoutingPath.Widgets.Last();
+
+		FReply Reply = Lambda( ArrangedWidget, PointerEventCopy ).SetHandler( ArrangedWidget.Widget );
+		const FWidgetPath* WidgetsUnderCursorPtr = &RoutingPath;
+		ThisApplication->ProcessReply( RoutingPath, Reply, WidgetsUnderCursorPtr, &PointerEventCopy );
+
+		LogEvent(ThisApplication, PointerEventCopy, Reply);
+
+		return Reply;
+	}
+};
+
+
+
 DECLARE_CYCLE_STAT( TEXT("Message Tick Time"), STAT_SlateMessageTick, STATGROUP_Slate );
 DECLARE_CYCLE_STAT( TEXT("Update Tooltip Time"), STAT_SlateUpdateTooltip, STATGROUP_Slate );
 DECLARE_CYCLE_STAT( TEXT("Tick Window And Children Time"), STAT_SlateTickWindowAndChildren, STATGROUP_Slate );
@@ -3072,23 +3174,17 @@ bool FSlateApplication::ProcessKeyCharEvent( FCharacterEvent& InCharacterEvent )
 	
 		// Switch worlds for widgets in the current path
 		FScopedSwitchWorldHack SwitchWorld( EventPath );
-		TSharedPtr<SWidget> WidgetToLog;
 
 		Reply = FReply::Unhandled();
 		// Send out mouse enter events.
 		InCharacterEvent.SetEventPath( EventPath );
-		for( int32 WidgetIndex = EventPath.Widgets.Num()-1; !Reply.IsEventHandled() && WidgetIndex >= 0; --WidgetIndex )
+		
+		Reply = FEventRouter::BubbleAlongFocusPath( this, EventPath, InCharacterEvent, []( const FArrangedWidget& SomeWidgetGettingEvent, const FCharacterEvent& Event )
 		{
-			FArrangedWidget& SomeWidgetGettingEvent = EventPath.Widgets[ WidgetIndex ];
-			if (SomeWidgetGettingEvent.Widget->IsEnabled())
-			{
-				// Widget newly under cursor, so send a MouseEnter.
-				Reply = SomeWidgetGettingEvent.Widget->OnKeyChar( SomeWidgetGettingEvent.Geometry, InCharacterEvent ).SetHandler( SomeWidgetGettingEvent.Widget );				
-				ProcessReply(EventPath, Reply, NULL, NULL);
-
-				WidgetToLog = SomeWidgetGettingEvent.Widget;
-			}
-		}
+			return SomeWidgetGettingEvent.Widget->IsEnabled()
+				? SomeWidgetGettingEvent.Widget->OnKeyChar( SomeWidgetGettingEvent.Geometry, Event )
+				: FReply::Unhandled();
+		});
 
 		LOG_EVENT_CONTENT( EEventLog::KeyChar, FString::Printf(TEXT("%c"), InCharacterEvent.GetCharacter()), Reply );
 	}
@@ -3162,16 +3258,14 @@ bool FSlateApplication::ProcessKeyDownEvent( FKeyboardEvent& InKeyboardEvent )
 			}
 
 			// Send out key down events.
-			for( int32 WidgetIndex = EventPath.Widgets.Num()-1; !Reply.IsEventHandled() && WidgetIndex >= 0; --WidgetIndex )
+			if (!Reply.IsEventHandled())
 			{
-				FArrangedWidget& SomeWidgetGettingEvent = EventPath.Widgets[ WidgetIndex ];
-				if (SomeWidgetGettingEvent.Widget->IsEnabled())
+				Reply = FEventRouter::BubbleAlongFocusPath(this, EventPath, InKeyboardEvent, [](const FArrangedWidget& SomeWidgetGettingEvent, const FKeyboardEvent& Event)
 				{
-					Reply = SomeWidgetGettingEvent.Widget->OnKeyDown( SomeWidgetGettingEvent.Geometry, InKeyboardEvent ).SetHandler(SomeWidgetGettingEvent.Widget);
-					ProcessReply(EventPath, Reply, NULL, NULL);
-
-					WidgetToLog = SomeWidgetGettingEvent.Widget;
-				}
+					return (SomeWidgetGettingEvent.Widget->IsEnabled())
+						? SomeWidgetGettingEvent.Widget->OnKeyDown( SomeWidgetGettingEvent.Geometry, Event )
+						: FReply::Unhandled();
+				});
 			}
 
 			LOG_EVENT_CONTENT( EEventLog::KeyDown, GetKeyName(InKeyboardEvent.GetKey()).ToString(), Reply );
@@ -3228,22 +3322,12 @@ bool FSlateApplication::ProcessKeyUpEvent( FKeyboardEvent& InKeyboardEvent )
 		// Switch worlds for widgets in the current path
 		FScopedSwitchWorldHack SwitchWorld( EventPath );
 
-		TSharedPtr<SWidget> WidgetToLog;
-
-		Reply = FReply::Unhandled();
-		// Send out mouse enter events.
-		for( int32 WidgetIndex = EventPath.Widgets.Num()-1; !Reply.IsEventHandled() && WidgetIndex >= 0; --WidgetIndex )
+		Reply = FEventRouter::BubbleAlongFocusPath(this, EventPath, InKeyboardEvent, [](const FArrangedWidget& SomeWidgetGettingEvent, const FKeyboardEvent& Event)
 		{
-			FArrangedWidget& SomeWidgetGettingEvent = EventPath.Widgets[ WidgetIndex ];
-			if (SomeWidgetGettingEvent.Widget->IsEnabled())
-			{
-				// Widget newly under cursor, so send a MouseEnter.
-				Reply = SomeWidgetGettingEvent.Widget->OnKeyUp( SomeWidgetGettingEvent.Geometry, InKeyboardEvent ).SetHandler( SomeWidgetGettingEvent.Widget );				
-				ProcessReply(EventPath, Reply, NULL, NULL);
-
-				WidgetToLog = SomeWidgetGettingEvent.Widget;
-			}
-		}
+			return (SomeWidgetGettingEvent.Widget->IsEnabled())
+				? SomeWidgetGettingEvent.Widget->OnKeyUp( SomeWidgetGettingEvent.Geometry, Event )
+				: FReply::Unhandled();
+		});
 
 		LOG_EVENT_CONTENT( EEventLog::KeyUp, GetKeyName(InKeyboardEvent.GetKey()).ToString(), Reply );
 	}
@@ -3339,29 +3423,35 @@ bool FSlateApplication::ProcessMouseButtonDownEvent( const TSharedPtr< FGenericW
 			FScopedSwitchWorldHack SwitchWorld( MouseCaptorPath );
 			bInGame = FApp::IsGame();
 
-			MouseEvent.SetEventPath( MouseCaptorPath );
 			
-			Reply = MouseCaptorWidget.Widget->OnPreviewMouseButtonDown( MouseCaptorWidget.Geometry, MouseEvent ).SetHandler(MouseCaptorWidget.Widget);
-			ProcessReply(MouseCaptorPath, Reply, &MouseCaptorPath, &MouseEvent);
+			Reply = FEventRouter::SendToLeafmost( this, MouseCaptorPath, nullptr, MouseEvent, []( const FArrangedWidget& MouseCaptorWidget, const FPointerEvent& Event )
+			{
+				return MouseCaptorWidget.Widget->OnPreviewMouseButtonDown( MouseCaptorWidget.Geometry, Event );
+			});
 			
 			if (!Reply.IsEventHandled())
 			{
-				if (MouseEvent.IsTouchEvent())
-				{
-					Reply = MouseCaptorWidget.Widget->OnTouchStarted( MouseCaptorWidget.Geometry, MouseEvent ).SetHandler( MouseCaptorWidget.Widget );
-				}
-				if (!MouseEvent.IsTouchEvent() || (!Reply.IsEventHandled() && bTouchFallbackToMouse))
-				{
-					Reply = MouseCaptorWidget.Widget->OnMouseButtonDown( MouseCaptorWidget.Geometry, MouseEvent ).SetHandler( MouseCaptorWidget.Widget );
-				}
-				ProcessReply(MouseCaptorPath, Reply, &MouseCaptorPath, &MouseEvent);
+				Reply = FEventRouter::SendToLeafmost( this, MouseCaptorPath, nullptr, MouseEvent,
+					[this]( const FArrangedWidget& MouseCaptorWidget, const FPointerEvent& Event )
+					{
+						FReply TempReply = FReply::Unhandled();
+						if (Event.IsTouchEvent())
+						{
+							TempReply = MouseCaptorWidget.Widget->OnTouchStarted( MouseCaptorWidget.Geometry, Event );
+						}
+						if (!Event.IsTouchEvent() || (!TempReply.IsEventHandled() && this->bTouchFallbackToMouse))
+						{
+							TempReply = MouseCaptorWidget.Widget->OnMouseButtonDown( MouseCaptorWidget.Geometry, Event );
+						}
+						return TempReply;
+					} );
+
 			}
 			LOG_EVENT( EEventLog::MouseButtonDown, Reply );
 		}
 		else
 		{
 			FWidgetPath WidgetsUnderCursor = LocateWindowUnderMouse( MouseEvent.GetScreenSpacePosition(), GetInteractiveTopLevelWindows() );
-			MouseEvent.SetEventPath(WidgetsUnderCursor);
 
 #if PLATFORM_MAC
 			NSWindow* ActiveWindow = [NSApp keyWindow];
@@ -3372,36 +3462,31 @@ bool FSlateApplication::ProcessMouseButtonDownEvent( const TSharedPtr< FGenericW
 			FScopedSwitchWorldHack SwitchWorld( WidgetsUnderCursor );
 			bInGame = FApp::IsGame();
 
-			TSharedPtr<SWidget> WidgetToLog;
-
 			const TSharedPtr<SWidget> PreviouslyFocusedWidget = GetKeyboardFocusedWidget();
 
-			for( int32 WidgetIndex = 0; !Reply.IsEventHandled() && WidgetIndex < WidgetsUnderCursor.Widgets.Num() ; ++WidgetIndex )
+			Reply = FEventRouter::BubbleAlongPointerPath(this, WidgetsUnderCursor, MouseEvent, [](const FArrangedWidget TargetWidget, const FPointerEvent& Event)
 			{
-				FArrangedWidget& CurWidget = WidgetsUnderCursor.Widgets[WidgetIndex];
+				return TargetWidget.Widget->OnPreviewMouseButtonDown(TargetWidget.Geometry, Event);
+			});
 
-				Reply = CurWidget.Widget->OnPreviewMouseButtonDown(CurWidget.Geometry, MouseEvent).SetHandler( CurWidget.Widget );
-				ProcessReply( WidgetsUnderCursor, Reply, &WidgetsUnderCursor, &MouseEvent);
-			}
-
-			for( int32 WidgetIndex = WidgetsUnderCursor.Widgets.Num()-1; !Reply.IsEventHandled() && WidgetIndex >= 0 ; --WidgetIndex )
+			if ( !Reply.IsEventHandled() )
 			{
-				FArrangedWidget& CurWidget = WidgetsUnderCursor.Widgets[WidgetIndex];				
-
-				if (!Reply.IsEventHandled())
+				Reply = FEventRouter::BubbleAlongPointerPath(this, WidgetsUnderCursor, MouseEvent, [this](const FArrangedWidget TargetWidget, const FPointerEvent& Event)
 				{
-					if (MouseEvent.IsTouchEvent())
-					{
-						Reply = CurWidget.Widget->OnTouchStarted( CurWidget.Geometry, MouseEvent ).SetHandler( CurWidget.Widget );
-					}
-					if (!MouseEvent.IsTouchEvent() || (!Reply.IsEventHandled() && bTouchFallbackToMouse))
+					FReply ThisReply = FReply::Unhandled();
+					if (!ThisReply.IsEventHandled())
 					{
-						Reply = CurWidget.Widget->OnMouseButtonDown( CurWidget.Geometry, MouseEvent ).SetHandler( CurWidget.Widget );
+						if (Event.IsTouchEvent())
+						{
+							ThisReply = TargetWidget.Widget->OnTouchStarted( TargetWidget.Geometry, Event );
+						}
+						if (!Event.IsTouchEvent() || (!ThisReply.IsEventHandled() && this->bTouchFallbackToMouse))
+						{
+							ThisReply = TargetWidget.Widget->OnMouseButtonDown( TargetWidget.Geometry, Event );
+						}
 					}
-					ProcessReply( WidgetsUnderCursor, Reply, &WidgetsUnderCursor, &MouseEvent);
-				}
-
-				WidgetToLog = CurWidget.Widget;
+					return ThisReply;
+				});
 			}
 			LOG_EVENT( EEventLog::MouseButtonDown, Reply );
 
@@ -3496,17 +3581,11 @@ bool FSlateApplication::ProcessMouseButtonDoubleClickEvent( const TSharedPtr< FG
 	// Switch worlds widgets in the current path
 	FScopedSwitchWorldHack SwitchWorld( WidgetsUnderCursor );
 
-	TSharedPtr<SWidget> WidgetToLog;
-
-	FReply Reply = FReply::Unhandled();
-	for( int32 WidgetIndex = WidgetsUnderCursor.Widgets.Num()-1; !Reply.IsEventHandled() && WidgetIndex >= 0; --WidgetIndex )
+	FReply Reply = FEventRouter::BubbleAlongPointerPath( this, WidgetsUnderCursor, InMouseEvent, [](const FArrangedWidget& TargetWidget, const FPointerEvent& Event)
 	{
-		FArrangedWidget& CurWidget = WidgetsUnderCursor.Widgets[WidgetIndex];
-		Reply = CurWidget.Widget->OnMouseButtonDoubleClick( CurWidget.Geometry, InMouseEvent ).SetHandler( CurWidget.Widget );
-		ProcessReply(WidgetsUnderCursor, Reply, &WidgetsUnderCursor, &InMouseEvent);
+		return TargetWidget.Widget->OnMouseButtonDoubleClick( TargetWidget.Geometry, Event );
+	} );
 
-		WidgetToLog = CurWidget.Widget;
-	}
 
 	LOG_EVENT( EEventLog::MouseButtonDoubleClick, Reply );
 
@@ -3558,22 +3637,23 @@ bool FSlateApplication::ProcessMouseButtonUpEvent( FPointerEvent& MouseEvent )
 #if PLATFORM_MAC
 			NSWindow* ActiveNativeWindow = [NSApp keyWindow];
 #endif
-			FArrangedWidget& MouseCaptorWidget = MouseCaptorPath.Widgets.Last();
-			MouseEvent.SetEventPath(MouseCaptorPath);
-
 			// Switch worlds widgets in the current path
 			FScopedSwitchWorldHack SwitchWorld( MouseCaptorPath );
 
-			FReply Reply = FReply::Unhandled();
-			if (MouseEvent.IsTouchEvent())
+			FReply Reply = FEventRouter::SendToLeafmost( this, MouseCaptorPath, nullptr, MouseEvent, [this]( const FArrangedWidget& TargetWidget, const FPointerEvent& Event )
 			{
-				Reply = MouseCaptorWidget.Widget->OnTouchEnded( MouseCaptorWidget.Geometry, MouseEvent ).SetHandler( MouseCaptorWidget.Widget );
-			}
-			if (!MouseEvent.IsTouchEvent() || (!Reply.IsEventHandled() && bTouchFallbackToMouse))
-			{
-				Reply = MouseCaptorWidget.Widget->OnMouseButtonUp( MouseCaptorWidget.Geometry, MouseEvent ).SetHandler( MouseCaptorWidget.Widget );
-			}
-			ProcessReply(MouseCaptorPath, Reply, &MouseCaptorPath, &MouseEvent);
+				FReply TempReply = FReply::Unhandled();
+				if (Event.IsTouchEvent())
+				{
+					TempReply = TargetWidget.Widget->OnTouchEnded( TargetWidget.Geometry, Event );
+				}
+				if (!Event.IsTouchEvent() || (!TempReply.IsEventHandled() && this->bTouchFallbackToMouse))
+				{
+					TempReply = TargetWidget.Widget->OnMouseButtonUp( TargetWidget.Geometry, Event );
+				}
+				return TempReply;
+			});
+
 #if PLATFORM_MAC
 			// Activate a window under the mouse if it's inactive and mouse up didn't bring any window to front
 			TSharedPtr<SWindow> ActiveWindow = GetActiveTopLevelWindow();
@@ -3588,10 +3668,6 @@ bool FSlateApplication::ProcessMouseButtonUpEvent( FPointerEvent& MouseEvent )
 	else
 	{
 		FWidgetPath WidgetsUnderCursor = LocateWindowUnderMouse( MouseEvent.GetScreenSpacePosition(), GetInteractiveTopLevelWindows() );
-		MouseEvent.SetEventPath(WidgetsUnderCursor);
-
-		// If we are doing a drag and drop, we will send this event instead.
-		FDragDropEvent DragDropEvent( MouseEvent, DragDropContent );
 
 		// Cache the drag drop content and reset the pointer in case OnMouseButtonUpMessage re-enters as a result of OnDrop
 		const bool bIsDragDropping = IsDragDropping();
@@ -3600,30 +3676,30 @@ bool FSlateApplication::ProcessMouseButtonUpEvent( FPointerEvent& MouseEvent )
 
 		// Switch worlds widgets in the current path
 		FScopedSwitchWorldHack SwitchWorld( WidgetsUnderCursor );
-
-		FReply Reply = FReply::Unhandled();
-		for( int32 WidgetIndex = WidgetsUnderCursor.Widgets.Num()-1; !Reply.IsEventHandled() && WidgetIndex >= 0; --WidgetIndex )
+		
+		FReply Reply = FEventRouter::BubbleAlongPointerPath( this, WidgetsUnderCursor, MouseEvent, [&](const FArrangedWidget& CurWidget, const FPointerEvent& Event)
 		{
-			FArrangedWidget& CurWidget = WidgetsUnderCursor.Widgets[WidgetIndex];
-			if (MouseEvent.IsTouchEvent())
+			FReply TempReply = FReply::Unhandled();
+			if (Event.IsTouchEvent())
 			{
-				Reply = CurWidget.Widget->OnTouchEnded( CurWidget.Geometry, MouseEvent ).SetHandler( CurWidget.Widget );
+				TempReply = CurWidget.Widget->OnTouchEnded( CurWidget.Geometry, Event );
 			}
-			if (!MouseEvent.IsTouchEvent() || (!Reply.IsEventHandled() && bTouchFallbackToMouse))
+			if (!Event.IsTouchEvent() || (!TempReply.IsEventHandled() && bTouchFallbackToMouse))
 			{
-				Reply = (bIsDragDropping)
-					? CurWidget.Widget->OnDrop( CurWidget.Geometry, DragDropEvent ).SetHandler( CurWidget.Widget )
-					: CurWidget.Widget->OnMouseButtonUp( CurWidget.Geometry, MouseEvent ).SetHandler( CurWidget.Widget );
+				TempReply = (bIsDragDropping)
+					? CurWidget.Widget->OnDrop( CurWidget.Geometry, FDragDropEvent( Event, LocalDragDropContent ) )
+					: CurWidget.Widget->OnMouseButtonUp( CurWidget.Geometry, Event );
 			}
-
-			ProcessReply(WidgetsUnderCursor, Reply, &WidgetsUnderCursor, &MouseEvent);
-		}
+			return TempReply;		
+		});
 
 		LOG_EVENT( bIsDragDropping ? EEventLog::DragDrop : EEventLog::MouseButtonUp, Reply );
 
 		// If we were dragging, notify the content
 		if ( bIsDragDropping )
 		{
+			// @todo slate : depending on SetEventPath() is not ideal.
+			MouseEvent.SetEventPath( WidgetsUnderCursor );
 			LocalDragDropContent->OnDrop( Reply.IsEventHandled(), MouseEvent );
 		}
 	}
@@ -3683,29 +3759,24 @@ bool FSlateApplication::ProcessMouseWheelOrGestureEvent( FPointerEvent& InWheelE
 	// Switch worlds widgets in the current path
 	FScopedSwitchWorldHack SwitchWorld( EventPath );
 
-	TSharedPtr<SWidget> WidgetToLog;
-
-	FReply Reply = FReply::Unhandled();
-	for( int32 WidgetIndex = EventPath.Widgets.Num()-1; !Reply.IsEventHandled() && WidgetIndex >= 0; --WidgetIndex )
+	FReply Reply = FEventRouter::BubbleAlongPointerPath( this, EventPath, InWheelEvent, [&InGestureEvent](const FArrangedWidget& CurWidget, const FPointerEvent& Event)
 	{
-		FArrangedWidget& CurWidget = EventPath.Widgets[WidgetIndex];
+		FReply TempReply = FReply::Unhandled();
 		// Gesture event gets first shot, if slate doesn't respond to it, we'll try the wheel event.
 		if( InGestureEvent!=NULL )
 		{
-			Reply = CurWidget.Widget->OnTouchGesture( CurWidget.Geometry, *InGestureEvent ).SetHandler( CurWidget.Widget );
-			ProcessReply(EventPath, Reply, &EventPath, InGestureEvent);
+			TempReply = CurWidget.Widget->OnTouchGesture( CurWidget.Geometry, *InGestureEvent );
 		}
 		
 		// Send the mouse wheel event if we haven't already handled the gesture version of this event.
-		if( !Reply.IsEventHandled() )
+		if( !TempReply.IsEventHandled() )
 		{
-			Reply = CurWidget.Widget->OnMouseWheel( CurWidget.Geometry, InWheelEvent ).SetHandler( CurWidget.Widget );
-			ProcessReply(EventPath, Reply, &EventPath, &InWheelEvent);
+			TempReply = CurWidget.Widget->OnMouseWheel( CurWidget.Geometry, Event );
 		}
 
-		WidgetToLog = CurWidget.Widget;
-	}
-
+		return TempReply;
+	} );
+	
 	LOG_EVENT( bIsGestureEvent ? EEventLog::TouchGesture : EEventLog::MouseWheel, Reply );
 
 	return Reply.IsEventHandled();
@@ -3885,35 +3956,33 @@ bool FSlateApplication::ProcessMouseMoveEvent( FPointerEvent& MouseEvent, bool b
 	{
 		if ( !bIsSynthetic )
 		{
-			FArrangedWidget& MouseCaptorWidget = MouseCaptorPath.Widgets.Last();
-			MouseEvent.SetEventPath(MouseCaptorPath);
-
 			// Switch worlds widgets in the current path
-			FScopedSwitchWorldHack SwitchWorld(MouseCaptorPath);
-
-			FReply Reply = FReply::Unhandled();
-			if (MouseEvent.IsTouchEvent())
-			{
-				Reply = MouseCaptorWidget.Widget->OnTouchMoved(MouseCaptorWidget.Geometry, MouseEvent).SetHandler(MouseCaptorWidget.Widget);
-			}
-			if (!MouseEvent.IsTouchEvent() || (!Reply.IsEventHandled() && bTouchFallbackToMouse))
+			FScopedSwitchWorldHack SwitchWorld( MouseCaptorPath );
+			
+			FReply Reply = FEventRouter::SendToLeafmost( this, MouseCaptorPath, nullptr, MouseEvent, [this]( const FArrangedWidget& MouseCaptorWidget, const FPointerEvent& Event )
 			{
-				Reply = MouseCaptorWidget.Widget->OnMouseMove(MouseCaptorWidget.Geometry, MouseEvent).SetHandler(MouseCaptorWidget.Widget);
-			}
-			ProcessReply(MouseCaptorPath, Reply, &MouseCaptorPath, &MouseEvent);
+				FReply TempReply = FReply::Unhandled();
+				if (Event.IsTouchEvent())
+				{
+					TempReply = MouseCaptorWidget.Widget->OnTouchMoved( MouseCaptorWidget.Geometry, Event );
+				}
+				if (!Event.IsTouchEvent() || (!TempReply.IsEventHandled() && this->bTouchFallbackToMouse))
+				{
+					TempReply = MouseCaptorWidget.Widget->OnMouseMove( MouseCaptorWidget.Geometry, Event );
+				}
+				return TempReply;
+			} );
 			bHandled = Reply.IsEventHandled();
 		}
 	}
 	else
 	{	
-		FReply Reply = FReply::Unhandled();
-
-		MouseEvent.SetEventPath(WidgetsUnderCursor);
-		FDragDropEvent DragDropEvent( MouseEvent, DragDropContent );
 		// Switch worlds widgets in the current path
 		FScopedSwitchWorldHack SwitchWorld( WidgetsUnderCursor );
-
+		
 		// Send out mouse enter events.
+		FDragDropEvent DragDropEvent( MouseEvent, DragDropContent );
+		DragDropEvent.SetEventPath( WidgetsUnderCursor );
 		for( int32 WidgetIndex = WidgetsUnderCursor.Widgets.Num()-1; WidgetIndex >= 0; --WidgetIndex )
 		{
 			FArrangedWidget& SomeWidgetUnderCursor = WidgetsUnderCursor.Widgets[ WidgetIndex ];
@@ -3936,27 +4005,27 @@ bool FSlateApplication::ProcessMouseMoveEvent( FPointerEvent& MouseEvent, bool b
 		}
 
 		// Bubble the MouseMove event.
-		for( int32 WidgetIndex = WidgetsUnderCursor.Widgets.Num()-1; !Reply.IsEventHandled() && WidgetIndex >= 0; --WidgetIndex )
+		FReply Reply = FEventRouter::BubbleAlongPointerPath(this, WidgetsUnderCursor, MouseEvent, [&](const FArrangedWidget& CurWidget, const FPointerEvent& Event)
 		{
-			FArrangedWidget& CurWidget = WidgetsUnderCursor.Widgets[WidgetIndex];
+			FReply TempReply = FReply::Unhandled();
 
-			if (MouseEvent.IsTouchEvent())
+			if (Event.IsTouchEvent())
 			{
-				Reply = CurWidget.Widget->OnTouchMoved( CurWidget.Geometry, MouseEvent ).SetHandler( CurWidget.Widget );
+				TempReply = CurWidget.Widget->OnTouchMoved( CurWidget.Geometry, Event );
 			}
-			if (!Reply.IsEventHandled())
+			if (!TempReply.IsEventHandled())
 			{
-				Reply = (IsDragDropping())
-					? CurWidget.Widget->OnDragOver( CurWidget.Geometry, DragDropEvent ).SetHandler( CurWidget.Widget )
-					: CurWidget.Widget->OnMouseMove( CurWidget.Geometry, MouseEvent ).SetHandler( CurWidget.Widget );
+				TempReply = (IsDragDropping())
+					? CurWidget.Widget->OnDragOver( CurWidget.Geometry, FDragDropEvent( Event, DragDropContent ) )
+					: CurWidget.Widget->OnMouseMove( CurWidget.Geometry, Event );
 			}
 
-			ProcessReply(WidgetsUnderCursor, Reply, &WidgetsUnderCursor, &MouseEvent);
-		}
+			return TempReply;
+		});
 
 		LOG_EVENT( IsDragDropping() ? EEventLog::DragOver : EEventLog::MouseMove, Reply )
 
-			bHandled = Reply.IsEventHandled();
+		bHandled = Reply.IsEventHandled();
 	}
 
 	// Give the current drag drop operation a chance to do something
diff --git a/Engine/Source/Runtime/Slate/Public/Framework/Application/IWidgetReflector.h b/Engine/Source/Runtime/Slate/Public/Framework/Application/IWidgetReflector.h
--- a/Engine/Source/Runtime/Slate/Public/Framework/Application/IWidgetReflector.h
+++ b/Engine/Source/Runtime/Slate/Public/Framework/Application/IWidgetReflector.h
@@ -3,12 +3,18 @@
 #pragma once
 
 
+/** A Delegate for passing along a string of a source code location to access */
+DECLARE_DELEGATE_RetVal_ThreeParams(bool, FAccessSourceCode, const FString& /*FileName*/, int32 /*InLineNumber*/, int32 /*InColumnNumber*/);
+
 
 /**
  * Interface for widget reflectors.
  */
 class IWidgetReflector
 {
+public:
+	virtual void OnEventProcessed( const FInputEvent& Event, const FReply& InReply ) = 0;
+
 public:
 
 	/** Called when the user has picked a widget to observe. */
diff --git a/Engine/Source/Runtime/Slate/Public/Framework/Application/SlateApplication.h b/Engine/Source/Runtime/Slate/Public/Framework/Application/SlateApplication.h
--- a/Engine/Source/Runtime/Slate/Public/Framework/Application/SlateApplication.h
+++ b/Engine/Source/Runtime/Slate/Public/Framework/Application/SlateApplication.h
@@ -628,6 +628,8 @@ public:
 
 protected:
 
+	friend class FEventRouter;
+
 	virtual bool HasMouseCapture(const TSharedPtr<const SWidget> Widget) const override;
 
 	/** 
diff --git a/Engine/Source/Runtime/SlateCore/Private/Input/Events.cpp b/Engine/Source/Runtime/SlateCore/Private/Input/Events.cpp
--- a/Engine/Source/Runtime/SlateCore/Private/Input/Events.cpp
+++ b/Engine/Source/Runtime/SlateCore/Private/Input/Events.cpp
@@ -8,3 +8,29 @@
 
 const FTouchKeySet FTouchKeySet::StandardSet(EKeys::LeftMouseButton);
 const FTouchKeySet FTouchKeySet::EmptySet(EKeys::Invalid);
+
+FText FInputEvent::ToText() const
+{
+	return NSLOCTEXT("Events", "Unimplemented", "Unimplemented");
+}
+
+FText FCharacterEvent::ToText() const
+{
+	return FText::Format( NSLOCTEXT("Events","Char","Char({0})"), FText::FromString(FString(1, &Character)) );
+}
+
+FText FControllerEvent::ToText() const
+{
+	return NSLOCTEXT("Events","Gamepad","Gamepad()");
+}
+
+FText FKeyboardEvent::ToText() const
+{
+	return FText::Format( NSLOCTEXT("Events","Key","Key({0})"), Key.GetDisplayName() );
+}
+
+FText FPointerEvent::ToText() const
+{
+	return FText::Format( NSLOCTEXT("Events","Pointer","Pointer({0}, {1})"), EffectingButton.GetDisplayName() );
+}
+
diff --git a/Engine/Source/Runtime/SlateCore/Public/Input/Events.h b/Engine/Source/Runtime/SlateCore/Public/Input/Events.h
--- a/Engine/Source/Runtime/SlateCore/Public/Input/Events.h
+++ b/Engine/Source/Runtime/SlateCore/Public/Input/Events.h
@@ -277,6 +277,8 @@ public:
 		EventPath = &InEventPath;
 	}
 
+	SLATECORE_API virtual FText ToText() const;
+
 protected:
 
 	// State of modifier keys when this event happened.
@@ -356,6 +358,8 @@ public:
 		return KeyCode;
 	}
 
+	SLATECORE_API virtual FText ToText() const override;
+
 private:
 
 	// Name of the key that was pressed.
@@ -403,6 +407,8 @@ public:
 		return Character;
 	}
 
+	SLATECORE_API virtual FText ToText() const override;	
+
 private:
 
 	// The character that was pressed.
@@ -611,6 +617,8 @@ public:
 		WheelOrGestureDelta = Other.WheelOrGestureDelta;
 	}
 
+	SLATECORE_API virtual FText ToText() const override;
+
 private:
 
 	FVector2D ScreenSpacePosition;
@@ -665,6 +673,8 @@ public:
 	/** @return Analog value between 0 and 1.  1 being fully pressed, 0 being not pressed at all */
 	float GetAnalogValue() const { return AnalogValue; }
 
+	SLATECORE_API virtual FText ToText() const override;
+
 private:
 
 	// The controller button that caused this event.
diff --git a/Engine/Source/Runtime/SlateCore/Public/Input/Reply.h b/Engine/Source/Runtime/SlateCore/Public/Input/Reply.h
--- a/Engine/Source/Runtime/SlateCore/Public/Input/Reply.h
+++ b/Engine/Source/Runtime/SlateCore/Public/Input/Reply.h
@@ -135,13 +135,6 @@ public:
 		return *this;
 	}
 
-	/** Set the widget that handled the event; undefined if never handled. This method is to be used by SlateApplication only! */
-	FReply& SetHandler( const TSharedRef<SWidget>& InHandler )
-	{
-		this->EventHandler = InHandler;
-		return *this;
-	}
-		
 	/** Ensures throttling for Slate UI responsiveness is not done on mouse down */
 	FReply& PreventThrottling()
 	{
@@ -251,6 +244,19 @@ private:
 		, bEndDragDrop(false)
 	{ }
 		
+
+private:
+
+	friend class FEventRouter;
+	friend class FSlateApplication;
+
+	/** Set the widget that handled the event; undefined if never handled. This method is to be used by SlateApplication only! */
+	FReply& SetHandler( const TSharedRef<SWidget>& InHandler )
+	{
+		this->EventHandler = InHandler;
+		return *this;
+	}
+
 private:
 
 	TOptional<FIntPoint> RequestedMousePos;
