Exposing overloads that allow external code to specify which GenericApplication Slate should interface with.

diff --git a/Engine/Source/Runtime/Slate/Private/Framework/Application/SlateApplication.cpp b/Engine/Source/Runtime/Slate/Private/Framework/Application/SlateApplication.cpp
--- a/Engine/Source/Runtime/Slate/Private/Framework/Application/SlateApplication.cpp
+++ b/Engine/Source/Runtime/Slate/Private/Framework/Application/SlateApplication.cpp
@@ -427,6 +427,12 @@ void FPopupSupport::SendNotifications( const FWidgetPath& WidgetsUnderCursor )
 
 
 void FSlateApplication::Create()
+{
+	Create(MakeShareable(FPlatformMisc::CreateApplication()));
+}
+
+
+TSharedRef<FSlateApplication> FSlateApplication::Create(const TSharedRef<class GenericApplication>& InPlatformApplication)
 {
 	EKeys::Initialize();
 
@@ -435,7 +441,7 @@ void FSlateApplication::Create()
 	CurrentApplication = MakeShareable( new FSlateApplication() );
 	CurrentBaseApplication = CurrentApplication;
 
-	PlatformApplication = MakeShareable( FPlatformMisc::CreateApplication() );
+	PlatformApplication = InPlatformApplication;
 	PlatformApplication->SetMessageHandler( CurrentApplication.ToSharedRef() );
 
 	// The grid needs to know the size and coordinate system of the desktop.
@@ -454,6 +460,8 @@ void FSlateApplication::Create()
 		// Sign up for updates from the OS. Polling this every frame is too expensive on at least some OSs.
 		PlatformApplication->OnDisplayMetricsChanged().AddSP(CurrentApplication.ToSharedRef(), &FSlateApplication::OnVirtualDesktopSizeChanged);
 	}
+
+	return CurrentApplication.ToSharedRef();
 }
 
 void FSlateApplication::Shutdown()
@@ -3180,7 +3188,7 @@ bool FSlateApplication::ProcessKeyCharEvent( FCharacterEvent& InCharacterEvent )
 		InCharacterEvent.SetEventPath( EventPath );
 		
 		Reply = FEventRouter::BubbleAlongFocusPath( this, EventPath, InCharacterEvent, []( const FArrangedWidget& SomeWidgetGettingEvent, const FCharacterEvent& Event )
-		{
+			{
 			return SomeWidgetGettingEvent.Widget->IsEnabled()
 				? SomeWidgetGettingEvent.Widget->OnKeyChar( SomeWidgetGettingEvent.Geometry, Event )
 				: FReply::Unhandled();
@@ -3226,12 +3234,12 @@ bool FSlateApplication::ProcessKeyDownEvent( FKeyboardEvent& InKeyboardEvent )
 				const bool bIsWidgetReflectorPicking = WidgetReflector.IsValid() && WidgetReflector->IsInPickingMode();
 				if ( bIsWidgetReflectorPicking )
 				{
-					WidgetReflector->OnWidgetPicked();
-					Reply = FReply::Handled();
+						WidgetReflector->OnWidgetPicked();
+						Reply = FReply::Handled();
 
-					return Reply.IsEventHandled();
+						return Reply.IsEventHandled();
+					}
 				}
-			}
 
 			// Bubble the keyboard event
 			FWidgetPath EventPath = FocusedWidgetPath.ToWidgetPath();
@@ -3433,16 +3441,16 @@ bool FSlateApplication::ProcessMouseButtonDownEvent( const TSharedPtr< FGenericW
 			{
 				Reply = FEventRouter::SendToLeafmost( this, MouseCaptorPath, nullptr, MouseEvent,
 					[this]( const FArrangedWidget& MouseCaptorWidget, const FPointerEvent& Event )
-					{
+				{
 						FReply TempReply = FReply::Unhandled();
 						if (Event.IsTouchEvent())
 						{
 							TempReply = MouseCaptorWidget.Widget->OnTouchStarted( MouseCaptorWidget.Geometry, Event );
-						}
+				}
 						if (!Event.IsTouchEvent() || (!TempReply.IsEventHandled() && this->bTouchFallbackToMouse))
-						{
+				{
 							TempReply = MouseCaptorWidget.Widget->OnMouseButtonDown( MouseCaptorWidget.Geometry, Event );
-						}
+				}
 						return TempReply;
 					} );
 
@@ -3472,19 +3480,19 @@ bool FSlateApplication::ProcessMouseButtonDownEvent( const TSharedPtr< FGenericW
 			if ( !Reply.IsEventHandled() )
 			{
 				Reply = FEventRouter::BubbleAlongPointerPath(this, WidgetsUnderCursor, MouseEvent, [this](const FArrangedWidget TargetWidget, const FPointerEvent& Event)
-				{
+			{
 					FReply ThisReply = FReply::Unhandled();
 					if (!ThisReply.IsEventHandled())
-					{
+				{
 						if (Event.IsTouchEvent())
-						{
+					{
 							ThisReply = TargetWidget.Widget->OnTouchStarted( TargetWidget.Geometry, Event );
-						}
+					}
 						if (!Event.IsTouchEvent() || (!ThisReply.IsEventHandled() && this->bTouchFallbackToMouse))
-						{
+					{
 							ThisReply = TargetWidget.Widget->OnMouseButtonDown( TargetWidget.Geometry, Event );
-						}
 					}
+				}
 					return ThisReply;
 				});
 			}
@@ -3644,13 +3652,13 @@ bool FSlateApplication::ProcessMouseButtonUpEvent( FPointerEvent& MouseEvent )
 			{
 				FReply TempReply = FReply::Unhandled();
 				if (Event.IsTouchEvent())
-				{
+			{
 					TempReply = TargetWidget.Widget->OnTouchEnded( TargetWidget.Geometry, Event );
-				}
+			}
 				if (!Event.IsTouchEvent() || (!TempReply.IsEventHandled() && this->bTouchFallbackToMouse))
-				{
+			{
 					TempReply = TargetWidget.Widget->OnMouseButtonUp( TargetWidget.Geometry, Event );
-				}
+			}
 				return TempReply;
 			});
 
@@ -3676,7 +3684,7 @@ bool FSlateApplication::ProcessMouseButtonUpEvent( FPointerEvent& MouseEvent )
 
 		// Switch worlds widgets in the current path
 		FScopedSwitchWorldHack SwitchWorld( WidgetsUnderCursor );
-		
+
 		FReply Reply = FEventRouter::BubbleAlongPointerPath( this, WidgetsUnderCursor, MouseEvent, [&](const FArrangedWidget& CurWidget, const FPointerEvent& Event)
 		{
 			FReply TempReply = FReply::Unhandled();
@@ -3689,7 +3697,7 @@ bool FSlateApplication::ProcessMouseButtonUpEvent( FPointerEvent& MouseEvent )
 				TempReply = (bIsDragDropping)
 					? CurWidget.Widget->OnDrop( CurWidget.Geometry, FDragDropEvent( Event, LocalDragDropContent ) )
 					: CurWidget.Widget->OnMouseButtonUp( CurWidget.Geometry, Event );
-			}
+		}
 			return TempReply;		
 		});
 
@@ -3776,7 +3784,7 @@ bool FSlateApplication::ProcessMouseWheelOrGestureEvent( FPointerEvent& InWheelE
 
 		return TempReply;
 	} );
-	
+
 	LOG_EVENT( bIsGestureEvent ? EEventLog::TouchGesture : EEventLog::MouseWheel, Reply );
 
 	return Reply.IsEventHandled();
@@ -3956,30 +3964,30 @@ bool FSlateApplication::ProcessMouseMoveEvent( FPointerEvent& MouseEvent, bool b
 	{
 		if ( !bIsSynthetic )
 		{
-			// Switch worlds widgets in the current path
-			FScopedSwitchWorldHack SwitchWorld( MouseCaptorPath );
-			
+		// Switch worlds widgets in the current path
+		FScopedSwitchWorldHack SwitchWorld( MouseCaptorPath );
+
 			FReply Reply = FEventRouter::SendToLeafmost( this, MouseCaptorPath, nullptr, MouseEvent, [this]( const FArrangedWidget& MouseCaptorWidget, const FPointerEvent& Event )
 			{
 				FReply TempReply = FReply::Unhandled();
 				if (Event.IsTouchEvent())
-				{
+		{
 					TempReply = MouseCaptorWidget.Widget->OnTouchMoved( MouseCaptorWidget.Geometry, Event );
-				}
+		}
 				if (!Event.IsTouchEvent() || (!TempReply.IsEventHandled() && this->bTouchFallbackToMouse))
-				{
+		{
 					TempReply = MouseCaptorWidget.Widget->OnMouseMove( MouseCaptorWidget.Geometry, Event );
-				}
+		}
 				return TempReply;
 			} );
-			bHandled = Reply.IsEventHandled();
-		}
+		bHandled = Reply.IsEventHandled();
+	}
 	}
 	else
 	{	
 		// Switch worlds widgets in the current path
 		FScopedSwitchWorldHack SwitchWorld( WidgetsUnderCursor );
-		
+
 		// Send out mouse enter events.
 		FDragDropEvent DragDropEvent( MouseEvent, DragDropContent );
 		DragDropEvent.SetEventPath( WidgetsUnderCursor );
@@ -4025,7 +4033,7 @@ bool FSlateApplication::ProcessMouseMoveEvent( FPointerEvent& MouseEvent, bool b
 
 		LOG_EVENT( IsDragDropping() ? EEventLog::DragOver : EEventLog::MouseMove, Reply )
 
-		bHandled = Reply.IsEventHandled();
+			bHandled = Reply.IsEventHandled();
 	}
 
 	// Give the current drag drop operation a chance to do something
@@ -4819,16 +4827,24 @@ void FSlateApplication::OnVirtualDesktopSizeChanged(const FDisplayMetrics& NewDi
 /* 
  *****************************************************************************/
 
-void FSlateApplication::InitializeAsStandaloneApplication( const TSharedRef<FSlateRenderer>& PlatformRenderer)
+TSharedRef<FSlateApplication> FSlateApplication::InitializeAsStandaloneApplication(const TSharedRef<FSlateRenderer>& PlatformRenderer)
+{
+	return InitializeAsStandaloneApplication(PlatformRenderer, MakeShareable(FPlatformMisc::CreateApplication()));
+}
+
+
+TSharedRef<FSlateApplication> FSlateApplication::InitializeAsStandaloneApplication(const TSharedRef< class FSlateRenderer >& PlatformRenderer, const TSharedRef<class GenericApplication>& PlatformApplication)
 {
 	// create the platform slate application (what FSlateApplication::Get() returns)
-	FSlateApplication::Create();
+	TSharedRef<FSlateApplication> Slate = FSlateApplication::Create(PlatformApplication);
 
 	// initialize renderer
 	FSlateApplication::Get().InitializeRenderer(PlatformRenderer);
 
 	// set the normal UE4 GIsRequestingExit when outer frame is closed
 	FSlateApplication::Get().SetExitRequestedHandler(FSimpleDelegate::CreateStatic(&OnRequestExit));
+
+	return Slate;
 }
 
 void FSlateApplication::SetWidgetReflector(const TSharedRef<IWidgetReflector>& WidgetReflector)
diff --git a/Engine/Source/Runtime/Slate/Public/Framework/Application/SlateApplication.h b/Engine/Source/Runtime/Slate/Public/Framework/Application/SlateApplication.h
--- a/Engine/Source/Runtime/Slate/Public/Framework/Application/SlateApplication.h
+++ b/Engine/Source/Runtime/Slate/Public/Framework/Application/SlateApplication.h
@@ -136,8 +136,11 @@ public:
 public:
 
 	static void Create();
+	static TSharedRef<FSlateApplication> Create(const TSharedRef<class GenericApplication>& InPlatformApplication);
 
-	static void InitializeAsStandaloneApplication( const TSharedRef< class FSlateRenderer >& PlatformRenderer);
+	static TSharedRef<FSlateApplication> InitializeAsStandaloneApplication(const TSharedRef< class FSlateRenderer >& PlatformRenderer);
+	
+	static TSharedRef<FSlateApplication> InitializeAsStandaloneApplication(const TSharedRef< class FSlateRenderer >& PlatformRenderer, const TSharedRef<class GenericApplication>& PlatformApplication);
 
 	/**
 	 * Returns true if a Slate application instance is currently initialized and ready
