Adding the new still in development/experimental Responsive Layout widget.

diff --git a/Engine/Source/Runtime/AppFramework/Private/Framework/Testing/TestStyle.cpp b/Engine/Source/Runtime/AppFramework/Private/Framework/Testing/TestStyle.cpp
--- a/Engine/Source/Runtime/AppFramework/Private/Framework/Testing/TestStyle.cpp
+++ b/Engine/Source/Runtime/AppFramework/Private/Framework/Testing/TestStyle.cpp
@@ -81,6 +81,10 @@ TSharedRef< ISlateStyle > FTestStyle::Create()
 
 	Style->Set( "NewLevelBlank", new IMAGE_BRUSH( "Testing/NewLevelBlank", FVector2D(256,256) ) );
 
+	Style->Set("White", new BOX_BRUSH("Testing/FlatColorSquare", FVector2D(1.0f, 1.0f), FMargin(1), FLinearColor::White));
+	Style->Set("Gray", new BOX_BRUSH("Testing/FlatColorSquare", FVector2D(1.0f, 1.0f), FMargin(1), FLinearColor::Gray));
+	Style->Set("Black", new BOX_BRUSH("Testing/FlatColorSquare", FVector2D(1.0f, 1.0f), FMargin(1), FLinearColor::Black));
+
 	// Testing images in the multi box test
 	{
 		Style->Set( "MultiBoxTest.FirstCommandInfo", new IMAGE_BRUSH( "Testing/BrushWireframe", Icon16x16 ) );
@@ -138,7 +142,8 @@ TSharedRef< ISlateStyle > FTestStyle::Create()
 
 		//Tagline
 		{
-			Style->Set( "RichText.Tagline.Background", new BOX_BRUSH( "Testing/FlatColorSquare", FVector2D(1.0f, 1.0f), FMargin(1), FLinearColor(FColor(0xffdbe4e4)) ) );
+			Style->Set("RichText.Tagline.Background", new BOX_BRUSH("Testing/FlatColorSquare", FVector2D(1.0f, 1.0f), FMargin(1), FLinearColor(FColor(0xffdbe4e4))));
+			Style->Set("RichText.Tagline.DarkBackground", new BOX_BRUSH("Testing/FlatColorSquare", FVector2D(1.0f, 1.0f), FMargin(1), FLinearColor(0.55423, 0.60548, 0.60548)));
 			Style->Set( "RichText.Tagline.Text", FTextBlockStyle(NormalText)
 				.SetFont( TTF_FONT("Fonts/Roboto-Bold", 24) )
 				.SetColorAndOpacity(FLinearColor(FColor( 0xff2c3e50 )))
diff --git a/Engine/Source/Runtime/AppFramework/Private/Widgets/Testing/STestSuite.cpp b/Engine/Source/Runtime/AppFramework/Private/Widgets/Testing/STestSuite.cpp
--- a/Engine/Source/Runtime/AppFramework/Private/Widgets/Testing/STestSuite.cpp
+++ b/Engine/Source/Runtime/AppFramework/Private/Widgets/Testing/STestSuite.cpp
@@ -21,6 +21,7 @@
 #include "SDockTab.h"
 #include "SSearchBox.h"
 #include "SVolumeControl.h"
+#include "SResponsiveGridPanel.h"
 
 #define LOCTEXT_NAMESPACE "STestSuite"
 
@@ -4298,6 +4299,353 @@ public:
 	END_SLATE_FUNCTION_BUILD_OPTIMIZATION
 };
 
+class SResponsiveGridPanelTestWidget : public SUserWidget
+{
+
+public:
+
+	SLATE_USER_ARGS(SResponsiveGridPanelTestWidget)
+	{}
+
+	SLATE_END_ARGS()
+
+	virtual void Construct(const FArguments& InArgs) = 0;
+
+};
+
+namespace
+{
+	class SResponsiveGridPanelTestWidgetImpl : public SResponsiveGridPanelTestWidget
+	{
+	public:
+		virtual void Construct(const FArguments& InArgs) override;
+
+	private:
+		TSharedRef<SWidget> ConstructBox(const FString& Text) const;
+	};
+}
+
+TSharedRef<SWidget> SResponsiveGridPanelTestWidgetImpl::ConstructBox(const FString& Text) const
+{
+	return
+		SNew(SBorder)
+		.Padding(1)
+		.BorderImage(FTestStyle::Get().GetBrush("Gray"))
+		[
+			SNew(SBorder)
+			.HAlign(HAlign_Center)
+			.VAlign(VAlign_Center)
+			.Padding(10)
+			.BorderImage(FTestStyle::Get().GetBrush("RichText.Tagline.Background"))
+			[
+				SNew(STextBlock).Text(Text)
+			]
+		];
+}
+
+void SResponsiveGridPanelTestWidgetImpl::Construct(const FArguments& InArgs)
+{
+	TSharedPtr<SResponsiveGridPanel> GridPanel;
+
+	SUserWidget::Construct(SUserWidget::FArguments()
+	[
+		SNew(SBorder)
+		.BorderImage(FTestStyle::Get().GetBrush("RichText.Background"))
+		.ForegroundColor(FLinearColor::Black)
+		[
+			SNew(SScrollBox)
+			+ SScrollBox::Slot()
+			[
+				SNew(SBox)
+				.Padding(FMargin(50, 10))
+				[
+					SNew(SBorder)
+					.BorderImage(FTestStyle::Get().GetBrush("RichText.Tagline.DarkBackground"))
+					.Padding(10.0f)
+					[
+						SAssignNew(GridPanel, SResponsiveGridPanel, 12)
+						.ColumnGutter(5)
+						.RowGutter(5)
+					]
+				]
+			]
+
+			+ SScrollBox::Slot()
+			.Padding(0, 20, 0, 0)
+			[
+				SNew(SBox)
+				.Padding(FMargin(50, 10))
+				[
+					SNew(SBorder)
+					.BorderImage(FTestStyle::Get().GetBrush("RichText.Tagline.DarkBackground"))
+					.Padding(10.0f)
+					[
+						SNew(SResponsiveGridPanel, 12)
+						.ColumnGutter(5)
+						.RowGutter(5)
+						+ SResponsiveGridPanel::Slot(0)
+						.ColumnSpan(SResponsiveGridSize::MediumDevice, 8)
+						.ColumnSpan(SResponsiveGridSize::Mobile, 12)
+						[
+							ConstructBox(TEXT(".mobile-col-12 .desktop-col-8"))
+						]
+						+ SResponsiveGridPanel::Slot(0)
+						.ColumnSpan(SResponsiveGridSize::MediumDevice, 4)
+						.ColumnSpan(SResponsiveGridSize::Mobile, 6)
+						[
+							ConstructBox(TEXT(".mobile-col-6 .desktop-col-4"))
+						]
+						+ SResponsiveGridPanel::Slot(1)
+						.ColumnSpan(SResponsiveGridSize::MediumDevice, 4)
+						.ColumnSpan(SResponsiveGridSize::Mobile, 6)
+						[
+							ConstructBox(TEXT(".mobile-col-6 .desktop-col-4"))
+						]
+						+ SResponsiveGridPanel::Slot(1)
+						.ColumnSpan(SResponsiveGridSize::MediumDevice, 4)
+						.ColumnSpan(SResponsiveGridSize::Mobile, 6)
+						[
+							ConstructBox(TEXT(".mobile-col-6 .desktop-col-4"))
+						]
+						+ SResponsiveGridPanel::Slot(1)
+						.ColumnSpan(SResponsiveGridSize::MediumDevice, 4)
+						.ColumnSpan(SResponsiveGridSize::Mobile, 6)
+						[
+							ConstructBox(TEXT(".mobile-col-6 .desktop-col-4"))
+						]
+						+ SResponsiveGridPanel::Slot(2)
+						.ColumnSpan(SResponsiveGridSize::Mobile, 6)
+						[
+							ConstructBox(TEXT(".mobile-col-6"))
+						]
+						+ SResponsiveGridPanel::Slot(2)
+						.ColumnSpan(SResponsiveGridSize::Mobile, 6)
+						[
+							ConstructBox(TEXT(".mobile-col-6"))
+						]
+					]
+				]
+			]
+
+			+ SScrollBox::Slot()
+			.Padding(0, 20, 0, 0)
+			[
+				SNew(SBox)
+				.Padding(FMargin(50, 10))
+				[
+					SNew(SBorder)
+					.BorderImage(FTestStyle::Get().GetBrush("RichText.Tagline.DarkBackground"))
+					.Padding(10.0f)
+					[
+						SNew(SResponsiveGridPanel, 12)
+						.ColumnGutter(5)
+						.RowGutter(5)
+						+ SResponsiveGridPanel::Slot(0)
+						.ColumnSpan(SResponsiveGridSize::MediumDevice, 8)
+						.ColumnSpan(SResponsiveGridSize::Tablet, 6)
+						.ColumnSpan(SResponsiveGridSize::Mobile, 12)
+						[
+							ConstructBox(TEXT(".mobile-col-12 .tablet-col-6 .desktop-col-8"))
+						]
+						+ SResponsiveGridPanel::Slot(0)
+						.ColumnSpan(SResponsiveGridSize::MediumDevice, 4)
+						.ColumnSpan(SResponsiveGridSize::Mobile, 6)
+						[
+							ConstructBox(TEXT(".mobile-col-6 .desktop-col-4"))
+						]
+						+ SResponsiveGridPanel::Slot(1)
+						.ColumnSpan(SResponsiveGridSize::Tablet, 4)
+						.ColumnSpan(SResponsiveGridSize::Mobile, 6)
+						[
+							ConstructBox(TEXT(".mobile-col-6 .tablet-col-4"))
+						]
+						+ SResponsiveGridPanel::Slot(1)
+						.ColumnSpan(SResponsiveGridSize::Tablet, 4)
+						.ColumnSpan(SResponsiveGridSize::Mobile, 6)
+						[
+							ConstructBox(TEXT(".mobile-col-6 .tablet-col-4"))
+						]
+						+ SResponsiveGridPanel::Slot(1)
+						.ColumnSpan(SResponsiveGridSize::Tablet, 4)
+						.ColumnSpan(SResponsiveGridSize::Mobile, 6)
+						[
+							ConstructBox(TEXT(".mobile-col-6 .tablet-col-4"))
+						]
+					]
+				]
+			]
+
+			+ SScrollBox::Slot()
+			.Padding(0, 20, 0, 0)
+			[
+				SNew(SBox)
+				.Padding(FMargin(50, 10))
+				[
+					SNew(SBorder)
+					.BorderImage(FTestStyle::Get().GetBrush("RichText.Tagline.DarkBackground"))
+					.Padding(10.0f)
+					[
+						SNew(SResponsiveGridPanel, 12)
+						.ColumnGutter(5)
+						.RowGutter(5)
+						+ SResponsiveGridPanel::Slot(0)
+						.ColumnSpan(SResponsiveGridSize::Mobile, 9)
+						[
+							ConstructBox(TEXT(".mobile-col-9"))
+						]
+						+ SResponsiveGridPanel::Slot(0)
+						.ColumnSpan(SResponsiveGridSize::Mobile, 4)
+						[
+							ConstructBox(TEXT(".mobile-col-4 \nSince 9 + 4 = 13 > 12, this 4 column wide slot\ngets wrapped onto a new line as one contiguous unit."))
+						]
+						+ SResponsiveGridPanel::Slot(0)
+						.ColumnSpan(SResponsiveGridSize::Mobile, 6)
+						[
+							ConstructBox(TEXT(".mobile-col-6 \nSubsequent columns continue along the new line."))
+						]
+					]
+				]
+			]
+
+			+ SScrollBox::Slot()
+			.Padding(0, 20, 0, 0)
+			[
+				SNew(SBox)
+				.Padding(FMargin(50, 10))
+				[
+					SNew(SBorder)
+					.BorderImage(FTestStyle::Get().GetBrush("RichText.Tagline.DarkBackground"))
+					.Padding(10.0f)
+					[
+						SNew(SResponsiveGridPanel, 12)
+						.ColumnGutter(5)
+						.RowGutter(5)
+						+ SResponsiveGridPanel::Slot(0)
+						.ColumnSpan(SResponsiveGridSize::MediumDevice, 4)
+						[
+							ConstructBox(TEXT(".desktop-col-4"))
+						]
+						+ SResponsiveGridPanel::Slot(0)
+						.ColumnSpan(SResponsiveGridSize::MediumDevice, 4, 4)
+						[
+							ConstructBox(TEXT(".desktop-col-4-offset-4"))
+						]
+						+ SResponsiveGridPanel::Slot(1)
+						.ColumnSpan(SResponsiveGridSize::MediumDevice, 3, 3)
+						[
+							ConstructBox(TEXT(".desktop-col-3-offset-3"))
+						]
+						+ SResponsiveGridPanel::Slot(1)
+						.ColumnSpan(SResponsiveGridSize::MediumDevice, 3, 3)
+						[
+							ConstructBox(TEXT(".desktop-col-3-offset-3"))
+						]
+						+ SResponsiveGridPanel::Slot(2)
+						.ColumnSpan(SResponsiveGridSize::MediumDevice, 6, 3)
+						[
+							ConstructBox(TEXT(".desktop-col-6-offset-3"))
+						]
+					]
+				]
+			]
+
+			+ SScrollBox::Slot()
+			.Padding(0, 20, 0, 0)
+			[
+				SNew(SBox)
+				.Padding(FMargin(50, 10))
+				.HeightOverride(800)
+				[
+					SNew(SBorder)
+					.BorderImage(FTestStyle::Get().GetBrush("RichText.Tagline.DarkBackground"))
+					.Padding(10.0f)
+					[
+						SNew(SResponsiveGridPanel, 12)
+						.ColumnGutter(5)
+						.RowGutter(5)
+						.FillRow(1, 1.0f)
+						+ SResponsiveGridPanel::Slot(0)
+						.ColumnSpan(SResponsiveGridSize::MediumDevice, 4)
+						[
+							ConstructBox(TEXT(".desktop-col-4"))
+						]
+						+ SResponsiveGridPanel::Slot(0)
+						.ColumnSpan(SResponsiveGridSize::MediumDevice, 4)
+						[
+							ConstructBox(TEXT(".desktop-col-4"))
+						]
+						+ SResponsiveGridPanel::Slot(0)
+						.ColumnSpan(SResponsiveGridSize::MediumDevice, 4)
+						[
+							ConstructBox(TEXT(".desktop-col-4"))
+						]
+						+ SResponsiveGridPanel::Slot(1)
+						.ColumnSpan(SResponsiveGridSize::MediumDevice, 3)
+						.ColumnSpan(SResponsiveGridSize::Mobile, 0)
+						[
+							ConstructBox(TEXT(".desktop-col-3"))
+						]
+						+ SResponsiveGridPanel::Slot(1)
+						.ColumnSpan(SResponsiveGridSize::MediumDevice, 9)
+						[
+							ConstructBox(TEXT(".desktop-col-9"))
+						]
+						+ SResponsiveGridPanel::Slot(2)
+						.ColumnSpan(SResponsiveGridSize::MediumDevice, 6, 3)
+						[
+							ConstructBox(TEXT(".desktop-col-6-offset-3"))
+						]
+					]
+				]
+			]
+		]
+	]);
+
+	for (int32 Index = 0; Index < 12; ++Index)
+	{
+		GridPanel->AddSlot(0)
+		.ColumnSpan(SResponsiveGridSize::MediumDevice, 1)
+		[
+			ConstructBox(TEXT(".desktop-col-1"))
+		];
+	}
+
+	GridPanel->AddSlot(1)
+	.ColumnSpan(SResponsiveGridSize::MediumDevice, 8)
+	[
+		ConstructBox(TEXT(".desktop-col-8"))
+	];
+
+	GridPanel->AddSlot(1)
+	.ColumnSpan(SResponsiveGridSize::MediumDevice, 4)
+	[
+		ConstructBox(TEXT(".desktop-col-4"))
+	];
+
+	for (int32 Index = 0; Index < 3; ++Index)
+	{
+		GridPanel->AddSlot(2)
+		.ColumnSpan(SResponsiveGridSize::MediumDevice, 4)
+		[
+			ConstructBox(TEXT(".desktop-col-4"))
+		];
+	}
+	
+	for (int32 Index = 0; Index < 2; ++Index)
+	{
+		GridPanel->AddSlot(3)
+		.ColumnSpan(SResponsiveGridSize::MediumDevice, 6)
+		[
+			ConstructBox(TEXT(".desktop-col-6"))
+		];
+	}
+}
+
+TSharedRef<SResponsiveGridPanelTestWidget> SResponsiveGridPanelTestWidget::New()
+{
+	return MakeShareable(new SResponsiveGridPanelTestWidgetImpl());
+}
+
 /**
  * User Widget wrapper for a this widget. Legacy as this widget used to be in a separate header.
  */
@@ -5038,6 +5386,21 @@ TSharedRef<SDockTab> SpawnRenderTransformManipulator( const FSpawnTabArgs& Args
 	return RenderTransformManipulatorTab;
 }
 
+TSharedRef<SDockTab> SpawnResponsiveGrid(const FSpawnTabArgs& Args)
+{
+	TSharedRef<SDockTab> ResponsiveGridTab =
+		SNew(SDockTab)
+		.TabRole(ETabRole::MajorTab)
+		.Label(LOCTEXT("ResponsiveGridTabLabel", "Responsive Grid"))
+		.ToolTipText(LOCTEXT("ResponsiveGridTabToolTip", ""));
+
+	ResponsiveGridTab->SetContent
+	(
+		SNew(SResponsiveGridPanelTestWidget)
+	);
+
+	return ResponsiveGridTab;
+}
 
 void RestoreSlateTestSuite()
 {
@@ -5045,7 +5408,8 @@ void RestoreSlateTestSuite()
 
 	FGlobalTabmanager::Get()->RegisterTabSpawner("TestSuite1", FOnSpawnTab::CreateStatic( &SpawnTestSuite1 ) );
 	FGlobalTabmanager::Get()->RegisterTabSpawner("TestSuite2", FOnSpawnTab::CreateStatic( &SpawnTestSuite2 ) );
-	FGlobalTabmanager::Get()->RegisterTabSpawner("RenderTrasnformManipulator", FOnSpawnTab::CreateStatic(&SpawnRenderTransformManipulator));
+	FGlobalTabmanager::Get()->RegisterTabSpawner("RenderTransformManipulator", FOnSpawnTab::CreateStatic(&SpawnRenderTransformManipulator));
+	FGlobalTabmanager::Get()->RegisterTabSpawner("ResponsiveGrid", FOnSpawnTab::CreateStatic(&SpawnResponsiveGrid));
 
 	TSharedRef<FTabManager::FLayout> Layout = FTabManager::NewLayout( "SlateTestSuite_Layout" )
 	->AddArea
@@ -5057,7 +5421,8 @@ void RestoreSlateTestSuite()
 			FTabManager::NewStack()
 			->AddTab( "TestSuite2", ETabState::OpenedTab )
 			->AddTab( "TestSuite1", ETabState::OpenedTab )
-			->AddTab( "RenderTrasnformManipulator", ETabState::OpenedTab )
+			->AddTab("RenderTransformManipulator", ETabState::OpenedTab)
+			->AddTab("ResponsiveGrid", ETabState::OpenedTab)
 		)		
 	)
 	#if PLATFORM_SUPPORTS_MULTIPLE_NATIVE_WINDOWS
diff --git a/Engine/Source/Runtime/Slate/Private/Widgets/Layout/SResponsiveGridPanel.cpp b/Engine/Source/Runtime/Slate/Private/Widgets/Layout/SResponsiveGridPanel.cpp
new file mode 100644
--- /dev/null
+++ b/Engine/Source/Runtime/Slate/Private/Widgets/Layout/SResponsiveGridPanel.cpp
@@ -0,0 +1,469 @@
+// Copyright 1998-2014 Epic Games, Inc. All Rights Reserved.
+#include "SlatePrivatePCH.h"
+#include "SResponsiveGridPanel.h"
+#include "LayoutUtils.h"
+
+/**
+ *  !!!!!!!!!!!!!!!!!   EXPERIMENTAL  !!!!!!!!!!!!!!!!!   
+ * The SResponsiveGridPanel is still in development and the API may change drastically in the future
+ * or maybe removed entirely.
+ */
+SResponsiveGridPanel::SResponsiveGridPanel()
+: Slots()
+{}
+
+SResponsiveGridPanel::FSlot& SResponsiveGridPanel::AddSlot(int32 Row)
+{
+	return InsertSlot(new FSlot(Row));
+}
+
+bool SResponsiveGridPanel::RemoveSlot(const TSharedRef<SWidget>& SlotWidget)
+{
+	for (int32 SlotIdx = 0; SlotIdx < Slots.Num(); ++SlotIdx)
+	{
+		if (SlotWidget == Slots[SlotIdx].GetWidget())
+		{
+			Slots.RemoveAt(SlotIdx);
+			return true;
+		}
+	}
+
+	return false;
+}
+
+void SResponsiveGridPanel::ClearChildren()
+{
+	Slots.Empty();
+}
+
+void SResponsiveGridPanel::Construct(const FArguments& InArgs, int32 TotalColumns)
+{
+	this->TotalColumns = TotalColumns;
+	ColumnGutter = InArgs._ColumnGutter;
+	RowGutter = InArgs._RowGutter;
+
+	PreviousWidth = 0;
+
+	RowFillCoefficients = InArgs.RowFillCoefficients;
+
+	for (int32 SlotIdx = 0; SlotIdx < InArgs.Slots.Num(); ++SlotIdx)
+	{
+		InsertSlot(InArgs.Slots[SlotIdx]);
+	}
+}
+
+int32 SResponsiveGridPanel::OnPaint( const FPaintArgs& Args, const FGeometry& AllottedGeometry, const FSlateRect& MyClippingRect, FSlateWindowElementList& OutDrawElements, int32 LayerId, const FWidgetStyle& InWidgetStyle, bool bParentEnabled ) const
+{
+	FArrangedChildren ArrangedChildren(EVisibility::All);
+	this->ArrangeChildren(AllottedGeometry, ArrangedChildren);
+
+	// Because we paint multiple children, we must track the maximum layer id that they produced in case one of our parents
+	// wants to an overlay for all of its contents.
+	int32 MaxLayerId = LayerId;
+
+	// We need to iterate over slots, because slots know the GridLayers. This isn't available in the arranged children.
+	// Some slots do not show up (they are hidden/collapsed). We need a 2nd index to skip over them.
+	 
+	for (int32 ChildIndex = 0; ChildIndex < ArrangedChildren.Num(); ++ChildIndex)
+	{
+		FArrangedWidget& CurWidget = ArrangedChildren[ChildIndex];
+		if (CurWidget.Widget->GetVisibility().IsVisible())
+		{
+			FSlateRect ChildClipRect = MyClippingRect.IntersectionWith( CurWidget.Geometry.GetClippingRect() );
+
+			const int32 CurWidgetsMaxLayerId = CurWidget.Widget->Paint(
+				Args.WithNewParent(this),
+				CurWidget.Geometry,
+				ChildClipRect,
+				OutDrawElements,
+				LayerId,
+				InWidgetStyle,
+				ShouldBeEnabled( bParentEnabled )
+			);
+			
+			MaxLayerId = FMath::Max( MaxLayerId, CurWidgetsMaxLayerId );
+		}
+	}
+
+#ifdef LAYOUT_DEBUG
+	LayerId = LayoutDebugPaint( AllottedGeometry, MyClippingRect, OutDrawElements, LayerId );
+#endif
+
+	return MaxLayerId;
+}
+
+void SResponsiveGridPanel::OnArrangeChildren( const FGeometry& AllottedGeometry, FArrangedChildren& ArrangedChildren ) const
+{
+	// PREPARE PHASE
+	// Prepare some data for arranging children.
+	// FinalColumns will be populated with column sizes that include the stretched column sizes.
+	// Then we will build partial sums so that we can easily handle column spans.
+
+	const FVector2D LocalSize = AllottedGeometry.GetLocalSize();
+	FVector2D FlexSpace = LocalSize;
+
+	PreviousWidth = LocalSize.X;
+	const float FullColumnGutter = (ColumnGutter * 2);
+	const float FullRowGutter = (RowGutter * 2);
+
+	TArray<float> Rows;
+	TArray<float> RowToSlot;
+	TArray<float> Columns;
+	ComputeDesiredCellSizes(LocalSize.X, Columns, Rows, RowToSlot);
+	check(Rows.Num() == RowToSlot.Num());
+
+	TArray<float> FinalColumns;
+	if (Columns.Num() > 0)
+	{
+		FinalColumns.AddUninitialized(Columns.Num());
+		FinalColumns[FinalColumns.Num() - 1] = 0.0f;
+
+		// We need an extra cell at the end for easily figuring out the size across any number of cells
+		FinalColumns.AddZeroed((TotalColumns + 1) - Columns.Num());
+	}
+
+	// You can't remove the gutter space from the flexspace for the columns because you don't know how 
+	// many columns there actually are in a single row at this point
+
+	float ColumnCoeffTotal = TotalColumns;
+	for (int32 ColIndex = 0; ColIndex < Columns.Num(); ++ColIndex)
+	{
+		// Figure out how big each column needs to be
+		FinalColumns[ColIndex] = (1.0f / ColumnCoeffTotal * FlexSpace.X);
+	}
+
+	// FinalColumns will be populated with column sizes that include the stretched column sizes.
+	// Then we will build partial sums so that we can easily handle column spans.
+	float RowCoeffTotal = 0.0f;
+	TArray<float> FinalRows;
+	if (Rows.Num() > 0)
+	{
+		FinalRows.AddUninitialized(Rows.Num());
+		FinalRows[FinalRows.Num() - 1] = 0.0f;
+
+		// We need an extra cell at the end for easily figuring out the size across any number of cells
+		FinalRows.AddZeroed(1);
+	}
+
+	FlexSpace.Y -= (RowGutter * 2)  * (Slots[Slots.Num() - 1].RowParam);
+
+	const int32 RowFillCoeffsLength = RowFillCoefficients.Num();
+	for (int32 RowIndex = 0; RowIndex < Rows.Num(); ++RowIndex)
+	{
+		// Compute the total space available for stretchy rows.
+		if (RowToSlot[RowIndex] >= RowFillCoeffsLength || RowFillCoefficients[RowToSlot[RowIndex]] == 0)
+		{
+			FlexSpace.Y -= Rows[RowIndex];
+		}
+		else //(RowIndex < RowFillCoeffsLength)
+		{
+			// Compute the denominator for dividing up the stretchy row space
+			RowCoeffTotal += RowFillCoefficients[RowToSlot[RowIndex]];
+		}
+	}
+
+	for (int32 RowIndex = 0; RowIndex < Rows.Num(); ++RowIndex)
+	{
+		// Compute how big each row needs to be
+		FinalRows[RowIndex] = (RowToSlot[RowIndex] < RowFillCoeffsLength && RowFillCoefficients[RowToSlot[RowIndex]] != 0)
+			? (RowFillCoefficients[RowToSlot[RowIndex]] / RowCoeffTotal * FlexSpace.Y)
+			: Rows[RowIndex];
+	}
+
+	// Build up partial sums for row and column sizes so that we can handle column and row spans conveniently.
+	ComputePartialSums(FinalColumns);
+	ComputePartialSums(FinalRows);
+
+	// ARRANGE PHASE
+	int32 ColumnsSoFar = 0;
+	int32 CurrentRow = INDEX_NONE;
+	int32 LastRowParam = INDEX_NONE;
+	float RowGuttersSoFar = 0;
+	for (int32 SlotIndex = 0; SlotIndex < Slots.Num(); ++SlotIndex)
+	{
+		const FSlot& CurSlot = Slots[SlotIndex];
+
+		const EVisibility ChildVisibility = CurSlot.GetWidget()->GetVisibility();
+		if (ArrangedChildren.Accepts(ChildVisibility))
+		{
+			// Find the appropriate column layout for the slot
+			FSlot::FColumnLayout ColumnLayout;
+			ColumnLayout.Span = TotalColumns;
+			ColumnLayout.Offset = 0;
+			for (int32 Index = CurSlot.ColumnLayouts.Num() - 1; Index >= 0; Index--)
+			{
+				if (CurSlot.ColumnLayouts[Index].LayoutSize < LocalSize.X)
+				{
+					ColumnLayout = CurSlot.ColumnLayouts[Index];
+					break;
+				}
+			}
+		
+			if (ColumnLayout.Span == 0)
+			{
+				continue;
+			}
+
+			if (CurSlot.RowParam != LastRowParam)
+			{
+				ColumnsSoFar = 0;
+				LastRowParam = CurSlot.RowParam;
+				++CurrentRow;
+
+				if (LastRowParam > 0)
+				{
+					RowGuttersSoFar += FullRowGutter;
+				}
+			}
+
+			// Figure out the position of this cell.
+			int32 StartColumn = ColumnsSoFar + ColumnLayout.Offset;
+			int32 EndColumn = StartColumn + ColumnLayout.Span;
+			ColumnsSoFar = FMath::Max(EndColumn, ColumnsSoFar);
+
+			if (ColumnsSoFar > TotalColumns)
+			{
+				StartColumn = 0;
+				EndColumn = StartColumn + ColumnLayout.Span;
+				ColumnsSoFar = EndColumn - StartColumn;
+				++CurrentRow;
+			}
+
+			FVector2D ThisCellOffset(FinalColumns[StartColumn], FinalRows[CurrentRow]);
+
+			// Account for the gutters applied to columns before the starting column of this cell
+			if (StartColumn > 0)
+			{
+				ThisCellOffset.X += FullColumnGutter;
+			}
+
+			// Figure out the size of this slot; takes row span into account.
+			// We use the properties of partial sums arrays to achieve this.
+			FVector2D CellSize(
+				FinalColumns[EndColumn] - ThisCellOffset.X,
+				FinalRows[CurrentRow + 1] - ThisCellOffset.Y);
+
+
+			// Do the standard arrangement of elements within a slot
+			// Takes care of alignment and padding.
+			FMargin SlotPadding(CurSlot.SlotPadding.Get());
+
+			AlignmentArrangeResult XAxisResult = AlignChild<Orient_Horizontal>(CellSize.X, CurSlot, SlotPadding);
+			AlignmentArrangeResult YAxisResult = AlignChild<Orient_Vertical>(CellSize.Y, CurSlot, SlotPadding);
+
+			// The row gutters have already been accounted for in the cell size by removing them from the flexspace, 
+			// so we just need to offset the cells appropriately
+			ThisCellOffset.Y += RowGuttersSoFar;
+
+			// Output the result
+			ArrangedChildren.AddWidget(ChildVisibility, AllottedGeometry.MakeChild(
+				CurSlot.GetWidget(),
+				ThisCellOffset + FVector2D(XAxisResult.Offset, YAxisResult.Offset),
+				FVector2D(XAxisResult.Size, YAxisResult.Size)
+				));
+		}
+	}
+}
+
+void SResponsiveGridPanel::CacheDesiredSize()
+{
+	// The desired size of the grid is the sum of the desires sizes for every row and column.
+	TArray<float> Columns;
+	TArray<float> Rows;
+	TArray<float> RowToSlot;
+	ComputeDesiredCellSizes(PreviousWidth, Columns, Rows, RowToSlot);
+
+	TotalDesiredSizes = FVector2D::ZeroVector;
+	for (int ColId = 0; ColId < Columns.Num(); ++ColId)
+	{
+		TotalDesiredSizes.X += Columns[ColId];
+	}
+	TotalDesiredSizes.X += (ColumnGutter * 2) * (TotalColumns - 1);
+
+	for (int RowId = 0; RowId < Rows.Num(); ++RowId)
+	{
+		TotalDesiredSizes.Y += Rows[RowId];
+	}
+	TotalDesiredSizes.Y += (RowGutter * 2)  * (Slots[Slots.Num() - 1].RowParam);
+
+	SPanel::CacheDesiredSize();
+}
+
+FVector2D SResponsiveGridPanel::ComputeDesiredSize() const
+{
+	return TotalDesiredSizes;
+}
+
+void SResponsiveGridPanel::ComputeDesiredCellSizes(float AvailableWidth, TArray<float>& OutColumns, TArray<float>& OutRows, TArray<float>& OutRowToSlot) const
+{
+	FMemory::Memzero(OutColumns.GetData(), OutColumns.Num() * sizeof(float));
+	FMemory::Memzero(OutRows.GetData(), OutRows.Num() * sizeof(float));
+
+	int32 ColumnsSoFar = 0;
+	int32 CurrentRow = INDEX_NONE;
+	int32 LastRowParam = INDEX_NONE;
+	for (int32 SlotIndex = 0; SlotIndex < Slots.Num(); ++SlotIndex)
+	{
+		const FSlot& CurSlot = Slots[SlotIndex];
+		if (CurSlot.GetWidget()->GetVisibility() != EVisibility::Collapsed)
+		{
+			// Find the appropriate column layout for the slot
+			FSlot::FColumnLayout ColumnLayout;
+			ColumnLayout.Span = TotalColumns;
+			ColumnLayout.Offset = 0;
+
+			for (int32 Index = CurSlot.ColumnLayouts.Num() - 1; Index >= 0; Index--)
+			{
+				if (CurSlot.ColumnLayouts[Index].LayoutSize < AvailableWidth)
+				{
+					ColumnLayout = CurSlot.ColumnLayouts[Index];
+					break;
+				}
+			}
+
+			if (ColumnLayout.Span == 0)
+			{
+				continue;
+			}
+
+			if (CurSlot.RowParam != LastRowParam)
+			{
+				ColumnsSoFar = 0;
+				LastRowParam = CurSlot.RowParam;
+				++CurrentRow;
+
+				OutRowToSlot.AddZeroed((CurrentRow + 1) - OutRowToSlot.Num());
+				OutRowToSlot[CurrentRow] = CurSlot.RowParam;
+			}
+
+			// The slots want to be as big as its content along with the required padding.
+			const FVector2D SlotDesiredSize = CurSlot.GetWidget()->GetDesiredSize() + CurSlot.SlotPadding.Get().GetDesiredSize();
+
+			// If the slot has a (colspan,rowspan) of (1,1) it will only affect that cell.
+			// For larger spans, the slots size will be evenly distributed across all the affected cells.
+			const FVector2D SizeContribution(SlotDesiredSize.X / ColumnLayout.Span, SlotDesiredSize.Y);
+
+			int32 StartColumnIndex = ColumnsSoFar + ColumnLayout.Offset;
+			int32 EndColumnIndex = StartColumnIndex + ColumnLayout.Span;
+			ColumnsSoFar = FMath::Max(EndColumnIndex, ColumnsSoFar);
+
+			if (ColumnsSoFar > TotalColumns)
+			{
+				StartColumnIndex = 0;
+				EndColumnIndex = StartColumnIndex + ColumnLayout.Span;
+				ColumnsSoFar = EndColumnIndex - StartColumnIndex;
+				++CurrentRow;
+
+				OutRowToSlot.AddZeroed((CurrentRow + 1) - OutRowToSlot.Num());
+				OutRowToSlot[CurrentRow] = CurSlot.RowParam;
+			}
+
+			OutColumns.AddZeroed(FMath::Max(0, ColumnsSoFar - OutColumns.Num()));
+			OutRows.AddZeroed((CurrentRow + 1) - OutRows.Num());
+
+			// Distribute the size contributions over all the columns and rows that this slot spans
+			DistributeSizeContributions(SizeContribution.X, OutColumns, StartColumnIndex, EndColumnIndex);
+			DistributeSizeContributions(SizeContribution.Y, OutRows, CurrentRow, CurrentRow + 1);
+		}
+	}
+}
+
+void SResponsiveGridPanel::DistributeSizeContributions(float SizeContribution, TArray<float>& DistributeOverMe, int32 StartIndex, int32 UpperBound)
+{
+	for (int32 Index = StartIndex; Index < UpperBound; ++Index)
+	{
+		// Each column or row only needs to get bigger if its current size does not already accommodate it.
+		DistributeOverMe[Index] = FMath::Max(SizeContribution, DistributeOverMe[Index]);
+	}
+}
+
+FChildren* SResponsiveGridPanel::GetChildren()
+{
+	return &Slots;
+}
+
+void SResponsiveGridPanel::ComputePartialSums( TArray<float>& TurnMeIntoPartialSums )
+{
+	// We assume there is a 0-valued item already at the end of this array.
+	// We need it so that we can  compute the original values
+	// by doing Array[N] - Array[N-1];
+	
+	float LastValue = 0;
+	float SumSoFar = 0;
+	for (int32 Index = 0; Index < TurnMeIntoPartialSums.Num(); ++Index)
+	{
+		LastValue = TurnMeIntoPartialSums[Index];
+		TurnMeIntoPartialSums[Index] = SumSoFar;
+		SumSoFar += LastValue;
+	}
+}
+
+void SResponsiveGridPanel::SetRowFill(int32 RowId, float Coefficient)
+{
+	if (RowFillCoefficients.Num() <= RowId)
+	{
+		RowFillCoefficients.AddZeroed(RowId - RowFillCoefficients.Num() + 1);
+	}
+
+	RowFillCoefficients[RowId] = Coefficient;
+}
+
+SResponsiveGridPanel::FSlot& SResponsiveGridPanel::InsertSlot(SResponsiveGridPanel::FSlot* InSlot)
+{
+	InSlot->Panel = SharedThis(this);
+
+	bool bInserted = false;
+
+	// Insert the slot in the list such that slots are sorted by LayerOffset.
+	for (int32 SlotIndex = 0; !bInserted && SlotIndex < Slots.Num(); ++SlotIndex)
+	{
+		if (InSlot->RowParam < this->Slots[SlotIndex].RowParam)
+		{
+			Slots.Insert(InSlot, SlotIndex);
+			bInserted = true;
+		}
+	}
+
+	// We haven't inserted yet, so add to the end of the list.
+	if (!bInserted)
+	{
+		Slots.Add(InSlot);
+	}
+
+	NotifySlotChanged(InSlot);
+
+	return *InSlot;
+}
+
+void SResponsiveGridPanel::NotifySlotChanged(SResponsiveGridPanel::FSlot* InSlot)
+{
+	//// Keep the size of the grid up to date.
+	//// We need an extra cell at the end for easily figuring out the size across any number of cells
+
+	//Rows.AddZeroed();
+}
+
+int32 SResponsiveGridPanel::LayoutDebugPaint(const FGeometry& AllottedGeometry, const FSlateRect& MyClippingRect, FSlateWindowElementList& OutDrawElements, int32 LayerId ) const
+{
+	//float XOffset = 0;
+	//for (int32 Column=0; Column<Columns.Num(); ++Column)
+	//{
+	//	float YOffset = 0;
+	//	for (int32 Row=0; Row<Rows.Num(); ++Row)
+	//	{
+	//		FSlateDrawElement::MakeDebugQuad
+	//		(
+	//			OutDrawElements, 
+	//			LayerId,
+	//			AllottedGeometry.ToPaintGeometry( FVector2D(XOffset, YOffset), FVector2D( Columns[Column], Rows[Row] ) ),
+	//			MyClippingRect
+	//		);
+
+	//		YOffset += Rows[Row];
+	//	}
+	//	XOffset += Columns[Column];
+	//}
+
+	return LayerId;
+}
+
diff --git a/Engine/Source/Runtime/Slate/Public/Widgets/Layout/SResponsiveGridPanel.h b/Engine/Source/Runtime/Slate/Public/Widgets/Layout/SResponsiveGridPanel.h
new file mode 100644
--- /dev/null
+++ b/Engine/Source/Runtime/Slate/Public/Widgets/Layout/SResponsiveGridPanel.h
@@ -0,0 +1,201 @@
+// Copyright 1998-2014 Epic Games, Inc. All Rights Reserved.
+#pragma once
+
+/**
+ *  !!!!!!!!!!!!!!!!!   EXPERIMENTAL  !!!!!!!!!!!!!!!!!   
+ * These sizes are subject to change in the future 
+ */
+namespace SResponsiveGridSize
+{
+	enum Type
+	{
+		Mobile = 0,
+		Tablet = 768,
+		MediumDevice = 992,
+		LargeDevice = 1200,
+	};
+}
+
+/**
+ *  !!!!!!!!!!!!!!!!!   EXPERIMENTAL  !!!!!!!!!!!!!!!!!   
+ * The SResponsiveGridPanel is still in development and the API may change drastically in the future
+ * or maybe removed entirely.
+ */
+class SLATE_API SResponsiveGridPanel
+	: public SPanel
+{
+public:
+
+	class FSlot : public TSlotBase<FSlot>, public TSupportsContentAlignmentMixin<FSlot>, public TSupportsContentPaddingMixin < FSlot >
+	{
+	public:
+
+		struct FColumnLayout
+		{
+			float LayoutSize;
+			int32 Span;
+			int32 Offset;
+		};
+
+	public:
+		/** Default values for a slot. */
+		FSlot(int32 Row)
+			: TSlotBase<FSlot>()
+			, TSupportsContentAlignmentMixin<FSlot>(HAlign_Fill, VAlign_Fill)
+			, RowParam(FMath::Max(0, Row))
+			, ColumnLayouts()
+		{ }
+
+		/** The panel that contains this slot */
+		TWeakPtr<SResponsiveGridPanel> Panel;
+
+		/** The row index*/
+		int32 RowParam;
+
+		/** Layout information for the column */
+		TArray<FColumnLayout> ColumnLayouts;
+
+		/** How many columns this slot spans over */
+		FSlot& ColumnSpan(float LayoutSize, int32 ColumnSpan, int32 ColumnOffset = 0)
+		{
+			FColumnLayout ColumnLayout;
+			ColumnLayout.LayoutSize = LayoutSize;
+			ColumnLayout.Span = FMath::Max(0, ColumnSpan);
+			ColumnLayout.Offset = ColumnOffset;
+
+			bool Inserted = false;
+			for (int32 Index = 0; Index < ColumnLayouts.Num(); Index++)
+			{
+				if (ColumnLayout.LayoutSize < ColumnLayouts[Index].LayoutSize)
+				{
+					ColumnLayouts.Insert(ColumnLayout, Index);
+					Inserted = true;
+					break;
+				}
+			}
+
+			if (!Inserted)
+			{
+				ColumnLayouts.Add(ColumnLayout);
+			}
+
+			if (Panel.IsValid())
+			{
+				Panel.Pin()->NotifySlotChanged(this);
+			}
+
+			return *this;
+		}
+	};
+
+	/**
+	 * Used by declarative syntax to create a Slot
+	 */
+	static FSlot& Slot(int32 Row)
+	{
+		return *(new FSlot(Row));
+	}
+
+	SLATE_BEGIN_ARGS(SResponsiveGridPanel)
+		: _ColumnGutter(0)
+		, _RowGutter(0)
+		{
+			_Visibility = EVisibility::SelfHitTestInvisible;
+		}
+
+		SLATE_ARGUMENT(float, ColumnGutter)
+		SLATE_ARGUMENT(float, RowGutter)
+
+		SLATE_SUPPORTS_SLOT(FSlot)
+
+		/** Specify a row to stretch instead of sizing to content. */
+		FArguments& FillRow(int32 RowId, float Coefficient)
+		{
+			if (RowFillCoefficients.Num() <= RowId)
+			{
+				RowFillCoefficients.AddZeroed(RowId - RowFillCoefficients.Num() + 1);
+			}
+			RowFillCoefficients[RowId] = Coefficient;
+			return Me();
+		}
+
+		/** Coefficients for rows that need to stretch instead of size to content */
+		TArray<float> RowFillCoefficients;
+
+	SLATE_END_ARGS()
+
+	SResponsiveGridPanel();
+
+	/** Removes all rows from the panel */
+	void ClearChildren();
+	
+	void Construct( const FArguments& InArgs, int32 TotalColumns );
+
+	FSlot& AddSlot(int32 Row);
+
+	bool RemoveSlot(const TSharedRef<SWidget>& SlotWidget);
+
+	void SetRowFill(int32 RowId, float Coefficient);
+
+public:
+
+	// SWidget interface
+
+	virtual int32 OnPaint( const FPaintArgs& Args, const FGeometry& AllottedGeometry, const FSlateRect& MyClippingRect, FSlateWindowElementList& OutDrawElements, int32 LayerId, const FWidgetStyle& InWidgetStyle, bool bParentEnabled ) const override;
+	virtual void OnArrangeChildren( const FGeometry& AllottedGeometry, FArrangedChildren& ArrangedChildren ) const override;
+	virtual void CacheDesiredSize() override;
+	virtual FVector2D ComputeDesiredSize() const override;
+	virtual FChildren* GetChildren() override;
+
+private:
+
+	/**
+	 * Given an array of values, re-populate the array such that every contains the partial sums up to that element.
+	 * i.e. Array[N] = Array.Sum(0 .. N-1)
+	 *
+	 * The resulting array is 1-element longer.
+	 */
+	static void ComputePartialSums( TArray<float>& TurnMeIntoPartialSums );
+	
+	/** Given a SizeContribution, distribute it to the elements in DistributeOverMe at indexes from [StartIndex .. UpperBound) */
+	static void DistributeSizeContributions( float SizeContribution, TArray<float>& DistributeOverMe, int32 StartIndex, int32 UpperBound );
+
+	/**
+	 * Inserts the given slot into the list of Slots based on its LayerParam, such that Slots are sorter by layer.
+	 *
+	 * @param The newly-allocated slot to insert.
+	 * @return A reference to the added slot
+	 */
+	FSlot& InsertSlot(FSlot* InSlot);
+
+	/** Compute the sizes of columns and rows needed to fit all the slots in this grid. */
+	void ComputeDesiredCellSizes(float AvailableWidth, TArray<float>& OutColumns, TArray<float>& OutRows, TArray<float>& OutRowToSlot) const;
+
+	/** Draw the debug grid of rows and colummns; useful for inspecting the GridPanel's logic. See OnPaint() for parameter meaning */
+	int32 LayoutDebugPaint(const FGeometry& AllottedGeometry, const FSlateRect& MyClippingRect, FSlateWindowElementList& OutDrawElements, int32 LayerId ) const;
+
+	/** 
+	 * Callback used to resize our internal arrays if a slot (or slot span) is changed. 
+	 *
+	 * @param The slot that has just changed.
+	 */
+	void NotifySlotChanged(FSlot* InSlot);
+
+
+private:
+
+	/** The rows*/
+	TPanelChildren<FSlot> Slots;
+
+	int32 TotalColumns;
+	float ColumnGutter;
+	float RowGutter;
+
+	/** Fill coefficients for the rows */
+	TArray<float> RowFillCoefficients;
+
+	/** Total desires size along each axis. */
+	FVector2D TotalDesiredSizes;
+
+	mutable float PreviousWidth;
+};
