Slate: Exposed ViewportInterface as an attribute in SViewport

diff --git a/Engine/Source/Runtime/Slate/Private/Widgets/SViewport.cpp b/Engine/Source/Runtime/Slate/Private/Widgets/SViewport.cpp
--- a/Engine/Source/Runtime/Slate/Private/Widgets/SViewport.cpp
+++ b/Engine/Source/Runtime/Slate/Private/Widgets/SViewport.cpp
@@ -7,9 +7,8 @@ DECLARE_CYCLE_STAT( TEXT("OnPaint SViewport"), STAT_SlateOnPaint_SViewport, STAT
 /* SViewport structors
  *****************************************************************************/
 
-SViewport::SViewport( )
-	: ViewportInterface(nullptr)
-	, bRenderDirectlyToWindow(false)
+SViewport::SViewport()
+	: bRenderDirectlyToWindow(false)
 	, bEnableGammaCorrection(true)
 { }
 
@@ -24,6 +23,7 @@ void SViewport::Construct( const FArguments& InArgs )
 	bEnableGammaCorrection = InArgs._EnableGammaCorrection;
 	bEnableBlending = InArgs._EnableBlending;
 	bIgnoreTextureAlpha = InArgs._IgnoreTextureAlpha;
+	ViewportInterface = InArgs._ViewportInterface;
 	ViewportSize = InArgs._ViewportSize;
 
 	this->ChildSlot
diff --git a/Engine/Source/Runtime/Slate/Public/Widgets/SViewport.h b/Engine/Source/Runtime/Slate/Public/Widgets/SViewport.h
--- a/Engine/Source/Runtime/Slate/Public/Widgets/SViewport.h
+++ b/Engine/Source/Runtime/Slate/Public/Widgets/SViewport.h
@@ -3,7 +3,8 @@
 #pragma once
 
 
-class SLATE_API SViewport : public SCompoundWidget
+class SLATE_API SViewport
+	: public SCompoundWidget
 {
 public:
 	
@@ -14,8 +15,8 @@ public:
 		, _EnableGammaCorrection(true)
 		, _EnableBlending(false)
 		, _IgnoreTextureAlpha(true)
-		, _ViewportSize(FVector2D(320, 240))
-		{}
+		, _ViewportSize(FVector2D(320.0f, 240.0f))
+	{ }
 
 		SLATE_DEFAULT_SLOT( FArguments, Content )
 
@@ -29,14 +30,10 @@ public:
 		 */
 		SLATE_ARGUMENT( bool, RenderDirectlyToWindow )
 
-		/**
-		 * Whether or not to enable gamma correction.  Doesn't apply when rendering directly to a backbuffer 
-		 */
+		/** Whether or not to enable gamma correction.  Doesn't apply when rendering directly to a backbuffer . */
 		SLATE_ARGUMENT( bool, EnableGammaCorrection )
 
-		/**
-		 * Allow this viewport to blend with its background
-		 */
+		/** Allow this viewport to blend with its background. */
 		SLATE_ARGUMENT( bool, EnableBlending )
 
 		/**
@@ -45,15 +42,16 @@ public:
 		 */
 		SLATE_ARGUMENT( bool, IgnoreTextureAlpha )
 
-		/**
-		 * Size of the viewport widget
-		 */
+		/** The interface to be used by this viewport for rendering and I/O. */
+		SLATE_ARGUMENT(TSharedPtr<ISlateViewport>, ViewportInterface)
+		
+		/** Size of the viewport widget. */
 		SLATE_ATTRIBUTE(FVector2D, ViewportSize);
 
 	SLATE_END_ARGS()
 
 	/** Default constructor. */
-	SViewport( );
+	SViewport();
 
 	/**
 	 * Construct the widget.
@@ -118,7 +116,7 @@ public:
 	 *
 	 * @param	InWidget	The widget to set focus to when this window is activated
 	 */
-	void SetWidgetToFocusOnActivate( const TSharedPtr< SWidget >& InWidget )
+	void SetWidgetToFocusOnActivate( const TSharedPtr<SWidget>& InWidget )
 	{
 		
 		WidgetToFocusOnActivate = InWidget;
@@ -168,30 +166,32 @@ private:
 
 protected:
 
-	// Interface to the rendering and I/O implementation of the viewport.
+	/** Interface to the rendering and I/O implementation of the viewport. */
 	TWeakPtr<ISlateViewport> ViewportInterface;
 	
 private:
 
-	// Whether or not to show the disabled effect when this viewport is disabled.
+	/** Whether or not to show the disabled effect when this viewport is disabled. */
 	TAttribute<bool> ShowDisabledEffect;
 
-	// Size of the viewport
+	/** Size of the viewport. */
 	TAttribute<FVector2D> ViewportSize;
 
-	// Whether or not this viewport renders directly to the window back-buffer.
+	/** Whether or not this viewport renders directly to the window back-buffer. */
 	bool bRenderDirectlyToWindow;
 
-	// Whether or not to apply gamma correction on the render target supplied by the ISlateViewport.
+	/** Whether or not to apply gamma correction on the render target supplied by the ISlateViewport. */
 	bool bEnableGammaCorrection;
 
-	// Whether or not to blend this viewport with the background.
+	/** Whether or not to blend this viewport with the background. */
 	bool bEnableBlending;
 
-	// Whether or not to allow texture alpha to be used in blending calculations.
+	/** Whether or not to allow texture alpha to be used in blending calculations. */
 	bool bIgnoreTextureAlpha;
 
-	// Widget to transfer keyboard focus to when this window becomes active, if any.
-	// This is used to restore focus to a widget after a popup has been dismissed.
-	TWeakPtr< SWidget > WidgetToFocusOnActivate;
+	/**
+	 * Widget to transfer keyboard focus to when this window becomes active, if any.
+	 * This is used to restore focus to a widget after a pop-up has been dismissed.
+	 */
+	TWeakPtr<SWidget> WidgetToFocusOnActivate;
 };
