Fixed compile error

diff --git a/Engine/Source/Runtime/Slate/Public/Framework/MarqueeRect.h b/Engine/Source/Runtime/Slate/Public/Framework/MarqueeRect.h
--- a/Engine/Source/Runtime/Slate/Public/Framework/MarqueeRect.h
+++ b/Engine/Source/Runtime/Slate/Public/Framework/MarqueeRect.h
@@ -1,5 +1,7 @@
 // Copyright 1998-2014 Epic Games, Inc. All Rights Reserved.
 
+#pragma once
+
 /**
  * A convenient representation of a marquee selection
  */
