SSpinBox needs to cache the cursor position on MouseButtonDown if it wants to later SetCursorPos().

Fixes https://jira.ol.epicgames.net/browse/UETOOL-130

diff --git a/Engine/Source/Runtime/Slate/Public/Widgets/Input/SSpinBox.h b/Engine/Source/Runtime/Slate/Public/Widgets/Input/SSpinBox.h
--- a/Engine/Source/Runtime/Slate/Public/Widgets/Input/SSpinBox.h
+++ b/Engine/Source/Runtime/Slate/Public/Widgets/Input/SSpinBox.h
@@ -248,6 +248,7 @@ public:
 		{
 			DistanceDragged = 0;
 			PreDragValue = InternalValue;
+			CachedMousePosition = MouseEvent.GetScreenSpacePosition().IntPoint();
 			return FReply::Handled().CaptureMouse( SharedThis(this) ).UseHighPrecisionMouseMovement( SharedThis(this) ).SetKeyboardFocus( SharedThis(this), EKeyboardFocusCause::Mouse);
 		}
 		else
