Slate: added preview events (tunneling events) to FEventRouter.
Fixed a bug where one of the tunneling events was actually bubbeling.

diff --git a/Engine/Source/Runtime/Slate/Private/Framework/Application/SlateApplication.cpp b/Engine/Source/Runtime/Slate/Private/Framework/Application/SlateApplication.cpp
--- a/Engine/Source/Runtime/Slate/Private/Framework/Application/SlateApplication.cpp
+++ b/Engine/Source/Runtime/Slate/Private/Framework/Application/SlateApplication.cpp
@@ -13,8 +13,6 @@
 class FEventRouter
 {
 
-// @todo slate : handle preview events.
-
 // @todo slate : Widget Reflector should log: (1) Every process reply (2) Every time the event is handled and by who.
 
 // @todo slate : Remove remaining [&]-style mass captures.
@@ -24,6 +22,109 @@ class FEventRouter
 // @todo slate : Remove CALL_WIDGET_FUNCTION
 
 public:
+	class FToLeafmostPolicy
+	{
+	public:
+		FToLeafmostPolicy( const FWidgetPath& InRoutingPath )
+		: bEventSent(false)
+		, RoutingPath(InRoutingPath)
+		{
+		}
+
+		bool ShouldKeepGoing() const
+		{
+			return !bEventSent && RoutingPath.Widgets.Num() > 0;
+		}
+
+		void Next()
+		{
+			bEventSent = true;
+		}
+
+		const FArrangedWidget& GetWidget() const
+		{
+			return RoutingPath.Widgets[RoutingPath.Widgets.Num()-1];
+		}
+
+		const FWidgetPath& GetRoutingPath() const
+		{
+			return RoutingPath;
+		}
+
+	private:
+		bool bEventSent;
+		const FWidgetPath& RoutingPath;
+	};
+
+	class FTunnelPolicy
+	{
+	public:
+		FTunnelPolicy( const FWidgetPath& InRoutingPath )
+		: WidgetIndex(0)
+		, RoutingPath(InRoutingPath)
+		{
+		}
+
+		bool ShouldKeepGoing() const
+		{
+			return WidgetIndex < RoutingPath.Widgets.Num();
+		}
+
+		void Next()
+		{
+			++WidgetIndex;
+		}
+
+		const FArrangedWidget& GetWidget() const
+		{
+			return RoutingPath.Widgets[WidgetIndex];
+		}
+		
+		const FWidgetPath& GetRoutingPath() const
+		{
+			return RoutingPath;
+		}
+
+
+	private:
+		int32 WidgetIndex;
+		const FWidgetPath& RoutingPath;
+	};
+
+	class FBubblePolicy
+	{
+	public:
+		FBubblePolicy( const FWidgetPath& InRoutingPath )
+		: WidgetIndex( InRoutingPath.Widgets.Num()-1 )
+		, RoutingPath (InRoutingPath)
+		{
+		}
+
+		bool ShouldKeepGoing() const
+		{
+			return WidgetIndex >= 0;
+		}
+
+		void Next()
+		{
+			--WidgetIndex;
+		}
+
+		const FArrangedWidget& GetWidget() const
+		{
+			return RoutingPath.Widgets[WidgetIndex];
+		}
+
+		const FWidgetPath& GetRoutingPath() const
+		{
+			return RoutingPath;
+		}
+	
+	private:
+		int32 WidgetIndex;
+		const FWidgetPath& RoutingPath;
+	};
+
 	static void LogEvent( FSlateApplication* ThisApplication, const FInputEvent& Event, const FReply& Reply )
 	{
 		TSharedPtr<IWidgetReflector> Reflector = ThisApplication->WidgetReflectorPtr.Pin();
@@ -40,16 +141,17 @@ public:
 	 * Focus paths change when the user navigates focus (e.g. Tab or
 	 * Shift Tab or clicks on a focusable widget.)
 	 */
-	template< typename FuncType, typename EventType >
-	static FReply BubbleAlongFocusPath( FSlateApplication* ThisApplication, const FWidgetPath& RoutingPath, EventType KeyEventCopy, const FuncType& Lambda )
+	template< typename RoutingPolicyType, typename FuncType, typename EventType >
+	static FReply RouteAlongFocusPath( FSlateApplication* ThisApplication, RoutingPolicyType RoutingPolicy, EventType KeyEventCopy, const FuncType& Lambda )
 	{
+		const FWidgetPath& RoutingPath = RoutingPolicy.GetRoutingPath();
 		FReply Reply = FReply::Unhandled();
 
 		KeyEventCopy.SetEventPath(RoutingPath);
 
-		for ( int32 WidgetIndex = RoutingPath.Widgets.Num() - 1; !Reply.IsEventHandled() && WidgetIndex >= 0; --WidgetIndex )
+		for ( ; !Reply.IsEventHandled() && RoutingPolicy.ShouldKeepGoing(); RoutingPolicy.Next() )
 		{
-			const FArrangedWidget& ArrangedWidget = RoutingPath.Widgets[WidgetIndex];
+			const FArrangedWidget& ArrangedWidget = RoutingPolicy.GetWidget();
 			Reply = Lambda( ArrangedWidget, KeyEventCopy ).SetHandler( ArrangedWidget.Widget );
 			ThisApplication->ProcessReply( RoutingPath, Reply, NULL, NULL );
 		}
@@ -69,16 +171,17 @@ public:
 	 * with that pointing device until the pointer capture is
 	 * released.
 	 */
-	template< typename FuncType >
-	static FReply BubbleAlongPointerPath( FSlateApplication* ThisApplication, const FWidgetPath& RoutingPath, FPointerEvent PointerEventCopy, const FuncType& Lambda )
+	template< typename RoutingPolicyType,  typename FuncType >
+	static FReply RouteAlongPointerPath( FSlateApplication* ThisApplication, RoutingPolicyType RoutingPolicy, FPointerEvent PointerEventCopy, const FuncType& Lambda )
 	{
+		const FWidgetPath& RoutingPath = RoutingPolicy.GetRoutingPath();
 		FReply Reply = FReply::Unhandled();
 
 		PointerEventCopy.SetEventPath( RoutingPath );
 
-		for ( int32 WidgetIndex = RoutingPath.Widgets.Num() - 1; !Reply.IsEventHandled() && WidgetIndex >= 0; --WidgetIndex )
+		for ( ; !Reply.IsEventHandled() && RoutingPolicy.ShouldKeepGoing(); RoutingPolicy.Next() )
 		{
-			const FArrangedWidget& ArrangedWidget = RoutingPath.Widgets[WidgetIndex];
+			const FArrangedWidget& ArrangedWidget = RoutingPolicy.GetWidget();
 			Reply = Lambda( ArrangedWidget, PointerEventCopy ).SetHandler( ArrangedWidget.Widget );
 			const FWidgetPath* WidgetsUnderCursorPtr = &RoutingPath;
 			ThisApplication->ProcessReply( RoutingPath, Reply, WidgetsUnderCursorPtr, &PointerEventCopy );
@@ -88,25 +191,6 @@ public:
 
 		return Reply;
 	}
-
-	/** Just like BubbleAlongPointerPath, but only the TargetWidget in the path gets the event. */
-	template< typename FuncType >
-	static FReply SendToLeafmost( FSlateApplication* ThisApplication, const FWidgetPath& RoutingPath, const TSharedPtr<SWidget> TargetWidget, FPointerEvent PointerEventCopy, const FuncType& Lambda )
-	{
-		PointerEventCopy.SetEventPath( RoutingPath );
-
-		const FArrangedWidget& ArrangedWidget = TargetWidget.IsValid()
-			? RoutingPath.FindArrangedWidget( TargetWidget.ToSharedRef() )
-			: RoutingPath.Widgets.Last();
-
-		FReply Reply = Lambda( ArrangedWidget, PointerEventCopy ).SetHandler( ArrangedWidget.Widget );
-		const FWidgetPath* WidgetsUnderCursorPtr = &RoutingPath;
-		ThisApplication->ProcessReply( RoutingPath, Reply, WidgetsUnderCursorPtr, &PointerEventCopy );
-
-		LogEvent(ThisApplication, PointerEventCopy, Reply);
-
-		return Reply;
-	}
 };
 
 
@@ -3187,8 +3271,8 @@ bool FSlateApplication::ProcessKeyCharEvent( FCharacterEvent& InCharacterEvent )
 		// Send out mouse enter events.
 		InCharacterEvent.SetEventPath( EventPath );
 		
-		Reply = FEventRouter::BubbleAlongFocusPath( this, EventPath, InCharacterEvent, []( const FArrangedWidget& SomeWidgetGettingEvent, const FCharacterEvent& Event )
-			{
+		Reply = FEventRouter::RouteAlongFocusPath( this, FEventRouter::FBubblePolicy(EventPath), InCharacterEvent, []( const FArrangedWidget& SomeWidgetGettingEvent, const FCharacterEvent& Event )
+		{
 			return SomeWidgetGettingEvent.Widget->IsEnabled()
 				? SomeWidgetGettingEvent.Widget->OnKeyChar( SomeWidgetGettingEvent.Geometry, Event )
 				: FReply::Unhandled();
@@ -3245,7 +3329,7 @@ bool FSlateApplication::ProcessKeyDownEvent( FKeyboardEvent& InKeyboardEvent )
 			FWidgetPath EventPath = FocusedWidgetPath.ToWidgetPath();
 			InKeyboardEvent.SetEventPath(EventPath);
 
-			// Switch worlds for widgets in the current path
+			// Switch worlds for widgets inOnPreviewMouseButtonDown the current path
 			FScopedSwitchWorldHack SwitchWorld( EventPath );
 
 			TSharedPtr<SWidget> WidgetToLog;
@@ -3253,22 +3337,17 @@ bool FSlateApplication::ProcessKeyDownEvent( FKeyboardEvent& InKeyboardEvent )
 			Reply = FReply::Unhandled();
 
 			// Tunnel the keyboard event
-			for( int32 WidgetIndex = 0; !Reply.IsEventHandled() && WidgetIndex < EventPath.Widgets.Num(); ++WidgetIndex )
+			Reply = FEventRouter::RouteAlongFocusPath( this, FEventRouter::FTunnelPolicy(EventPath), InKeyboardEvent, []( const FArrangedWidget& CurrentWidget, const FKeyboardEvent& Event )
 			{
-				FArrangedWidget& SomeWidgetGettingEvent = EventPath.Widgets[ WidgetIndex ];
-				if (SomeWidgetGettingEvent.Widget->IsEnabled())
-				{
-					Reply = SomeWidgetGettingEvent.Widget->OnPreviewKeyDown( SomeWidgetGettingEvent.Geometry, InKeyboardEvent ).SetHandler(SomeWidgetGettingEvent.Widget);
-					ProcessReply(EventPath, Reply, NULL, NULL);
-
-					WidgetToLog = SomeWidgetGettingEvent.Widget;
-				}
-			}
+				return (CurrentWidget.Widget->IsEnabled())
+					? CurrentWidget.Widget->OnPreviewKeyDown( CurrentWidget.Geometry, Event )
+					: FReply::Unhandled();
+			});
 
 			// Send out key down events.
 			if (!Reply.IsEventHandled())
 			{
-				Reply = FEventRouter::BubbleAlongFocusPath(this, EventPath, InKeyboardEvent, [](const FArrangedWidget& SomeWidgetGettingEvent, const FKeyboardEvent& Event)
+				Reply = FEventRouter::RouteAlongFocusPath(this, FEventRouter::FBubblePolicy(EventPath), InKeyboardEvent, [](const FArrangedWidget& SomeWidgetGettingEvent, const FKeyboardEvent& Event)
 				{
 					return (SomeWidgetGettingEvent.Widget->IsEnabled())
 						? SomeWidgetGettingEvent.Widget->OnKeyDown( SomeWidgetGettingEvent.Geometry, Event )
@@ -3330,7 +3409,7 @@ bool FSlateApplication::ProcessKeyUpEvent( FKeyboardEvent& InKeyboardEvent )
 		// Switch worlds for widgets in the current path
 		FScopedSwitchWorldHack SwitchWorld( EventPath );
 
-		Reply = FEventRouter::BubbleAlongFocusPath(this, EventPath, InKeyboardEvent, [](const FArrangedWidget& SomeWidgetGettingEvent, const FKeyboardEvent& Event)
+		Reply = FEventRouter::RouteAlongFocusPath(this, FEventRouter::FBubblePolicy(EventPath), InKeyboardEvent, [](const FArrangedWidget& SomeWidgetGettingEvent, const FKeyboardEvent& Event)
 		{
 			return (SomeWidgetGettingEvent.Widget->IsEnabled())
 				? SomeWidgetGettingEvent.Widget->OnKeyUp( SomeWidgetGettingEvent.Geometry, Event )
@@ -3432,14 +3511,14 @@ bool FSlateApplication::ProcessMouseButtonDownEvent( const TSharedPtr< FGenericW
 			bInGame = FApp::IsGame();
 
 			
-			Reply = FEventRouter::SendToLeafmost( this, MouseCaptorPath, nullptr, MouseEvent, []( const FArrangedWidget& MouseCaptorWidget, const FPointerEvent& Event )
+			Reply = FEventRouter::RouteAlongPointerPath( this, FEventRouter::FToLeafmostPolicy(MouseCaptorPath), MouseEvent, []( const FArrangedWidget& MouseCaptorWidget, const FPointerEvent& Event )
 			{
 				return MouseCaptorWidget.Widget->OnPreviewMouseButtonDown( MouseCaptorWidget.Geometry, Event );
 			});
 			
 			if (!Reply.IsEventHandled())
 			{
-				Reply = FEventRouter::SendToLeafmost( this, MouseCaptorPath, nullptr, MouseEvent,
+				Reply = FEventRouter::RouteAlongPointerPath( this, FEventRouter::FToLeafmostPolicy(MouseCaptorPath), MouseEvent,
 					[this]( const FArrangedWidget& MouseCaptorWidget, const FPointerEvent& Event )
 				{
 						FReply TempReply = FReply::Unhandled();
@@ -3472,27 +3551,27 @@ bool FSlateApplication::ProcessMouseButtonDownEvent( const TSharedPtr< FGenericW
 
 			const TSharedPtr<SWidget> PreviouslyFocusedWidget = GetKeyboardFocusedWidget();
 
-			Reply = FEventRouter::BubbleAlongPointerPath(this, WidgetsUnderCursor, MouseEvent, [](const FArrangedWidget TargetWidget, const FPointerEvent& Event)
+			Reply = FEventRouter::RouteAlongPointerPath(this, FEventRouter::FTunnelPolicy(WidgetsUnderCursor), MouseEvent, [](const FArrangedWidget TargetWidget, const FPointerEvent& Event)
 			{
 				return TargetWidget.Widget->OnPreviewMouseButtonDown(TargetWidget.Geometry, Event);
 			});
 
 			if ( !Reply.IsEventHandled() )
 			{
-				Reply = FEventRouter::BubbleAlongPointerPath(this, WidgetsUnderCursor, MouseEvent, [this](const FArrangedWidget TargetWidget, const FPointerEvent& Event)
-			{
+				Reply = FEventRouter::RouteAlongPointerPath(this, FEventRouter::FBubblePolicy(WidgetsUnderCursor), MouseEvent, [this](const FArrangedWidget TargetWidget, const FPointerEvent& Event)
+				{
 					FReply ThisReply = FReply::Unhandled();
 					if (!ThisReply.IsEventHandled())
-				{
-						if (Event.IsTouchEvent())
 					{
+						if (Event.IsTouchEvent())
+						{
 							ThisReply = TargetWidget.Widget->OnTouchStarted( TargetWidget.Geometry, Event );
-					}
+						}
 						if (!Event.IsTouchEvent() || (!ThisReply.IsEventHandled() && this->bTouchFallbackToMouse))
-					{
+						{
 							ThisReply = TargetWidget.Widget->OnMouseButtonDown( TargetWidget.Geometry, Event );
+						}
 					}
-				}
 					return ThisReply;
 				});
 			}
@@ -3589,7 +3668,7 @@ bool FSlateApplication::ProcessMouseButtonDoubleClickEvent( const TSharedPtr< FG
 	// Switch worlds widgets in the current path
 	FScopedSwitchWorldHack SwitchWorld( WidgetsUnderCursor );
 
-	FReply Reply = FEventRouter::BubbleAlongPointerPath( this, WidgetsUnderCursor, InMouseEvent, [](const FArrangedWidget& TargetWidget, const FPointerEvent& Event)
+	FReply Reply = FEventRouter::RouteAlongPointerPath( this, FEventRouter::FBubblePolicy(WidgetsUnderCursor), InMouseEvent, [](const FArrangedWidget& TargetWidget, const FPointerEvent& Event)
 	{
 		return TargetWidget.Widget->OnMouseButtonDoubleClick( TargetWidget.Geometry, Event );
 	} );
@@ -3648,7 +3727,7 @@ bool FSlateApplication::ProcessMouseButtonUpEvent( FPointerEvent& MouseEvent )
 			// Switch worlds widgets in the current path
 			FScopedSwitchWorldHack SwitchWorld( MouseCaptorPath );
 
-			FReply Reply = FEventRouter::SendToLeafmost( this, MouseCaptorPath, nullptr, MouseEvent, [this]( const FArrangedWidget& TargetWidget, const FPointerEvent& Event )
+			FReply Reply = FEventRouter::RouteAlongPointerPath( this, FEventRouter::FToLeafmostPolicy(MouseCaptorPath), MouseEvent, [this]( const FArrangedWidget& TargetWidget, const FPointerEvent& Event )
 			{
 				FReply TempReply = FReply::Unhandled();
 				if (Event.IsTouchEvent())
@@ -3684,8 +3763,8 @@ bool FSlateApplication::ProcessMouseButtonUpEvent( FPointerEvent& MouseEvent )
 
 		// Switch worlds widgets in the current path
 		FScopedSwitchWorldHack SwitchWorld( WidgetsUnderCursor );
-
-		FReply Reply = FEventRouter::BubbleAlongPointerPath( this, WidgetsUnderCursor, MouseEvent, [&](const FArrangedWidget& CurWidget, const FPointerEvent& Event)
+		
+		FReply Reply = FEventRouter::RouteAlongPointerPath( this, FEventRouter::FBubblePolicy(WidgetsUnderCursor), MouseEvent, [&](const FArrangedWidget& CurWidget, const FPointerEvent& Event)
 		{
 			FReply TempReply = FReply::Unhandled();
 			if (Event.IsTouchEvent())
@@ -3697,7 +3776,7 @@ bool FSlateApplication::ProcessMouseButtonUpEvent( FPointerEvent& MouseEvent )
 				TempReply = (bIsDragDropping)
 					? CurWidget.Widget->OnDrop( CurWidget.Geometry, FDragDropEvent( Event, LocalDragDropContent ) )
 					: CurWidget.Widget->OnMouseButtonUp( CurWidget.Geometry, Event );
-		}
+			}
 			return TempReply;		
 		});
 
@@ -3767,7 +3846,7 @@ bool FSlateApplication::ProcessMouseWheelOrGestureEvent( FPointerEvent& InWheelE
 	// Switch worlds widgets in the current path
 	FScopedSwitchWorldHack SwitchWorld( EventPath );
 
-	FReply Reply = FEventRouter::BubbleAlongPointerPath( this, EventPath, InWheelEvent, [&InGestureEvent](const FArrangedWidget& CurWidget, const FPointerEvent& Event)
+	FReply Reply = FEventRouter::RouteAlongPointerPath( this, FEventRouter::FBubblePolicy(EventPath), InWheelEvent, [&InGestureEvent](const FArrangedWidget& CurWidget, const FPointerEvent& Event)
 	{
 		FReply TempReply = FReply::Unhandled();
 		// Gesture event gets first shot, if slate doesn't respond to it, we'll try the wheel event.
@@ -3964,24 +4043,24 @@ bool FSlateApplication::ProcessMouseMoveEvent( FPointerEvent& MouseEvent, bool b
 	{
 		if ( !bIsSynthetic )
 		{
-		// Switch worlds widgets in the current path
-		FScopedSwitchWorldHack SwitchWorld( MouseCaptorPath );
-
-			FReply Reply = FEventRouter::SendToLeafmost( this, MouseCaptorPath, nullptr, MouseEvent, [this]( const FArrangedWidget& MouseCaptorWidget, const FPointerEvent& Event )
+			// Switch worlds widgets in the current path
+			FScopedSwitchWorldHack SwitchWorld( MouseCaptorPath );
+			
+			FReply Reply = FEventRouter::RouteAlongPointerPath( this, FEventRouter::FToLeafmostPolicy(MouseCaptorPath), MouseEvent, [this]( const FArrangedWidget& MouseCaptorWidget, const FPointerEvent& Event )
 			{
 				FReply TempReply = FReply::Unhandled();
 				if (Event.IsTouchEvent())
-		{
+				{
 					TempReply = MouseCaptorWidget.Widget->OnTouchMoved( MouseCaptorWidget.Geometry, Event );
-		}
+				}
 				if (!Event.IsTouchEvent() || (!TempReply.IsEventHandled() && this->bTouchFallbackToMouse))
-		{
+				{
 					TempReply = MouseCaptorWidget.Widget->OnMouseMove( MouseCaptorWidget.Geometry, Event );
-		}
+				}
 				return TempReply;
 			} );
-		bHandled = Reply.IsEventHandled();
-	}
+			bHandled = Reply.IsEventHandled();
+		}
 	}
 	else
 	{	
@@ -4013,7 +4092,7 @@ bool FSlateApplication::ProcessMouseMoveEvent( FPointerEvent& MouseEvent, bool b
 		}
 
 		// Bubble the MouseMove event.
-		FReply Reply = FEventRouter::BubbleAlongPointerPath(this, WidgetsUnderCursor, MouseEvent, [&](const FArrangedWidget& CurWidget, const FPointerEvent& Event)
+		FReply Reply = FEventRouter::RouteAlongPointerPath(this, FEventRouter::FBubblePolicy(WidgetsUnderCursor), MouseEvent, [&](const FArrangedWidget& CurWidget, const FPointerEvent& Event)
 		{
 			FReply TempReply = FReply::Unhandled();
 
