Slate/UMG - Exposing more properties on the slate Scrollbar and ScrollBox.  Can now change the visibility of the scrollbox's internal scrollbar as well as thickness.

diff --git a/Engine/Source/Runtime/Slate/Private/Widgets/Layout/SScrollBar.cpp b/Engine/Source/Runtime/Slate/Private/Widgets/Layout/SScrollBar.cpp
--- a/Engine/Source/Runtime/Slate/Private/Widgets/Layout/SScrollBar.cpp
+++ b/Engine/Source/Runtime/Slate/Private/Widgets/Layout/SScrollBar.cpp
@@ -277,7 +277,9 @@ void SScrollBar::Construct(const FArguments& InArgs)
 		.ColorAndOpacity( this, &SScrollBar::GetThumbOpacity )
 		[
 			SNew(SVerticalBox)
-			+SVerticalBox::Slot() .FillHeight( 1 )
+
+			+ SVerticalBox::Slot()
+			.FillHeight( 1 )
 			[
 				SAssignNew(Track, SScrollBarTrack)
 				.Orientation(InArgs._Orientation)
@@ -293,12 +295,12 @@ void SScrollBar::Construct(const FArguments& InArgs)
 				]
 				.ThumbSlot()
 				[
-					SAssignNew(DragThumb , SBorder)
+					SAssignNew(DragThumb, SBorder)
 					.BorderImage( this, &SScrollBar::GetDragThumbImage )
 					.HAlign(HAlign_Center)
 					.VAlign(VAlign_Center)
 					[
-						SNew(SSpacer)
+						SAssignNew(ThicknessSpacer, SSpacer)
 						.Size(InArgs._Thickness)
 					]
 				]
@@ -317,10 +319,7 @@ void SScrollBar::Construct(const FArguments& InArgs)
 	);
 
 	this->EnabledState = TAttribute<bool>( Track.ToSharedRef(), &SScrollBarTrack::IsNeeded );
-	if ( ! InArgs._AlwaysShowScrollbar )
-	{
-		this->Visibility = TAttribute<EVisibility>( SharedThis(this), &SScrollBar::ShouldBeVisible );
-	}
+	SetScrollBarAlwaysVisible(InArgs._AlwaysShowScrollbar);
 }
 
 void SScrollBar::SetOnUserScrolled( const FOnUserScrolled& InHandler )
@@ -525,3 +524,20 @@ bool SScrollBar::IsScrolling() const
 {
 	return bDraggingThumb;
 }
+
+void SScrollBar::SetThickness(TAttribute<FVector2D> InThickness)
+{
+	ThicknessSpacer->SetSize(InThickness);
+}
+
+void SScrollBar::SetScrollBarAlwaysVisible(bool InAlwaysVisible)
+{
+	if ( InAlwaysVisible )
+	{
+		Visibility = EVisibility::Visible;
+	}
+	else
+	{
+		Visibility = TAttribute<EVisibility>(SharedThis(this), &SScrollBar::ShouldBeVisible);
+	}
+}
diff --git a/Engine/Source/Runtime/Slate/Private/Widgets/Layout/SScrollBox.cpp b/Engine/Source/Runtime/Slate/Private/Widgets/Layout/SScrollBox.cpp
--- a/Engine/Source/Runtime/Slate/Private/Widgets/Layout/SScrollBox.cpp
+++ b/Engine/Source/Runtime/Slate/Private/Widgets/Layout/SScrollBox.cpp
@@ -168,6 +168,10 @@ void SScrollBox::Construct( const FArguments& InArgs )
 	{
 		// Make a scroll bar 
 		ScrollBar = ConstructScrollBar();
+		ScrollBar->SetThickness(InArgs._ScrollBarThickness);
+		ScrollBar->SetUserVisibility(InArgs._ScrollBarVisibility);
+		ScrollBar->SetScrollBarAlwaysVisible(InArgs._ScrollBarAlwaysVisible);
+
 		bScrollBarIsExternal = false;
 	}
 
@@ -189,7 +193,6 @@ void SScrollBox::Construct( const FArguments& InArgs )
 TSharedPtr<SScrollBar> SScrollBox::ConstructScrollBar()
 {
 	return TSharedPtr<SScrollBar>(SNew(SScrollBar)
-		.Thickness(FVector2D(5.0f, 5.0f))
 		.Style(ScrollBarStyle)
 		.Orientation(Orientation)
 		.OnUserScrolled(this, &SScrollBox::ScrollBar_OnUserScrolled));
@@ -372,6 +375,21 @@ void SScrollBox::SetOrientation(EOrientation InOrientation)
 	}
 }
 
+void SScrollBox::SetScrollBarVisibility(EVisibility InVisibility)
+{
+	ScrollBar->SetUserVisibility(InVisibility);
+}
+
+void SScrollBox::SetScrollBarAlwaysVisible(bool InAlwaysVisible)
+{
+	ScrollBar->SetScrollBarAlwaysVisible(InAlwaysVisible);
+}
+
+void SScrollBox::SetScrollBarThickness(FVector2D InThickness)
+{
+	ScrollBar->SetThickness(InThickness);
+}
+
 void SScrollBox::Tick( const FGeometry& AllottedGeometry, const double InCurrentTime, const float InDeltaTime )
 {
 	if ( !IsRightClickScrolling() )
@@ -406,10 +424,8 @@ void SScrollBox::Tick( const FGeometry& AllottedGeometry, const double InCurrent
 		// We cannot scroll, so ensure that there is no offset.
 		ScrollPanel->PhysicalOffset = 0.0f;
 	}
-
 }
 
-
 FReply SScrollBox::OnMouseButtonDown( const FGeometry& MyGeometry, const FPointerEvent& MouseEvent )
 {
 	// Zero the scroll velocity so the panel stops immediately on mouse down, even if the user does not drag
diff --git a/Engine/Source/Runtime/Slate/Public/Widgets/Layout/SScrollBar.h b/Engine/Source/Runtime/Slate/Public/Widgets/Layout/SScrollBar.h
--- a/Engine/Source/Runtime/Slate/Public/Widgets/Layout/SScrollBar.h
+++ b/Engine/Source/Runtime/Slate/Public/Widgets/Layout/SScrollBar.h
@@ -16,7 +16,7 @@ public:
 
 	SLATE_BEGIN_ARGS( SScrollBar )
 		: _Style( &FCoreStyle::Get().GetWidgetStyle<FScrollBarStyle>("Scrollbar") )
-        , _OnUserScrolled()
+		, _OnUserScrolled()
 		, _AlwaysShowScrollbar(false)
 		, _Orientation( Orient_Vertical )
 		, _Thickness( FVector2D(12.0f, 12.0f) )
@@ -100,8 +100,15 @@ public:
 	/** @return True if the user is scrolling by dragging the scroll bar thumb. */
 	bool IsScrolling() const;
 
+	/** See UserVisibility attribute */
 	void SetUserVisibility(TAttribute<EVisibility> InUserVisibility) { UserVisibility = InUserVisibility; }
 
+	/** See Thickness attribute */
+	void SetThickness(TAttribute<FVector2D> InThickness);
+
+	/** See ScrollBarAlwaysVisible attribute */
+	void SetScrollBarAlwaysVisible(bool InAlwaysVisible);
+
 	SScrollBar();
 
 protected:
@@ -122,6 +129,7 @@ protected:
 	TAttribute<EVisibility> UserVisibility;
 
 	TSharedPtr<SBorder> DragThumb;
+	TSharedPtr<SSpacer> ThicknessSpacer;
 	bool bDraggingThumb;
 	TSharedPtr<SScrollBarTrack> Track;
 	FOnUserScrolled OnUserScrolled;
diff --git a/Engine/Source/Runtime/Slate/Public/Widgets/Layout/SScrollBox.h b/Engine/Source/Runtime/Slate/Public/Widgets/Layout/SScrollBox.h
--- a/Engine/Source/Runtime/Slate/Public/Widgets/Layout/SScrollBox.h
+++ b/Engine/Source/Runtime/Slate/Public/Widgets/Layout/SScrollBox.h
@@ -43,6 +43,9 @@ public:
 		, _ScrollBarStyle( &FCoreStyle::Get().GetWidgetStyle<FScrollBarStyle>("ScrollBar") )
 		, _ExternalScrollbar()
 		, _Orientation(Orient_Vertical)
+		, _ScrollBarVisibility(EVisibility::Visible)
+		, _ScrollBarAlwaysVisible(false)
+		, _ScrollBarThickness(FVector2D(5, 5))
 		, _OnUserScrolled()
 		{}
 		
@@ -58,7 +61,13 @@ public:
 		SLATE_ARGUMENT( TSharedPtr<SScrollBar>, ExternalScrollbar )
 
 		/** The direction that children will be stacked, and also the direction the box will scroll. */
-		SLATE_ARGUMENT(EOrientation, Orientation)
+		SLATE_ARGUMENT( EOrientation, Orientation )
+
+		SLATE_ARGUMENT( EVisibility, ScrollBarVisibility )
+
+		SLATE_ARGUMENT( bool, ScrollBarAlwaysVisible )
+
+		SLATE_ARGUMENT( FVector2D, ScrollBarThickness )
 
 		/** Called when the button is clicked */
 		SLATE_EVENT(FOnUserScrolled, OnUserScrolled)
@@ -94,6 +103,12 @@ public:
 	/** Sets the current orientation of the scrollbox and updates the layout */
 	void SetOrientation(EOrientation InOrientation);
 
+	void SetScrollBarVisibility(EVisibility InVisibility);
+
+	void SetScrollBarAlwaysVisible(bool InAlwaysVisible);
+
+	void SetScrollBarThickness(FVector2D InThickness);
+
 public:
 
 	// SWidget interface
diff --git a/Engine/Source/Runtime/UMG/Private/Components/ScrollBox.cpp b/Engine/Source/Runtime/UMG/Private/Components/ScrollBox.cpp
--- a/Engine/Source/Runtime/UMG/Private/Components/ScrollBox.cpp
+++ b/Engine/Source/Runtime/UMG/Private/Components/ScrollBox.cpp
@@ -19,6 +19,10 @@ UScrollBox::UScrollBox(const FObjectInitializer& ObjectInitializer)
 
 	WidgetStyle = *Defaults._Style;
 	WidgetBarStyle = *Defaults._ScrollBarStyle;
+
+	AlwaysShowScrollbar = false;
+	ScrollbarThickness = FVector2D(5, 5);
+	ScrollBarVisibility = ESlateVisibility::Visible;
 }
 
 void UScrollBox::ReleaseSlateResources(bool bReleaseChildren)
@@ -59,7 +63,6 @@ TSharedRef<SWidget> UScrollBox::RebuildWidget()
 {
 	MyScrollBox = SNew(SScrollBox)
 		.Style(&WidgetStyle)
-		.ScrollBarStyle(&WidgetBarStyle)
 		.Orientation(Orientation);
 
 	for ( UPanelSlot* Slot : Slots )
@@ -78,6 +81,9 @@ void UScrollBox::SynchronizeProperties()
 {
 	MyScrollBox->SetScrollOffset(DesiredScrollOffset);
 	MyScrollBox->SetOrientation(Orientation);
+	MyScrollBox->SetScrollBarVisibility(UWidget::ConvertSerializedVisibilityToRuntime(ScrollBarVisibility));
+	MyScrollBox->SetScrollBarThickness(ScrollbarThickness);
+	MyScrollBox->SetScrollBarAlwaysVisible(AlwaysShowScrollbar);
 }
 
 void UScrollBox::SetScrollOffset(float NewScrollOffset)
diff --git a/Engine/Source/Runtime/UMG/Public/Components/ScrollBox.h b/Engine/Source/Runtime/UMG/Public/Components/ScrollBox.h
--- a/Engine/Source/Runtime/UMG/Public/Components/ScrollBox.h
+++ b/Engine/Source/Runtime/UMG/Public/Components/ScrollBox.h
@@ -29,10 +29,25 @@ public:
 	USlateWidgetStyleAsset* BarStyle_DEPRECATED;
 
 	/** The orientation of the scrolling and stacking in the box. */
-	UPROPERTY(EditDefaultsOnly, BlueprintReadOnly, Category = "Behavior" )
+	UPROPERTY(EditDefaultsOnly, BlueprintReadOnly, Category="Scroll")
 	TEnumAsByte<EOrientation> Orientation;
 
+	/** Visibility */
+	UPROPERTY(EditDefaultsOnly, BlueprintReadOnly, Category="Scroll")
+	TEnumAsByte<ESlateVisibility::Type> ScrollBarVisibility;
+
+	/**  */
+	UPROPERTY(EditDefaultsOnly, BlueprintReadOnly, Category="Scroll")
+	FVector2D ScrollbarThickness;
+
+	/**  */
+	UPROPERTY(EditDefaultsOnly, BlueprintReadOnly, Category="Scroll")
+	bool AlwaysShowScrollbar;
+	
 	//TODO UMG Add SetOrientation
+	//TODO UMG Add SetScrollBarVisibility
+	//TODO UMG Add SetScrollbarThickness
+	//TODO UMG Add SetAlwaysShowScrollbar
 
 public:
 
