Slate - Adding PreciseClick for mouse interaction with SButton to allow for drag/droppable buttons.

diff --git a/Engine/Source/Runtime/Slate/Private/Widgets/Input/SButton.cpp b/Engine/Source/Runtime/Slate/Private/Widgets/Input/SButton.cpp
--- a/Engine/Source/Runtime/Slate/Private/Widgets/Input/SButton.cpp
+++ b/Engine/Source/Runtime/Slate/Private/Widgets/Input/SButton.cpp
@@ -121,7 +121,7 @@ FReply SButton::OnMouseButtonDown( const FGeometry& MyGeometry, const FPointerEv
 	{
 		bIsPressed = true;
 
-		if(IsEnabled())
+		if ( IsEnabled() )
 		{
 			PlayPressedSound();
 		}
@@ -134,9 +134,9 @@ FReply SButton::OnMouseButtonDown( const FGeometry& MyGeometry, const FPointerEv
 			//You should ALWAYS handle the OnClicked event.
 			ensure(Reply.IsEventHandled() == true);
 		}
-		else if (TouchMethod == EButtonTouchMethod::PreciseTap && MouseEvent.IsTouchEvent() )
+		else if ( IsPreciseTapOrClick(MouseEvent) )
 		{
-			// do not capture the pointer for precise taps
+			// do not capture the pointer for precise taps or clicks
 			// 
 		}
 		else
@@ -156,7 +156,6 @@ FReply SButton::OnMouseButtonDoubleClick( const FGeometry& InMyGeometry, const F
 	return OnMouseButtonDown( InMyGeometry, InMouseEvent );
 }
 
-
 FReply SButton::OnMouseButtonUp( const FGeometry& MyGeometry, const FPointerEvent& MouseEvent )
 {
 	FReply Reply = FReply::Unhandled();
@@ -176,8 +175,7 @@ FReply SButton::OnMouseButtonUp( const FGeometry& MyGeometry, const FPointerEven
 			if( bIsUnderMouse )
 			{
 				// If we asked for a precide tap, all we need is for the user to have not moved their pointer very far.
-				// Precise taps only occur for mobile devices, 
-				const bool bTriggerForTouchEvent = MouseEvent.IsTouchEvent() && TouchMethod == EButtonTouchMethod::PreciseTap;
+				const bool bTriggerForTouchEvent = IsPreciseTapOrClick(MouseEvent);
 
 				// If we were asked to allow the button to be clicked on mouse up, regardless of whether the user
 				// pressed the button down first, then we'll allow the click to proceed without an active capture
@@ -208,11 +206,10 @@ FReply SButton::OnMouseButtonUp( const FGeometry& MyGeometry, const FPointerEven
 	return Reply;
 }
 
-
 FReply SButton::OnMouseMove( const FGeometry& MyGeometry, const FPointerEvent& MouseEvent )
 {
 	const float SlateDragStartDistance = FSlateApplication::Get().GetDragTriggerDistnace();
-	if ( MouseEvent.IsTouchEvent() && TouchMethod == EButtonTouchMethod::PreciseTap && MouseEvent.GetCursorDelta().SizeSquared() > (SlateDragStartDistance*SlateDragStartDistance) )
+	if ( IsPreciseTapOrClick(MouseEvent) && MouseEvent.GetCursorDelta().SizeSquared() > ( SlateDragStartDistance*SlateDragStartDistance ) )
 	{
 		bIsPressed = false;
 	}
@@ -221,7 +218,7 @@ FReply SButton::OnMouseMove( const FGeometry& MyGeometry, const FPointerEvent& M
 
 void SButton::OnMouseEnter( const FGeometry& MyGeometry, const FPointerEvent& MouseEvent )
 {
-	if(IsEnabled())
+	if ( IsEnabled() )
 	{
 		PlayHoverSound();
 	}
@@ -236,12 +233,18 @@ void SButton::OnMouseLeave( const FPointerEvent& MouseEvent )
 
 	// If we're setup to click on mouse-down, then we never capture the mouse and may not receive a
 	// mouse up event, so we need to make sure our pressed state is reset properly here
-	if( ClickMethod == EButtonClickMethod::MouseDown || ( TouchMethod == EButtonTouchMethod::PreciseTap && MouseEvent.IsTouchEvent() ) )
+	if ( ClickMethod == EButtonClickMethod::MouseDown || IsPreciseTapOrClick(MouseEvent) )
 	{
 		bIsPressed = false;
 	}
 }
 
+bool SButton::IsPreciseTapOrClick(const FPointerEvent& MouseEvent) const
+{
+	return ( TouchMethod == EButtonTouchMethod::PreciseTap && MouseEvent.IsTouchEvent() ) ||
+		   ( ClickMethod == EButtonClickMethod::PreciseClick && !MouseEvent.IsTouchEvent() );
+}
+
 void SButton::PlayPressedSound() const
 {
 	FSlateApplication::Get().PlaySound( PressedSound );
diff --git a/Engine/Source/Runtime/Slate/Public/Widgets/Input/SButton.h b/Engine/Source/Runtime/Slate/Public/Widgets/Input/SButton.h
--- a/Engine/Source/Runtime/Slate/Public/Widgets/Input/SButton.h
+++ b/Engine/Source/Runtime/Slate/Public/Widgets/Input/SButton.h
@@ -172,6 +172,9 @@ protected:
 	/** Can this button be focused? */
 	bool bIsFocusable;
 
+	/** Utility function to determine if the incoming mouse event is for a precise tap or click */
+	bool IsPreciseTapOrClick(const FPointerEvent& MouseEvent) const;
+
 	/** Play the pressed sound */
 	void PlayPressedSound() const;
 
diff --git a/Engine/Source/Runtime/SlateCore/Public/Types/SlateEnums.h b/Engine/Source/Runtime/SlateCore/Public/Types/SlateEnums.h
--- a/Engine/Source/Runtime/SlateCore/Public/Types/SlateEnums.h
+++ b/Engine/Source/Runtime/SlateCore/Public/Types/SlateEnums.h
@@ -29,6 +29,12 @@ namespace EButtonClickMethod
 		 * even if the button wasn't pressed down over it.
 		 */
 		MouseUp,
+
+		/**
+		 * Inside a list, buttons can only be clicked with precise tap.
+		 * Moving the pointer will scroll the list, also allows drag-droppable buttons.
+		 */
+		 PreciseClick
 	};
 }
 
