Cleaning up event routing code after the addition of 3D WidgetComponent and hit-testing of 3D Widgets.
- Renamed FVirtualCursorPosition to FVirtualPointerPosition; we are going with Pointer as a generic name for all things Mouse, Touch, and Stylus.
- Slate's FArrangedWidget no longer has a VirtualPointerPosition; this is now stored in FWidgetAndPointer.
- WidgetPaths are aware of a virtual pointer position associated with every arranged widget.

diff --git a/Engine/Source/Editor/UMGEditor/Private/Designer/DesignTimeUtils.cpp b/Engine/Source/Editor/UMGEditor/Private/Designer/DesignTimeUtils.cpp
--- a/Engine/Source/Editor/UMGEditor/Private/Designer/DesignTimeUtils.cpp
+++ b/Engine/Source/Editor/UMGEditor/Private/Designer/DesignTimeUtils.cpp
@@ -17,7 +17,7 @@ bool FDesignTimeUtils::GetArrangedWidget(TSharedRef<SWidget> Widget, FArrangedWi
 	FWidgetPath WidgetPath;
 	if ( FSlateApplication::Get().GeneratePathToWidgetUnchecked(Widget, WidgetPath) )
 	{
-		ArrangedWidget = WidgetPath.FindArrangedWidget(Widget);
+		ArrangedWidget = WidgetPath.FindArrangedWidget(Widget).Get(FArrangedWidget::NullWidget);
 		return true;
 	}
 
@@ -37,7 +37,7 @@ bool FDesignTimeUtils::GetArrangedWidgetRelativeToWindow(TSharedRef<SWidget> Wid
 	FWidgetPath WidgetPath;
 	if ( FSlateApplication::Get().GeneratePathToWidgetUnchecked(Widget, WidgetPath) )
 	{
-		ArrangedWidget = WidgetPath.FindArrangedWidget(Widget);
+		ArrangedWidget = WidgetPath.FindArrangedWidget(Widget).Get(FArrangedWidget::NullWidget);
 		ArrangedWidget.Geometry.AppendTransform(FSlateLayoutTransform(Inverse(CurrentWindowRef->GetPositionInScreen())));
 		return true;
 	}
@@ -59,9 +59,9 @@ bool FDesignTimeUtils::GetArrangedWidgetRelativeToParent(TSharedRef<const SWidge
 
 void FDesignTimeUtils::GetArrangedWidgetRelativeToParent(FWidgetPath& WidgetPath, TSharedRef<const SWidget> Widget, TSharedRef<const SWidget> Parent, FArrangedWidget& ArrangedWidget)
 {
-	FArrangedWidget ArrangedDesigner = WidgetPath.FindArrangedWidget(Parent);
+	FArrangedWidget ArrangedDesigner = WidgetPath.FindArrangedWidget(Parent).Get(FArrangedWidget::NullWidget);
 
-	ArrangedWidget = WidgetPath.FindArrangedWidget(Widget);
+	ArrangedWidget = WidgetPath.FindArrangedWidget(Widget).Get(FArrangedWidget::NullWidget);
 	// !!! WRH 2014/08/26 - This code tries to mutate an FGeometry in an unsupported way. This code should be refactored.
 	ArrangedWidget.Geometry.AppendTransform(FSlateLayoutTransform(Inverse(ArrangedDesigner.Geometry.AbsolutePosition)));
 }
diff --git a/Engine/Source/Runtime/Slate/Private/Framework/Application/SlateApplication.cpp b/Engine/Source/Runtime/Slate/Private/Framework/Application/SlateApplication.cpp
--- a/Engine/Source/Runtime/Slate/Private/Framework/Application/SlateApplication.cpp
+++ b/Engine/Source/Runtime/Slate/Private/Framework/Application/SlateApplication.cpp
@@ -10,17 +10,17 @@
 #include "NotificationManager.h"
 
 
-// @todo slate : Move TranslatePointerEvent into FEventRouter
+// @todo slate : Move TranslatePointerEvent into FEventRouter. Remove external calls to it.
 template <typename PointerEventType>
-static PointerEventType TranslatePointerEvent( const FArrangedWidget& InWidget, const PointerEventType& InEvent )
+static PointerEventType TranslatePointerEvent( const TSharedPtr<FVirtualPointerPosition>& InPosition, const PointerEventType& InEvent )
 {
-	if ( !InWidget.VirtualCursorPosition.IsValid() )
+	if ( !InPosition.IsValid() )
 	{
 		return InEvent;
 	}
 	else
 	{
-		return FPointerEvent::MakeTranslatedEvent<PointerEventType>( InEvent, *InWidget.VirtualCursorPosition );
+		return FPointerEvent::MakeTranslatedEvent<PointerEventType>( InEvent, *InPosition );
 	}
 }
 
@@ -55,9 +55,10 @@ public:
 			bEventSent = true;
 		}
 
-		const FArrangedWidget& GetWidget() const
+		FWidgetAndPointer GetWidget() const
 		{
-			return RoutingPath.Widgets[RoutingPath.Widgets.Num()-1];
+			const int32 WidgetIndex = RoutingPath.Widgets.Num()-1;
+			return FWidgetAndPointer(RoutingPath.Widgets[WidgetIndex], RoutingPath.VirtualPointerPositions[WidgetIndex]);
 		}
 
 		const FWidgetPath& GetRoutingPath() const
@@ -89,9 +90,9 @@ public:
 			++WidgetIndex;
 		}
 
-		const FArrangedWidget& GetWidget() const
+		FWidgetAndPointer GetWidget() const
 		{
-			return RoutingPath.Widgets[WidgetIndex];
+			return FWidgetAndPointer(RoutingPath.Widgets[WidgetIndex], RoutingPath.VirtualPointerPositions[WidgetIndex]);
 		}
 		
 		const FWidgetPath& GetRoutingPath() const
@@ -124,9 +125,9 @@ public:
 			--WidgetIndex;
 		}
 
-		const FArrangedWidget& GetWidget() const
+		FWidgetAndPointer GetWidget() const
 		{
-			return RoutingPath.Widgets[WidgetIndex];
+			return FWidgetAndPointer(RoutingPath.Widgets[WidgetIndex], RoutingPath.VirtualPointerPositions[WidgetIndex]);
 		}
 
 		const FWidgetPath& GetRoutingPath() const
@@ -165,7 +166,7 @@ public:
 
 		for ( ; !Reply.IsEventHandled() && RoutingPolicy.ShouldKeepGoing(); RoutingPolicy.Next() )
 		{
-			const FArrangedWidget& ArrangedWidget = RoutingPolicy.GetWidget();
+			const FWidgetAndPointer& ArrangedWidget = RoutingPolicy.GetWidget();
 			Reply = Lambda( ArrangedWidget, KeyEventCopy ).SetHandler( ArrangedWidget.Widget );
 			ThisApplication->ProcessReply( RoutingPath, Reply, NULL, NULL );
 		}
@@ -195,8 +196,8 @@ public:
 
 		for ( ; !Reply.IsEventHandled() && RoutingPolicy.ShouldKeepGoing(); RoutingPolicy.Next() )
 		{
-			const FArrangedWidget& ArrangedWidget = RoutingPolicy.GetWidget();
-			const FPointerEvent TranslatedEvent = TranslatePointerEvent( ArrangedWidget, PointerEventCopy );
+			const FWidgetAndPointer& ArrangedWidget = RoutingPolicy.GetWidget();
+			const FPointerEvent TranslatedEvent = TranslatePointerEvent( ArrangedWidget.PointerPosition, PointerEventCopy );
 			Reply = Lambda( ArrangedWidget, TranslatedEvent ).SetHandler( ArrangedWidget.Widget );
 			const FWidgetPath* WidgetsUnderCursorPtr = &RoutingPath;
 			ThisApplication->ProcessReply( RoutingPath, Reply, WidgetsUnderCursorPtr, &TranslatedEvent );
@@ -758,8 +759,8 @@ FWidgetPath FSlateApplication::LocateWindowUnderMouse( FVector2D ScreenspaceMous
 
 		if ( Window->IsVisible() && AcceptsInput && Window->IsScreenspaceMouseWithin(ScreenspaceMouseCoordinate) && !bPrevWindowWasModal )
 		{
-			const TArray<FArrangedWidget> ArrangedWidgets = HittestGrid->GetBubblePath( ScreenspaceMouseCoordinate, bIgnoreEnabledStatus );
-			return FWidgetPath( ArrangedWidgets );
+			const TArray<FWidgetAndPointer> WidgetsAndCursors = HittestGrid->GetBubblePath( ScreenspaceMouseCoordinate, bIgnoreEnabledStatus );
+			return FWidgetPath( WidgetsAndCursors );
 		}
 	}
 
@@ -2109,14 +2110,14 @@ void FSlateApplication::ProcessReply( const FWidgetPath& CurrentEventPath, const
 		for (int32 WidgetIndex=0; WidgetIndex < WidgetsUnderMouse->Widgets.Num(); ++WidgetIndex)
 		{
 			const FArrangedWidget& SomeWidget = WidgetsUnderMouse->Widgets[WidgetIndex];
-			SomeWidget.Widget->OnMouseLeave( TranslatePointerEvent( SomeWidget, *InMouseEvent ) );
+			SomeWidget.Widget->OnMouseLeave( TranslatePointerEvent( WidgetsUnderMouse->VirtualPointerPositions[WidgetIndex], *InMouseEvent ) );
 		}
 
 		FDragDropEvent DragDropEvent( *InMouseEvent, ReplyDragDropContent );
 		for (int32 WidgetIndex=0; WidgetIndex < WidgetsUnderMouse->Widgets.Num(); ++WidgetIndex)
 		{
 			const FArrangedWidget& SomeWidget = WidgetsUnderMouse->Widgets[WidgetIndex];
-			SomeWidget.Widget->OnDragEnter( SomeWidget.Geometry, TranslatePointerEvent( SomeWidget, DragDropEvent ) );
+			SomeWidget.Widget->OnDragEnter( SomeWidget.Geometry, TranslatePointerEvent( WidgetsUnderMouse->VirtualPointerPositions[WidgetIndex], DragDropEvent ) );
 		}
 	}
 	
@@ -2338,7 +2339,7 @@ void FSlateApplication::QueryCursor()
 				for( int32 WidgetIndex = WidgetsToQueryForCursor.Widgets.Num() - 1; !CursorResult.IsEventHandled() && WidgetIndex >= 0; --WidgetIndex )
 				{
 					const FArrangedWidget& WidgetToQuery = WidgetsToQueryForCursor.Widgets[ WidgetIndex ];
-					CursorResult = WidgetToQuery.Widget->OnCursorQuery( WidgetToQuery.Geometry, TranslatePointerEvent( WidgetToQuery, CursorEvent ) );
+					CursorResult = WidgetToQuery.Widget->OnCursorQuery( WidgetToQuery.Geometry, TranslatePointerEvent( WidgetsToQueryForCursor.VirtualPointerPositions[WidgetIndex], CursorEvent ) );
 				}
 
 				if (CursorResult.IsEventHandled())
@@ -3152,7 +3153,7 @@ bool FSlateApplication::TakeScreenshot(TSharedRef<SWidget>& Widget, const FIntRe
 	FWidgetPath WidgetPath;
 	FSlateApplication::Get().GeneratePathToWidgetChecked(Widget, WidgetPath);
 
-	FArrangedWidget ArrangedWidget = WidgetPath.FindArrangedWidget(Widget);
+	FArrangedWidget ArrangedWidget = WidgetPath.FindArrangedWidget(Widget).Get(FArrangedWidget::NullWidget);
 	FVector2D Position = ArrangedWidget.Geometry.AbsolutePosition;
 	FVector2D Size = ArrangedWidget.Geometry.GetDrawSize();
 	FVector2D WindowPosition = WidgetWindow->GetPositionInScreen();
@@ -3997,7 +3998,7 @@ bool FSlateApplication::ProcessMouseMoveEvent( FPointerEvent& MouseEvent, bool b
 			FWidgetPath DragDetectPath = DragDetector.DetectDragForWidget.ToWidgetPath();
 			if( DragDetectPath.IsValid() && DragDetector.DetectDragForWidget.GetLastWidget().IsValid() )
 			{
-				FArrangedWidget DetectDragForMe = DragDetectPath.FindArrangedWidget(DragDetector.DetectDragForWidget.GetLastWidget().Pin().ToSharedRef());
+				FWidgetAndPointer DetectDragForMe = DragDetectPath.FindArrangedWidgetAndCursor(DragDetector.DetectDragForWidget.GetLastWidget().Pin().ToSharedRef()).Get(FWidgetAndPointer());
 
 				// A drag has been triggered. The cursor exited some widgets as a result.
 				// This assignment ensures that we will send OnLeave notifications to those widgets.
@@ -4012,7 +4013,7 @@ bool FSlateApplication::ProcessMouseMoveEvent( FPointerEvent& MouseEvent, bool b
 				// Switch worlds widgets in the current path
 				FScopedSwitchWorldHack SwitchWorld( DragDetectPath );
 
-				FPointerEvent TranslatedMouseEvent = TranslatePointerEvent( DetectDragForMe, MouseEvent );
+				FPointerEvent TranslatedMouseEvent = TranslatePointerEvent( DetectDragForMe.PointerPosition, MouseEvent );
 				FReply Reply = DetectDragForMe.Widget->OnDragDetected(DetectDragForMe.Geometry, TranslatedMouseEvent ).SetHandler(DetectDragForMe.Widget);
 				ProcessReply( DragDetectPath, Reply, &DragDetectPath, &TranslatedMouseEvent );
 				LOG_EVENT( EEventLog::DragDetected, Reply );
@@ -4059,14 +4060,17 @@ bool FSlateApplication::ProcessMouseMoveEvent( FPointerEvent& MouseEvent, bool b
 				const TSharedPtr<SWidget>& SomeWidgetPreviouslyUnderCursor = LastWidgetsUnderCursor.Widgets[WidgetIndex].Pin();
 				if( SomeWidgetPreviouslyUnderCursor.IsValid() )
 				{
-					//@todo dhertzka - return a TOptional<FArrangedWidget>
-					FArrangedWidget FoundWidget = WidgetsUnderCursor.FindArrangedWidget( SomeWidgetPreviouslyUnderCursor.ToSharedRef() );
-					if ( FoundWidget.Widget == SNullWidget::NullWidget )
+					TOptional<FArrangedWidget> FoundWidget = WidgetsUnderCursor.FindArrangedWidget( SomeWidgetPreviouslyUnderCursor.ToSharedRef() );
+					const bool bWidgetNoLongerUnderMouse = !FoundWidget.IsSet();
+					if ( bWidgetNoLongerUnderMouse )
 					{
 						// Widget is no longer under cursor, so send a MouseLeave.
+						// The widget might not even be in the hierarchy any more!
+						// Thus, we cannot translate the PointerPosition into the appropriate space for this event.
 						if ( IsDragDropping() )
 						{
-							SomeWidgetPreviouslyUnderCursor->OnDragLeave( TranslatePointerEvent( FoundWidget, DragDropEvent ) );
+							// Note that the event's pointer position is not translated.
+							SomeWidgetPreviouslyUnderCursor->OnDragLeave( DragDropEvent );
 							LOG_EVENT( EEventLog::DragLeave, SomeWidgetPreviouslyUnderCursor );
 
 							// Reset the cursor override
@@ -4074,7 +4078,8 @@ bool FSlateApplication::ProcessMouseMoveEvent( FPointerEvent& MouseEvent, bool b
 						}
 						else
 						{
-							SomeWidgetPreviouslyUnderCursor->OnMouseLeave( TranslatePointerEvent( FoundWidget, MouseEvent ) );
+							// Note that the event's pointer position is not translated.
+							SomeWidgetPreviouslyUnderCursor->OnMouseLeave( MouseEvent );
 							LOG_EVENT( EEventLog::MouseLeave, SomeWidgetPreviouslyUnderCursor );
 						}			
 					}
@@ -4131,13 +4136,13 @@ bool FSlateApplication::ProcessMouseMoveEvent( FPointerEvent& MouseEvent, bool b
 				if ( IsDragDropping() )
 				{
 					// Doing a drag and drop; send a DragDropEvent
-					SomeWidgetUnderCursor.Widget->OnDragEnter( SomeWidgetUnderCursor.Geometry, TranslatePointerEvent( SomeWidgetUnderCursor, DragDropEvent ) );
+					SomeWidgetUnderCursor.Widget->OnDragEnter( SomeWidgetUnderCursor.Geometry, TranslatePointerEvent( WidgetsUnderCursor.VirtualPointerPositions[WidgetIndex], DragDropEvent ) );
 					LOG_EVENT( EEventLog::DragEnter, SomeWidgetUnderCursor.Widget );
 				}
 				else
 				{
 					// Not drag dropping; send regular mouse event
-					SomeWidgetUnderCursor.Widget->OnMouseEnter( SomeWidgetUnderCursor.Geometry, TranslatePointerEvent( SomeWidgetUnderCursor, MouseEvent ) );
+					SomeWidgetUnderCursor.Widget->OnMouseEnter( SomeWidgetUnderCursor.Geometry, TranslatePointerEvent( WidgetsUnderCursor.VirtualPointerPositions[WidgetIndex], MouseEvent ) );
 					LOG_EVENT( EEventLog::MouseEnter, SomeWidgetUnderCursor.Widget );
 				}
 			}
diff --git a/Engine/Source/Runtime/Slate/Private/Widgets/SViewport.cpp b/Engine/Source/Runtime/Slate/Private/Widgets/SViewport.cpp
--- a/Engine/Source/Runtime/Slate/Private/Widgets/SViewport.cpp
+++ b/Engine/Source/Runtime/Slate/Private/Widgets/SViewport.cpp
@@ -320,7 +320,7 @@ void SViewport::OnArrangeChildren( const FGeometry& AllottedGeometry, FArrangedC
 	}
 }
 
-TSharedPtr<FVirtualCursorPosition> SViewport::TranslateMouseCoordinateFor3DChild(const TSharedRef<SWidget>& ChildWidget, const FGeometry& MyGeometry, const FVector2D& ScreenSpaceMouseCoordinate, const FVector2D& LastScreenSpaceMouseCoordinate) const
+TSharedPtr<FVirtualPointerPosition> SViewport::TranslateMouseCoordinateFor3DChild(const TSharedRef<SWidget>& ChildWidget, const FGeometry& MyGeometry, const FVector2D& ScreenSpaceMouseCoordinate, const FVector2D& LastScreenSpaceMouseCoordinate) const
 {
 	if( CustomHitTestPath.IsValid() )
 	{
diff --git a/Engine/Source/Runtime/Slate/Public/Widgets/SViewport.h b/Engine/Source/Runtime/Slate/Public/Widgets/SViewport.h
--- a/Engine/Source/Runtime/Slate/Public/Widgets/SViewport.h
+++ b/Engine/Source/Runtime/Slate/Public/Widgets/SViewport.h
@@ -155,7 +155,7 @@ public:
 	virtual FReply OnMotionDetected( const FGeometry& MyGeometry, const FMotionEvent& MotionEvent ) override;
 	virtual void OnFinishedPointerInput() override;
 	virtual void OnArrangeChildren( const FGeometry& AllottedGeometry, FArrangedChildren& ArrangedChildren ) const override;
-	virtual TSharedPtr<struct FVirtualCursorPosition> TranslateMouseCoordinateFor3DChild(const TSharedRef<SWidget>& ChildWidget, const FGeometry& MyGeometry, const FVector2D& ScreenSpaceMouseCoordinate, const FVector2D& LastScreenSpaceMouseCoordinate) const;
+	virtual TSharedPtr<struct FVirtualPointerPosition> TranslateMouseCoordinateFor3DChild(const TSharedRef<SWidget>& ChildWidget, const FGeometry& MyGeometry, const FVector2D& ScreenSpaceMouseCoordinate, const FVector2D& LastScreenSpaceMouseCoordinate) const;
 
 private:
 
diff --git a/Engine/Source/Runtime/SlateCore/Private/Input/HittestGrid.cpp b/Engine/Source/Runtime/SlateCore/Private/Input/HittestGrid.cpp
--- a/Engine/Source/Runtime/SlateCore/Private/Input/HittestGrid.cpp
+++ b/Engine/Source/Runtime/SlateCore/Private/Input/HittestGrid.cpp
@@ -39,7 +39,7 @@ FHittestGrid::FHittestGrid()
 {
 }
 
-TArray<FArrangedWidget> FHittestGrid::GetBubblePath( FVector2D DesktopSpaceCoordinate, bool bIgnoreEnabledStatus )
+TArray<FWidgetAndPointer> FHittestGrid::GetBubblePath( FVector2D DesktopSpaceCoordinate, bool bIgnoreEnabledStatus )
 {
 	if (WidgetsCachedThisFrame->Num() > 0 && Cells.Num() > 0)
 	{
@@ -82,18 +82,18 @@ TArray<FArrangedWidget> FHittestGrid::GetBubblePath( FVector2D DesktopSpaceCoord
 			}
 		}
 		
-		TArray<FArrangedWidget> BubblePath =[&]()
+		TArray<FWidgetAndPointer> BubblePath =[&]()
 		{
 			if( HitWidgetIndex != INDEX_NONE )
 			{
 				const FCachedWidget& PhysicallyHitWidget = (*WidgetsCachedThisFrame)[HitWidgetIndex];
 				if(PhysicallyHitWidget.CustomPath.IsValid())
 				{
-					return PhysicallyHitWidget.CustomPath.Pin()->GetBubblePath(PhysicallyHitWidget.CachedGeometry, DesktopSpaceCoordinate, bIgnoreEnabledStatus);
+					return PhysicallyHitWidget.CustomPath.Pin()->GetBubblePathAndVirtualCursors(PhysicallyHitWidget.CachedGeometry, DesktopSpaceCoordinate, bIgnoreEnabledStatus);
 				}
 			}
 
-			return TArray<FArrangedWidget>();
+			return TArray<FWidgetAndPointer>();
 
 		}();
 
@@ -111,7 +111,7 @@ TArray<FArrangedWidget> FHittestGrid::GetBubblePath( FVector2D DesktopSpaceCoord
 				bPathUninterrupted = CachedWidgetPtr.IsValid();
 				if (bPathUninterrupted)
 				{
-					BubblePath.Insert(FArrangedWidget(CachedWidgetPtr.ToSharedRef(), CurCachedWidget.CachedGeometry), 0);
+					BubblePath.Insert( FWidgetAndPointer( FArrangedWidget(CachedWidgetPtr.ToSharedRef(), CurCachedWidget.CachedGeometry), TSharedPtr<FVirtualPointerPosition>() ), 0 );
 					CurWidgetIndex = CurCachedWidget.ParentIndex;
 				}
 			}
@@ -122,7 +122,7 @@ TArray<FArrangedWidget> FHittestGrid::GetBubblePath( FVector2D DesktopSpaceCoord
 				// A widget in the path to the root has been removed, so anything
 				// we thought we had hittest is no longer actually there.
 				// Pretend we didn't hit anything.
-				BubblePath = TArray<FArrangedWidget>();
+				BubblePath = TArray<FWidgetAndPointer>();
 			}
 
 			// Disabling a widget disables all of its logical children
@@ -141,12 +141,12 @@ TArray<FArrangedWidget> FHittestGrid::GetBubblePath( FVector2D DesktopSpaceCoord
 		}
 		else
 		{
-			return TArray<FArrangedWidget>();
+			return TArray<FWidgetAndPointer>();
 		}
 	}
 	else
 	{
-		return TArray<FArrangedWidget>();
+		return TArray<FWidgetAndPointer>();
 	}
 }
 
diff --git a/Engine/Source/Runtime/SlateCore/Private/Layout/ArrangedWidget.cpp b/Engine/Source/Runtime/SlateCore/Private/Layout/ArrangedWidget.cpp
--- a/Engine/Source/Runtime/SlateCore/Private/Layout/ArrangedWidget.cpp
+++ b/Engine/Source/Runtime/SlateCore/Private/Layout/ArrangedWidget.cpp
@@ -3,10 +3,20 @@
 #include "SlateCorePrivatePCH.h"
 
 
-/* FArrangedWidget interface
- *****************************************************************************/
+FArrangedWidget FArrangedWidget::NullWidget(SNullWidget::NullWidget, FGeometry());
 
 FString FArrangedWidget::ToString( ) const
 {
 	return FString::Printf(TEXT("%s @ %s"), *Widget->ToString(), *Geometry.ToString());
 }
+
+
+FWidgetAndPointer::FWidgetAndPointer()
+: FArrangedWidget(FArrangedWidget::NullWidget)
+, PointerPosition(TSharedPtr<FVirtualPointerPosition>())
+{}
+
+FWidgetAndPointer::FWidgetAndPointer( const FArrangedWidget& InWidget, const TSharedPtr<FVirtualPointerPosition>& InPosition )
+: FArrangedWidget(InWidget)
+, PointerPosition(InPosition)
+{}
\ No newline at end of file
diff --git a/Engine/Source/Runtime/SlateCore/Private/Layout/WidgetPath.cpp b/Engine/Source/Runtime/SlateCore/Private/Layout/WidgetPath.cpp
--- a/Engine/Source/Runtime/SlateCore/Private/Layout/WidgetPath.cpp
+++ b/Engine/Source/Runtime/SlateCore/Private/Layout/WidgetPath.cpp
@@ -3,7 +3,79 @@
 #include "SlateCorePrivatePCH.h"
 
 
-FArrangedWidget FWidgetPath::FindArrangedWidget( TSharedRef<const SWidget> WidgetToFind ) const
+FWidgetPath::FWidgetPath()
+: Widgets( EVisibility::Visible )
+, TopLevelWindow()
+, VirtualPointerPositions()
+{
+}
+
+FWidgetPath::FWidgetPath( TSharedPtr<SWindow> InTopLevelWindow, const FArrangedChildren& InWidgetPath )
+: Widgets( InWidgetPath )
+, TopLevelWindow(InTopLevelWindow)
+, VirtualPointerPositions()
+{
+}
+
+FWidgetPath::FWidgetPath( TArray<FWidgetAndPointer> InWidgetsAndPointers )
+: Widgets( FArrangedChildren::Hittest2_FromArray(InWidgetsAndPointers) )
+, TopLevelWindow( InWidgetsAndPointers.Num() > 0 ? StaticCastSharedRef<SWindow>(InWidgetsAndPointers[0].Widget) : TSharedPtr<SWindow>(nullptr) )
+, VirtualPointerPositions( [&InWidgetsAndPointers]()
+	{ 
+		TArray< TSharedPtr<FVirtualPointerPosition> > Pointers;
+		Pointers.Reserve(InWidgetsAndPointers.Num());
+		for ( const FWidgetAndPointer& WidgetAndPointer : InWidgetsAndPointers )
+		{
+			Pointers.Add( WidgetAndPointer.PointerPosition );
+		};
+		return Pointers;
+	}())
+{
+}
+
+FWidgetPath FWidgetPath::GetPathDownTo( TSharedRef<const SWidget> MarkerWidget ) const
+{
+	FArrangedChildren ClippedPath(EVisibility::Visible);
+	bool bCopiedMarker = false;
+	for( int32 WidgetIndex = 0; !bCopiedMarker && WidgetIndex < Widgets.Num(); ++WidgetIndex )
+	{
+		ClippedPath.AddWidget( Widgets[WidgetIndex] );
+		bCopiedMarker = (Widgets[WidgetIndex].Widget == MarkerWidget);
+	}
+		
+	if ( bCopiedMarker )
+	{
+		// We found the MarkerWidget and copied the path down to (and including) it.
+		return FWidgetPath( TopLevelWindow, ClippedPath );
+	}
+	else
+	{
+		// The MarkerWidget was not in the widget path. We failed.
+		return FWidgetPath( nullptr, FArrangedChildren(EVisibility::Visible) );		
+	}	
+}
+
+const TSharedPtr<FVirtualPointerPosition>& FWidgetPath::GetCursorAt( int32 Index ) const
+{
+	return VirtualPointerPositions[Index];
+}
+
+
+bool FWidgetPath::ContainsWidget( TSharedRef<const SWidget> WidgetToFind ) const
+{
+	for(int32 WidgetIndex = 0; WidgetIndex < Widgets.Num(); ++WidgetIndex)
+	{
+		if ( Widgets[WidgetIndex].Widget == WidgetToFind )
+		{
+			return true;
+		}
+	}
+		
+	return false;
+}
+
+
+TOptional<FArrangedWidget> FWidgetPath::FindArrangedWidget( TSharedRef<const SWidget> WidgetToFind ) const
 {
 	for(int32 WidgetIndex = 0; WidgetIndex < Widgets.Num(); ++WidgetIndex)
 	{
@@ -13,7 +85,55 @@ FArrangedWidget FWidgetPath::FindArrangedWidget( TSharedRef<const SWidget> Widge
 		}
 	}
 
-	return FArrangedWidget( SNullWidget::NullWidget, FGeometry() );
+	return TOptional<FArrangedWidget>();
+}
+
+TOptional<FWidgetAndPointer> FWidgetPath::FindArrangedWidgetAndCursor( TSharedRef<const SWidget> WidgetToFind ) const
+{
+	const int32 Index = Widgets.IndexOfByPredicate( [&WidgetToFind]( const FArrangedWidget& SomeWidget )
+	{
+		return SomeWidget.Widget == WidgetToFind;
+	} );
+
+	return (Index != INDEX_NONE)
+		? FWidgetAndPointer( Widgets[Index], VirtualPointerPositions[Index] )
+		: FWidgetAndPointer();
+}
+
+	
+TSharedRef<SWindow> FWidgetPath::GetWindow()
+{
+	check(IsValid());
+
+	TSharedRef<SWindow> FirstWidgetWindow = StaticCastSharedRef<SWindow>(Widgets[0].Widget);
+	return FirstWidgetWindow;
+}
+
+
+TSharedRef<SWindow> FWidgetPath::GetWindow() const
+{
+	check(IsValid());
+
+	TSharedRef<SWindow> FirstWidgetWindow = StaticCastSharedRef<SWindow>(Widgets[0].Widget);
+	return FirstWidgetWindow;
+}
+
+
+bool FWidgetPath::IsValid() const
+{
+	return Widgets.Num() > 0;
+}
+
+	
+FString FWidgetPath::ToString() const
+{
+	FString StringBuffer;
+	for( int32 WidgetIndex = Widgets.Num()-1; WidgetIndex >= 0; --WidgetIndex )
+	{
+		StringBuffer += Widgets[WidgetIndex].ToString();
+		StringBuffer += TEXT("\n");
+	}
+	return StringBuffer;
 }
 
 
@@ -120,9 +240,9 @@ FWidgetPath FWeakWidgetPath::ToWidgetPath(EInterruptedPathHandling::Type Interru
 
 FWeakWidgetPath::EPathResolutionResult::Result FWeakWidgetPath::ToWidgetPath( FWidgetPath& WidgetPath, EInterruptedPathHandling::Type InterruptedPathHandling, const FPointerEvent* PointerEvent ) const
 {
-	FArrangedChildren PathWithGeometries(EVisibility::Visible);
+	TArray<FWidgetAndPointer> PathWithGeometries;
 	TArray< TSharedPtr<SWidget> > WidgetPtrs;
-	
+		
 	// Convert the weak pointers into shared pointers because we are about to do something with this path instead of just observe it.
 	TSharedPtr<SWindow> TopLevelWindowPtr = Window.Pin();
 	for( TArray< TWeakPtr<SWidget> >::TConstIterator SomeWeakWidgetPtr( Widgets ); SomeWeakWidgetPtr; ++SomeWeakWidgetPtr )
@@ -139,11 +259,14 @@ FWeakWidgetPath::EPathResolutionResult::Result FWeakWidgetPath::ToWidgetPath( FW
 		bPathUninterrupted = true;
 
 		FGeometry ParentGeometry = TopLevelWindowPtr->GetWindowGeometryInScreen();
-		PathWithGeometries.AddWidget( FArrangedWidget( TopLevelWindowPtr.ToSharedRef(), ParentGeometry ) );
+		PathWithGeometries.Add( FWidgetAndPointer(
+			FArrangedWidget( TopLevelWindowPtr.ToSharedRef(), ParentGeometry ),
+			// @todo slate: this should be the cursor's virtual position in window space.
+			TSharedPtr<FVirtualPointerPosition>() ) );
 		
 		FArrangedChildren ArrangedChildren(EVisibility::Visible, true);
 		
-		TSharedPtr<FVirtualCursorPosition> VirtualCursorPos;
+		TSharedPtr<FVirtualPointerPosition> VirtualPointerPos;
 		// For every widget in the vertical slice...
 		for( int32 WidgetIndex = 0; bPathUninterrupted && WidgetIndex < WidgetPtrs.Num()-1; ++WidgetIndex )
 		{
@@ -163,16 +286,14 @@ FWeakWidgetPath::EPathResolutionResult::Result FWeakWidgetPath::ToWidgetPath( FW
 
 					if ( ArrangedWidget.Widget == WidgetPtrs[WidgetIndex+1] )
 					{
-						if( PointerEvent && !VirtualCursorPos.IsValid() )
+						if( PointerEvent && !VirtualPointerPos.IsValid() )
 						{
-							VirtualCursorPos = CurWidget->TranslateMouseCoordinateFor3DChild( ArrangedWidget.Widget, ParentGeometry, PointerEvent->GetScreenSpacePosition(), PointerEvent->GetLastScreenSpacePosition() );
+							VirtualPointerPos = CurWidget->TranslateMouseCoordinateFor3DChild( ArrangedWidget.Widget, ParentGeometry, PointerEvent->GetScreenSpacePosition(), PointerEvent->GetLastScreenSpacePosition() );
 						}
 
-						ArrangedWidget.VirtualCursorPosition = VirtualCursorPos;
-
 						bFoundChild = true;
-						// Remember the widget and the associated geometry.
-						PathWithGeometries.AddWidget( ArrangedChildren[SearchIndex] );
+						// Remember the widget, the associated geometry, and the pointer position in a transformed space.
+						PathWithGeometries.Add( FWidgetAndPointer(ArrangedChildren[SearchIndex], VirtualPointerPos) );
 						// The next child in the vertical slice will be arranged with respect to its parent's geometry.
 						ParentGeometry = ArrangedChildren[SearchIndex].Geometry;
 					}
@@ -187,7 +308,7 @@ FWeakWidgetPath::EPathResolutionResult::Result FWeakWidgetPath::ToWidgetPath( FW
 		}			
 	}
 	
-	WidgetPath = FWidgetPath( TopLevelWindowPtr, PathWithGeometries );
+	WidgetPath = FWidgetPath( PathWithGeometries );
 	return bPathUninterrupted ? EPathResolutionResult::Live : EPathResolutionResult::Truncated;
 }
 
diff --git a/Engine/Source/Runtime/SlateCore/Public/Input/Events.h b/Engine/Source/Runtime/SlateCore/Public/Input/Events.h
--- a/Engine/Source/Runtime/SlateCore/Public/Input/Events.h
+++ b/Engine/Source/Runtime/SlateCore/Public/Input/Events.h
@@ -84,8 +84,12 @@ private:
 /**
  * Represents the current and last cursor position in a "virtual window" for events that are routed to widgets transformed in a 3D scene.
  */
-struct FVirtualCursorPosition
+struct FVirtualPointerPosition
 {
+	FVirtualPointerPosition()
+	: CurrentCursorPosition(FVector2D::ZeroVector)
+	, LastCursorPosition(FVector2D::ZeroVector)
+	{}
 	FVector2D CurrentCursorPosition;
 	FVector2D LastCursorPosition;
 };
@@ -271,7 +275,7 @@ public:
 	/** The event path provides additional context for handling */
 	FGeometry FindGeometry( const TSharedRef<SWidget>& WidgetToFind ) const
 	{
-		return EventPath->FindArrangedWidget(WidgetToFind).Geometry;
+		return EventPath->FindArrangedWidget(WidgetToFind).Get(FArrangedWidget::NullWidget).Geometry;
 	}
 
 	TSharedRef<SWindow> GetWindow() const
@@ -628,7 +632,7 @@ public:
 	SLATECORE_API virtual FText ToText() const override;
 
 	template<typename PointerEventType>
-	static PointerEventType MakeTranslatedEvent( const PointerEventType& InPointerEvent, const FVirtualCursorPosition& VirtualPosition )
+	static PointerEventType MakeTranslatedEvent( const PointerEventType& InPointerEvent, const FVirtualPointerPosition& VirtualPosition )
 	{
 		PointerEventType NewEvent = InPointerEvent;
 		NewEvent.ScreenSpacePosition = VirtualPosition.CurrentCursorPosition;
diff --git a/Engine/Source/Runtime/SlateCore/Public/Input/HittestGrid.h b/Engine/Source/Runtime/SlateCore/Public/Input/HittestGrid.h
--- a/Engine/Source/Runtime/SlateCore/Public/Input/HittestGrid.h
+++ b/Engine/Source/Runtime/SlateCore/Public/Input/HittestGrid.h
@@ -9,11 +9,11 @@ class ICustomHitTestPath
 public:
 	virtual ~ICustomHitTestPath(){}
 
-	virtual TArray<FArrangedWidget> GetBubblePath( const FGeometry& InGeometry, FVector2D DesktopSpaceCoordinate, bool bIgnoreEnabledStatus ) const = 0;
+	virtual TArray<FWidgetAndPointer> GetBubblePathAndVirtualCursors( const FGeometry& InGeometry, FVector2D DesktopSpaceCoordinate, bool bIgnoreEnabledStatus ) const = 0;
 
 	virtual void ArrangeChildren( FArrangedChildren& ArrangedChildren ) const = 0;
 
-	virtual TSharedPtr<struct FVirtualCursorPosition> TranslateMouseCoordinateFor3DChild( const TSharedRef<SWidget>& ChildWidget, const FGeometry& ViewportGeometry, const FVector2D& ScreenSpaceMouseCoordinate, const FVector2D& LastScreenSpaceMouseCoordinate ) const = 0;
+	virtual TSharedPtr<struct FVirtualPointerPosition> TranslateMouseCoordinateFor3DChild( const TSharedRef<SWidget>& ChildWidget, const FGeometry& ViewportGeometry, const FVector2D& ScreenSpaceMouseCoordinate, const FVector2D& LastScreenSpaceMouseCoordinate ) const = 0;
 };
 
 class SLATECORE_API FHittestGrid
@@ -26,7 +26,7 @@ public:
 	 * Given a Slate Units coordinate in virtual desktop space, perform a hittest
 	 * and return the path along which the corresponding event would be bubbled.
 	 */
-	TArray<FArrangedWidget> GetBubblePath( FVector2D DesktopSpaceCoordinate, bool bIgnoreEnabledStatus );
+	TArray<FWidgetAndPointer> GetBubblePath( FVector2D DesktopSpaceCoordinate, bool bIgnoreEnabledStatus );
 
 	/**
 	 * Clear the hittesting area and prepare to execute a new frame.
diff --git a/Engine/Source/Runtime/SlateCore/Public/Layout/ArrangedChildren.h b/Engine/Source/Runtime/SlateCore/Public/Layout/ArrangedChildren.h
--- a/Engine/Source/Runtime/SlateCore/Public/Layout/ArrangedChildren.h
+++ b/Engine/Source/Runtime/SlateCore/Public/Layout/ArrangedChildren.h
@@ -30,10 +30,14 @@ class SLATECORE_API FArrangedChildren
 	}
 
 	// @todo hittest2.0 : we should get rid of this eventually.
-	static FArrangedChildren Hittest2_FromArray(const TArray<FArrangedWidget>& InWidgets)
+	static FArrangedChildren Hittest2_FromArray(const TArray<FWidgetAndPointer>& InWidgets)
 	{
 		FArrangedChildren Temp( EVisibility::All );
-		Temp.Array = InWidgets;
+		Temp.Array.Reserve(InWidgets.Num());
+		for (const FWidgetAndPointer& WidgetAndPointer : InWidgets)
+		{
+			Temp.Array.Add(WidgetAndPointer);
+		}		
 		return Temp;
 	}
 
@@ -138,6 +142,12 @@ public:
 		return Array.Find(ItemToFind);
 	}
 
+	template<typename PredicateType>
+	int32 IndexOfByPredicate( const PredicateType& Pred ) const
+	{
+		return Array.IndexOfByPredicate( Pred );
+	}
+
 	void Remove( int32 Index, int32 Count=1 )
 	{
 		Array.RemoveAt(Index, Count);
diff --git a/Engine/Source/Runtime/SlateCore/Public/Layout/ArrangedWidget.h b/Engine/Source/Runtime/SlateCore/Public/Layout/ArrangedWidget.h
--- a/Engine/Source/Runtime/SlateCore/Public/Layout/ArrangedWidget.h
+++ b/Engine/Source/Runtime/SlateCore/Public/Layout/ArrangedWidget.h
@@ -4,7 +4,7 @@
 
 #include "Geometry.h"
 
-struct FVirtualCursorPosition;
+struct FVirtualPointerPosition;
 
 /**
  * A pair: Widget and its Geometry. Widgets populate an list of WidgetGeometries
@@ -19,6 +19,8 @@ public:
 		, Widget(InWidget)
 	{ }
 
+	SLATECORE_API static FArrangedWidget NullWidget;
+
 public:
 
 	/**
@@ -48,7 +50,13 @@ public:
 
 	/** The widget that is being arranged. */
 	TSharedRef<SWidget> Widget;
-
-	/** Virtual cursor position if this widget is transformed in a 3D scene */
-	TSharedPtr<FVirtualCursorPosition> VirtualCursorPosition;
 };
+
+struct SLATECORE_API FWidgetAndPointer : public FArrangedWidget
+{
+	FWidgetAndPointer();
+
+	FWidgetAndPointer( const FArrangedWidget& InWidget, const TSharedPtr<FVirtualPointerPosition>& InPosition );
+
+	TSharedPtr<FVirtualPointerPosition> PointerPosition;
+};
\ No newline at end of file
diff --git a/Engine/Source/Runtime/SlateCore/Public/Layout/WidgetPath.h b/Engine/Source/Runtime/SlateCore/Public/Layout/WidgetPath.h
--- a/Engine/Source/Runtime/SlateCore/Public/Layout/WidgetPath.h
+++ b/Engine/Source/Runtime/SlateCore/Public/Layout/WidgetPath.h
@@ -28,111 +28,50 @@ struct FWidgetMatcher
  * The window is needed for its ability to determine its own geometry, from which the geometries of the rest
  * of the widget can be determined.
  */
-class FWidgetPath
+class SLATECORE_API FWidgetPath
 {
 public:
-	/** Constructor */
-	FWidgetPath()
-	: Widgets( EVisibility::Visible )
-	{
-	}
+	FWidgetPath();
 
-	FWidgetPath( TSharedPtr<SWindow> InTopLevelWindow, const FArrangedChildren& InWidgetPath )
-	: Widgets( InWidgetPath )
-	, TopLevelWindow(InTopLevelWindow)
-	{
-	
-	}
+	FWidgetPath( TSharedPtr<SWindow> InTopLevelWindow, const FArrangedChildren& InWidgetPath );
 
-	FWidgetPath( const TArray<FArrangedWidget>& InWidgetPath )
-	: Widgets( FArrangedChildren::Hittest2_FromArray(InWidgetPath) )
-	, TopLevelWindow( InWidgetPath.Num() > 0 ? StaticCastSharedRef<SWindow>(InWidgetPath[0].Widget) : TSharedPtr<SWindow>(nullptr) )
-	{
-	
-	}
+	FWidgetPath( TArray<FWidgetAndPointer> InWidgetsAndPointers );
 
 	/**
 	 * @param MarkerWidget Copy the path up to and including this widget 
 	 *
 	 * @return a copy of the widget path down to and including the MarkerWidget. If the MarkerWidget is not found in the path, return an invalid path.
 	 */
-	FWidgetPath GetPathDownTo( TSharedRef<const SWidget> MarkerWidget ) const
-	{
-		FArrangedChildren ClippedPath(EVisibility::Visible);
-		bool bCopiedMarker = false;
-		for( int32 WidgetIndex = 0; !bCopiedMarker && WidgetIndex < Widgets.Num(); ++WidgetIndex )
-		{
-			ClippedPath.AddWidget( Widgets[WidgetIndex] );
-			bCopiedMarker = (Widgets[WidgetIndex].Widget == MarkerWidget);
-		}
-		
-		if ( bCopiedMarker )
-		{
-			// We found the MarkerWidget and copied the path down to (and including) it.
-			return FWidgetPath( TopLevelWindow, ClippedPath );
-		}
-		else
-		{
-			// The MarkerWidget was not in the widget path. We failed.
-			return FWidgetPath( nullptr, FArrangedChildren(EVisibility::Visible) );		
-		}	
-	}
+	FWidgetPath GetPathDownTo( TSharedRef<const SWidget> MarkerWidget ) const;
 	
+	/** Get the widget and associated virtual cursor at a given index. */
+	const TSharedPtr<FVirtualPointerPosition>& GetCursorAt( int32 Index ) const;
+
 	/** @return true if the WidgetToFind is in this WidgetPath, false otherwise. */
-	bool ContainsWidget( TSharedRef<const SWidget> WidgetToFind ) const
-	{
-		for(int32 WidgetIndex = 0; WidgetIndex < Widgets.Num(); ++WidgetIndex)
-		{
-			if ( Widgets[WidgetIndex].Widget == WidgetToFind )
-			{
-				return true;
-			}
-		}
-		
-		return false;
-	}
+	bool ContainsWidget( TSharedRef<const SWidget> WidgetToFind ) const;
 
-	SLATECORE_API FArrangedWidget FindArrangedWidget( TSharedRef<const SWidget> WidgetToFind ) const;
+	TOptional<FArrangedWidget> FindArrangedWidget( TSharedRef<const SWidget> WidgetToFind ) const;
+
+	TOptional<FWidgetAndPointer> FindArrangedWidgetAndCursor( TSharedRef<const SWidget> WidgetToFind ) const;
 	
 	/**
 	 * Get the first (top-most) widget in this path, which is always a window; assumes path is valid
 	 *
 	 * @return Window at the top of this path
 	 */
-	TSharedRef<SWindow> GetWindow()
-	{
-		check(IsValid());
-
-		TSharedRef<SWindow> FirstWidgetWindow = StaticCastSharedRef<SWindow>(Widgets[0].Widget);
-		return FirstWidgetWindow;
-	}
+	TSharedRef<SWindow> GetWindow();
 
 	/**
 	 * Get the first (top-most) widget in this path, which is always a window; assumes path is valid
 	 *
 	 * @return Window at the top of this path
 	 */
-	TSharedRef<SWindow> GetWindow() const
-	{
-		check(IsValid());
-
-		TSharedRef<SWindow> FirstWidgetWindow = StaticCastSharedRef<SWindow>(Widgets[0].Widget);
-		return FirstWidgetWindow;
-	}
+	TSharedRef<SWindow> GetWindow() const;
 	
 	/** A valid path has at least one widget in it */
-	bool IsValid() const { return Widgets.Num() > 0; }
+	bool IsValid() const;
 	
-	FString ToString() const
-	{
-		FString StringBuffer;
-		for( int32 WidgetIndex = Widgets.Num()-1; WidgetIndex >= 0; --WidgetIndex )
-		{
-			StringBuffer += Widgets[WidgetIndex].ToString();
-			StringBuffer += TEXT("\n");
-		}
-		return StringBuffer;
-	}
+	FString ToString() const;
 
 	/**
 	 * Extend the current path such that it reaches some widget that qualifies as a Match
@@ -202,6 +141,8 @@ public:
 
 	FArrangedChildren Widgets;
 	TSharedPtr< SWindow > TopLevelWindow;
+	TArray< TSharedPtr<FVirtualPointerPosition> > VirtualPointerPositions;
+
 
 private:
 
@@ -232,7 +173,7 @@ class SLATECORE_API FWeakWidgetPath
 {
 public:
 	/** Construct a weak widget path from a widget path. Defaults to an invalid path. */
-	FWeakWidgetPath( const FWidgetPath& InWidgetPath = FWidgetPath( TSharedPtr<SWindow>(nullptr), FArrangedChildren(EVisibility::Visible) ) );
+	FWeakWidgetPath( const FWidgetPath& InWidgetPath = FWidgetPath() );
 
 	/** Should interrupted paths truncate or return an invalid path? */
 	struct EInterruptedPathHandling
diff --git a/Engine/Source/Runtime/SlateCore/Public/Widgets/SWidget.h b/Engine/Source/Runtime/SlateCore/Public/Widgets/SWidget.h
--- a/Engine/Source/Runtime/SlateCore/Public/Widgets/SWidget.h
+++ b/Engine/Source/Runtime/SlateCore/Public/Widgets/SWidget.h
@@ -403,7 +403,7 @@ public:
 	 */
 	virtual FReply OnMotionDetected( const FGeometry& MyGeometry, const FMotionEvent& InMotionEvent );
 
-	virtual TSharedPtr<struct FVirtualCursorPosition> TranslateMouseCoordinateFor3DChild( const TSharedRef<SWidget>& ChildWidget, const FGeometry& MyGeometry, const FVector2D& ScreenSpaceMouseCoordinate, const FVector2D& LastScreenSpaceMouseCoordinate ) const { return nullptr; }
+	virtual TSharedPtr<struct FVirtualPointerPosition> TranslateMouseCoordinateFor3DChild( const TSharedRef<SWidget>& ChildWidget, const FGeometry& MyGeometry, const FVector2D& ScreenSpaceMouseCoordinate, const FVector2D& LastScreenSpaceMouseCoordinate ) const { return nullptr; }
 
 
 	/**
diff --git a/Engine/Source/Runtime/UMG/Private/Components/WidgetComponent.cpp b/Engine/Source/Runtime/UMG/Private/Components/WidgetComponent.cpp
--- a/Engine/Source/Runtime/UMG/Private/Components/WidgetComponent.cpp
+++ b/Engine/Source/Runtime/UMG/Private/Components/WidgetComponent.cpp
@@ -187,7 +187,7 @@ public:
 	{}
 
 	// ICustomHitTestPath implementation
-	virtual TArray<FArrangedWidget> GetBubblePath(const FGeometry& InGeometry, FVector2D DesktopSpaceCoordinate, bool bIgnoreEnabledStatus) const override
+	virtual TArray<FWidgetAndPointer> GetBubblePathAndVirtualCursors(const FGeometry& InGeometry, FVector2D DesktopSpaceCoordinate, bool bIgnoreEnabledStatus) const override
 	{
 		if( World.IsValid() && ensure( World->IsGameWorld() ) )
 		{
@@ -207,7 +207,7 @@ public:
 				}
 			}
 		}
-		return TArray<FArrangedWidget>();
+		return TArray<FWidgetAndPointer>();
 	}
 
 	virtual void ArrangeChildren( FArrangedChildren& ArrangedChildren ) const override
@@ -225,7 +225,7 @@ public:
 		}
 	}
 
-	virtual TSharedPtr<struct FVirtualCursorPosition> TranslateMouseCoordinateFor3DChild( const TSharedRef<SWidget>& ChildWidget, const FGeometry& ViewportGeometry, const FVector2D& ScreenSpaceMouseCoordinate, const FVector2D& LastScreenSpaceMouseCoordinate ) const override
+	virtual TSharedPtr<struct FVirtualPointerPosition> TranslateMouseCoordinateFor3DChild( const TSharedRef<SWidget>& ChildWidget, const FGeometry& ViewportGeometry, const FVector2D& ScreenSpaceMouseCoordinate, const FVector2D& LastScreenSpaceMouseCoordinate ) const override
 	{
 		if(World.IsValid() && ensure(World->IsGameWorld()))
 		{
@@ -248,7 +248,7 @@ public:
 						{
 							if ( WidgetComponent == HitResult.Component.Get() )
 							{
-								TSharedPtr<FVirtualCursorPosition> VirtualCursorPos = MakeShareable(new FVirtualCursorPosition);
+								TSharedPtr<FVirtualPointerPosition> VirtualCursorPos = MakeShareable(new FVirtualPointerPosition);
 
 								FVector2D LocalHitLocation = FVector2D( WidgetComponent->ComponentToWorld.InverseTransformPosition( HitResult.Location ) );
 
@@ -680,11 +680,11 @@ void UWidgetComponent::UpdateBodySetup( bool bDrawSizeChanged )
 	}	
 }
 
-TArray<FArrangedWidget> UWidgetComponent::GetHitWidgetPath( const FHitResult& HitResult, bool bIgnoreEnabledStatus  )
+TArray<FWidgetAndPointer> UWidgetComponent::GetHitWidgetPath( const FHitResult& HitResult, bool bIgnoreEnabledStatus  )
 {
 	FVector2D LocalHitLocation = FVector2D( ComponentToWorld.InverseTransformPosition( HitResult.Location ) );
 
-	TSharedRef<FVirtualCursorPosition> VirtualMouseCoordinate = MakeShareable( new FVirtualCursorPosition );
+	TSharedRef<FVirtualPointerPosition> VirtualMouseCoordinate = MakeShareable( new FVirtualPointerPosition );
 
 	VirtualMouseCoordinate->CurrentCursorPosition = LocalHitLocation;
 	VirtualMouseCoordinate->LastCursorPosition = LastLocalHitLocation;
@@ -692,11 +692,11 @@ TArray<FArrangedWidget> UWidgetComponent::GetHitWidgetPath( const FHitResult& Hi
 	// Cache the location of the hit
 	LastLocalHitLocation = LocalHitLocation;
 
-	TArray<FArrangedWidget> ArrangedWidgets = HitTestGrid->GetBubblePath( LocalHitLocation, bIgnoreEnabledStatus );
+	TArray<FWidgetAndPointer> ArrangedWidgets = HitTestGrid->GetBubblePath( LocalHitLocation, bIgnoreEnabledStatus );
 
-	for( FArrangedWidget& Widget : ArrangedWidgets )
+	for( FWidgetAndPointer& Widget : ArrangedWidgets )
 	{
-		Widget.VirtualCursorPosition = VirtualMouseCoordinate;
+		Widget.PointerPosition = VirtualMouseCoordinate;
 	}
 
 	return ArrangedWidgets;
diff --git a/Engine/Source/Runtime/UMG/Public/Components/WidgetComponent.h b/Engine/Source/Runtime/UMG/Public/Components/WidgetComponent.h
--- a/Engine/Source/Runtime/UMG/Public/Components/WidgetComponent.h
+++ b/Engine/Source/Runtime/UMG/Public/Components/WidgetComponent.h
@@ -4,6 +4,8 @@
 
 #include "WidgetComponent.generated.h"
 
+struct FVirtualPointerPosition;
+
 /**
 * Beware! This feature is experimental and may be substantially changed or removed in future releases.
 * A 3D instance of a Widget Blueprint that can be interacted with in the world.
@@ -50,8 +52,8 @@ public:
 	/** @return The user widget object displayed by this component */
 	UUserWidget* GetUserWidgetObject() const { return Widget; }
 
-	/** @return The path of hit widgets that were found when the component was hit */
-	TArray<FArrangedWidget> GetHitWidgetPath( const FHitResult& HitResult, bool bIgnoreEnabledStatus );
+	/** @return List of widgets with their geometry and the cursor position transformed into this Widget component's space. */
+	TArray<FWidgetAndPointer> GetHitWidgetPath( const FHitResult& HitResult, bool bIgnoreEnabledStatus );
 
 	/** @return The render target to which the user widget is rendered */
 	UTextureRenderTarget2D* GetRenderTarget() const { return RenderTarget; };
