Slate - when an SListPanel is pending refresh, it no longer reports any children in OnArrangeChildren().

diff --git a/Engine/Source/Runtime/Slate/Private/Widgets/Views/SListPanel.cpp b/Engine/Source/Runtime/Slate/Private/Widgets/Views/SListPanel.cpp
--- a/Engine/Source/Runtime/Slate/Private/Widgets/Views/SListPanel.cpp
+++ b/Engine/Source/Runtime/Slate/Private/Widgets/Views/SListPanel.cpp
@@ -54,50 +54,53 @@ SListPanel::FSlot& SListPanel::AddSlot(int32 InsertAtIndex)
  */
 void SListPanel::OnArrangeChildren( const FGeometry& AllottedGeometry, FArrangedChildren& ArrangedChildren ) const
 {
-	if ( ShouldArrangeHorizontally() )
+	if (!bIsRefreshPending)
 	{
-		// This is a tile view list, arrange items horizontally until there is no more room then create a new row.
-		const float AllottedWidth = AllottedGeometry.Size.X;
-		const float ItemPadding = GetItemPadding(AllottedGeometry);
-		const float HalfItemPadding = ItemPadding * 0.5;
-		float WidthSoFar = 0;
-		float LocalItemWidth = ItemWidth.Get();
-		float LocalItemHeight = ItemHeight.Get();
-		float HeightSoFar = -FMath::FloorToInt(SmoothScrollOffsetInItems * LocalItemHeight) - OverscrollAmount;
-
-		for( int32 ItemIndex = 0; ItemIndex < Children.Num(); ++ItemIndex )
+		if (ShouldArrangeHorizontally())
 		{
-			ArrangedChildren.AddWidget(
-				AllottedGeometry.MakeChild( Children[ItemIndex].GetWidget(), FVector2D(WidthSoFar + HalfItemPadding, HeightSoFar), FVector2D(LocalItemWidth, LocalItemHeight) )
-				);
+			// This is a tile view list, arrange items horizontally until there is no more room then create a new row.
+			const float AllottedWidth = AllottedGeometry.Size.X;
+			const float ItemPadding = GetItemPadding(AllottedGeometry);
+			const float HalfItemPadding = ItemPadding * 0.5;
+			float WidthSoFar = 0;
+			float LocalItemWidth = ItemWidth.Get();
+			float LocalItemHeight = ItemHeight.Get();
+			float HeightSoFar = -FMath::FloorToInt(SmoothScrollOffsetInItems * LocalItemHeight) - OverscrollAmount;
+
+			for (int32 ItemIndex = 0; ItemIndex < Children.Num(); ++ItemIndex)
+			{
+				ArrangedChildren.AddWidget(
+					AllottedGeometry.MakeChild(Children[ItemIndex].GetWidget(), FVector2D(WidthSoFar + HalfItemPadding, HeightSoFar), FVector2D(LocalItemWidth, LocalItemHeight))
+					);
 
-			WidthSoFar += LocalItemWidth + ItemPadding;
+				WidthSoFar += LocalItemWidth + ItemPadding;
 
-			if ( WidthSoFar + LocalItemWidth + ItemPadding > AllottedWidth )
-			{
-				WidthSoFar = 0;
-				HeightSoFar += LocalItemHeight;
+				if (WidthSoFar + LocalItemWidth + ItemPadding > AllottedWidth)
+				{
+					WidthSoFar = 0;
+					HeightSoFar += LocalItemHeight;
+				}
 			}
 		}
-	}
-	else
-	{
-		if (Children.Num() > 0)
+		else
 		{
-			// This is a normal list, arrange items vertically
-			float HeightSoFar = -FMath::FloorToInt(SmoothScrollOffsetInItems * Children[0].GetWidget()->GetDesiredSize().Y)-OverscrollAmount;
-			for( int32 ItemIndex = 0; ItemIndex < Children.Num(); ++ItemIndex )
+			if (Children.Num() > 0)
 			{
+				// This is a normal list, arrange items vertically
+				float HeightSoFar = -FMath::FloorToInt(SmoothScrollOffsetInItems * Children[0].GetWidget()->GetDesiredSize().Y) - OverscrollAmount;
+				for (int32 ItemIndex = 0; ItemIndex < Children.Num(); ++ItemIndex)
+				{
 					const FVector2D ItemDesiredSize = Children[ItemIndex].GetWidget()->GetDesiredSize();
 					const float LocalItemHeight = ItemDesiredSize.Y;
 
-				// Note that ListPanel does not respect child Visibility.
-				// It is simply not useful for ListPanels.
-				ArrangedChildren.AddWidget(
-					AllottedGeometry.MakeChild( Children[ItemIndex].GetWidget(), FVector2D(0, HeightSoFar), FVector2D(AllottedGeometry.Size.X, LocalItemHeight) )
-					);
+					// Note that ListPanel does not respect child Visibility.
+					// It is simply not useful for ListPanels.
+					ArrangedChildren.AddWidget(
+						AllottedGeometry.MakeChild(Children[ItemIndex].GetWidget(), FVector2D(0, HeightSoFar), FVector2D(AllottedGeometry.Size.X, LocalItemHeight))
+						);
 
-				HeightSoFar += LocalItemHeight;
+					HeightSoFar += LocalItemHeight;
+				}
 			}
 		}
 	}
