Added the ability to specify the row alignment for items in STileView

diff --git a/Engine/Source/Developer/AutomationWindow/Private/SAutomationTestTreeView.h b/Engine/Source/Developer/AutomationWindow/Private/SAutomationTestTreeView.h
--- a/Engine/Source/Developer/AutomationWindow/Private/SAutomationTestTreeView.h
+++ b/Engine/Source/Developer/AutomationWindow/Private/SAutomationTestTreeView.h
@@ -122,7 +122,7 @@ public:
 			else
 			{
 				// Make the TableView
-				this->ConstructChildren( 0, InArgs._ItemHeight, InArgs._HeaderRow, InArgs._ExternalScrollbar );
+				this->ConstructChildren( 0, InArgs._ItemHeight, EListItemAlignment::LeftAligned, InArgs._HeaderRow, InArgs._ExternalScrollbar );
 			}
 		}
 
diff --git a/Engine/Source/Runtime/Slate/Private/Widgets/Views/SListPanel.cpp b/Engine/Source/Runtime/Slate/Private/Widgets/Views/SListPanel.cpp
--- a/Engine/Source/Runtime/Slate/Private/Widgets/Views/SListPanel.cpp
+++ b/Engine/Source/Runtime/Slate/Private/Widgets/Views/SListPanel.cpp
@@ -24,6 +24,7 @@ void SListPanel::Construct( const FArguments& InArgs )
 	ItemWidth = InArgs._ItemWidth;
 	ItemHeight = InArgs._ItemHeight;
 	NumDesiredItems = InArgs._NumDesiredItems;
+	ItemAlignment = InArgs._ItemAlignment;
 	bIsRefreshPending = false;
 }
 
@@ -56,17 +57,37 @@ void SListPanel::OnArrangeChildren( const FGeometry& AllottedGeometry, FArranged
 {
 	if ( ShouldArrangeHorizontally() )
 	{
+		const EListItemAlignment ListItemAlignment = ItemAlignment.Get();
+
 		// This is a tile view list, arrange items horizontally until there is no more room then create a new row.
 		const float AllottedWidth = AllottedGeometry.Size.X;
-		const float ItemPadding = GetItemPadding(AllottedGeometry);
+		const float ItemPadding = (ListItemAlignment == EListItemAlignment::EvenlyDistributed) ? GetItemPadding(AllottedGeometry) : 0.0f;
 		const float HalfItemPadding = ItemPadding * 0.5;
 		float WidthSoFar = 0;
 		float LocalItemWidth = ItemWidth.Get();
 		float LocalItemHeight = ItemHeight.Get();
 		float HeightSoFar = -FMath::FloorToInt(SmoothScrollOffsetInItems * LocalItemHeight) - OverscrollAmount;
 
+		bool bIsNewLine = true;
 		for( int32 ItemIndex = 0; ItemIndex < Children.Num(); ++ItemIndex )
 		{
+			if ( bIsNewLine )
+			{
+				if ( ListItemAlignment == EListItemAlignment::RightAligned || ListItemAlignment == EListItemAlignment::CenterAligned )
+				{
+					const float LinePadding = GetLinePadding(AllottedGeometry, ItemIndex);
+					if ( ListItemAlignment == EListItemAlignment::RightAligned )
+					{
+						WidthSoFar += LinePadding;
+					}
+					else
+					{
+						const float HalfLinePadding = LinePadding * 0.5;
+						WidthSoFar += HalfLinePadding;
+					}
+				}
+			}
+
 			ArrangedChildren.AddWidget(
 				AllottedGeometry.MakeChild( Children[ItemIndex].GetWidget(), FVector2D(WidthSoFar + HalfItemPadding, HeightSoFar), FVector2D(LocalItemWidth, LocalItemHeight) )
 				);
@@ -77,6 +98,7 @@ void SListPanel::OnArrangeChildren( const FGeometry& AllottedGeometry, FArranged
 			{
 				WidthSoFar = 0;
 				HeightSoFar += LocalItemHeight;
+				bIsNewLine = true;
 			}
 		}
 	}
@@ -202,7 +224,7 @@ float SListPanel::GetItemWidth() const
 
 float SListPanel::GetItemPadding(const FGeometry& AllottedGeometry) const
 {
-	float LocalItemWidth = ItemWidth.Get();
+	const float LocalItemWidth = ItemWidth.Get();
 	const int32 NumItemsWide = LocalItemWidth > 0 ? FMath::FloorToInt(AllottedGeometry.Size.X / LocalItemWidth) : 0;
 	
 	// Only add padding between items if we have more total items that we can fit on a single row.  Otherwise,
@@ -218,6 +240,24 @@ float SListPanel::GetItemPadding(const FGeometry& AllottedGeometry) const
 	return Padding;
 }
 
+float SListPanel::GetLinePadding(const FGeometry& AllottedGeometry, const int32 LineStartIndex) const
+{
+	const int32 NumItemsLeft = Children.Num() - LineStartIndex;
+	if(NumItemsLeft <= 0)
+	{
+		return 0.0f;
+	}
+
+	const float LocalItemWidth = ItemWidth.Get();
+	const int32 NumItemsWide = LocalItemWidth > 0 ? FMath::FloorToInt(AllottedGeometry.Size.X / LocalItemWidth) : 0;
+	const int32 NumItemsOnLine = FMath::Min(NumItemsLeft, NumItemsWide);
+
+	// Subtract a tiny amount from the available width to avoid floating point precision problems when arranging children
+	static const float FloatingPointPrecisionOffset = 0.001f;
+
+	return (AllottedGeometry.Size.X - FloatingPointPrecisionOffset - NumItemsOnLine * LocalItemWidth);
+}
+
 /** @return the uniform item height used when arranging children. */
 float SListPanel::GetItemHeight() const
 {
diff --git a/Engine/Source/Runtime/Slate/Private/Widgets/Views/SListPanel.h b/Engine/Source/Runtime/Slate/Private/Widgets/Views/SListPanel.h
--- a/Engine/Source/Runtime/Slate/Private/Widgets/Views/SListPanel.h
+++ b/Engine/Source/Runtime/Slate/Private/Widgets/Views/SListPanel.h
@@ -30,6 +30,7 @@ public:
 		: _ItemWidth(0)
 		, _ItemHeight(16)
 		, _NumDesiredItems(0)
+		, _ItemAlignment(EListItemAlignment::EvenlyDistributed)
 		{
 			_Visibility = EVisibility::SelfHitTestInvisible;
 		}
@@ -37,6 +38,7 @@ public:
 		SLATE_ATTRIBUTE( float, ItemWidth )
 		SLATE_ATTRIBUTE( float, ItemHeight )
 		SLATE_ATTRIBUTE( int32, NumDesiredItems )
+		SLATE_ATTRIBUTE( EListItemAlignment, ItemAlignment )
 
 	SLATE_END_ARGS()
 	
@@ -72,6 +74,9 @@ public:
 
 	/** @return the horizontal padding applied to each tile item */
 	float GetItemPadding(const FGeometry& AllottedGeometry) const;
+
+	/** @return the horizontal padding applied to all the items on a line */
+	float GetLinePadding(const FGeometry& AllottedGeometry, const int32 LineStartIndex) const;
 	
 	/** @return the uniform item height used when arranging children. */
 	float GetItemHeight() const;
@@ -107,6 +112,9 @@ protected:
 	/** Total number of items that the tree wants to visualize */
 	TAttribute<int32> NumDesiredItems;
 	
+	/** How should be horizontally aligned? Only relevant for tile views. */
+	TAttribute<EListItemAlignment> ItemAlignment;
+
 	/**
 	 * The offset of the view area from the top of the list in item heights.
 	 * Translate to physical units based on first item in list.
diff --git a/Engine/Source/Runtime/Slate/Private/Widgets/Views/STableViewBase.cpp b/Engine/Source/Runtime/Slate/Private/Widgets/Views/STableViewBase.cpp
--- a/Engine/Source/Runtime/Slate/Private/Widgets/Views/STableViewBase.cpp
+++ b/Engine/Source/Runtime/Slate/Private/Widgets/Views/STableViewBase.cpp
@@ -9,7 +9,7 @@ namespace ListConstants
 	static const float OvershootBounceRate = 250.0f;
 }
 
-void STableViewBase::ConstructChildren( const TAttribute<float>& InItemWidth, const TAttribute<float>& InItemHeight, const TSharedPtr<SHeaderRow>& InHeaderRow, const TSharedPtr<SScrollBar>& InScrollBar )
+void STableViewBase::ConstructChildren( const TAttribute<float>& InItemWidth, const TAttribute<float>& InItemHeight, const TAttribute<EListItemAlignment>& InItemAlignment, const TSharedPtr<SHeaderRow>& InHeaderRow, const TSharedPtr<SScrollBar>& InScrollBar )
 {
 	HeaderRow = InHeaderRow;
 
@@ -33,6 +33,7 @@ void STableViewBase::ConstructChildren( const TAttribute<float>& InItemWidth, co
 				.ItemWidth( InItemWidth )
 				.ItemHeight( InItemHeight )
 				.NumDesiredItems( this, &STableViewBase::GetNumItemsBeingObserved )
+				.ItemAlignment( InItemAlignment )
 			]
 			+SHorizontalBox::Slot()
 			.AutoWidth()
@@ -52,6 +53,7 @@ void STableViewBase::ConstructChildren( const TAttribute<float>& InItemWidth, co
 			.ItemWidth( InItemWidth )
 			.ItemHeight( InItemHeight )
 			.NumDesiredItems( this, &STableViewBase::GetNumItemsBeingObserved )
+			.ItemAlignment( InItemAlignment )
 		);
 
 	if (InHeaderRow.IsValid())
diff --git a/Engine/Source/Runtime/Slate/Public/Widgets/Views/SListView.h b/Engine/Source/Runtime/Slate/Public/Widgets/Views/SListView.h
--- a/Engine/Source/Runtime/Slate/Public/Widgets/Views/SListView.h
+++ b/Engine/Source/Runtime/Slate/Public/Widgets/Views/SListView.h
@@ -146,7 +146,7 @@ public:
 		else
 		{
 			// Make the TableView
-			ConstructChildren( 0, InArgs._ItemHeight, InArgs._HeaderRow, InArgs._ExternalScrollbar );
+			ConstructChildren( 0, InArgs._ItemHeight, EListItemAlignment::LeftAligned, InArgs._HeaderRow, InArgs._ExternalScrollbar );
 			if(ScrollBar.IsValid())
 			{
 				ScrollBar->SetUserVisibility(InArgs._ScrollbarVisibility);
diff --git a/Engine/Source/Runtime/Slate/Public/Widgets/Views/STableViewBase.h b/Engine/Source/Runtime/Slate/Public/Widgets/Views/STableViewBase.h
--- a/Engine/Source/Runtime/Slate/Public/Widgets/Views/STableViewBase.h
+++ b/Engine/Source/Runtime/Slate/Public/Widgets/Views/STableViewBase.h
@@ -16,6 +16,23 @@ enum class EAllowOverscroll
 };
 
 
+/** If the list panel is arranging items horizontally, this enum dictates how the items should be aligned (basically, where any extra space is placed) */
+enum class EListItemAlignment : uint8
+{
+	/** Items are distributed evenly along the row (any extra space is added as padding between the items) */
+	EvenlyDistributed,
+
+	/** Items are left aligned on the row (any extra space is added to the right of the items) */
+	LeftAligned,
+
+	/** Items are right aligned on the row (any extra space is added to the left of the items) */
+	RightAligned,
+
+	/** Items are center aligned on the row (any extra space is halved and added to the left of the items) */
+	CenterAligned,
+};
+
+
 /**
  * Contains ListView functionality that does not depend on the type of data being observed by the ListView.
  */
@@ -26,7 +43,7 @@ class SLATE_API STableViewBase
 public:
 
 	/** Create the child widgets that comprise the list */
-	void ConstructChildren( const TAttribute<float>& InItemWidth, const TAttribute<float>& InItemHeight, const TSharedPtr<SHeaderRow>& InColumnHeaders, const TSharedPtr<SScrollBar>& InScrollBar  );
+	void ConstructChildren( const TAttribute<float>& InItemWidth, const TAttribute<float>& InItemHeight, const TAttribute<EListItemAlignment>& InItemAlignment, const TSharedPtr<SHeaderRow>& InColumnHeaders, const TSharedPtr<SScrollBar>& InScrollBar  );
 
 	/** Sets the item height */
 	void SetItemHeight(TAttribute<float> Height);
diff --git a/Engine/Source/Runtime/Slate/Public/Widgets/Views/STileView.h b/Engine/Source/Runtime/Slate/Public/Widgets/Views/STileView.h
--- a/Engine/Source/Runtime/Slate/Public/Widgets/Views/STileView.h
+++ b/Engine/Source/Runtime/Slate/Public/Widgets/Views/STileView.h
@@ -25,6 +25,7 @@ public:
 		, _ListItemsSource( static_cast<TArray<ItemType>*>(nullptr) ) //@todo Slate Syntax: Initializing from nullptr without a cast
 		, _ItemHeight(128)
 		, _ItemWidth(128)
+		, _ItemAlignment(EListItemAlignment::EvenlyDistributed)
 		, _OnContextMenuOpening()
 		, _OnMouseButtonDoubleClick()
 		, _OnSelectionChanged()
@@ -45,6 +46,8 @@ public:
 
 		SLATE_ATTRIBUTE( float, ItemWidth )
 
+		SLATE_ATTRIBUTE( EListItemAlignment, ItemAlignment )
+
 		SLATE_EVENT( FOnContextMenuOpening, OnContextMenuOpening )
 
 		SLATE_EVENT( FOnMouseButtonDoubleClick, OnMouseButtonDoubleClick )
@@ -111,7 +114,7 @@ public:
 		else
 		{
 			// Make the TableView
-			this->ConstructChildren(InArgs._ItemWidth, InArgs._ItemHeight, TSharedPtr<SHeaderRow>(), InArgs._ExternalScrollbar);
+			this->ConstructChildren(InArgs._ItemWidth, InArgs._ItemHeight, InArgs._ItemAlignment, TSharedPtr<SHeaderRow>(), InArgs._ExternalScrollbar);
 			if (this->ScrollBar.IsValid())
 			{
 				this->ScrollBar->SetUserVisibility(InArgs._ScrollbarVisibility);
diff --git a/Engine/Source/Runtime/Slate/Public/Widgets/Views/STreeView.h b/Engine/Source/Runtime/Slate/Public/Widgets/Views/STreeView.h
--- a/Engine/Source/Runtime/Slate/Public/Widgets/Views/STreeView.h
+++ b/Engine/Source/Runtime/Slate/Public/Widgets/Views/STreeView.h
@@ -174,7 +174,7 @@ public:
 		else
 		{
 			// Make the TableView
-			this->ConstructChildren( 0, InArgs._ItemHeight, InArgs._HeaderRow, InArgs._ExternalScrollbar );
+			this->ConstructChildren( 0, InArgs._ItemHeight, EListItemAlignment::LeftAligned, InArgs._HeaderRow, InArgs._ExternalScrollbar );
 		}
 	}
 
