Added support for disabling tab spawner menu entries

You could previously only have them enabled or hidden, you can now add them as disabled.

You can also now bind the menu item state to a delegate to let it update after the tab spawner has been registered.

diff --git a/Engine/Source/Editor/EditorSettingsViewer/Private/EditorSettingsViewerModule.cpp b/Engine/Source/Editor/EditorSettingsViewer/Private/EditorSettingsViewerModule.cpp
--- a/Engine/Source/Editor/EditorSettingsViewer/Private/EditorSettingsViewerModule.cpp
+++ b/Engine/Source/Editor/EditorSettingsViewer/Private/EditorSettingsViewerModule.cpp
@@ -66,7 +66,7 @@ public:
 
 		FGlobalTabmanager::Get()->RegisterNomadTabSpawner(EditorSettingsTabName, FOnSpawnTab::CreateRaw(this, &FEditorSettingsViewerModule::HandleSpawnSettingsTab))
 			.SetDisplayName(LOCTEXT("EditorSettingsTabTitle", "Editor Preferences"))
-			.SetMenuType(ETabSpawnerMenuType::Hide);
+			.SetMenuType(ETabSpawnerMenuType::Hidden);
 	}
 
 	virtual void ShutdownModule() override
diff --git a/Engine/Source/Editor/IntroTutorials/Private/IntroTutorials.cpp b/Engine/Source/Editor/IntroTutorials/Private/IntroTutorials.cpp
--- a/Engine/Source/Editor/IntroTutorials/Private/IntroTutorials.cpp
+++ b/Engine/Source/Editor/IntroTutorials/Private/IntroTutorials.cpp
@@ -151,7 +151,7 @@ void FIntroTutorials::StartupModule()
 	}
 
 	FGlobalTabmanager::Get()->RegisterNomadTabSpawner(TEXT("TutorialsBrowser"), FOnSpawnTab::CreateRaw(this, &FIntroTutorials::SpawnTutorialsBrowserTab))
-		.SetMenuType(ETabSpawnerMenuType::Hide);
+		.SetMenuType(ETabSpawnerMenuType::Hidden);
 }
 
 void FIntroTutorials::ShutdownModule()
diff --git a/Engine/Source/Editor/LevelEditor/Private/SLevelEditor.cpp b/Engine/Source/Editor/LevelEditor/Private/SLevelEditor.cpp
--- a/Engine/Source/Editor/LevelEditor/Private/SLevelEditor.cpp
+++ b/Engine/Source/Editor/LevelEditor/Private/SLevelEditor.cpp
@@ -1020,13 +1020,13 @@ TSharedRef<SWidget> SLevelEditor::RestoreContentArea( const TSharedRef<SDockTab>
 				.SetIcon( FSlateIcon(FEditorStyle::GetStyleSetName(), "LevelEditor.Tabs.WorldBrowser") );
 			
 			LevelEditorTabManager->RegisterTabSpawner( WorldBrowserDetailsTab, FOnSpawnTab::CreateSP<SLevelEditor, FName, FString>(this, &SLevelEditor::SpawnLevelEditorTab, WorldBrowserDetailsTab, FString()) )
-				.SetMenuType( ETabSpawnerMenuType::Hide )
+				.SetMenuType( ETabSpawnerMenuType::Hidden )
 				.SetDisplayName(NSLOCTEXT("LevelEditorTabs", "WorldBrowserDetails", "Level Details"))
 				.SetGroup( WorkspaceMenu::GetMenuStructure().GetLevelEditorCategory() )
 				.SetIcon( FSlateIcon(FEditorStyle::GetStyleSetName(), "LevelEditor.Tabs.WorldBrowserDetails") );
 		
 			LevelEditorTabManager->RegisterTabSpawner( WorldBrowserCompositionTab, FOnSpawnTab::CreateSP<SLevelEditor, FName, FString>(this, &SLevelEditor::SpawnLevelEditorTab, WorldBrowserCompositionTab, FString()) )
-				.SetMenuType( ETabSpawnerMenuType::Hide )
+				.SetMenuType( ETabSpawnerMenuType::Hidden )
 				.SetDisplayName(NSLOCTEXT("LevelEditorTabs", "WorldBrowserComposition", "World Composition"))
 				.SetGroup( WorkspaceMenu::GetMenuStructure().GetLevelEditorCategory() )
 				.SetIcon( FSlateIcon(FEditorStyle::GetStyleSetName(), "LevelEditor.Tabs.WorldBrowserComposition") );
diff --git a/Engine/Source/Editor/ProjectSettingsViewer/Private/ProjectSettingsViewerModule.cpp b/Engine/Source/Editor/ProjectSettingsViewer/Private/ProjectSettingsViewerModule.cpp
--- a/Engine/Source/Editor/ProjectSettingsViewer/Private/ProjectSettingsViewerModule.cpp
+++ b/Engine/Source/Editor/ProjectSettingsViewer/Private/ProjectSettingsViewerModule.cpp
@@ -68,7 +68,7 @@ public:
 
 		FGlobalTabmanager::Get()->RegisterNomadTabSpawner(ProjectSettingsTabName, FOnSpawnTab::CreateRaw(this, &FProjectSettingsViewerModule::HandleSpawnSettingsTab))
 			.SetDisplayName(LOCTEXT("ProjectSettingsTabTitle", "Project Settings"))
-			.SetMenuType(ETabSpawnerMenuType::Hide);
+			.SetMenuType(ETabSpawnerMenuType::Hidden);
 	}
 
 	virtual void ShutdownModule() override
diff --git a/Engine/Source/Editor/ReferenceViewer/Private/ReferenceViewer.cpp b/Engine/Source/Editor/ReferenceViewer/Private/ReferenceViewer.cpp
--- a/Engine/Source/Editor/ReferenceViewer/Private/ReferenceViewer.cpp
+++ b/Engine/Source/Editor/ReferenceViewer/Private/ReferenceViewer.cpp
@@ -38,7 +38,7 @@ public:
 
 		FGlobalTabmanager::Get()->RegisterNomadTabSpawner(ReferenceViewerTabId, FOnSpawnTab::CreateRaw( this, &FReferenceViewerModule::SpawnReferenceViewerTab ) )
 			.SetDisplayName( LOCTEXT("ReferenceViewerTitle", "Reference Viewer") )
-			.SetMenuType( ETabSpawnerMenuType::Hide );
+			.SetMenuType( ETabSpawnerMenuType::Hidden );
 	}
 
 	virtual void ShutdownModule() override
diff --git a/Engine/Source/Runtime/Slate/Private/Framework/Docking/TabManager.cpp b/Engine/Source/Runtime/Slate/Private/Framework/Docking/TabManager.cpp
--- a/Engine/Source/Runtime/Slate/Private/Framework/Docking/TabManager.cpp
+++ b/Engine/Source/Runtime/Slate/Private/Framework/Docking/TabManager.cpp
@@ -760,7 +760,12 @@ void FTabManager::PopulateTabSpawnerMenu_Helper( FMenuBuilder& PopulateMe, FPopu
 
 void FTabManager::MakeSpawnerMenuEntry( FMenuBuilder &PopulateMe, const TSharedPtr<FTabSpawnerEntry> &SpawnerNode ) 
 {
-	if ( SpawnerNode->MenuType != ETabSpawnerMenuType::Hide )
+	auto CanExecuteMenuEntry = [](TAttribute<ETabSpawnerMenuType::Type> InMenuType) -> bool
+	{
+		return InMenuType.Get() == ETabSpawnerMenuType::Enabled;
+	};
+
+	if ( SpawnerNode->MenuType.Get() != ETabSpawnerMenuType::Hidden )
 	{
 		PopulateMe.AddMenuEntry(
 			SpawnerNode->GetDisplayName().IsEmpty() ? FText::FromName( SpawnerNode->TabType ) : SpawnerNode->GetDisplayName(),
@@ -768,7 +773,7 @@ void FTabManager::MakeSpawnerMenuEntry( FMenuBuilder &PopulateMe, const TSharedP
 			SpawnerNode->GetIcon(),
 			FUIAction(
 			FExecuteAction::CreateSP(SharedThis(this), &FTabManager::InvokeTabForMenu, SpawnerNode->TabType),
-			FCanExecuteAction(),
+			FCanExecuteAction::CreateStatic(CanExecuteMenuEntry, SpawnerNode->MenuType),
 			FIsActionChecked::CreateSP(SpawnerNode.ToSharedRef(), &FTabSpawnerEntry::IsSoleTabInstanceSpawned)
 			),
 			NAME_None,
diff --git a/Engine/Source/Runtime/Slate/Public/Framework/Docking/TabManager.h b/Engine/Source/Runtime/Slate/Public/Framework/Docking/TabManager.h
--- a/Engine/Source/Runtime/Slate/Public/Framework/Docking/TabManager.h
+++ b/Engine/Source/Runtime/Slate/Public/Framework/Docking/TabManager.h
@@ -107,8 +107,9 @@ namespace ETabSpawnerMenuType
 {
 	enum Type
 	{
-		Display,		// Display this spawner in menus
-		Hide			// Do not display this spawner in menus, it will be invoked manually
+		Enabled,		// Display this spawner in menus
+		Disabled,		// Display this spawner in menus, but make it disabled
+		Hidden,			// Do not display this spawner in menus, it will be invoked manually
 	};
 }
 
@@ -119,10 +120,9 @@ struct FTabSpawnerEntry : public FWorkspaceItem
 		, TabType( InTabType )
 		, OnSpawnTab( InSpawnTabMethod )
 		, OnFindTabToReuse()
-		, MenuType(ETabSpawnerMenuType::Display)
+		, MenuType(ETabSpawnerMenuType::Enabled)
 		, bAutoGenerateMenuEntry(true)
 		, SpawnedTabPtr()
-
 	{
 	}
 
@@ -156,7 +156,7 @@ struct FTabSpawnerEntry : public FWorkspaceItem
 		return *this;
 	}
 
-	FTabSpawnerEntry& SetMenuType( ETabSpawnerMenuType::Type InMenuType )
+	FTabSpawnerEntry& SetMenuType( const TAttribute<ETabSpawnerMenuType::Type>& InMenuType )
 	{
 		MenuType = InMenuType;
 		return *this;
@@ -178,7 +178,8 @@ private:
 	FOnSpawnTab OnSpawnTab;
 	/** When this method is not provided, we assume that the tab should only allow 0 or 1 instances */
 	FOnFindTabToReuse OnFindTabToReuse;
-	ETabSpawnerMenuType::Type MenuType;
+	/** Whether this menu item should be enabled, disabled, or hidden */
+	TAttribute<ETabSpawnerMenuType::Type> MenuType;
 	/** Whether to automatically generate a menu entry for this tab spawner */
 	bool bAutoGenerateMenuEntry;
 
