Slate : further unification of even routing logic.
- All event types are now routed by FEventRouter.
- All replies are now based on a single FReplyBase, s.t. we can reason about them uniformly.
- Unified Focus and Pointer path routing.
- Removed all instances of TransformPointerEvent outside of FEventRouter.

diff --git a/Engine/Source/Developer/SlateReflector/Private/Widgets/SWidgetReflector.cpp b/Engine/Source/Developer/SlateReflector/Private/Widgets/SWidgetReflector.cpp
--- a/Engine/Source/Developer/SlateReflector/Private/Widgets/SWidgetReflector.cpp
+++ b/Engine/Source/Developer/SlateReflector/Private/Widgets/SWidgetReflector.cpp
@@ -11,7 +11,7 @@ static const int32 MaxLoggedEvents = 100;
 
 struct FLoggedEvent
 {
-	FLoggedEvent( const FInputEvent& InEvent, const FReply& InReply )
+	FLoggedEvent( const FInputEvent& InEvent, const FReplyBase& InReply )
 	: Event( InEvent )
 	, Handler( InReply.GetHandler() )
 	, EventText( InEvent.ToText() )
@@ -178,7 +178,7 @@ void SWidgetReflector::Tick( const FGeometry& AllottedGeometry, const double InC
 	SCompoundWidget::Tick(AllottedGeometry, InCurrentTime, InDeltaTime);
 }
 
-void SWidgetReflector::OnEventProcessed( const FInputEvent& Event, const FReply& InReply )
+void SWidgetReflector::OnEventProcessed( const FInputEvent& Event, const FReplyBase& InReply )
 {
 	#if WITH_EVENT_LOGGING
 		if (LoggedEvents.Num() >= MaxLoggedEvents)
diff --git a/Engine/Source/Developer/SlateReflector/Private/Widgets/SWidgetReflector.h b/Engine/Source/Developer/SlateReflector/Private/Widgets/SWidgetReflector.h
--- a/Engine/Source/Developer/SlateReflector/Private/Widgets/SWidgetReflector.h
+++ b/Engine/Source/Developer/SlateReflector/Private/Widgets/SWidgetReflector.h
@@ -38,7 +38,7 @@ public:
 
 	// IWidgetReflector interface
 
-	virtual void OnEventProcessed( const FInputEvent& Event, const FReply& InReply ) override;
+	virtual void OnEventProcessed( const FInputEvent& Event, const FReplyBase& InReply ) override;
 
 	virtual bool IsInPickingMode( ) const override
 	{
diff --git a/Engine/Source/Runtime/Slate/Private/Framework/Application/SlateApplication.cpp b/Engine/Source/Runtime/Slate/Private/Framework/Application/SlateApplication.cpp
--- a/Engine/Source/Runtime/Slate/Private/Framework/Application/SlateApplication.cpp
+++ b/Engine/Source/Runtime/Slate/Private/Framework/Application/SlateApplication.cpp
@@ -10,23 +10,14 @@
 #include "NotificationManager.h"
 
 
-// @todo slate : Move TranslatePointerEvent into FEventRouter. Remove external calls to it.
-template <typename PointerEventType>
-static PointerEventType TranslatePointerEvent( const TSharedPtr<FVirtualPointerPosition>& InPosition, const PointerEventType& InEvent )
-{
-	if ( !InPosition.IsValid() )
-	{
-		return InEvent;
-	}
-	else
-	{
-		return FPointerEvent::MakeTranslatedEvent<PointerEventType>( InEvent, *InPosition );
-	}
-}
 
 class FEventRouter
 {
 
+// @todo slate : making too many event copies when translating events( i.e. Translate<EventType>::PointerEvent ).
+
+
+
 // @todo slate : Widget Reflector should log: (1) Every process reply (2) Every time the event is handled and by who.
 
 // @todo slate : Remove remaining [&]-style mass captures.
@@ -36,6 +27,43 @@ class FEventRouter
 // @todo slate : Remove CALL_WIDGET_FUNCTION
 
 public:
+
+	class FDirectPolicy
+	{
+	public:
+		FDirectPolicy( const FWidgetAndPointer& InTarget, const FWidgetPath& InRoutingPath )
+		: bEventSent(false)
+		, RoutingPath(InRoutingPath)
+		, Target(InTarget)
+		{
+		}
+
+		bool ShouldKeepGoing() const
+		{
+			return !bEventSent;
+		}
+
+		void Next()
+		{
+			bEventSent = true;
+		}
+
+		FWidgetAndPointer GetWidget() const
+		{
+			return Target;
+		}
+
+		const FWidgetPath& GetRoutingPath() const
+		{
+			return RoutingPath;
+		}
+
+	private:
+		bool bEventSent;
+		const FWidgetPath& RoutingPath;
+		const FWidgetAndPointer& Target;
+	};
+
 	class FToLeafmostPolicy
 	{
 	public:
@@ -140,7 +168,7 @@ public:
 		const FWidgetPath& RoutingPath;
 	};
 
-	static void LogEvent( FSlateApplication* ThisApplication, const FInputEvent& Event, const FReply& Reply )
+	static void LogEvent( FSlateApplication* ThisApplication, const FInputEvent& Event, const FReplyBase& Reply )
 	{
 		TSharedPtr<IWidgetReflector> Reflector = ThisApplication->WidgetReflectorPtr.Pin();
 		if (Reflector.IsValid() && Reply.IsEventHandled())
@@ -150,7 +178,7 @@ public:
 	}
 
 	/**
-	 * Bubble an event along a focus path (as opposed to PointerPath)
+	 * Route an event along a focus path (as opposed to PointerPath)
 	 *
 	 * Focus paths are determined by the Keyboard Focus or Controller Focus.
 	 * Focus paths change when the user navigates focus (e.g. Tab or
@@ -159,54 +187,78 @@ public:
 	template< typename RoutingPolicyType, typename FuncType, typename EventType >
 	static FReply RouteAlongFocusPath( FSlateApplication* ThisApplication, RoutingPolicyType RoutingPolicy, EventType KeyEventCopy, const FuncType& Lambda )
 	{
-		const FWidgetPath& RoutingPath = RoutingPolicy.GetRoutingPath();
-		FReply Reply = FReply::Unhandled();
+		return Route<FReply>(ThisApplication, RoutingPolicy, KeyEventCopy, Lambda);
+	}
 
-		KeyEventCopy.SetEventPath(RoutingPath);
+	/**
+	 * Route an event based on the Routing Policy.
+	 */
+	template< typename ReplyType, typename RoutingPolicyType, typename EventType, typename FuncType >
+	static ReplyType Route( FSlateApplication* ThisApplication, RoutingPolicyType RoutingPolicy, EventType EventCopy, const FuncType& Lambda )
+	{
+		ReplyType Reply = ReplyType::Unhandled();
+		const FWidgetPath& RoutingPath = RoutingPolicy.GetRoutingPath();
+		
+		EventCopy.SetEventPath( RoutingPath );
 
 		for ( ; !Reply.IsEventHandled() && RoutingPolicy.ShouldKeepGoing(); RoutingPolicy.Next() )
 		{
 			const FWidgetAndPointer& ArrangedWidget = RoutingPolicy.GetWidget();
-			Reply = Lambda( ArrangedWidget, KeyEventCopy ).SetHandler( ArrangedWidget.Widget );
-			ThisApplication->ProcessReply( RoutingPath, Reply, NULL, NULL );
+			const EventType TranslatedEvent = Translate<EventType>::PointerEvent( ArrangedWidget.PointerPosition, EventCopy );
+			Reply = Lambda( ArrangedWidget, TranslatedEvent ).SetHandler( ArrangedWidget.Widget );
+			ProcessReply(ThisApplication, RoutingPath, Reply, &RoutingPath, &TranslatedEvent);
 		}
 
-		LogEvent(ThisApplication, KeyEventCopy, Reply);
-		
+		LogEvent(ThisApplication, EventCopy, Reply);
+
 		return Reply;
 	}
 
-	/**
-	 * Bubble an event along a PointerPath (as opposed to FocusPath)
-	 *
-	 * The pointer path is determined on the fly: whatever is under
-	 * the cursor will get a shot at the event. Widgets can request
-	 * pointer capture. Wile pointer capture is active, only the
-	 * widget in question will receive pointer events associated
-	 * with that pointing device until the pointer capture is
-	 * released.
-	 */
-	template< typename RoutingPolicyType,  typename FuncType >
-	static FReply RouteAlongPointerPath( FSlateApplication* ThisApplication, RoutingPolicyType RoutingPolicy, FPointerEvent PointerEventCopy, const FuncType& Lambda )
+	static void ProcessReply( FSlateApplication* Application, const FWidgetPath& RoutingPath, const FNoReply& Reply, const FWidgetPath* WidgetsUnderCursor, const FInputEvent* )
 	{
-		const FWidgetPath& RoutingPath = RoutingPolicy.GetRoutingPath();
-		FReply Reply = FReply::Unhandled();
+	}
 
-		PointerEventCopy.SetEventPath( RoutingPath );
+	static void ProcessReply( FSlateApplication* Application, const FWidgetPath& RoutingPath, const FCursorReply& Reply, const FWidgetPath* WidgetsUnderCursor, const FInputEvent* )
+	{
+	}
 
-		for ( ; !Reply.IsEventHandled() && RoutingPolicy.ShouldKeepGoing(); RoutingPolicy.Next() )
+	static void ProcessReply( FSlateApplication* Application, const FWidgetPath& RoutingPath, const FReply& Reply, const FWidgetPath* WidgetsUnderCursor, const FInputEvent* PointerEvent )
+	{
+		Application->ProcessReply(RoutingPath, Reply, WidgetsUnderCursor, NULL);
+	}
+
+	static void ProcessReply( FSlateApplication* Application, const FWidgetPath& RoutingPath, const FReply& Reply, const FWidgetPath* WidgetsUnderCursor, const FPointerEvent* PointerEvent )
+	{
+		Application->ProcessReply(RoutingPath, Reply, WidgetsUnderCursor, PointerEvent);
+	}
+
+	template<typename EventType>
+	struct Translate
+	{
+		static EventType PointerEvent( const TSharedPtr<FVirtualPointerPosition>& InPosition, const EventType& InEvent )
 		{
-			const FWidgetAndPointer& ArrangedWidget = RoutingPolicy.GetWidget();
-			const FPointerEvent TranslatedEvent = TranslatePointerEvent( ArrangedWidget.PointerPosition, PointerEventCopy );
-			Reply = Lambda( ArrangedWidget, TranslatedEvent ).SetHandler( ArrangedWidget.Widget );
-			const FWidgetPath* WidgetsUnderCursorPtr = &RoutingPath;
-			ThisApplication->ProcessReply( RoutingPath, Reply, WidgetsUnderCursorPtr, &TranslatedEvent );
+			// Most events do not do any coordinate translation.
+			return InEvent;
 		}
+	};
 
-		LogEvent(ThisApplication, PointerEventCopy, Reply);
+	template<>
+	struct Translate<FPointerEvent>
+	{
+		static  FPointerEvent PointerEvent( const TSharedPtr<FVirtualPointerPosition>& InPosition, const FPointerEvent& InEvent )
+		{
+			// Pointer events are translated into the virtual window space. For 3D Widget Components this means
+			if ( !InPosition.IsValid() )
+			{
+				return InEvent;
+			}
+			else
+			{
+				return FPointerEvent::MakeTranslatedEvent<FPointerEvent>( InEvent, *InPosition );
+			}
+		}
+	};
 
-		return Reply;
-	}
 };
 
 
@@ -2107,18 +2159,17 @@ void FSlateApplication::ProcessReply( const FWidgetPath& CurrentEventPath, const
 
 		// We have entered drag and drop mode.
 		// Pretend that the mouse left all the previously hovered widgets, and a drag entered them.
-		for (int32 WidgetIndex=0; WidgetIndex < WidgetsUnderMouse->Widgets.Num(); ++WidgetIndex)
+		FEventRouter::Route<FNoReply>(this, FEventRouter::FBubblePolicy(*WidgetsUnderMouse), *InMouseEvent, [](const FArrangedWidget& SomeWidget, const FPointerEvent& PointerEvent)
 		{
-			const FArrangedWidget& SomeWidget = WidgetsUnderMouse->Widgets[WidgetIndex];
-			SomeWidget.Widget->OnMouseLeave( TranslatePointerEvent( WidgetsUnderMouse->VirtualPointerPositions[WidgetIndex], *InMouseEvent ) );
-		}
+			SomeWidget.Widget->OnMouseLeave( PointerEvent );
+			return FNoReply();
+		});
 
-		FDragDropEvent DragDropEvent( *InMouseEvent, ReplyDragDropContent );
-		for (int32 WidgetIndex=0; WidgetIndex < WidgetsUnderMouse->Widgets.Num(); ++WidgetIndex)
+		FEventRouter::Route<FNoReply>(this, FEventRouter::FBubblePolicy(*WidgetsUnderMouse), FDragDropEvent( *InMouseEvent, ReplyDragDropContent ), [](const FArrangedWidget& SomeWidget, const FDragDropEvent& DragDropEvent )
 		{
-			const FArrangedWidget& SomeWidget = WidgetsUnderMouse->Widgets[WidgetIndex];
-			SomeWidget.Widget->OnDragEnter( SomeWidget.Geometry, TranslatePointerEvent( WidgetsUnderMouse->VirtualPointerPositions[WidgetIndex], DragDropEvent ) );
-		}
+			SomeWidget.Widget->OnDragEnter( SomeWidget.Geometry, DragDropEvent );
+			return FNoReply();
+		});
 	}
 	
 	TSharedPtr<SWidget> RequestedMouseCaptor = TheReply.GetMouseCaptor();
@@ -2299,8 +2350,6 @@ void FSlateApplication::QueryCursor()
 			// Query widgets with mouse capture for the cursor
 			if (MouseCaptor.HasCaptureForPointerIndex(CursorPointerIndex))
 			{
-				//@todo dhertzka - See if old ToWidgetPath should be replaced entirely
-				//FWidgetPath MouseCaptorPath = MouseCaptor.ToWidgetPath(CursorPointerIndex);
 				FWidgetPath MouseCaptorPath = MouseCaptor.ToWidgetPath( FWeakWidgetPath::EInterruptedPathHandling::Truncate, &CursorEvent);
 				if ( MouseCaptorPath.IsValid() )
 				{
@@ -2323,7 +2372,6 @@ void FSlateApplication::QueryCursor()
 				// Switch worlds for widgets in the current path
 				FScopedSwitchWorldHack SwitchWorld( WidgetsToQueryForCursor );
 
-				//@TODO dhertzka: This was deleted in UMG 3D 2 => investigate
 				const FVector2D CurrentCursorPosition = GetCursorPos();
 				const FVector2D LastCursorPosition = GetLastCursorPos();
 				const FPointerEvent CursorEvent(
@@ -2335,12 +2383,10 @@ void FSlateApplication::QueryCursor()
 					PlatformApplication->GetModifierKeys()
 				);
 
-				CursorResult = FCursorReply::Unhandled();
-				for( int32 WidgetIndex = WidgetsToQueryForCursor.Widgets.Num() - 1; !CursorResult.IsEventHandled() && WidgetIndex >= 0; --WidgetIndex )
+				CursorResult = FEventRouter::Route<FCursorReply>(this, FEventRouter::FBubblePolicy(WidgetsToQueryForCursor), CursorEvent, [](const FArrangedWidget& WidgetToQuery, const FPointerEvent& PointerEvent)
 				{
-					const FArrangedWidget& WidgetToQuery = WidgetsToQueryForCursor.Widgets[ WidgetIndex ];
-					CursorResult = WidgetToQuery.Widget->OnCursorQuery( WidgetToQuery.Geometry, TranslatePointerEvent( WidgetsToQueryForCursor.VirtualPointerPositions[WidgetIndex], CursorEvent ) );
-				}
+					return WidgetToQuery.Widget->OnCursorQuery( WidgetToQuery.Geometry, PointerEvent );
+				});
 
 				if (CursorResult.IsEventHandled())
 				{
@@ -3301,31 +3347,22 @@ bool FSlateApplication::ProcessKeyCharEvent( FCharacterEvent& InCharacterEvent )
 {
 	FReply Reply = FReply::Unhandled();
 
-	const int32 EventCount = InCharacterEvent.IsRepeat() ? SlateApplicationDefs::NumRepeatsPerActualRepeat : 1;
-	for( int32 CurEventIndex = 0; CurEventIndex < EventCount; ++CurEventIndex )
-	{
-		// NOTE: We intentionally don't reset LastUserInteractionTimeForThrottling here so that the UI can be responsive while typing
+	// NOTE: We intentionally don't reset LastUserInteractionTimeForThrottling here so that the UI can be responsive while typing
 
-		// Bubble the keyboard event
-		FWidgetPath EventPath = FocusedWidgetPath.ToWidgetPath();
-		InCharacterEvent.SetEventPath(EventPath);
+	// Bubble the keyboard event
+	FWidgetPath EventPath = FocusedWidgetPath.ToWidgetPath();
 	
-		// Switch worlds for widgets in the current path
-		FScopedSwitchWorldHack SwitchWorld( EventPath );
+	// Switch worlds for widgets in the current path
+	FScopedSwitchWorldHack SwitchWorld( EventPath );
 
-		Reply = FReply::Unhandled();
-		// Send out mouse enter events.
-		InCharacterEvent.SetEventPath( EventPath );
-		
-		Reply = FEventRouter::RouteAlongFocusPath( this, FEventRouter::FBubblePolicy(EventPath), InCharacterEvent, []( const FArrangedWidget& SomeWidgetGettingEvent, const FCharacterEvent& Event )
-		{
-			return SomeWidgetGettingEvent.Widget->IsEnabled()
-				? SomeWidgetGettingEvent.Widget->OnKeyChar( SomeWidgetGettingEvent.Geometry, Event )
-				: FReply::Unhandled();
-		});
+	Reply = FEventRouter::RouteAlongFocusPath( this, FEventRouter::FBubblePolicy(EventPath), InCharacterEvent, []( const FArrangedWidget& SomeWidgetGettingEvent, const FCharacterEvent& Event )
+	{
+		return SomeWidgetGettingEvent.Widget->IsEnabled()
+			? SomeWidgetGettingEvent.Widget->OnKeyChar( SomeWidgetGettingEvent.Geometry, Event )
+			: FReply::Unhandled();
+	});
 
-		LOG_EVENT_CONTENT( EEventLog::KeyChar, FString::Printf(TEXT("%c"), InCharacterEvent.GetCharacter()), Reply );
-	}
+	LOG_EVENT_CONTENT( EEventLog::KeyChar, FString::Printf(TEXT("%c"), InCharacterEvent.GetCharacter()), Reply );
 
 	return Reply.IsEventHandled();
 }
@@ -3352,77 +3389,72 @@ bool FSlateApplication::ProcessKeyDownEvent( FKeyboardEvent& InKeyboardEvent )
 	}
 	else
 	{
-		const int32 EventCount = InKeyboardEvent.IsRepeat() ? SlateApplicationDefs::NumRepeatsPerActualRepeat : 1;
-		for( int32 CurEventIndex = 0; CurEventIndex < EventCount; ++CurEventIndex )
-		{
-			LastUserInteractionTimeForThrottling = LastUserInteractionTime;
+		LastUserInteractionTimeForThrottling = LastUserInteractionTime;
 
-			// If we are inspecting, pressing ESC exits inspection mode.
-			if ( InKeyboardEvent.GetKey() == EKeys::Escape )
+		// If we are inspecting, pressing ESC exits inspection mode.
+		if ( InKeyboardEvent.GetKey() == EKeys::Escape )
+		{
+			TSharedPtr<IWidgetReflector> WidgetReflector = WidgetReflectorPtr.Pin();
+			const bool bIsWidgetReflectorPicking = WidgetReflector.IsValid() && WidgetReflector->IsInPickingMode();
+			if ( bIsWidgetReflectorPicking )
 			{
-				TSharedPtr<IWidgetReflector> WidgetReflector = WidgetReflectorPtr.Pin();
-				const bool bIsWidgetReflectorPicking = WidgetReflector.IsValid() && WidgetReflector->IsInPickingMode();
-				if ( bIsWidgetReflectorPicking )
-				{
-						WidgetReflector->OnWidgetPicked();
-						Reply = FReply::Handled();
+					WidgetReflector->OnWidgetPicked();
+					Reply = FReply::Handled();
 
-						return Reply.IsEventHandled();
-					}
+					return Reply.IsEventHandled();
 				}
+			}
 
-			// Bubble the keyboard event
-			FWidgetPath EventPath = FocusedWidgetPath.ToWidgetPath();
-			InKeyboardEvent.SetEventPath(EventPath);
+		// Bubble the keyboard event
+		FWidgetPath EventPath = FocusedWidgetPath.ToWidgetPath();
 
-			// Switch worlds for widgets inOnPreviewMouseButtonDown the current path
-			FScopedSwitchWorldHack SwitchWorld( EventPath );
+		// Switch worlds for widgets inOnPreviewMouseButtonDown the current path
+		FScopedSwitchWorldHack SwitchWorld( EventPath );
 
-			TSharedPtr<SWidget> WidgetToLog;
+		TSharedPtr<SWidget> WidgetToLog;
+
+		Reply = FReply::Unhandled();
 
-			Reply = FReply::Unhandled();
+		// Tunnel the keyboard event
+		Reply = FEventRouter::RouteAlongFocusPath( this, FEventRouter::FTunnelPolicy(EventPath), InKeyboardEvent, []( const FArrangedWidget& CurrentWidget, const FKeyboardEvent& Event )
+		{
+			return (CurrentWidget.Widget->IsEnabled())
+				? CurrentWidget.Widget->OnPreviewKeyDown( CurrentWidget.Geometry, Event )
+				: FReply::Unhandled();
+		});
 
-			// Tunnel the keyboard event
-			Reply = FEventRouter::RouteAlongFocusPath( this, FEventRouter::FTunnelPolicy(EventPath), InKeyboardEvent, []( const FArrangedWidget& CurrentWidget, const FKeyboardEvent& Event )
+		// Send out key down events.
+		if (!Reply.IsEventHandled())
+		{
+			Reply = FEventRouter::RouteAlongFocusPath(this, FEventRouter::FBubblePolicy(EventPath), InKeyboardEvent, [](const FArrangedWidget& SomeWidgetGettingEvent, const FKeyboardEvent& Event)
 			{
-				return (CurrentWidget.Widget->IsEnabled())
-					? CurrentWidget.Widget->OnPreviewKeyDown( CurrentWidget.Geometry, Event )
+				return (SomeWidgetGettingEvent.Widget->IsEnabled())
+					? SomeWidgetGettingEvent.Widget->OnKeyDown( SomeWidgetGettingEvent.Geometry, Event )
 					: FReply::Unhandled();
 			});
+		}
 
-			// Send out key down events.
-			if (!Reply.IsEventHandled())
-			{
-				Reply = FEventRouter::RouteAlongFocusPath(this, FEventRouter::FBubblePolicy(EventPath), InKeyboardEvent, [](const FArrangedWidget& SomeWidgetGettingEvent, const FKeyboardEvent& Event)
-				{
-					return (SomeWidgetGettingEvent.Widget->IsEnabled())
-						? SomeWidgetGettingEvent.Widget->OnKeyDown( SomeWidgetGettingEvent.Geometry, Event )
-						: FReply::Unhandled();
-				});
-			}
-
-			LOG_EVENT_CONTENT( EEventLog::KeyDown, GetKeyName(InKeyboardEvent.GetKey()).ToString(), Reply );
+		LOG_EVENT_CONTENT( EEventLog::KeyDown, GetKeyName(InKeyboardEvent.GetKey()).ToString(), Reply );
 
-			// If the keyboard event was not processed by any widget...
-			if ( !Reply.IsEventHandled() )
+		// If the keyboard event was not processed by any widget...
+		if ( !Reply.IsEventHandled() )
+		{
+			// If the key was Tab, interpret as an attempt to move focus.
+			if ( InKeyboardEvent.GetKey() == EKeys::Tab )
 			{
-				// If the key was Tab, interpret as an attempt to move focus.
-				if ( InKeyboardEvent.GetKey() == EKeys::Tab )
-				{
-					if (FocusedWidgetPath.IsValid())
-					{
-						EFocusMoveDirection::Type MoveDirection = ( InKeyboardEvent.IsShiftDown() )
-							? EFocusMoveDirection::Previous
-							: EFocusMoveDirection::Next;
-						SetKeyboardFocus( FocusedWidgetPath.ToNextFocusedPath( MoveDirection ), EKeyboardFocusCause::Keyboard );
-					}
-				}
-				else if ( UnhandledKeyDownEventHandler.IsBound() )
+				if (FocusedWidgetPath.IsValid())
 				{
-					// Nothing else handled this event, give external code a chance to handle it.
-					Reply = UnhandledKeyDownEventHandler.Execute( InKeyboardEvent );
+					EFocusMoveDirection::Type MoveDirection = ( InKeyboardEvent.IsShiftDown() )
+						? EFocusMoveDirection::Previous
+						: EFocusMoveDirection::Next;
+					SetKeyboardFocus( FocusedWidgetPath.ToNextFocusedPath( MoveDirection ), EKeyboardFocusCause::Keyboard );
 				}
 			}
+			else if ( UnhandledKeyDownEventHandler.IsBound() )
+			{
+				// Nothing else handled this event, give external code a chance to handle it.
+				Reply = UnhandledKeyDownEventHandler.Execute( InKeyboardEvent );
+			}
 		}
 	}
 
@@ -3443,27 +3475,22 @@ bool FSlateApplication::ProcessKeyUpEvent( FKeyboardEvent& InKeyboardEvent )
 
 	LastUserInteractionTime = this->GetCurrentTime();
 	
-	const int32 EventCount = InKeyboardEvent.IsRepeat() ? SlateApplicationDefs::NumRepeatsPerActualRepeat : 1;
-	for( int32 CurEventIndex = 0; CurEventIndex < EventCount; ++CurEventIndex )
-	{
-		LastUserInteractionTimeForThrottling = LastUserInteractionTime;
+	LastUserInteractionTimeForThrottling = LastUserInteractionTime;
 
-		// Bubble the keyboard event
-		FWidgetPath EventPath = FocusedWidgetPath.ToWidgetPath();
-		InKeyboardEvent.SetEventPath(EventPath);
+	// Bubble the keyboard event
+	FWidgetPath EventPath = FocusedWidgetPath.ToWidgetPath();
 
-		// Switch worlds for widgets in the current path
-		FScopedSwitchWorldHack SwitchWorld( EventPath );
+	// Switch worlds for widgets in the current path
+	FScopedSwitchWorldHack SwitchWorld( EventPath );
 
-		Reply = FEventRouter::RouteAlongFocusPath(this, FEventRouter::FBubblePolicy(EventPath), InKeyboardEvent, [](const FArrangedWidget& SomeWidgetGettingEvent, const FKeyboardEvent& Event)
-		{
-			return (SomeWidgetGettingEvent.Widget->IsEnabled())
-				? SomeWidgetGettingEvent.Widget->OnKeyUp( SomeWidgetGettingEvent.Geometry, Event )
-				: FReply::Unhandled();
-		});
+	Reply = FEventRouter::RouteAlongFocusPath(this, FEventRouter::FBubblePolicy(EventPath), InKeyboardEvent, [](const FArrangedWidget& SomeWidgetGettingEvent, const FKeyboardEvent& Event)
+	{
+		return (SomeWidgetGettingEvent.Widget->IsEnabled())
+			? SomeWidgetGettingEvent.Widget->OnKeyUp( SomeWidgetGettingEvent.Geometry, Event )
+			: FReply::Unhandled();
+	});
 
-		LOG_EVENT_CONTENT( EEventLog::KeyUp, GetKeyName(InKeyboardEvent.GetKey()).ToString(), Reply );
-	}
+	LOG_EVENT_CONTENT( EEventLog::KeyUp, GetKeyName(InKeyboardEvent.GetKey()).ToString(), Reply );
 
 	return Reply.IsEventHandled();
 }
@@ -3549,7 +3576,6 @@ bool FSlateApplication::ProcessMouseButtonDownEvent( const TSharedPtr< FGenericW
 		FReply Reply = FReply::Unhandled();
 		if (MouseCaptor.HasCaptureForPointerIndex(MouseEvent.GetPointerIndex()))
 		{
-			//FWidgetPath MouseCaptorPath = MouseCaptor.ToWidgetPath(MouseEvent.GetPointerIndex());
 			FWidgetPath MouseCaptorPath = MouseCaptor.ToWidgetPath( FWeakWidgetPath::EInterruptedPathHandling::Truncate, &MouseEvent );
 			FArrangedWidget& MouseCaptorWidget = MouseCaptorPath.Widgets.Last();
 
@@ -3557,16 +3583,14 @@ bool FSlateApplication::ProcessMouseButtonDownEvent( const TSharedPtr< FGenericW
 			FScopedSwitchWorldHack SwitchWorld( MouseCaptorPath );
 			bInGame = FApp::IsGame();
 
-			MouseEvent.SetEventPath( MouseCaptorPath );
-			
-			Reply = FEventRouter::RouteAlongPointerPath( this, FEventRouter::FToLeafmostPolicy(MouseCaptorPath), MouseEvent, []( const FArrangedWidget& MouseCaptorWidget, const FPointerEvent& Event )
+			Reply = FEventRouter::Route<FReply>( this, FEventRouter::FToLeafmostPolicy(MouseCaptorPath), MouseEvent, []( const FArrangedWidget& MouseCaptorWidget, const FPointerEvent& Event )
 			{
 				return MouseCaptorWidget.Widget->OnPreviewMouseButtonDown( MouseCaptorWidget.Geometry, Event );
 			});
 			
 			if (!Reply.IsEventHandled())
 			{
-				Reply = FEventRouter::RouteAlongPointerPath( this, FEventRouter::FToLeafmostPolicy(MouseCaptorPath), MouseEvent,
+				Reply = FEventRouter::Route<FReply>( this, FEventRouter::FToLeafmostPolicy(MouseCaptorPath), MouseEvent,
 					[this]( const FArrangedWidget& MouseCaptorWidget, const FPointerEvent& Event )
 					{
 						FReply TempReply = FReply::Unhandled();
@@ -3598,14 +3622,14 @@ bool FSlateApplication::ProcessMouseButtonDownEvent( const TSharedPtr< FGenericW
 
 			const TSharedPtr<SWidget> PreviouslyFocusedWidget = GetKeyboardFocusedWidget();
 
-			Reply = FEventRouter::RouteAlongPointerPath(this, FEventRouter::FTunnelPolicy(WidgetsUnderCursor), MouseEvent, [](const FArrangedWidget TargetWidget, const FPointerEvent& Event)
+			Reply = FEventRouter::Route<FReply>(this, FEventRouter::FTunnelPolicy(WidgetsUnderCursor), MouseEvent, [](const FArrangedWidget TargetWidget, const FPointerEvent& Event)
 			{
 				return TargetWidget.Widget->OnPreviewMouseButtonDown(TargetWidget.Geometry, Event);
 			});
 
 			if ( !Reply.IsEventHandled() )
 			{
-				Reply = FEventRouter::RouteAlongPointerPath(this, FEventRouter::FBubblePolicy(WidgetsUnderCursor), MouseEvent, [this](const FArrangedWidget TargetWidget, const FPointerEvent& Event)
+				Reply = FEventRouter::Route<FReply>(this, FEventRouter::FBubblePolicy(WidgetsUnderCursor), MouseEvent, [this](const FArrangedWidget TargetWidget, const FPointerEvent& Event)
 				{
 					FReply ThisReply = FReply::Unhandled();
 					if (!ThisReply.IsEventHandled())
@@ -3710,12 +3734,11 @@ bool FSlateApplication::ProcessMouseButtonDoubleClickEvent( const TSharedPtr< FG
 	PressedMouseButtons.Add( InMouseEvent.GetEffectingButton() );
 
 	FWidgetPath WidgetsUnderCursor = LocateWindowUnderMouse( InMouseEvent.GetScreenSpacePosition(), GetInteractiveTopLevelWindows() );
-	InMouseEvent.SetEventPath(WidgetsUnderCursor);
 
 	// Switch worlds widgets in the current path
 	FScopedSwitchWorldHack SwitchWorld( WidgetsUnderCursor );
 
-	FReply Reply = FEventRouter::RouteAlongPointerPath( this, FEventRouter::FBubblePolicy(WidgetsUnderCursor), InMouseEvent, [](const FArrangedWidget& TargetWidget, const FPointerEvent& Event)
+	FReply Reply = FEventRouter::Route<FReply>( this, FEventRouter::FBubblePolicy(WidgetsUnderCursor), InMouseEvent, [](const FArrangedWidget& TargetWidget, const FPointerEvent& Event)
 	{
 		return TargetWidget.Widget->OnMouseButtonDoubleClick( TargetWidget.Geometry, Event );
 	} );
@@ -3775,7 +3798,7 @@ bool FSlateApplication::ProcessMouseButtonUpEvent( FPointerEvent& MouseEvent )
 			// Switch worlds widgets in the current path
 			FScopedSwitchWorldHack SwitchWorld( MouseCaptorPath );
 
-			FReply Reply = FEventRouter::RouteAlongPointerPath( this, FEventRouter::FToLeafmostPolicy(MouseCaptorPath), MouseEvent, [this]( const FArrangedWidget& TargetWidget, const FPointerEvent& Event )
+			FReply Reply = FEventRouter::Route<FReply>( this, FEventRouter::FToLeafmostPolicy(MouseCaptorPath), MouseEvent, [this]( const FArrangedWidget& TargetWidget, const FPointerEvent& Event )
 			{
 				FReply TempReply = FReply::Unhandled();
 				if (Event.IsTouchEvent())
@@ -3797,7 +3820,6 @@ bool FSlateApplication::ProcessMouseButtonUpEvent( FPointerEvent& MouseEvent )
 				MouseCaptorPath.TopLevelWindow->BringToFront(true);
 			}
 #endif
-			//ProcessReply(MouseCaptorPath, Reply, &MouseCaptorPath, &TranslatedMouseEvent);
 			LOG_EVENT( EEventLog::MouseButtonUp, Reply );
 		}
 	}
@@ -3813,7 +3835,7 @@ bool FSlateApplication::ProcessMouseButtonUpEvent( FPointerEvent& MouseEvent )
 		// Switch worlds widgets in the current path
 		FScopedSwitchWorldHack SwitchWorld( WidgetsUnderCursor );
 		
-		FReply Reply = FEventRouter::RouteAlongPointerPath( this, FEventRouter::FBubblePolicy(WidgetsUnderCursor), MouseEvent, [&](const FArrangedWidget& CurWidget, const FPointerEvent& Event)
+		FReply Reply = FEventRouter::Route<FReply>( this, FEventRouter::FBubblePolicy(WidgetsUnderCursor), MouseEvent, [&](const FArrangedWidget& CurWidget, const FPointerEvent& Event)
 		{
 			FReply TempReply = FReply::Unhandled();
 			if (Event.IsTouchEvent())
@@ -3890,12 +3912,10 @@ bool FSlateApplication::ProcessMouseWheelOrGestureEvent( FPointerEvent& InWheelE
 		? MouseCaptor.ToWidgetPath(InWheelEvent.GetPointerIndex())
 		: LocateWindowUnderMouse( InWheelEvent.GetScreenSpacePosition(), GetInteractiveTopLevelWindows() );
 
-	InWheelEvent.SetEventPath( EventPath );
-
 	// Switch worlds widgets in the current path
 	FScopedSwitchWorldHack SwitchWorld( EventPath );
 
-	FReply Reply = FEventRouter::RouteAlongPointerPath( this, FEventRouter::FBubblePolicy(EventPath), InWheelEvent, [&InGestureEvent](const FArrangedWidget& CurWidget, const FPointerEvent& Event)
+	FReply Reply = FEventRouter::Route<FReply>( this, FEventRouter::FBubblePolicy(EventPath), InWheelEvent, [&InGestureEvent](const FArrangedWidget& CurWidget, const FPointerEvent& Event)
 	{
 		FReply TempReply = FReply::Unhandled();
 		// Gesture event gets first shot, if slate doesn't respond to it, we'll try the wheel event.
@@ -4008,14 +4028,15 @@ bool FSlateApplication::ProcessMouseMoveEvent( FPointerEvent& MouseEvent, bool b
 				DragDetector = FDragDetector();
 
 				// Send an OnDragDetected to the widget that requested drag-detection.
-				MouseEvent.SetEventPath( DragDetectPath );
-
+				
 				// Switch worlds widgets in the current path
 				FScopedSwitchWorldHack SwitchWorld( DragDetectPath );
 
-				FPointerEvent TranslatedMouseEvent = TranslatePointerEvent( DetectDragForMe.PointerPosition, MouseEvent );
-				FReply Reply = DetectDragForMe.Widget->OnDragDetected(DetectDragForMe.Geometry, TranslatedMouseEvent ).SetHandler(DetectDragForMe.Widget);
-				ProcessReply( DragDetectPath, Reply, &DragDetectPath, &TranslatedMouseEvent );
+				FReply Reply = FEventRouter::Route<FReply>(this, FEventRouter::FDirectPolicy(DetectDragForMe,DragDetectPath), MouseEvent, []( const FArrangedWidget& DetectDragForMe, const FPointerEvent& TranslatedMouseEvent )
+				{
+					return DetectDragForMe.Widget->OnDragDetected(DetectDragForMe.Geometry, TranslatedMouseEvent );
+				});
+				
 				LOG_EVENT( EEventLog::DragDetected, Reply );
 			}
 			else
@@ -4037,9 +4058,6 @@ bool FSlateApplication::ProcessMouseMoveEvent( FPointerEvent& MouseEvent, bool b
 		LastWidgetsUnderCursor = WidgetsUnderCursorLastEvent;
 	}
 
-	// In the case of drag leave/enter events there is no path to speak of
-	MouseEvent.SetEventPath( FWidgetPath() );
-
 	// Send out mouse leave events
 	// If we are doing a drag and drop, we will send this event instead.
 	{
@@ -4103,7 +4121,7 @@ bool FSlateApplication::ProcessMouseMoveEvent( FPointerEvent& MouseEvent, bool b
 			// Switch worlds widgets in the current path
 			FScopedSwitchWorldHack SwitchWorld( MouseCaptorPath );
 			
-			FReply Reply = FEventRouter::RouteAlongPointerPath( this, FEventRouter::FToLeafmostPolicy(MouseCaptorPath), MouseEvent, [this]( const FArrangedWidget& MouseCaptorWidget, const FPointerEvent& Event )
+			FReply Reply = FEventRouter::Route<FReply>( this, FEventRouter::FToLeafmostPolicy(MouseCaptorPath), MouseEvent, [this]( const FArrangedWidget& MouseCaptorWidget, const FPointerEvent& Event )
 			{
 				FReply TempReply = FReply::Unhandled();
 				if (Event.IsTouchEvent())
@@ -4125,31 +4143,32 @@ bool FSlateApplication::ProcessMouseMoveEvent( FPointerEvent& MouseEvent, bool b
 		FScopedSwitchWorldHack SwitchWorld( WidgetsUnderCursor );
 
 		// Send out mouse enter events.
-		FDragDropEvent DragDropEvent( MouseEvent, DragDropContent );
-		DragDropEvent.SetEventPath( WidgetsUnderCursor );
-		for( int32 WidgetIndex = WidgetsUnderCursor.Widgets.Num()-1; WidgetIndex >= 0; --WidgetIndex )
+		if (IsDragDropping())
 		{
-			FArrangedWidget& SomeWidgetUnderCursor = WidgetsUnderCursor.Widgets[ WidgetIndex ];
-			if ( !LastWidgetsUnderCursor.ContainsWidget(SomeWidgetUnderCursor.Widget) )
+			FDragDropEvent DragDropEvent( MouseEvent, DragDropContent );
+			FEventRouter::Route<FNoReply>( this, FEventRouter::FBubblePolicy(WidgetsUnderCursor), DragDropEvent, [&LastWidgetsUnderCursor](const FArrangedWidget& WidgetUnderCursor, const FDragDropEvent& DragDropEvent)
 			{
-				// Widget newly under cursor, so send a MouseEnter.
-				if ( IsDragDropping() )
+				if ( !LastWidgetsUnderCursor.ContainsWidget(WidgetUnderCursor.Widget) )
 				{
-					// Doing a drag and drop; send a DragDropEvent
-					SomeWidgetUnderCursor.Widget->OnDragEnter( SomeWidgetUnderCursor.Geometry, TranslatePointerEvent( WidgetsUnderCursor.VirtualPointerPositions[WidgetIndex], DragDropEvent ) );
-					LOG_EVENT( EEventLog::DragEnter, SomeWidgetUnderCursor.Widget );
+					WidgetUnderCursor.Widget->OnDragEnter( WidgetUnderCursor.Geometry, DragDropEvent );
 				}
-				else
+				return FNoReply();
+			} );
+		}
+		else
+		{
+			FEventRouter::Route<FNoReply>( this, FEventRouter::FBubblePolicy(WidgetsUnderCursor), MouseEvent, [&LastWidgetsUnderCursor](const FArrangedWidget& WidgetUnderCursor, const FPointerEvent& PointerEvent)
+			{
+				if ( !LastWidgetsUnderCursor.ContainsWidget(WidgetUnderCursor.Widget) )
 				{
-					// Not drag dropping; send regular mouse event
-					SomeWidgetUnderCursor.Widget->OnMouseEnter( SomeWidgetUnderCursor.Geometry, TranslatePointerEvent( WidgetsUnderCursor.VirtualPointerPositions[WidgetIndex], MouseEvent ) );
-					LOG_EVENT( EEventLog::MouseEnter, SomeWidgetUnderCursor.Widget );
+					WidgetUnderCursor.Widget->OnMouseEnter( WidgetUnderCursor.Geometry, PointerEvent );
 				}
-			}
+				return FNoReply();
+			} );
 		}
-
+		
 		// Bubble the MouseMove event.
-		FReply Reply = FEventRouter::RouteAlongPointerPath(this, FEventRouter::FBubblePolicy(WidgetsUnderCursor), MouseEvent, [&](const FArrangedWidget& CurWidget, const FPointerEvent& Event)
+		FReply Reply = FEventRouter::Route<FReply>(this, FEventRouter::FBubblePolicy(WidgetsUnderCursor), MouseEvent, [&](const FArrangedWidget& CurWidget, const FPointerEvent& Event)
 		{
 			FReply TempReply = FReply::Unhandled();
 
@@ -4857,7 +4876,6 @@ bool FSlateApplication::ProcessDragEnterEvent( TSharedRef<SWindow> WindowEntered
 	LastUserInteractionTime = this->GetCurrentTime();
 	
 	FWidgetPath WidgetsUnderCursor = LocateWindowUnderMouse( DragDropEvent.GetScreenSpacePosition(), GetInteractiveTopLevelWindows() );
-	DragDropEvent.SetEventPath( WidgetsUnderCursor );
 
 	// Switch worlds for widgets in the current path
 	FScopedSwitchWorldHack SwitchWorld( WidgetsUnderCursor );
diff --git a/Engine/Source/Runtime/Slate/Public/Framework/Application/IWidgetReflector.h b/Engine/Source/Runtime/Slate/Public/Framework/Application/IWidgetReflector.h
--- a/Engine/Source/Runtime/Slate/Public/Framework/Application/IWidgetReflector.h
+++ b/Engine/Source/Runtime/Slate/Public/Framework/Application/IWidgetReflector.h
@@ -13,7 +13,7 @@ DECLARE_DELEGATE_RetVal_ThreeParams(bool, FAccessSourceCode, const FString& /*Fi
 class IWidgetReflector
 {
 public:
-	virtual void OnEventProcessed( const FInputEvent& Event, const FReply& InReply ) = 0;
+	virtual void OnEventProcessed( const FInputEvent& Event, const FReplyBase& InReply ) = 0;
 
 public:
 
diff --git a/Engine/Source/Runtime/Slate/Public/Framework/Application/SlateApplication.h b/Engine/Source/Runtime/Slate/Public/Framework/Application/SlateApplication.h
--- a/Engine/Source/Runtime/Slate/Public/Framework/Application/SlateApplication.h
+++ b/Engine/Source/Runtime/Slate/Public/Framework/Application/SlateApplication.h
@@ -27,9 +27,6 @@ extern SLATE_API const FName NAME_UnrealOS;
  */
 namespace SlateApplicationDefs
 {
-	/** How much to amplify keyboard repeat rate by */
-	static const int32 NumRepeatsPerActualRepeat = 2;
-
 	/** How many users can we support at once? */
 	static const int32 MaxUsers = 8;
 }
diff --git a/Engine/Source/Runtime/SlateCore/Public/Input/CursorReply.h b/Engine/Source/Runtime/SlateCore/Public/Input/CursorReply.h
--- a/Engine/Source/Runtime/SlateCore/Public/Input/CursorReply.h
+++ b/Engine/Source/Runtime/SlateCore/Public/Input/CursorReply.h
@@ -6,44 +6,46 @@
 /**
  * A reply to the OnQueryCursor event.
  */
-class FCursorReply
+class FCursorReply : public TReplyBase<FCursorReply>
 {
 public:
 
-	/** Makes a NULL response; this response will not change the previous value. */
+	/**
+	 * Makes a NULL response meaning no prefersce.
+	 * i.e. If your widget returns this, its parent will get to decide what the cursor shoudl be.
+	 * This is the default behavior for a widget.
+	 */
 	static FCursorReply Unhandled()
 	{
 		return FCursorReply();
 	}
 		
-	/** Respond with a cursor type. */
+	/**
+	 * Respond with a specific cursor.
+	 * This cursor will be used and no other widgets will be asked.
+	 */
 	static FCursorReply Cursor( EMouseCursor::Type InCursor )
 	{
 		return FCursorReply( InCursor );
 	}
 		
-	/** @return true if this this reply is a result of the event being handled; false otherwise. */
-	bool IsEventHandled(){ return bIsHandled; }
-		
 	/** @return The requested MouseCursor when the event was handled. Undefined otherwise. */
 	EMouseCursor::Type GetCursor() { return MouseCursor; }
-		
+
 private:
 
 	/** Internal constructor - makes a NULL result. */
 	FCursorReply()
-		: bIsHandled(false)
+		: TReplyBase<FCursorReply>(false)
 		, MouseCursor( EMouseCursor::Default )
 	{ }
 		
 	/** Internal constructor - makes a non-NULL result. */
 	FCursorReply( EMouseCursor::Type InCursorType )
-		: bIsHandled(true)
+		: TReplyBase<FCursorReply>(true)
 		, MouseCursor( InCursorType )
 	{ }
-		
-	/** Does this reply have any meaning */ 
-	bool bIsHandled;
+	
 		
 	/** The value of the reply, if bHasValue is true. */
 	EMouseCursor::Type MouseCursor;
diff --git a/Engine/Source/Runtime/SlateCore/Public/Input/Reply.h b/Engine/Source/Runtime/SlateCore/Public/Input/Reply.h
--- a/Engine/Source/Runtime/SlateCore/Public/Input/Reply.h
+++ b/Engine/Source/Runtime/SlateCore/Public/Input/Reply.h
@@ -13,7 +13,7 @@ class SWidget;
  * For example, a widget may handle an OnMouseDown event by asking the system to give mouse capture to a specific Widget.
  * To do this, return FReply::CaptureMouse( NewMouseCapture ).
  */
-class FReply
+class FReply : public TReplyBase<FReply>
 {
 public:
 		
@@ -21,14 +21,14 @@ public:
 	FReply& CaptureMouse( TSharedRef<SWidget> InMouseCaptor )
 	{
 		this->MouseCaptor = InMouseCaptor;
-		return *this;
+		return Me();
 	}
 
 	FReply& CaptureJoystick( TSharedRef<SWidget> InJoystickCaptor, bool bInAllJoysticks = false )
 	{
 		this->JoystickCaptor = InJoystickCaptor;
 		this->bAllJoysticks = bInAllJoysticks;
-		return *this;
+		return Me();
 	}
 
 	/**
@@ -39,7 +39,7 @@ public:
 	{
 		this->MouseCaptor = InMouseCaptor;
 		this->bUseHighPrecisionMouse = true;
-		return *this;
+		return Me();
 	}
 
 	/**
@@ -48,7 +48,7 @@ public:
 	FReply& SetMousePos( const FIntPoint& NewMousePos )
 	{
 		this->RequestedMousePos = NewMousePos;
-		return *this;
+		return Me();
 	}
 
 	/** An event should return FReply::Handled().SetKeyboardFocus( SomeWidget ) as a means of asking the system to set keyboard focus to the provided widget*/
@@ -67,7 +67,7 @@ public:
 	{
 		this->MouseLockWidget = InWidget;
 		this->bShouldReleaseMouseLock = false;
-		return *this;
+		return Me();
 	}
 
 	/** 
@@ -88,7 +88,7 @@ public:
 	{
 		this->bReleaseMouseCapture = true;
 		this->bUseHighPrecisionMouse = false;
-		return *this;		
+		return Me();
 	}
 
 	/**
@@ -98,7 +98,7 @@ public:
 	{
 		this->bReleaseJoystickCapture = true;
 		this->bAllJoysticks = bInAllJoysticks;
-		return *this;		
+		return Me();
 	}
 
 	/**
@@ -112,7 +112,7 @@ public:
 	{
 		this->DetectDragForWidget = DetectDragInMe;
 		this->DetectDragForMouseButton = MouseButton;
-		return *this;
+		return Me();
 	}
 
 	/**
@@ -125,31 +125,25 @@ public:
 	FReply& BeginDragDrop(TSharedRef<FDragDropOperation> InDragDropContent)
 	{
 		this->DragDropContent = InDragDropContent;
-		return *this;
+		return Me();
 	}
 
 	/** An event should return FReply::Handled().EndDragDrop() to request that the current drag/drop operation be terminated. */
 	FReply& EndDragDrop()
 	{
 		this->bEndDragDrop = true;
-		return *this;
+		return Me();
 	}
 
 	/** Ensures throttling for Slate UI responsiveness is not done on mouse down */
 	FReply& PreventThrottling()
 	{
 		this->bPreventThrottling = true;
-		return *this;
+		return Me();
 	}
 
 public:
 
-	/** True if this reply indicated that the event was handled */
-	bool IsEventHandled() const { return bIsHandled; }
-
-	/** The widget that ultimately handled the event */
-	const TSharedPtr<SWidget> GetHandler() const { return EventHandler; }
-
 	/** True if this reply indicated that we should release mouse capture as a result of the event being handled */
 	bool ShouldReleaseMouse() const { return bReleaseMouseCapture; }
 
@@ -205,9 +199,7 @@ public:
 	 */
 	static FReply Handled( )
 	{
-		FReply NewReply;
-		NewReply.bIsHandled = true;
-		return NewReply;
+		return FReply(true);
 	}
 
 	/**
@@ -215,9 +207,7 @@ public:
 	 */
 	static FReply Unhandled( )
 	{
-		FReply NewReply;
-		NewReply.bIsHandled = false;
-		return NewReply;
+		return FReply(false);
 	}
 
 private:
@@ -225,8 +215,9 @@ private:
 	/**
 	 * Hidden default constructor.
 	 */
-	FReply( )
-		: RequestedMousePos()
+	FReply( bool bIsHandled )
+		: TReplyBase<FReply>(bIsHandled)
+		, RequestedMousePos()
 		, EventHandler(nullptr)
 		, MouseCaptor(nullptr)
 		, JoystickCaptor(nullptr)
@@ -234,7 +225,6 @@ private:
 		, MouseLockWidget(nullptr)
 		, DragDropContent(nullptr)
 		, FocusChangeReason(EKeyboardFocusCause::SetDirectly)
-		, bIsHandled(false)
 		, bReleaseMouseCapture(false)
 		, bReleaseJoystickCapture(false)
 		, bAllJoysticks(false)
@@ -246,17 +236,8 @@ private:
 		
 
 private:
-
-	friend class FEventRouter;
 	friend class FSlateApplication;
 
-	/** Set the widget that handled the event; undefined if never handled. This method is to be used by SlateApplication only! */
-	FReply& SetHandler( const TSharedRef<SWidget>& InHandler )
-	{
-		this->EventHandler = InHandler;
-		return *this;
-	}
-
 private:
 
 	TOptional<FIntPoint> RequestedMousePos;
@@ -269,7 +250,6 @@ private:
 	FKey DetectDragForMouseButton;
 	TSharedPtr<FDragDropOperation> DragDropContent;
 	EKeyboardFocusCause::Type FocusChangeReason;
-	uint32 bIsHandled:1;
 	uint32 bReleaseMouseCapture:1;
 	uint32 bReleaseJoystickCapture:1;
 	uint32 bAllJoysticks:1;
diff --git a/Engine/Source/Runtime/SlateCore/Public/Input/ReplyBase.h b/Engine/Source/Runtime/SlateCore/Public/Input/ReplyBase.h
new file mode 100644
--- /dev/null
+++ b/Engine/Source/Runtime/SlateCore/Public/Input/ReplyBase.h
@@ -0,0 +1,67 @@
+// Copyright 1998-2014 Epic Games, Inc. All Rights Reserved.
+
+#pragma once
+
+/**
+ * Base class for all the ways that a Slate Widget can reply to events.
+ * Typical replies include requests to set focus, capture or release the pointer,
+ * signify a cursor preference, etc.
+ */
+class FReplyBase
+{
+public:
+
+	/** @return true if this this reply is a result of the event being handled; false otherwise. */
+	bool IsEventHandled() const { return bIsHandled; }
+	
+	/** The widget that ultimately handled the event */
+	const TSharedPtr<SWidget> GetHandler() const { return EventHandler; }
+		
+protected:
+
+	/** A reply can be handled or unhandled. Any widget handling events decides whether it has handled the event. */
+	FReplyBase( bool InIsHandled )
+	: bIsHandled(InIsHandled)
+	{
+	}
+
+	/** Has a widget handled an event. */ 
+	bool bIsHandled;
+	
+	/** Widget that handled the event that generated this reply. */
+	TSharedPtr<SWidget> EventHandler;
+};
+
+template<typename ReplyType>
+class TReplyBase : public FReplyBase
+{
+public:
+	TReplyBase( bool bIsHandled )
+	: FReplyBase( bIsHandled )
+	{
+	}
+
+protected:
+	friend class FEventRouter;
+
+	/** Set the widget that handled the event; undefined if never handled. This method is to be used by SlateApplication only! */
+	ReplyType& SetHandler( const TSharedRef<SWidget>& InHandler )
+	{
+		this->EventHandler = InHandler;
+		return Me();
+	}
+
+	/** @return a reference to this reply */
+	ReplyType& Me()
+	{
+		return static_cast<ReplyType&>(*this);
+	}
+};
+
+/** A reply type for events that return a void reply. e.g. OnMouseLeave() */
+class FNoReply : public TReplyBase<FNoReply>
+{
+public:
+	static FNoReply Unhandled(){ return FNoReply(); }
+	FNoReply() : TReplyBase<FNoReply>(false){}
+};
\ No newline at end of file
diff --git a/Engine/Source/Runtime/SlateCore/Public/SlateCore.h b/Engine/Source/Runtime/SlateCore/Public/SlateCore.h
--- a/Engine/Source/Runtime/SlateCore/Public/SlateCore.h
+++ b/Engine/Source/Runtime/SlateCore/Public/SlateCore.h
@@ -127,6 +127,7 @@ DECLARE_MEMORY_STAT_EXTERN(TEXT("Index Buffer Memory"), STAT_SlateIndexBufferMem
 #include "SlateStyle.h"
 
 // Input
+#include "ReplyBase.h"
 #include "CursorReply.h"
 #include "Events.h"
 #include "DragAndDrop.h"
