Fixed Clang compile error

diff --git a/Engine/Source/Runtime/Slate/Private/Framework/Application/SlateApplication.cpp b/Engine/Source/Runtime/Slate/Private/Framework/Application/SlateApplication.cpp
--- a/Engine/Source/Runtime/Slate/Private/Framework/Application/SlateApplication.cpp
+++ b/Engine/Source/Runtime/Slate/Private/Framework/Application/SlateApplication.cpp
@@ -242,27 +242,26 @@ public:
 		}
 	};
 
-	template<>
-	struct Translate<FPointerEvent>
+};
+
+template<>
+struct FEventRouter::Translate<FPointerEvent>
+{
+	static  FPointerEvent PointerEvent( const TSharedPtr<FVirtualPointerPosition>& InPosition, const FPointerEvent& InEvent )
 	{
-		static  FPointerEvent PointerEvent( const TSharedPtr<FVirtualPointerPosition>& InPosition, const FPointerEvent& InEvent )
+		// Pointer events are translated into the virtual window space. For 3D Widget Components this means
+		if ( !InPosition.IsValid() )
 		{
-			// Pointer events are translated into the virtual window space. For 3D Widget Components this means
-			if ( !InPosition.IsValid() )
-			{
-				return InEvent;
-			}
-			else
-			{
-				return FPointerEvent::MakeTranslatedEvent<FPointerEvent>( InEvent, *InPosition );
-			}
+			return InEvent;
 		}
-	};
-
+		else
+		{
+			return FPointerEvent::MakeTranslatedEvent<FPointerEvent>( InEvent, *InPosition );
+		}
+	}
 };
 
 
-
 DECLARE_CYCLE_STAT( TEXT("Message Tick Time"), STAT_SlateMessageTick, STATGROUP_Slate );
 DECLARE_CYCLE_STAT( TEXT("Update Tooltip Time"), STAT_SlateUpdateTooltip, STATGROUP_Slate );
 DECLARE_CYCLE_STAT( TEXT("Tick Window And Children Time"), STAT_SlateTickWindowAndChildren, STATGROUP_Slate );
