Fix focus issues on windows and mac when a window is activated while a modal dialog is open

diff --git a/Engine/Source/Runtime/Slate/Private/Framework/Application/SlateApplication.cpp b/Engine/Source/Runtime/Slate/Private/Framework/Application/SlateApplication.cpp
--- a/Engine/Source/Runtime/Slate/Private/Framework/Application/SlateApplication.cpp
+++ b/Engine/Source/Runtime/Slate/Private/Framework/Application/SlateApplication.cpp
@@ -4610,13 +4610,15 @@ bool FSlateApplication::ProcessWindowActivatedEvent( const FWindowActivateEvent&
 		// Only window activate considered a user interaction
 		LastUserInteractionTime = this->GetCurrentTime();
 		
+		// NOTE: The window is brought to front even when a modal window is active and this is not the modal window one of its children 
+		// The reason for this is so that the Slate window order is in sync with the OS window order when a modal window is open.  This is important so that when the modal window closes the proper window receives input from Slate.
+		// If you change this be sure to test windows are activated properly and receive input when they are opened when a modal dialog is open.
+		FSlateWindowHelper::BringWindowToFront(SlateWindows, ActivateEvent.GetAffectedWindow());
+
 		// Do not process activation messages unless we have no modal windows or the current window is modal
 		if( !ActiveModalWindow.IsValid() || ActivateEvent.GetAffectedWindow() == ActiveModalWindow || ActivateEvent.GetAffectedWindow()->IsDescendantOf(ActiveModalWindow) )
 		{
 			// Window being ACTIVATED
-
-			FSlateWindowHelper::BringWindowToFront(SlateWindows, ActivateEvent.GetAffectedWindow());
-
 			{
 				// Switch worlds widgets in the current path
 				FScopedSwitchWorldHack SwitchWorld( ActivateEvent.GetAffectedWindow() );
