Avoid changing the cursor location while the mouse is leaving the gameviewport

diff --git a/Engine/Source/Runtime/Engine/Private/GameViewportClient.cpp b/Engine/Source/Runtime/Engine/Private/GameViewportClient.cpp
--- a/Engine/Source/Runtime/Engine/Private/GameViewportClient.cpp
+++ b/Engine/Source/Runtime/Engine/Private/GameViewportClient.cpp
@@ -425,10 +425,8 @@ void UGameViewportClient::MouseLeave(FViewport* InViewport)
 	{
 		FIntPoint LastViewportCursorPos;
 		Viewport->GetMousePos(LastViewportCursorPos, false);
-		const FVector2D CurrentCursorPos = FSlateApplication::Get().GetCursorPos();
-		FSlateApplication::Get().SetCursorPos(FVector2D(LastViewportCursorPos.X,LastViewportCursorPos.Y));
-		FSlateApplication::Get().SetGameIsFakingTouchEvents(false);
-		FSlateApplication::Get().SetCursorPos(CurrentCursorPos);
+		FVector2D CursorPos(LastViewportCursorPos.X, LastViewportCursorPos.Y);
+		FSlateApplication::Get().SetGameIsFakingTouchEvents(false, &CursorPos);
 	}
 }
 
diff --git a/Engine/Source/Runtime/Slate/Private/Framework/Application/SlateApplication.cpp b/Engine/Source/Runtime/Slate/Private/Framework/Application/SlateApplication.cpp
--- a/Engine/Source/Runtime/Slate/Private/Framework/Application/SlateApplication.cpp
+++ b/Engine/Source/Runtime/Slate/Private/Framework/Application/SlateApplication.cpp
@@ -3520,11 +3520,11 @@ FKey TranslateMouseButtonToKey( const EMouseButtons::Type Button )
 	return Key;
 }
 
-void FSlateApplication::SetGameIsFakingTouchEvents(const bool bIsFaking)
+void FSlateApplication::SetGameIsFakingTouchEvents(const bool bIsFaking, FVector2D* CursorLocation)
 {
 	if (bIsFakingTouched && !bIsFaking && bIsGameFakingTouch && !bIsFakingTouch)
 	{
-		OnTouchEnded(PlatformApplication->Cursor->GetPosition(), 0, 0);
+		OnTouchEnded((CursorLocation ? *CursorLocation : PlatformApplication->Cursor->GetPosition()), 0, 0);
 	}
 	bIsGameFakingTouch = bIsFaking;
 }
diff --git a/Engine/Source/Runtime/Slate/Public/Framework/Application/SlateApplication.h b/Engine/Source/Runtime/Slate/Public/Framework/Application/SlateApplication.h
--- a/Engine/Source/Runtime/Slate/Public/Framework/Application/SlateApplication.h
+++ b/Engine/Source/Runtime/Slate/Public/Framework/Application/SlateApplication.h
@@ -902,7 +902,8 @@ public:
 	/** @return true if mouse events are being turned into touch events, and touch UI should be forced on */
 	bool IsFakingTouchEvents() const;
 
-	void SetGameIsFakingTouchEvents(const bool bIsFaking);
+	/** Sets whether the application is treating mouse events as imitating touch events.  Optional CursorLocation can be supplied to override the platform's belief of where the cursor is */
+	void SetGameIsFakingTouchEvents(const bool bIsFaking, FVector2D* CursorLocation = nullptr);
 
 	/** Sets the handler for otherwise unhandled key down events. This is used by the editor to provide a global action list, if the key was not consumed by any widget. */
 	void SetUnhandledKeyDownEventHandler( const FOnKeyboardEvent& NewHandler );
