Fixed a crash caused by routing input events intended for a modal window to a non modal window

The problem was that the slow task dialog has not been written to the hit test grid at the point that we Tick (and subsequently handle an input message for it). Although the message is intended for the slow task window, the widget path generated from the hittest grid said otherwise (returning an SEditorViewport, since that was the last widget rendered to it).

This addresses UE4750 - Crash painting landscape layers during material compilation due to Slate input re-entrancy

diff --git a/Engine/Source/Runtime/Slate/Private/Framework/Application/SlateApplication.cpp b/Engine/Source/Runtime/Slate/Private/Framework/Application/SlateApplication.cpp
--- a/Engine/Source/Runtime/Slate/Private/Framework/Application/SlateApplication.cpp
+++ b/Engine/Source/Runtime/Slate/Private/Framework/Application/SlateApplication.cpp
@@ -1454,6 +1454,11 @@ void FSlateApplication::AddModalWindow( TSharedRef<SWindow> InSlateWindow, const
 		return;
 	}
 
+	// If we are spawning a modal window, we don't want to allow the possibility of input events being routed to non-modal geometry that existed on the hittest grid last frame.
+	// To this end we clear the grid now. This means that input events that occur on the first frame will be missed, however since the user couldn't even see the window yet, this is ok.
+	// @todo slate : switch to FHittestGrid per window
+	HittestGrid->BeginFrame( VirtualDesktopRect );
+
 	// Push the active modal window onto the stack.  
 	ActiveModalWindows.AddUnique( InSlateWindow );
 
