Rename/move file(s)

Exposing SWindowTitleBar to the public. Exposed options to specify the title and for whether or not to display an app icon.

diff --git a/Engine/Source/Runtime/Slate/Private/Framework/Application/SWindowTitleBar.h b/Engine/Source/Runtime/Slate/Private/Framework/Application/SWindowTitleBar.h
deleted file mode 100644
diff --git a/Engine/Source/Runtime/Slate/Public/Framework/Application/SWindowTitleBar.h b/Engine/Source/Runtime/Slate/Public/Framework/Application/SWindowTitleBar.h
new file mode 100644
--- /dev/null
+++ b/Engine/Source/Runtime/Slate/Public/Framework/Application/SWindowTitleBar.h
@@ -0,0 +1,687 @@
+// Copyright 1998-2014 Epic Games, Inc. All Rights Reserved.
+
+#pragma once
+
+
+namespace SWindowTitleBarDefs
+{
+	/** Window flash rate. Flashes per second */
+	static const float WindowFlashFrequency = 5.5f;
+
+	/** Window flash duration. Seconds*/
+	static const float WindowFlashDuration = 1.0f;
+}
+
+
+/** Widget that represents the app icon + system menu button, usually drawn in the top left of a Windows app */
+class SAppIconWidget
+	: public SCompoundWidget
+{
+	SLATE_BEGIN_ARGS( SAppIconWidget )
+		: _IconColorAndOpacity( FLinearColor::White )
+	{}
+
+	/** Icon color and opacity */
+	SLATE_ATTRIBUTE( FSlateColor, IconColorAndOpacity )
+
+	SLATE_END_ARGS( )
+	
+	void Construct( const FArguments& Args )
+	{
+		this->ChildSlot
+			[
+				SNew(SVerticalBox)
+				+ SVerticalBox::Slot()
+				.Padding( FMargin( 0.0f, 0.0f, 2.0f, 0.0f ) )
+				[
+					FSlateApplicationBase::Get().MakeImage(
+						FSlateApplicationBase::Get().GetAppIcon(),
+						Args._IconColorAndOpacity,
+						EVisibility::HitTestInvisible
+					)
+				]
+			];
+	}
+
+	virtual EWindowZone::Type GetWindowZoneOverride() const override
+	{
+		// Pretend we are a REAL system menu so the user can click to open a menu, or double-click to close the app on Windows
+		return EWindowZone::SysMenu;
+	}
+};
+
+
+/**
+ * Implements a window title bar widget.
+ */
+class SWindowTitleBar
+	: public SCompoundWidget
+	, public IWindowTitleBar
+{
+public:
+
+	SLATE_BEGIN_ARGS(SWindowTitleBar)
+		: _Style(&FCoreStyle::Get().GetWidgetStyle<FWindowStyle>("Window"))
+		, _ShowAppIcon(true)
+	{ }
+		SLATE_STYLE_ARGUMENT(FWindowStyle, Style)
+		SLATE_ARGUMENT(bool, ShowAppIcon)
+		SLATE_ATTRIBUTE(FText, Title)
+	SLATE_END_ARGS()
+	
+public:
+
+	/**
+	 * Creates and initializes a new window title bar widget.
+	 *
+	 * @param InArgs The construction arguments.
+	 * @param InWindow The window to create the title bar for.
+	 * @param InCenterContent The content for the title bar's center area.
+	 * @param CenterContentAlignment The horizontal alignment of the center content.
+	 */
+	void Construct( const FArguments& InArgs, const TSharedRef<SWindow>& InWindow, const TSharedPtr<SWidget>& InCenterContent, EHorizontalAlignment InCenterContentAlignment )
+	{
+		OwnerWindowPtr = InWindow;
+		Style = InArgs._Style;
+		ShowAppIcon = InArgs._ShowAppIcon;
+		Title = InArgs._Title;
+
+		if (!Title.IsSet() && !Title.IsBound())
+		{
+			Title = TAttribute<FText>::Create(TAttribute<FText>::FGetter::CreateSP(this, &SWindowTitleBar::HandleWindowTitleText));
+		}
+
+		ChildSlot
+		[
+			SNew(SBorder)
+			.Padding(0.0f)
+			.Visibility(EVisibility::SelfHitTestInvisible)
+			.BorderImage(this, &SWindowTitleBar::GetWindowTitlebackgroundImage)
+			[
+				SNew(SOverlay)
+				.Visibility(EVisibility::SelfHitTestInvisible)
+				+ SOverlay::Slot()
+				[
+					SNew(SImage)
+						.Visibility(this, &SWindowTitleBar::GetWindowFlashVisibility)
+						.Image(&Style->FlashTitleBrush )
+						.ColorAndOpacity(this, &SWindowTitleBar::GetWindowTitleAreaColor)
+				]
+				+ SOverlay::Slot()
+				[
+					MakeTitleBarContent(InCenterContent, InCenterContentAlignment)
+				]
+			]
+		];
+	}
+
+	virtual EWindowZone::Type GetWindowZoneOverride() const override
+	{
+		return EWindowZone::TitleBar;
+	}
+
+public:
+
+	// Begin IWindowTitleBar interface
+
+	void Flash( )
+	{
+		TitleFlashSequence = FCurveSequence(0, SWindowTitleBarDefs::WindowFlashDuration, ECurveEaseFunction::Linear);
+		TitleFlashSequence.Play();
+	}
+
+	// Begin IWindowTitleBar interface
+	
+protected:
+
+	float GetFlashValue( ) const
+	{
+		if (TitleFlashSequence.IsPlaying())
+		{
+			float Lerp = TitleFlashSequence.GetLerp();
+
+			const float SinRateMultiplier = 2.0f * PI * SWindowTitleBarDefs::WindowFlashDuration * SWindowTitleBarDefs::WindowFlashFrequency;
+			float SinTerm = 0.5f * (FMath::Sin( Lerp * SinRateMultiplier ) + 1.0f);
+
+			float FadeTerm = 1.0f - Lerp;
+
+			return SinTerm * FadeTerm;
+		}
+		return 0.0f;
+	}
+
+	/**
+	 * Creates widgets for this window's title bar area.
+	 *
+	 * This is an advanced method, only for fancy windows that want to
+	 * override the look of the title area by arranging those widgets itself.
+	 */
+	void MakeTitleBarContentWidgets( TSharedPtr< SWidget >& OutLeftContent, TSharedPtr< SWidget >& OutRightContent )
+	{
+		TSharedPtr<SWindow> OwnerWindow = OwnerWindowPtr.Pin();
+
+		if (!OwnerWindow.IsValid())
+		{
+			return;
+		}
+
+		MinimizeButton = SNew(SButton)
+				.IsFocusable(false)
+				.IsEnabled(OwnerWindow->HasMinimizeBox())
+				.ContentPadding(0)
+				.OnClicked(this, &SWindowTitleBar::MinimizeButton_OnClicked)
+				.Cursor(EMouseCursor::Default)
+				.ButtonStyle(FCoreStyle::Get(), "NoBorder")
+				[
+					SNew(SImage)
+						.Image(this, &SWindowTitleBar::GetMinimizeImage)
+						.ColorAndOpacity(this, &SWindowTitleBar::GetWindowTitleContentColor)
+				]
+			;
+
+		MaximizeRestoreButton = SNew(SButton)
+				.IsFocusable(false)
+				.IsEnabled(OwnerWindow->HasMaximizeBox())
+				.ContentPadding(0.0f)
+				.OnClicked(this, &SWindowTitleBar::MaximizeRestoreButton_OnClicked)
+				.Cursor(EMouseCursor::Default)
+				.ButtonStyle(FCoreStyle::Get(), "NoBorder")
+				[
+					SNew(SImage)
+						.Image(this, &SWindowTitleBar::GetMaximizeRestoreImage)
+						.ColorAndOpacity(this, &SWindowTitleBar::GetWindowTitleContentColor)
+				]
+			;
+
+		CloseButton = SNew(SButton)
+				.IsFocusable(false)
+				.ContentPadding(0.0f)
+				.OnClicked(this, &SWindowTitleBar::CloseButton_OnClicked)
+				.Cursor(EMouseCursor::Default)
+				.ButtonStyle(FCoreStyle::Get(), "NoBorder")
+				[
+					SNew(SImage)
+						.Image(this, &SWindowTitleBar::GetCloseImage)
+						.ColorAndOpacity(this, &SWindowTitleBar::GetWindowTitleContentColor)
+				]
+			;
+
+#if PLATFORM_MAC
+
+		// Mac UI layout
+		OutLeftContent = SNew(SBox)
+			.Visibility(EVisibility::SelfHitTestInvisible)
+			.Padding(FMargin(0.0f, 2.0f, 0.0f, 0.0f))
+			[
+				SNew(SHorizontalBox)
+					.Visibility(EVisibility::SelfHitTestInvisible)
+
+				// Close button
+				+ SHorizontalBox::Slot()
+					.AutoWidth()
+					.Padding(FMargin( 4.0f, 0.0f, 0.0f, 0.0f))
+					[
+						CloseButton.ToSharedRef()
+					]
+
+				// Minimize
+				+ SHorizontalBox::Slot()
+					.AutoWidth()
+					[
+						MinimizeButton.ToSharedRef()
+					]
+
+				// Maximize/Restore
+				+ SHorizontalBox::Slot()
+					.AutoWidth()
+					[
+						MaximizeRestoreButton.ToSharedRef()
+					]
+			];
+
+		// No right content on Mac
+		OutRightContent = SNew(SSpacer);
+
+#else // PLATFORM_MAC
+
+		// Windows UI layout
+		if (ShowAppIcon)
+		{
+			OutLeftContent = SNew(SAppIconWidget)
+				.IconColorAndOpacity(this, &SWindowTitleBar::GetWindowTitleContentColor);
+		}
+		else
+		{
+			OutLeftContent = SNullWidget::NullWidget;
+		}
+
+		OutRightContent = SNew(SBox)
+			.Visibility(EVisibility::SelfHitTestInvisible)
+			.Padding(FMargin(2.0f, 0.0f, 0.0f, 0.0f))
+			[
+				// Minimize
+				SNew(SHorizontalBox)
+					.Visibility(EVisibility::SelfHitTestInvisible)
+
+				+ SHorizontalBox::Slot()
+					.AutoWidth()
+					[
+						MinimizeButton.ToSharedRef()
+					]
+
+				// Maximize/Restore
+				+ SHorizontalBox::Slot()
+					.AutoWidth()
+					[
+						MaximizeRestoreButton.ToSharedRef()
+					]
+
+				// Close button
+				+ SHorizontalBox::Slot()
+					.AutoWidth()
+					[
+						CloseButton.ToSharedRef()
+					]
+			];
+
+#endif // PLATFORM_MAC
+	}
+
+	/**
+	 * Creates the title bar's content.
+	 *
+	 * @param CenterContent The content for the title bar's center area.
+	 * @param CenterContentAlignment The horizontal alignment of the center content.
+	 *
+	 * @return The content widget.
+	 */
+	TSharedRef<SWidget> MakeTitleBarContent( TSharedPtr<SWidget> CenterContent, EHorizontalAlignment CenterContentAlignment )
+	{
+		TSharedPtr<SWidget> LeftContent;
+		TSharedPtr<SWidget> RightContent;
+
+		MakeTitleBarContentWidgets(LeftContent, RightContent);
+
+		// create window title if no content was provided
+		if (!CenterContent.IsValid())
+		{
+			CenterContent = SNew(SBox)
+				.HAlign(HAlign_Center)
+				.Visibility(EVisibility::SelfHitTestInvisible)
+				.Padding(FMargin(5.0f, 2.0f, 2.0f, 5.0f))
+				[
+					// NOTE: We bind the window's title text to our window's GetTitle method, so that if the
+					//       title is changed later, the text will always be visually up to date
+					SNew(STextBlock)
+						.Visibility(EVisibility::SelfHitTestInvisible)
+						.TextStyle(&Style->TitleTextStyle)
+						.Text(Title)
+				];
+		}
+
+		// calculate content dimensions
+		LeftContent->SlatePrepass();
+		RightContent->SlatePrepass();
+
+		FVector2D LeftSize = LeftContent->GetDesiredSize();
+		FVector2D RightSize = RightContent->GetDesiredSize();
+
+		if (CenterContentAlignment == HAlign_Center)
+		{
+			LeftSize = FMath::Max(LeftSize, RightSize);
+			RightSize = LeftSize;
+		}
+
+		float SpacerHeight = FMath::Max(LeftSize.Y, RightSize.Y);
+
+		// create title bar
+		return SAssignNew(TitleArea, SBox)
+			.Visibility(EVisibility::SelfHitTestInvisible)
+			[
+				SNew(SOverlay)
+					.Visibility(EVisibility::SelfHitTestInvisible)
+
+				+ SOverlay::Slot()
+					[
+						SNew(SHorizontalBox)
+							.Visibility(EVisibility::SelfHitTestInvisible)
+						
+						+ SHorizontalBox::Slot()
+							.AutoWidth()
+							.HAlign(HAlign_Left)
+							.VAlign(VAlign_Top)
+							[
+								SNew(SSpacer)
+									.Size(FVector2D(LeftSize.X, SpacerHeight))
+							]
+
+						+ SHorizontalBox::Slot()
+							.HAlign(CenterContentAlignment)
+							.VAlign(VAlign_Top)
+							.FillWidth(1.0f)
+							[
+								CenterContent.ToSharedRef()
+							]
+
+						+ SHorizontalBox::Slot()
+							.AutoWidth()
+							.HAlign(HAlign_Right)
+							.VAlign(VAlign_Top)
+							[
+								SNew(SSpacer)
+									.Size(FVector2D(RightSize.X, SpacerHeight))
+							]
+					]
+
+				+ SOverlay::Slot()
+					[
+						SNew(SHorizontalBox)
+							.Visibility(EVisibility::SelfHitTestInvisible)
+
+						+ SHorizontalBox::Slot()
+							.AutoWidth()
+							.HAlign(HAlign_Left)
+							.VAlign(VAlign_Top)
+							[
+								LeftContent.ToSharedRef()
+							]
+
+						+ SHorizontalBox::Slot()
+							.FillWidth(1.0f)
+
+						+ SHorizontalBox::Slot()
+							.AutoWidth()
+							.HAlign(HAlign_Right)
+							.VAlign(VAlign_Top)
+							[
+								RightContent.ToSharedRef()
+							]
+					]
+			];
+	}
+
+private:
+
+	// Callback for clicking the close button.
+	FReply CloseButton_OnClicked()
+	{
+		TSharedPtr<SWindow> OwnerWindow = OwnerWindowPtr.Pin();
+
+		if (OwnerWindow.IsValid())
+		{
+			OwnerWindow->RequestDestroyWindow();
+		}
+
+		return FReply::Handled();
+	}
+
+	// Callback for getting the image of the close button.
+	const FSlateBrush* GetCloseImage() const
+	{
+		TSharedPtr<SWindow> OwnerWindow = OwnerWindowPtr.Pin();
+
+		if (!OwnerWindow.IsValid())
+		{
+			return nullptr;
+		}
+
+		TSharedPtr<FGenericWindow> NativeWindow = OwnerWindow->GetNativeWindow();
+
+		if (CloseButton->IsPressed())
+		{
+			return &Style->CloseButtonStyle.Pressed;
+		}
+
+#if PLATFORM_MAC
+		if (this->MinimizeButton->IsHovered() || MaximizeRestoreButton->IsHovered() || CloseButton->IsHovered())
+#else
+		if (CloseButton->IsHovered())
+#endif
+		{
+			return &Style->CloseButtonStyle.Hovered;
+		}
+
+#if PLATFORM_MAC
+		if (NativeWindow.IsValid() && !NativeWindow->IsForegroundWindow())
+		{
+			return &Style->CloseButtonStyle.Disabled;
+		}
+#endif
+		return &Style->CloseButtonStyle.Normal;
+	}
+
+	// Callback for clicking the maximize button
+	FReply MaximizeRestoreButton_OnClicked( )
+	{
+		TSharedPtr<SWindow> OwnerWindow = OwnerWindowPtr.Pin();
+
+		if (OwnerWindow.IsValid())
+		{
+			TSharedPtr<FGenericWindow> NativeWindow = OwnerWindow->GetNativeWindow();
+
+			if (NativeWindow.IsValid())
+			{
+				if (NativeWindow->IsMaximized())
+				{
+					NativeWindow->Restore();
+				}
+				else
+				{
+					NativeWindow->Maximize();
+				}
+			}
+		}
+
+		return FReply::Handled();
+	}
+
+	// Callback for getting the image of the maximize/restore button.
+	const FSlateBrush* GetMaximizeRestoreImage( ) const
+	{
+		TSharedPtr<SWindow> OwnerWindow = OwnerWindowPtr.Pin();
+
+		if (!OwnerWindow.IsValid())
+		{
+			return nullptr;
+		}
+
+		TSharedPtr<FGenericWindow> NativeWindow = OwnerWindow->GetNativeWindow();
+
+		if (NativeWindow.IsValid() && NativeWindow->IsMaximized())
+		{
+			if (!OwnerWindow->HasMaximizeBox())
+			{
+				return &Style->MaximizeButtonStyle.Disabled;
+			}
+			else if (MaximizeRestoreButton->IsPressed())
+			{
+				return &Style->RestoreButtonStyle.Pressed;
+			}
+#if PLATFORM_MAC
+			else if (this->MinimizeButton->IsHovered() || MaximizeRestoreButton->IsHovered() || CloseButton->IsHovered())
+#else
+			else if (MaximizeRestoreButton->IsHovered())
+#endif
+			{
+				return &Style->RestoreButtonStyle.Hovered;
+			}
+#if PLATFORM_MAC
+			else if (NativeWindow.IsValid() && !NativeWindow->IsForegroundWindow())
+			{
+				return &Style->RestoreButtonStyle.Disabled;
+			}
+#endif
+			else
+			{
+				return &Style->RestoreButtonStyle.Normal;
+			}
+		}
+		else
+		{
+			if (!OwnerWindow->HasMaximizeBox())
+			{
+				return &Style->MaximizeButtonStyle.Disabled;
+			}
+			else if (MaximizeRestoreButton->IsPressed())
+			{
+				return &Style->MaximizeButtonStyle.Pressed;
+			}
+#if PLATFORM_MAC
+			else if (this->MinimizeButton->IsHovered() || MaximizeRestoreButton->IsHovered() || CloseButton->IsHovered())
+#else
+			else if (MaximizeRestoreButton->IsHovered())
+#endif
+			{
+				return &Style->MaximizeButtonStyle.Hovered;
+			}
+#if PLATFORM_MAC
+			else if (NativeWindow.IsValid() && !NativeWindow->IsForegroundWindow())
+			{
+				return &Style->MaximizeButtonStyle.Disabled;
+			}
+#endif
+			else
+			{
+				return &Style->MaximizeButtonStyle.Normal;
+			}
+		}
+	}
+
+	// Callback for clicking the minimize button.
+	FReply MinimizeButton_OnClicked( )
+	{
+		TSharedPtr<SWindow> OwnerWindow = OwnerWindowPtr.Pin();
+
+		if (OwnerWindow.IsValid())
+		{
+			TSharedPtr<FGenericWindow> NativeWindow = OwnerWindow->GetNativeWindow();
+
+			if (NativeWindow.IsValid())
+			{
+				NativeWindow->Minimize();
+			}
+		}
+
+		return FReply::Handled();
+	}
+
+	// Callback for getting the image of the minimize button.
+	const FSlateBrush* GetMinimizeImage( ) const
+	{
+		TSharedPtr<SWindow> OwnerWindow = OwnerWindowPtr.Pin();
+
+		if (!OwnerWindow.IsValid())
+		{
+			return nullptr;
+		}
+
+		TSharedPtr<FGenericWindow> NativeWindow = OwnerWindow->GetNativeWindow();
+
+		if (!OwnerWindow->HasMinimizeBox())
+		{
+			return &Style->MinimizeButtonStyle.Disabled;
+		}
+		else if (MinimizeButton->IsPressed())
+		{
+			return &Style->MinimizeButtonStyle.Pressed;
+		}
+	#if PLATFORM_MAC
+		else if (MinimizeButton->IsHovered() || MaximizeRestoreButton->IsHovered() || CloseButton->IsHovered())
+	#else
+		else if (MinimizeButton->IsHovered())
+	#endif
+		{
+			return &Style->MinimizeButtonStyle.Hovered;
+		}
+	#if PLATFORM_MAC
+		else if (NativeWindow.IsValid() && !NativeWindow->IsForegroundWindow())
+		{
+			return &Style->MinimizeButtonStyle.Disabled;
+		}
+	#endif
+		else
+		{
+			return &Style->MinimizeButtonStyle.Normal;
+		}
+	}
+	
+	/** @return An appropriate resource for the window title background depending on whether the window is active */
+	const FSlateBrush* GetWindowTitlebackgroundImage( ) const
+	{
+		TSharedPtr<SWindow> OwnerWindow = OwnerWindowPtr.Pin();
+
+		if (!OwnerWindow.IsValid())
+		{
+			return nullptr;
+		}
+
+		TSharedPtr<FGenericWindow> NativeWindow = OwnerWindow->GetNativeWindow();
+		const bool bIsActive = NativeWindow.IsValid() && NativeWindow->IsForegroundWindow();
+
+		return bIsActive ? &Style->ActiveTitleBrush : &Style->InactiveTitleBrush;
+	}
+	
+	EVisibility GetWindowFlashVisibility( ) const
+	{
+		return TitleFlashSequence.IsPlaying() ? EVisibility::SelfHitTestInvisible : EVisibility::Hidden;
+	}
+	
+	FSlateColor GetWindowTitleAreaColor( ) const
+	{	
+		// Color of the white flash in the title area
+		float Flash = GetFlashValue();
+		float Alpha = Flash * 0.4f;
+
+		FLinearColor Color = FLinearColor::White;
+		Color.A = Alpha;
+
+		return Color;
+	}
+	
+	FSlateColor GetWindowTitleContentColor( ) const
+	{	
+		// Color of the title area contents - modulates the icon and buttons
+		float Flash = GetFlashValue();
+
+		return FMath::Lerp(FLinearColor::White, FLinearColor::Black, Flash);
+	}
+
+	FText HandleWindowTitleText( ) const
+	{
+		TSharedPtr<SWindow> OwnerWindow = OwnerWindowPtr.Pin();
+
+		if (!OwnerWindow.IsValid())
+		{
+			return FText::GetEmpty();
+		}
+
+		return OwnerWindow->GetTitle();
+	}
+
+private:
+
+	// Holds a weak pointer to the owner window.
+	TWeakPtr<SWindow> OwnerWindowPtr;
+
+	// Holds the window style to use (for buttons, text, etc.).
+	const FWindowStyle* Style;
+
+	// Holds the content widget of the title area.
+	TSharedPtr<SWidget> TitleArea;
+
+	// Holds the curve sequence for the window flash animation.
+	FCurveSequence TitleFlashSequence;
+
+	// Holds the minimize button.
+	TSharedPtr<SButton> MinimizeButton;
+
+	// Holds the maximize/restore button.
+	TSharedPtr<SButton> MaximizeRestoreButton;
+
+	// Holds the close button.
+	TSharedPtr<SButton> CloseButton;
+
+	bool ShowAppIcon;
+
+	TAttribute<FText> Title;
+};
