Refactor to unify keyboard and controller input and focus.

diff --git a/Engine/Source/Runtime/Core/Public/Math/IntPoint.h b/Engine/Source/Runtime/Core/Public/Math/IntPoint.h
--- a/Engine/Source/Runtime/Core/Public/Math/IntPoint.h
+++ b/Engine/Source/Runtime/Core/Public/Math/IntPoint.h
@@ -173,6 +173,22 @@ public:
 	 */
 	FIntPoint operator/( const FIntPoint& Other ) const;
 
+	/**
+	* Gets specific component of the point.
+	*
+	* @param Index the index of point component
+	* @return reference to component.
+	*/
+	int32& operator[](int32 Index);
+
+	/**
+	* Gets specific component of the point.
+	*
+	* @param Index the index of point component
+	* @return copy of component value.
+	*/
+	int32 operator[](int32 Index) const;
+
 public:
 
 	/**
@@ -371,6 +387,17 @@ FORCEINLINE FIntPoint FIntPoint::operator/( int32 Divisor ) const
 	return FIntPoint(*this) /= Divisor;
 }
 
+FORCEINLINE int32& FIntPoint::operator[](int32 Index)
+{
+	check(Index >= 0 && Index < 2);
+	return ((Index == 0) ? X : Y);
+}
+
+FORCEINLINE int32 FIntPoint::operator[](int32 Index) const
+{
+	check(Index >= 0 && Index < 2);
+	return ((Index == 0) ? X : Y);
+}
 
 FORCEINLINE FIntPoint FIntPoint::DivideAndRoundUp( FIntPoint lhs, int32 Divisor )
 {
diff --git a/Engine/Source/Runtime/Engine/Classes/Kismet/KismetInputLibrary.h b/Engine/Source/Runtime/Engine/Classes/Kismet/KismetInputLibrary.h
--- a/Engine/Source/Runtime/Engine/Classes/Kismet/KismetInputLibrary.h
+++ b/Engine/Source/Runtime/Engine/Classes/Kismet/KismetInputLibrary.h
@@ -3,7 +3,7 @@
 #pragma once
 
 #include "SlateCore.h"
-
+typedef FKeyEvent FKeyEvent;
 #include "KismetInputLibrary.generated.h"
 
 UCLASS(MinimalAPI)
@@ -134,9 +134,14 @@ class UKismetInputLibrary : public UBlueprintFunctionLibrary
 	 *
 	 * @return  Key name
 	 */
-	UFUNCTION(BlueprintCallable, Category="Utilities|KeyboardEvent")
-	static FKey GetKey(const FKeyboardEvent& Input);
+	UFUNCTION(BlueprintCallable, Category="Utilities|KeyEvent")
+	static FKey GetKey(const FKeyEvent& Input);
+
+	UFUNCTION(BlueprintCallable, Category = "Utilities|KeyEvent")
+	static int32 GetUserIndex(const FKeyEvent& Input);
 
+	UFUNCTION(BlueprintCallable, Category = "Utilities|FAnalogInputEvent")
+	static float GetAnalogValue(const FAnalogInputEvent& Input);
 
 	/** @return The position of the cursor in screen space */
 	UFUNCTION(BlueprintPure, meta=( FriendlyName = "GetScreenSpacePosition" ), Category="Utilities|PointerEvent")
@@ -190,14 +195,14 @@ class UKismetInputLibrary : public UBlueprintFunctionLibrary
 
 
 	/** @return The controller button that caused this event */
-	UFUNCTION(BlueprintPure, meta=( FriendlyName = "GetEffectingButton" ), Category="Utilities|ControllerEvent")
+	UFUNCTION(BlueprintPure, meta = (DeprecatedFunction, DeprecationMessage = "Use GetKey() for KeyEvent instead", FriendlyName = "GetEffectingButton"), Category = "Utilities|ControllerEvent")
 	static FKey ControllerEvent_GetEffectingButton(const FControllerEvent& Input);
 
 	/** @return The index of the user that caused the event */
-	UFUNCTION(BlueprintPure, meta=( FriendlyName = "GetUserIndex" ), Category="Utilities|ControllerEvent")
+	UFUNCTION(BlueprintPure, meta = (DeprecatedFunction, DeprecationMessage = "Use GetUserIndex() for KeyEvent instead"), Category = "Utilities|ControllerEvent")
 	static int32 ControllerEvent_GetUserIndex(const FControllerEvent& Input);
 
 	/** @return Analog value between 0 and 1.  1 being fully pressed, 0 being not pressed at all */
-	UFUNCTION(BlueprintPure, meta=( FriendlyName = "GetAnalogValue" ), Category="Utilities|ControllerEvent")
+	UFUNCTION(BlueprintPure, meta = (DeprecatedFunction, DeprecationMessage = "Use GetAnalogValue() for AnalogInputEvent instead"), Category = "Utilities|ControllerEvent")
 	static float ControllerEvent_GetAnalogValue(const FControllerEvent& Input);
 };
diff --git a/Engine/Source/Runtime/Engine/Private/GameViewportClient.cpp b/Engine/Source/Runtime/Engine/Private/GameViewportClient.cpp
--- a/Engine/Source/Runtime/Engine/Private/GameViewportClient.cpp
+++ b/Engine/Source/Runtime/Engine/Private/GameViewportClient.cpp
@@ -1243,8 +1243,6 @@ void UGameViewportClient::LostFocus(FViewport* InViewport)
 
 void UGameViewportClient::ReceivedFocus(FViewport* InViewport)
 {
-	InViewport->CaptureJoystickInput(true);
-
 	if (GetDefault<UInputSettings>()->bUseMouseForTouch && !GetGameViewport()->GetPlayInEditorIsSimulate())
 	{
 		FSlateApplication::Get().SetGameIsFakingTouchEvents(true);
@@ -2400,7 +2398,7 @@ bool UGameViewportClient::HandlePrevViewModeCommand( const TCHAR* Cmd, FOutputDe
 #if WITH_EDITOR
 bool UGameViewportClient::HandleShowMouseCursorCommand( const TCHAR* Cmd, FOutputDevice& Ar )
 {
-	FSlateApplication::Get().ClearKeyboardFocus( EKeyboardFocusCause::SetDirectly );
+	FSlateApplication::Get().ClearKeyboardFocus( EFocusCause::SetDirectly );
 	FSlateApplication::Get().ResetToDefaultInputSettings();
 	return true;
 }
diff --git a/Engine/Source/Runtime/Engine/Private/KismetInputLibrary.cpp b/Engine/Source/Runtime/Engine/Private/KismetInputLibrary.cpp
--- a/Engine/Source/Runtime/Engine/Private/KismetInputLibrary.cpp
+++ b/Engine/Source/Runtime/Engine/Private/KismetInputLibrary.cpp
@@ -94,11 +94,21 @@ bool UKismetInputLibrary::InputEvent_IsRightCommandDown(const FInputEvent& Input
 }
 
 
-FKey UKismetInputLibrary::GetKey(const FKeyboardEvent& Input)
+FKey UKismetInputLibrary::GetKey(const FKeyEvent& Input)
 {
 	return Input.GetKey();
 }
 
+int32 UKismetInputLibrary::GetUserIndex(const FKeyEvent& Input)
+{
+	return Input.GetUserIndex();
+}
+
+float UKismetInputLibrary::GetAnalogValue(const FAnalogInputEvent& Input)
+{
+	return Input.GetAnalogValue();
+}
+
 
 FVector2D UKismetInputLibrary::PointerEvent_GetScreenSpacePosition(const FPointerEvent& Input)
 {
diff --git a/Engine/Source/Runtime/Engine/Private/LocalPlayer.cpp b/Engine/Source/Runtime/Engine/Private/LocalPlayer.cpp
--- a/Engine/Source/Runtime/Engine/Private/LocalPlayer.cpp
+++ b/Engine/Source/Runtime/Engine/Private/LocalPlayer.cpp
@@ -1148,7 +1148,7 @@ bool ULocalPlayer::HandlePauseCommand( const TCHAR* Cmd, FOutputDevice& Ar, UWor
 	{
 		if (ViewportClient && ViewportClient->Viewport)
 		{
-			ViewportClient->Viewport->CaptureJoystickInput(true);
+			ViewportClient->Viewport->SetUserFocus(true);
 			ViewportClient->Viewport->CaptureMouse(true);
 		}
 	}
diff --git a/Engine/Source/Runtime/Engine/Private/PlayerController.cpp b/Engine/Source/Runtime/Engine/Private/PlayerController.cpp
--- a/Engine/Source/Runtime/Engine/Private/PlayerController.cpp
+++ b/Engine/Source/Runtime/Engine/Private/PlayerController.cpp
@@ -4270,11 +4270,7 @@ void FInputModeDataBase::SetFocusAndLocking(FReply& SlateOperations, TSharedPtr<
 {
 	if (InWidgetToFocus.IsValid())
 	{
-		SlateOperations.CaptureJoystick(InWidgetToFocus.ToSharedRef());
-	}
-	else
-	{
-		SlateOperations.ReleaseJoystickCapture();
+		SlateOperations.SetUserFocus(InWidgetToFocus.ToSharedRef());
 	}
 
 	if (bLockMouseToViewport)
@@ -4311,6 +4307,7 @@ void FInputModeGameAndUI::ApplyInputMode(FReply& SlateOperations, class UGameVie
 		SlateOperations.ReleaseMouseCapture();
 
 		GameViewportClient.SetIgnoreInput(false);
+		GameViewportClient.SetHideCursorDuringCapture(bHideCursorDuringCapture);
 		GameViewportClient.SetCaptureMouseOnClick(EMouseCaptureMode::CaptureDuringMouseDown);
 	}
 }
@@ -4322,9 +4319,8 @@ void FInputModeGameOnly::ApplyInputMode(FReply& SlateOperations, class UGameView
 	{
 		TSharedRef<SViewport> ViewportWidgetRef = ViewportWidget.ToSharedRef();
 		SlateOperations.UseHighPrecisionMouseMovement(ViewportWidgetRef);
-		SlateOperations.CaptureJoystick(ViewportWidgetRef);
+		SlateOperations.SetUserFocus(ViewportWidgetRef);
 		SlateOperations.LockMouseToWidget(ViewportWidgetRef);
-		SlateOperations.SetKeyboardFocus(ViewportWidgetRef, EKeyboardFocusCause::SetDirectly);
 		GameViewportClient.SetIgnoreInput(false);
 		GameViewportClient.SetCaptureMouseOnClick(EMouseCaptureMode::CapturePermanently);
 	}
diff --git a/Engine/Source/Runtime/Engine/Private/Slate/SceneViewport.cpp b/Engine/Source/Runtime/Engine/Private/Slate/SceneViewport.cpp
--- a/Engine/Source/Runtime/Engine/Private/Slate/SceneViewport.cpp
+++ b/Engine/Source/Runtime/Engine/Private/Slate/SceneViewport.cpp
@@ -105,18 +105,18 @@ void FSceneViewport::ShowCursor( bool bVisible )
 	}
 }
 
-bool FSceneViewport::CaptureJoystickInput(bool Capture)
+bool FSceneViewport::SetUserFocus(bool bFocus)
 {
-	if( Capture )
+	if (bFocus)
 	{
-		CurrentReplyState.CaptureJoystick( ViewportWidget.Pin().ToSharedRef(), true );
+		CurrentReplyState.SetUserFocus(ViewportWidget.Pin().ToSharedRef(), EFocusCause::SetDirectly, true);
 	}
 	else
 	{
-		CurrentReplyState.ReleaseJoystickCapture(true);
+		CurrentReplyState.ClearUserFocus(true);
 	}
 
-	return Capture;
+	return bFocus;
 }
 
 bool FSceneViewport::KeyState( FKey Key ) const
@@ -371,8 +371,8 @@ FReply FSceneViewport::OnMouseButtonDown( const FGeometry& InGeometry, const FPo
 		{
 			TSharedRef<SViewport> ViewportWidgetRef = ViewportWidget.Pin().ToSharedRef();
 
-			// Mouse down should focus viewport for keyboard input
-			CurrentReplyState.SetKeyboardFocus(ViewportWidgetRef, EKeyboardFocusCause::Mouse);
+			// Mouse down should focus viewport for user input
+			CurrentReplyState.SetUserFocus(ViewportWidgetRef, EFocusCause::SetDirectly, true);
 			
 			UWorld* World = ViewportClient->GetWorld();
 			if (World && World->IsGameWorld() && World->GetFirstPlayerController())
@@ -551,67 +551,6 @@ FReply FSceneViewport::OnMouseButtonDoubleClick( const FGeometry& InGeometry, co
 	return CurrentReplyState;
 }
 
-FReply FSceneViewport::OnControllerButtonPressed( const FGeometry& MyGeometry, const FControllerEvent& ControllerEvent )
-{
-	// Start a new reply state
-	CurrentReplyState = FReply::Handled(); 
-
-	KeyStateMap.Add( ControllerEvent.GetEffectingButton(), true );
-
-	if( ViewportClient )
-	{
-		// Switch to the viewport clients world before processing input
-		FScopedConditionalWorldSwitcher WorldSwitcher( ViewportClient );
-
-		if( !ViewportClient->InputKey( this, ControllerEvent.GetUserIndex(), ControllerEvent.GetEffectingButton(), ControllerEvent.IsRepeat() ? IE_Repeat : IE_Pressed, 1.0f, true ) )
-		{
-			CurrentReplyState = FReply::Unhandled(); 
-		}
-	}
-	return CurrentReplyState;
-}
-
-FReply FSceneViewport::OnControllerButtonReleased( const FGeometry& MyGeometry, const FControllerEvent& ControllerEvent )
-{
-	// Start a new reply state
-	CurrentReplyState = FReply::Handled(); 
-
-	KeyStateMap.Add( ControllerEvent.GetEffectingButton(), false );
-
-	if( ViewportClient )
-	{
-		// Switch to the viewport clients world before processing input
-		FScopedConditionalWorldSwitcher WorldSwitcher( ViewportClient );
-
-		if( !ViewportClient->InputKey( this, ControllerEvent.GetUserIndex(), ControllerEvent.GetEffectingButton(), IE_Released, 1.0f, true ) )
-		{
-			CurrentReplyState = FReply::Unhandled(); 
-		}
-	}
-	return CurrentReplyState;
-}
-
-FReply FSceneViewport::OnControllerAnalogValueChanged( const FGeometry& MyGeometry, const FControllerEvent& ControllerEvent )
-{
-	// Start a new reply state
-	CurrentReplyState = FReply::Handled(); 
-
-	KeyStateMap.Add( ControllerEvent.GetEffectingButton(), true );
-
-	if( ViewportClient )
-	{
-		// Switch to the viewport clients world before processing input
-		FScopedConditionalWorldSwitcher WorldSwitcher( ViewportClient );
-
-		if (!ViewportClient->InputAxis(this, ControllerEvent.GetUserIndex(), ControllerEvent.GetEffectingButton(), ControllerEvent.GetEffectingButton() == EKeys::Gamepad_RightY ? -ControllerEvent.GetAnalogValue() : ControllerEvent.GetAnalogValue(), FApp::GetDeltaTime(), 1, true))
-		{
-			CurrentReplyState = FReply::Unhandled(); 
-		}
-	}
-
-	return CurrentReplyState;
-}
-
 FReply FSceneViewport::OnTouchStarted( const FGeometry& MyGeometry, const FPointerEvent& TouchEvent )
 {
 	// Start a new reply state
@@ -732,12 +671,12 @@ void FSceneViewport::OnFinishedPointerInput()
 	ProcessAccumulatedPointerInput();
 }
 
-FReply FSceneViewport::OnKeyDown( const FGeometry& InGeometry, const FKeyboardEvent& InKeyboardEvent )
+FReply FSceneViewport::OnKeyDown( const FGeometry& InGeometry, const FKeyEvent& InKeyEvent )
 {
 	// Start a new reply state
 	CurrentReplyState = FReply::Handled(); 
 
-	FKey Key = InKeyboardEvent.GetKey();
+	FKey Key = InKeyEvent.GetKey();
 	KeyStateMap.Add( Key, true );
 
 	//@todo Slate Viewports: FWindowsViewport checks for Alt+Enter or F11 and toggles fullscreen.  Unknown if fullscreen via this method will be needed for slate viewports. 
@@ -746,7 +685,7 @@ FReply FSceneViewport::OnKeyDown( const FGeometry& InGeometry, const FKeyboardEv
 		// Switch to the viewport clients world before processing input
 		FScopedConditionalWorldSwitcher WorldSwitcher( ViewportClient );
 
-		if( !ViewportClient->InputKey( this, 0, Key, InKeyboardEvent.IsRepeat() ? IE_Repeat : IE_Pressed ) )
+		if( !ViewportClient->InputKey( this, 0, Key, InKeyEvent.IsRepeat() ? IE_Repeat : IE_Pressed ) )
 		{
 			CurrentReplyState = FReply::Unhandled();
 		}
@@ -754,12 +693,12 @@ FReply FSceneViewport::OnKeyDown( const FGeometry& InGeometry, const FKeyboardEv
 	return CurrentReplyState;
 }
 
-FReply FSceneViewport::OnKeyUp( const FGeometry& InGeometry, const FKeyboardEvent& InKeyboardEvent )
+FReply FSceneViewport::OnKeyUp( const FGeometry& InGeometry, const FKeyEvent& InKeyEvent )
 {
 	// Start a new reply state
 	CurrentReplyState = FReply::Handled(); 
 
-	FKey Key = InKeyboardEvent.GetKey();
+	FKey Key = InKeyEvent.GetKey();
 	KeyStateMap.Add( Key, false );
 	
 	if( ViewportClient && GetSizeXY() != FIntPoint::ZeroValue  )
@@ -776,6 +715,27 @@ FReply FSceneViewport::OnKeyUp( const FGeometry& InGeometry, const FKeyboardEven
 	return CurrentReplyState;
 }
 
+FReply FSceneViewport::OnAnalogValueChanged(const FGeometry& MyGeometry, const FAnalogInputEvent& InAnalogInputEvent)
+{
+	// Start a new reply state
+	CurrentReplyState = FReply::Handled();
+
+	KeyStateMap.Add(InAnalogInputEvent.GetKey(), true);
+
+	if (ViewportClient)
+	{
+		// Switch to the viewport clients world before processing input
+		FScopedConditionalWorldSwitcher WorldSwitcher(ViewportClient);
+
+		if (!ViewportClient->InputAxis(this, InAnalogInputEvent.GetUserIndex(), InAnalogInputEvent.GetKey(), InAnalogInputEvent.GetKey() == EKeys::Gamepad_RightY ? -InAnalogInputEvent.GetAnalogValue() : InAnalogInputEvent.GetAnalogValue(), FApp::GetDeltaTime(), 1, true))
+		{
+			CurrentReplyState = FReply::Unhandled();
+		}
+	}
+
+	return CurrentReplyState;
+}
+
 FReply FSceneViewport::OnKeyChar( const FGeometry& InGeometry, const FCharacterEvent& InCharacterEvent )
 {
 	// Start a new reply state
@@ -794,23 +754,22 @@ FReply FSceneViewport::OnKeyChar( const FGeometry& InGeometry, const FCharacterE
 	return CurrentReplyState;
 }
 
-FReply FSceneViewport::OnKeyboardFocusReceived( const FKeyboardFocusEvent& InKeyboardFocusEvent )
+FReply FSceneViewport::OnFocusReceived(const FFocusEvent& InFocusEvent)
 {
 	CurrentReplyState = FReply::Handled(); 
 
-	if( ViewportClient )
+	if (ViewportClient != nullptr)
 	{
-		FScopedConditionalWorldSwitcher WorldSwitcher( ViewportClient );
-		ViewportClient->ReceivedFocus( this );
+		FScopedConditionalWorldSwitcher WorldSwitcher(ViewportClient);
+		ViewportClient->ReceivedFocus(this);
 
-		if( ( FApp::IsGame() && !GIsEditor ) || bIsPlayInEditorViewport )
+		if ((FApp::IsGame() && !GIsEditor) || bIsPlayInEditorViewport)
 		{
-			if( IsForegroundWindow() )
+			if (IsForegroundWindow())
 			{
-				const bool bIsCursorForcedVisible =
-					( ViewportClient != NULL && ViewportClient->GetCursor( this, GetMouseX(), GetMouseY() ) != EMouseCursor::None );
+				const bool bIsCursorForcedVisible = ViewportClient->GetCursor(this, GetMouseX(), GetMouseY()) != EMouseCursor::None;
 
-				const bool bPlayInEditorCapture = !bIsPlayInEditorViewport || InKeyboardFocusEvent.GetCause() != EKeyboardFocusCause::SetDirectly || bPlayInEditorGetsMouseControl;
+				const bool bPlayInEditorCapture = !bIsPlayInEditorViewport || InFocusEvent.GetCause() != EFocusCause::SetDirectly || bPlayInEditorGetsMouseControl;
 
 				// capturing the mouse interferes with slate UI (like the virtual joysticks)
 				if (FPlatformProperties::SupportsWindowedMode() && bPlayInEditorCapture && !bIsCursorForcedVisible && !FSlateApplication::Get().IsFakingTouchEvents())
@@ -818,18 +777,13 @@ FReply FSceneViewport::OnKeyboardFocusReceived( const FKeyboardFocusEvent& InKey
 					// Only require the user to click in the window the first time - after that return focus to the game so long as it was the last focused widget.
 					// Means that tabbing in/out will return the mouse control to where it was & the in-game console won't leave the mouse under editor control.
 					bPlayInEditorGetsMouseControl = true;
-					CurrentReplyState.UseHighPrecisionMouseMovement( ViewportWidget.Pin().ToSharedRef() );
-					CurrentReplyState.LockMouseToWidget( ViewportWidget.Pin().ToSharedRef() );
-				}
-				else if(!bPlayInEditorCapture)
-				{
-					FSlateApplication::Get().ClearKeyboardFocus( EKeyboardFocusCause::SetDirectly );
-					FSlateApplication::Get().ResetToDefaultInputSettings();
+					CurrentReplyState.UseHighPrecisionMouseMovement(ViewportWidget.Pin().ToSharedRef());
+					CurrentReplyState.LockMouseToWidget(ViewportWidget.Pin().ToSharedRef());
 				}
 			}
 			else
 			{
-				FSlateApplication::Get().ClearKeyboardFocus( EKeyboardFocusCause::Cleared );
+				CurrentReplyState.ClearUserFocus(true);
 			}
 		}
 	}
@@ -837,10 +791,10 @@ FReply FSceneViewport::OnKeyboardFocusReceived( const FKeyboardFocusEvent& InKey
 	return CurrentReplyState;
 }
 
-void FSceneViewport::OnKeyboardFocusLost( const FKeyboardFocusEvent& InKeyboardFocusEvent )
+void FSceneViewport::OnFocusLost( const FFocusEvent& InFocusEvent )
 {
 	KeyStateMap.Empty();
-	if( ViewportClient )
+	if (ViewportClient != nullptr)
 	{
 		FScopedConditionalWorldSwitcher WorldSwitcher( ViewportClient );
 		ViewportClient->LostFocus( this );
@@ -848,11 +802,11 @@ void FSceneViewport::OnKeyboardFocusLost( const FKeyboardFocusEvent& InKeyboardF
 		TSharedPtr<SWidget> ViewportWidgetPin = ViewportWidget.Pin();
 		if( ViewportWidgetPin.IsValid() )
 		{
-			for( int32 UserIndex = 0; UserIndex < SlateApplicationDefs::MaxUsers; ++UserIndex )
+			for (int32 UserIndex = 0; UserIndex < SlateApplicationDefs::MaxUsers; ++UserIndex)
 			{
-				if( FSlateApplication::Get().GetJoystickCaptor(UserIndex) == ViewportWidgetPin )
+				if (FSlateApplication::Get().GetUserFocusedWidget(UserIndex) == ViewportWidgetPin)
 				{
-					FSlateApplication::Get().ReleaseJoystickCapture( UserIndex );
+					FSlateApplication::Get().ClearUserFocus(UserIndex);
 				}
 			}
 		}
diff --git a/Engine/Source/Runtime/Engine/Private/UserInterface/Console.cpp b/Engine/Source/Runtime/Engine/Private/UserInterface/Console.cpp
--- a/Engine/Source/Runtime/Engine/Private/UserInterface/Console.cpp
+++ b/Engine/Source/Runtime/Engine/Private/UserInterface/Console.cpp
@@ -1247,7 +1247,7 @@ void UConsole::FakeGotoState(FName NextStateName)
 		// Since the viewport may not be the current focus, we need to re-focus whatever the current focus is,
 		// in order to ensure it gets a chance to reapply any custom input settings
 		auto CurrentFocus = FSlateApplication::Get().GetKeyboardFocusedWidget();
-		FSlateApplication::Get().ClearKeyboardFocus( EKeyboardFocusCause::SetDirectly );
+		FSlateApplication::Get().ClearKeyboardFocus( EFocusCause::SetDirectly );
 		FSlateApplication::Get().SetKeyboardFocus( CurrentFocus );
 	}
 
diff --git a/Engine/Source/Runtime/Engine/Public/Slate/SceneViewport.h b/Engine/Source/Runtime/Engine/Public/Slate/SceneViewport.h
--- a/Engine/Source/Runtime/Engine/Public/Slate/SceneViewport.h
+++ b/Engine/Source/Runtime/Engine/Public/Slate/SceneViewport.h
@@ -44,7 +44,7 @@ public:
 	 *
 	 * @param Capture	true if we should capture, false if we should uncapture
 	 */
-	virtual bool CaptureJoystickInput(bool Capture) override;
+	virtual bool SetUserFocus(bool bFocus) override;
 
 	/**
 	 * Returns the state of the provided key. 
@@ -195,20 +195,18 @@ public:
 	virtual FReply OnMouseMove( const FGeometry& InGeometry, const FPointerEvent& MouseEvent ) override;
 	virtual FReply OnMouseWheel( const FGeometry& InGeometry, const FPointerEvent& MouseEvent ) override;
 	virtual FReply OnMouseButtonDoubleClick( const FGeometry& InGeometry, const FPointerEvent& InMouseEvent ) override;
-	virtual FReply OnControllerButtonPressed( const FGeometry& MyGeometry, const FControllerEvent& ControllerEvent ) override;
-	virtual FReply OnControllerButtonReleased( const FGeometry& MyGeometry, const FControllerEvent& ControllerEvent ) override;
-	virtual FReply OnControllerAnalogValueChanged( const FGeometry& MyGeometry, const FControllerEvent& ControllerEvent ) override;
 	virtual FReply OnTouchStarted( const FGeometry& MyGeometry, const FPointerEvent& InTouchEvent ) override;
 	virtual FReply OnTouchMoved( const FGeometry& MyGeometry, const FPointerEvent& InTouchEvent ) override;
 	virtual FReply OnTouchEnded( const FGeometry& MyGeometry, const FPointerEvent& InTouchEvent ) override;
 	virtual FReply OnTouchGesture( const FGeometry& MyGeometry, const FPointerEvent& InGestureEvent ) override;
 	virtual FReply OnMotionDetected( const FGeometry& MyGeometry, const FMotionEvent& InMotionEvent ) override;
 	virtual void OnFinishedPointerInput() override;
-	virtual FReply OnKeyDown( const FGeometry& InGeometry, const FKeyboardEvent& InKeyboardEvent ) override;
-	virtual FReply OnKeyUp( const FGeometry& InGeometry, const FKeyboardEvent& InKeyboardEvent ) override;
+	virtual FReply OnKeyDown( const FGeometry& InGeometry, const FKeyEvent& InKeyEvent ) override;
+	virtual FReply OnKeyUp( const FGeometry& InGeometry, const FKeyEvent& InKeyEvent ) override;
+	virtual FReply OnAnalogValueChanged( const FGeometry& MyGeometry, const FAnalogInputEvent& InAnalogInputEvent ) override;
 	virtual FReply OnKeyChar( const FGeometry& InGeometry, const FCharacterEvent& InCharacterEvent ) override;
-	virtual FReply OnKeyboardFocusReceived( const FKeyboardFocusEvent& InKeyboardFocusEvent ) override;
-	virtual void OnKeyboardFocusLost( const FKeyboardFocusEvent& InKeyboardFocusEvent ) override;
+	virtual FReply OnFocusReceived( const FFocusEvent& InFocusEvent ) override;
+	virtual void OnFocusLost( const FFocusEvent& InFocusEvent ) override;
 	virtual void OnViewportClosed() override;
 	virtual FIntPoint GetSize() const override { return GetSizeXY(); }
 	
diff --git a/Engine/Source/Runtime/Engine/Public/UnrealClient.h b/Engine/Source/Runtime/Engine/Public/UnrealClient.h
--- a/Engine/Source/Runtime/Engine/Public/UnrealClient.h
+++ b/Engine/Source/Runtime/Engine/Public/UnrealClient.h
@@ -277,7 +277,7 @@ public:
 	 */
 	virtual bool IsCursorVisible() const { return true; }
 
-	virtual bool CaptureJoystickInput(bool Capture) = 0;
+	virtual bool SetUserFocus(bool bFocus) = 0;
 	virtual bool KeyState(FKey Key) const = 0;
 	virtual int32 GetMouseX() const = 0;
 	virtual int32 GetMouseY() const = 0;
@@ -937,7 +937,7 @@ public:
 	virtual void*	GetWindow() { return 0; }
 	virtual void	MoveWindow(int32 NewPosX, int32 NewPosY, int32 NewSizeX, int32 NewSizeY) {}
 	virtual void	Destroy() {}
-	virtual bool	CaptureJoystickInput(bool Capture) { return false; }
+	virtual bool SetUserFocus(bool bFocus) { return false; }
 	virtual bool	KeyState(FKey Key) const { return false; }
 	virtual int32	GetMouseX() const { return 0; }
 	virtual int32	GetMouseY() const { return 0; }
diff --git a/Engine/Source/Runtime/GameMenuBuilder/Private/SGameMenuPageWidget.cpp b/Engine/Source/Runtime/GameMenuBuilder/Private/SGameMenuPageWidget.cpp
--- a/Engine/Source/Runtime/GameMenuBuilder/Private/SGameMenuPageWidget.cpp
+++ b/Engine/Source/Runtime/GameMenuBuilder/Private/SGameMenuPageWidget.cpp
@@ -634,7 +634,7 @@ void SGameMenuPageWidget::Tick(const FGeometry& AllottedGeometry, const double I
 		if (!bConsoleVisible)
 		{
 			bConsoleVisible = true;
-			FSlateApplication::Get().SetFocusToGameViewport();
+			FSlateApplication::Get().SetAllUserFocusToGameViewport();
 		}
 	}
 	else
@@ -878,7 +878,7 @@ FReply SGameMenuPageWidget::SelectionChanged(int32 SelectionIndex)
 		ConfirmMenuItem();
 	}
 
-	return FReply::Handled().SetKeyboardFocus(SharedThis(this), EKeyboardFocusCause::SetDirectly);
+	return FReply::Handled().SetUserFocus(SharedThis(this), EFocusCause::SetDirectly);
 }
 
 void SGameMenuPageWidget::FadeIn()
@@ -901,7 +901,7 @@ FReply SGameMenuPageWidget::OnMouseButtonDown(const FGeometry& MyGeometry, const
 
 	//Set the keyboard focus 
 	return FReply::Handled()
-	.SetKeyboardFocus(SharedThis(this), EKeyboardFocusCause::SetDirectly);
+		.SetUserFocus(SharedThis(this), EFocusCause::SetDirectly);
 }
 
 void SGameMenuPageWidget::ChangeOption(int32 InMoveBy)
@@ -928,58 +928,18 @@ void SGameMenuPageWidget::ChangeOption(int32 InMoveBy)
 	}
 }
 
-FReply SGameMenuPageWidget::OnKeyDown(const FGeometry& MyGeometry, const FKeyboardEvent& InKeyboardEvent)
+FReply SGameMenuPageWidget::OnKeyDown(const FGeometry& MyGeometry, const FKeyEvent& InKeyEvent)
 {
 	FReply Result = FReply::Unhandled();
 	
-	if ( (CurrentMenu.IsValid() == true ) && ( !bControlsLocked) )
-	{
-		bool bNavigationLocked = bControlsLocked | PendingMainMenu.IsValid() | PendingSubMenu.IsValid();
-		const FKey Key = InKeyboardEvent.GetKey();
-	
-		if (bNavigationLocked == false )
-		{
-			if (Key == EKeys::Up || Key == EKeys::Down)
-			{
-				int32 MoveBy = Key == EKeys::Up ? -1 : 1;
-
-				if (SelectedIndex + MoveBy > -1 && SelectedIndex + MoveBy < CurrentMenu->NumItems())
-				{
-					SelectionChanged(SelectedIndex + MoveBy);
-				}
-				Result = FReply::Handled();
-			}
-			else if (Key == EKeys::Left || Key == EKeys::Right)
-			{
-				int32 MoveBy = Key == EKeys::Left ? -1 : 1;
-				ChangeOption(MoveBy);
-				Result = FReply::Handled();
-			}
-		}
-		if (Key == EKeys::Enter)
-		{
-			ConfirmMenuItem();
-			Result = FReply::Handled();
-		} 
-		else if (Key == EKeys::Escape )
-		{
-			MenuGoBack(true);
-			Result = FReply::Handled();
-		}
-	}
-	return Result;
-}
-
-FReply SGameMenuPageWidget::OnControllerButtonPressed( const FGeometry& MyGeometry, const FControllerEvent& ControllerEvent )
-{
-	FReply Result = FReply::Unhandled();
 	if ((CurrentMenu.IsValid() == true) && (!bControlsLocked))
 	{
 		bool bNavigationLocked = bControlsLocked | PendingMainMenu.IsValid() | PendingSubMenu.IsValid();
-		const FKey Key = ControllerEvent.GetEffectingButton();
+		const FKey Key = InKeyEvent.GetKey();
+
 		if (bNavigationLocked == false)
 		{
-			if (Key == EKeys::Gamepad_DPad_Up || Key == EKeys::Gamepad_LeftStick_Up)
+			if (Key == EKeys::Up || Key == EKeys::Gamepad_DPad_Up || Key == EKeys::Gamepad_LeftStick_Up)
 			{
 				if (SelectedIndex > 0)
 				{
@@ -987,7 +947,7 @@ FReply SGameMenuPageWidget::OnControllerButtonPressed( const FGeometry& MyGeomet
 				}
 				Result = FReply::Handled();
 			}
-			else if (Key == EKeys::Gamepad_DPad_Down || Key == EKeys::Gamepad_LeftStick_Down)
+			else if (Key == EKeys::Down || Key == EKeys::Gamepad_DPad_Down || Key == EKeys::Gamepad_LeftStick_Down)
 			{
 				if (SelectedIndex + 1 < CurrentMenu->NumItems())
 				{
@@ -995,38 +955,32 @@ FReply SGameMenuPageWidget::OnControllerButtonPressed( const FGeometry& MyGeomet
 				}
 				Result = FReply::Handled();
 			}
-			else if (Key == EKeys::Gamepad_DPad_Left || Key == EKeys::Gamepad_LeftStick_Left)
+			else if (Key == EKeys::Left || Key == EKeys::Gamepad_DPad_Left || Key == EKeys::Gamepad_LeftStick_Left)
 			{
 				ChangeOption(-1);
 				Result = FReply::Handled();
 			}
-			else if (Key == EKeys::Gamepad_DPad_Right || Key == EKeys::Gamepad_LeftStick_Right)
+			else if (Key == EKeys::Right || Key == EKeys::Gamepad_DPad_Right || Key == EKeys::Gamepad_LeftStick_Right)
 			{
 				ChangeOption(1);
 				Result = FReply::Handled();
 			}
 		}
-		if (Key == EKeys::Gamepad_FaceButton_Bottom)
+		if (Key == EKeys::Enter || Key == EKeys::Gamepad_FaceButton_Bottom)
 		{
 			ConfirmMenuItem();
 			Result = FReply::Handled();
 		} 
-		else if (Key == EKeys::Gamepad_FaceButton_Right || Key == EKeys::Gamepad_Special_Left)
+		else if (Key == EKeys::Escape || Key == EKeys::Gamepad_FaceButton_Right || Key == EKeys::Gamepad_Special_Left)
 		{
-			// Treat the right facepad button as CANCEL, but the menu button as simply back.
-			MenuGoBack(Key == EKeys::Gamepad_FaceButton_Right);
+			MenuGoBack(true);
 			Result = FReply::Handled();
-		} 
-// 		else if (Key == ControllerHideMenuKey)
-// 		{
-// 			OnToggleMenu.ExecuteIfBound();
-// 			Result = FReply::Handled();
-// 		}
+		}
 	}
 	return Result;
 }
 
-FReply SGameMenuPageWidget::OnKeyboardFocusReceived(const FGeometry& MyGeometry, const FKeyboardFocusEvent& InKeyboardFocusEvent)
+FReply SGameMenuPageWidget::OnFocusReceived(const FGeometry& MyGeometry, const FFocusEvent& InFocusEvent)
 {
 	if (CurrentMenu.IsValid() == false)
 	{
@@ -1036,10 +990,10 @@ FReply SGameMenuPageWidget::OnKeyboardFocusReceived(const FGeometry& MyGeometry,
 	//Focus the custom widget
 	if (CurrentMenu.IsValid() && CurrentMenu->NumItems() == 1 && CurrentMenu->GetItem(0)->MenuItemType == EGameMenuItemType::CustomWidget)
 	{
-		return FReply::Handled().ReleaseJoystickCapture().SetKeyboardFocus(CurrentMenu->GetItem(0)->CustomWidget.ToSharedRef(),EKeyboardFocusCause::SetDirectly);
+		return FReply::Handled().SetUserFocus(CurrentMenu->GetItem(0)->CustomWidget.ToSharedRef(),EFocusCause::SetDirectly);
 	}
 
-	return FReply::Handled().ReleaseMouseCapture().CaptureJoystick(SharedThis( this ));
+	return FReply::Handled().ReleaseMouseCapture().SetUserFocus(SharedThis(this));
 }
 
 void SGameMenuPageWidget::SetTitleAnimation(bool bShowTitle)
diff --git a/Engine/Source/Runtime/GameMenuBuilder/Public/SGameMenuPageWidget.h b/Engine/Source/Runtime/GameMenuBuilder/Public/SGameMenuPageWidget.h
--- a/Engine/Source/Runtime/GameMenuBuilder/Public/SGameMenuPageWidget.h
+++ b/Engine/Source/Runtime/GameMenuBuilder/Public/SGameMenuPageWidget.h
@@ -96,16 +96,13 @@ public:
 	virtual FReply OnMouseButtonDown(const FGeometry& MyGeometry, const FPointerEvent& MouseEvent) override;
 
 	/** Key down handler. */
-	virtual FReply OnKeyDown(const FGeometry& MyGeometry, const FKeyboardEvent& InKeyboardEvent) override;
-
-	/** Controller button handler. */
-	virtual FReply OnControllerButtonPressed(const FGeometry& MyGeometry, const FControllerEvent& ControllerEvent) override;
+	virtual FReply OnKeyDown(const FGeometry& MyGeometry, const FKeyEvent& InKeyEvent) override;
 
 	/** says that we can support keyboard focus. */
 	virtual bool SupportsKeyboardFocus() const override { return true; }
 
 	/** The menu sets up the appropriate mouse settings upon focus. */
-	virtual FReply OnKeyboardFocusReceived(const FGeometry& MyGeometry, const FKeyboardFocusEvent& InKeyboardFocusEvent) override;
+	virtual FReply OnFocusReceived(const FGeometry& MyGeometry, const FFocusEvent& InFocusEvent) override;
 
 	// End SWidget interface
 
diff --git a/Engine/Source/Runtime/Slate/Private/Framework/Application/MenuStack.cpp b/Engine/Source/Runtime/Slate/Private/Framework/Application/MenuStack.cpp
--- a/Engine/Source/Runtime/Slate/Private/Framework/Application/MenuStack.cpp
+++ b/Engine/Source/Runtime/Slate/Private/Framework/Application/MenuStack.cpp
@@ -44,11 +44,11 @@ namespace
 			return true;
 		}
 
-		virtual FReply OnKeyDown( const FGeometry& MyGeometry, const FKeyboardEvent& InKeyboardEvent ) override
+		virtual FReply OnKeyDown( const FGeometry& MyGeometry, const FKeyEvent& InKeyEvent ) override
 		{
 			if (OnKeyDownDelegate.IsBound())
 			{
-				return OnKeyDownDelegate.Execute(InKeyboardEvent.GetKey());
+				return OnKeyDownDelegate.Execute(InKeyEvent.GetKey());
 			}
 
 			return FReply::Unhandled();
diff --git a/Engine/Source/Runtime/Slate/Private/Framework/Application/SlateApplication.cpp b/Engine/Source/Runtime/Slate/Private/Framework/Application/SlateApplication.cpp
--- a/Engine/Source/Runtime/Slate/Private/Framework/Application/SlateApplication.cpp
+++ b/Engine/Source/Runtime/Slate/Private/Framework/Application/SlateApplication.cpp
@@ -180,9 +180,9 @@ public:
 	/**
 	 * Route an event along a focus path (as opposed to PointerPath)
 	 *
-	 * Focus paths are determined by the Keyboard Focus or Controller Focus.
+	 * Focus paths are used focus devices.(e.g. Keyboard or Game Pads)
 	 * Focus paths change when the user navigates focus (e.g. Tab or
-	 * Shift Tab or clicks on a focusable widget.)
+	 * Shift Tab, clicks on a focusable widget, or navigation with keyboard/game pad.)
 	 */
 	template< typename RoutingPolicyType, typename FuncType, typename EventType >
 	static FReply RouteAlongFocusPath( FSlateApplication* ThisApplication, RoutingPolicyType RoutingPolicy, EventType KeyEventCopy, const FuncType& Lambda )
@@ -878,15 +878,13 @@ void FSlateApplication::DrawWindows()
 
 struct FDrawWindowArgs
 {
-	FDrawWindowArgs( FSlateDrawBuffer& InDrawBuffer, const FWidgetPath& InFocusedPath, const FWidgetPath& InWidgetsUnderCursor )
+	FDrawWindowArgs( FSlateDrawBuffer& InDrawBuffer, const FWidgetPath& InWidgetsUnderCursor )
 		: OutDrawBuffer( InDrawBuffer )
-		, FocusedPath( InFocusedPath )
 		, WidgetsUnderCursor( InWidgetsUnderCursor )
 	{}
 
 	TArray<FGenericWindow*, TInlineAllocator<10> > OutDrawnWindows;
 	FSlateDrawBuffer& OutDrawBuffer;
-	const FWidgetPath& FocusedPath;
 	const FWidgetPath& WidgetsUnderCursor;
 };
 
@@ -914,9 +912,15 @@ void FSlateApplication::DrawWindowAndChildren( const TSharedRef<SWindow>& Window
 				WindowToDraw->IsEnabled() );
 		}
 
-		if (DrawWindowArgs.FocusedPath.IsValid() && DrawWindowArgs.FocusedPath.GetWindow() == WindowToDraw)
+
+		for (int32 UserIndex = 0; UserIndex < ARRAY_COUNT(UserFocusedWidgetPaths); UserIndex++)
 		{
-			MaxLayerId = DrawKeyboardFocus(DrawWindowArgs.FocusedPath, WindowElementList, MaxLayerId);
+			FWidgetPath ControllerFocusPath = UserFocusedWidgetPaths[UserIndex].ToWidgetPath();
+
+			if (ControllerFocusPath.IsValid() && ControllerFocusPath.GetWindow() == WindowToDraw && FocusCause == EFocusCause::Navigation)
+			{
+				MaxLayerId = DrawFocus(ControllerFocusPath, WindowElementList, MaxLayerId);
+			}
 		}
 
 		// The widget reflector may want to paint some additional stuff as part of the Widget introspection that it performs.
@@ -1011,14 +1015,12 @@ void FSlateApplication::PrivateDrawWindows( TSharedPtr<SWindow> DrawOnlyThisWind
 		? WidgetsUnderCursorLastEvent.ToWidgetPath()
 		: FWidgetPath();
 
-	FWidgetPath FocusPath = FocusedWidgetPath.ToWidgetPath();
-
 	if ( !SkipSecondPrepass.GetValueOnGameThread() )
 	{
 		DrawPrepass( DrawOnlyThisWindow );
 	}
 
-	FDrawWindowArgs DrawWindowArgs( Renderer->GetDrawBuffer(), FocusPath, WidgetsUnderCursor );
+	FDrawWindowArgs DrawWindowArgs( Renderer->GetDrawBuffer(), WidgetsUnderCursor );
 
 	{
 		SCOPE_CYCLE_COUNTER( STAT_SlateDrawWindowTime );
@@ -1115,16 +1117,19 @@ void FSlateApplication::FinishedInputThisFrame()
 
 	// Any widgets that may have recieved key events
 	// are given a chance to process accumulated values.
-	for( const TWeakPtr<SWidget>& WidgetPtr : FocusedWidgetPath.Widgets )
+	for (int32 SlateUserIndex = 0; SlateUserIndex < SlateApplicationDefs::MaxUsers; ++SlateUserIndex)
 	{
-		const TSharedPtr<SWidget>& Widget = WidgetPtr.Pin();
-		if (Widget.IsValid())
+		for (const TWeakPtr<SWidget>& WidgetPtr : UserFocusedWidgetPaths[SlateUserIndex].Widgets)
 		{
-			Widget->OnFinishedKeyInput();
-		}
-		else
-		{
-			break;
+			const TSharedPtr<SWidget>& Widget = WidgetPtr.Pin();
+			if (Widget.IsValid())
+			{
+				Widget->OnFinishedKeyInput();
+			}
+			else
+			{
+				break;
+			}
 		}
 	}
 }
@@ -1728,7 +1733,7 @@ void FSlateApplication::RegisterGameViewport( TSharedRef<SViewport> InViewport )
 	// If we cannot find the window it could have been destroyed.
 	if (FSlateWindowHelper::FindPathToWidget(SlateWindows, InViewport, PathToViewport, EVisibility::All))
 	{
-		FReply Reply = FReply::Handled().SetKeyboardFocus( InViewport, EKeyboardFocusCause::SetDirectly );
+		FReply Reply = FReply::Handled().SetUserFocus(InViewport, EFocusCause::SetDirectly);
 	
 		// Set keyboard focus on the actual OS window for the top level Slate window in the viewport path
 		// This is needed because some OS messages are only sent to the window with keyboard focus
@@ -1752,47 +1757,93 @@ TSharedPtr<SViewport> FSlateApplication::GetGameViewport() const
 	return GameViewportWidget.Pin();
 }
 
-void FSlateApplication::SetFocusToGameViewport()
+int32 FSlateApplication::GetUserIndexForKeyboard() const
 {
-	TSharedPtr< SViewport > CurrentGameViewportWidget = GameViewportWidget.Pin();
+	//@Todo Slate: Fix this to actual be a map and add API for the user to edit the mapping.
+	// HACK! Just directly mapping the keyboard to User Index 0.
+	return 0;
+}
+ 
+int32 FSlateApplication::GetUserIndexForController(int32 ControllerId) const
+{
+	//@Todo Slate: Fix this to actual be a map and add API for the user to edit the mapping.
+	// HACK! Just directly mapping a controller to a User Index.
+	return ControllerId;
+}
 
-	if ( CurrentGameViewportWidget.IsValid() )
+void FSlateApplication::SetUserFocusToGameViewport(uint32 UserIndex, EFocusCause ReasonFocusIsChanging /* = EFocusCause::SetDirectly*/)
+{
+	TSharedPtr<SViewport> CurrentGameViewportWidget = GameViewportWidget.Pin();
+	if (CurrentGameViewportWidget.IsValid())
 	{
-		SetKeyboardFocus( CurrentGameViewportWidget );
+		SetUserFocus(UserIndex, CurrentGameViewportWidget, ReasonFocusIsChanging);
 	}
 }
 
-void FSlateApplication::SetJoystickCaptorToGameViewport()
+void FSlateApplication::SetFocusToGameViewport()
+{
+	SetUserFocusToGameViewport(GetUserIndexForKeyboard(), EFocusCause::SetDirectly);
+}
+
+void FSlateApplication::SetAllUserFocusToGameViewport(EFocusCause ReasonFocusIsChanging /* = EFocusCause::SetDirectly*/)
 {
 	TSharedPtr< SViewport > CurrentGameViewportWidget = GameViewportWidget.Pin();
 
-	if ( CurrentGameViewportWidget.IsValid() )
+	if (CurrentGameViewportWidget.IsValid())
 	{
 		FWidgetPath PathToWidget;
 		FSlateWindowHelper::FindPathToWidget(SlateWindows, CurrentGameViewportWidget.ToSharedRef(), /*OUT*/ PathToWidget);
 
-		FReply Temp = FReply::Handled().CaptureJoystick(CurrentGameViewportWidget.ToSharedRef(), true);
-
-		ProcessReply(PathToWidget, Temp, NULL, NULL);
+		for (int32 SlateUserIndex = 0; SlateUserIndex < SlateApplicationDefs::MaxUsers; ++SlateUserIndex)
+		{
+			SetUserFocus(SlateUserIndex, PathToWidget, ReasonFocusIsChanging);
+		}
 	}
 }
 
-void FSlateApplication::SetKeyboardFocus( const TSharedPtr< SWidget >& OptionalWidgetToFocus, EKeyboardFocusCause::Type ReasonFocusIsChanging )
+void FSlateApplication::SetJoystickCaptorToGameViewport()
 {
-	if (OptionalWidgetToFocus.IsValid())
-	{
-		FWidgetPath PathToWidget;
-		FSlateWindowHelper::FindPathToWidget(SlateWindows, OptionalWidgetToFocus.ToSharedRef(), /*OUT*/ PathToWidget);
-
-		FReply Reply = FReply::Handled();
-		Reply.SetKeyboardFocus( OptionalWidgetToFocus.ToSharedRef(), EKeyboardFocusCause::SetDirectly );
+	SetAllUserFocusToGameViewport();
+}
 
-		ProcessReply( PathToWidget, Reply, NULL, NULL );	
-	}
-	else
+void FSlateApplication::SetUserFocus(uint32 UserIndex, const TSharedPtr<SWidget>& WidgetToFocus, EFocusCause ReasonFocusIsChanging /* = EFocusCause::SetDirectly*/)
+{
+	FWidgetPath PathToWidget;
+	if (WidgetToFocus.IsValid())
 	{
-		ClearKeyboardFocus(EKeyboardFocusCause::SetDirectly);
+		FSlateWindowHelper::FindPathToWidget(SlateWindows, WidgetToFocus.ToSharedRef(), /*OUT*/ PathToWidget);
 	}
+	SetUserFocus(UserIndex, PathToWidget, ReasonFocusIsChanging);
+}
+
+TSharedPtr<SWidget> FSlateApplication::GetUserFocusedWidget(uint32 UserIndex) const
+{
+	return UserFocusedWidgetPaths[UserIndex].IsValid() ? UserFocusedWidgetPaths[UserIndex].GetLastWidget().Pin() : TSharedPtr<SWidget>();
+}
+
+TSharedPtr<SWidget> FSlateApplication::GetJoystickCaptor(uint32 UserIndex) const
+{
+	return GetUserFocusedWidget(UserIndex);
+}
+
+void FSlateApplication::ClearUserFocus(uint32 UserIndex, EFocusCause ReasonFocusIsChanging /* = EFocusCause::Cleared*/)
+{
+	SetUserFocus(UserIndex, FWidgetPath(), ReasonFocusIsChanging);
+}
+
+void FSlateApplication::ReleaseJoystickCapture(uint32 UserIndex)
+{
+	ClearUserFocus(UserIndex);
+}
+
+void FSlateApplication::SetKeyboardFocus(const TSharedPtr< SWidget >& OptionalWidgetToFocus, EFocusCause ReasonFocusIsChanging /* = EFocusCause::SetDirectly*/)
+{
+	SetUserFocus(GetUserIndexForKeyboard(), OptionalWidgetToFocus, ReasonFocusIsChanging);
+}
+
+void FSlateApplication::ClearKeyboardFocus(const EFocusCause ReasonFocusIsChanging)
+{
+	SetUserFocus(GetUserIndexForKeyboard(), FWidgetPath(), ReasonFocusIsChanging);
 }
 
 void FSlateApplication::ResetToDefaultInputSettings()
@@ -1801,17 +1852,11 @@ void FSlateApplication::ResetToDefaultInputSettings()
 	{
 		ProcessReply(MouseCaptorPath, FReply::Handled().ReleaseMouseCapture(), NULL, NULL);
 	}
+	
+	ProcessReply(FWidgetPath(), FReply::Handled().ClearUserFocus(true), NULL, NULL);
 
-	for (int32 UserIndex = 0; UserIndex < ARRAY_COUNT(JoystickCaptorWeakPaths); UserIndex++)
-	{
-		if (JoystickCaptorWeakPaths[UserIndex].IsValid())
-		{
-			FWidgetPath JoystickCaptorPath = JoystickCaptorWeakPaths[UserIndex].ToWidgetPath();
-			ProcessReply( JoystickCaptorPath, FReply::Handled().ReleaseJoystickCapture(), NULL, NULL, UserIndex );
-		}
-	}
+	ProcessReply(FWidgetPath(), FReply::Handled().ReleaseMouseLock(), NULL, NULL);
 
-	ProcessReply( FWidgetPath(), FReply::Handled().ReleaseMouseLock(), NULL, NULL );
 	if ( PlatformApplication->Cursor.IsValid() )
 	{
 		PlatformApplication->Cursor->SetType(EMouseCursor::Default);
@@ -1830,19 +1875,6 @@ void FSlateApplication::ReleaseMouseCapture()
 	MouseCaptor.InvalidateCaptureForAllPointers();
 }
 
-
-TSharedPtr< SWidget > FSlateApplication::GetJoystickCaptor(uint32 UserIndex) const
-{
-	return JoystickCaptorWeakPaths[UserIndex].IsValid() ? JoystickCaptorWeakPaths[UserIndex].GetLastWidget().Pin() : TSharedPtr<SWidget>();
-}
-
-
-void FSlateApplication::ReleaseJoystickCapture(uint32 UserIndex)
-{
-	JoystickCaptorWeakPaths[UserIndex] = FWeakWidgetPath();
-}
-
-
 TSharedPtr<SWindow> FSlateApplication::GetActiveTopLevelWindow() const
 {
 	return ActiveTopLevelWindow.Pin();
@@ -1854,159 +1886,151 @@ TSharedPtr<SWindow> FSlateApplication::GetActiveModalWindow() const
 	return (ActiveModalWindows.Num() > 0) ? ActiveModalWindows.Last() : NULL;
 }
 
+bool FSlateApplication::SetKeyboardFocus(const FWidgetPath& InFocusPath, const EFocusCause InCause /*= EFocusCause::SetDirectly*/)
+{
+	return SetUserFocus(GetUserIndexForKeyboard(), InFocusPath, InCause);
+}
 
-bool FSlateApplication::SetKeyboardFocus( const FWidgetPath& InFocusPath, const EKeyboardFocusCause::Type InCause )
+bool FSlateApplication::SetUserFocus(const uint32 InUserIndex, const FWidgetPath& InFocusPath, const EFocusCause InCause)
 {
-	if (!InFocusPath.IsValid())
-	{
-		return false;
-	}
+	check(InUserIndex >= 0 && InUserIndex < SlateApplicationDefs::MaxUsers);
 
 	TSharedPtr<IWidgetReflector> WidgetReflector = WidgetReflectorPtr.Pin();
 	const bool bReflectorShowingFocus = WidgetReflector.IsValid() && WidgetReflector->IsShowingFocus();
 
-	bool bFocusTransferComplete = false;
+	// Get the old Widget information
+	const FWeakWidgetPath OldFocusedWidgetPath = UserFocusedWidgetPaths[InUserIndex];
+	TSharedPtr<SWidget> OldFocusedWidget = OldFocusedWidgetPath.IsValid() ? OldFocusedWidgetPath.GetLastWidget().Pin() : TSharedPtr< SWidget >();
+
+	// Get the new widget information by finding the first widget in the path that supports focus
+	FWidgetPath NewFocusedWidgetPath;
+	TSharedPtr<SWidget> NewFocusedWidget;
 
-	for (int32 WidgetIndex = InFocusPath.Widgets.Num()-1; !bFocusTransferComplete  && WidgetIndex >= 0; --WidgetIndex)
+	if (InFocusPath.IsValid())
 	{
-		const FArrangedWidget& WidgetToFocus = InFocusPath.Widgets[ WidgetIndex ];
-		
-		// Does this widget support keyboard focus?  If so, then we'll go ahead and set it!
-		if( WidgetToFocus.Widget->SupportsKeyboardFocus() )
+		for (int32 WidgetIndex = InFocusPath.Widgets.Num() - 1; WidgetIndex >= 0; --WidgetIndex)
 		{
-			// Has focus actually changed?
-			TSharedPtr< SWidget > OldFocusedWidget( GetKeyboardFocusedWidget() );
-			
-			// Is the focus actually changing?
-			if( WidgetToFocus.Widget != OldFocusedWidget )
+			const FArrangedWidget& WidgetToFocus = InFocusPath.Widgets[WidgetIndex];
+
+			// Does this widget support keyboard focus?  If so, then we'll go ahead and set it!
+			if (WidgetToFocus.Widget->SupportsKeyboardFocus())
 			{
-				FWidgetPath NewFocusPath = InFocusPath.GetPathDownTo(WidgetToFocus.Widget);
+				// Is we aren't changing focus then simply return
+				if (WidgetToFocus.Widget == OldFocusedWidget)
 				{
-					// Notify all affected widgets about the change in focus.
-					TArray< TSharedRef<SWidget> > NotifyUsAboutFocusChange;
-
-					// Notify widgets in the old focus path
-					{
-						for (int32 ChildIndex=0; ChildIndex < FocusedWidgetPath.Widgets.Num(); ++ChildIndex)
-						{
-							TSharedPtr<SWidget> SomeWidget = FocusedWidgetPath.Widgets[ ChildIndex ].Pin();
-							if (SomeWidget.IsValid())
-							{
-								NotifyUsAboutFocusChange.Add( SomeWidget.ToSharedRef() );
-							}
-						}
-
-						FScopedSwitchWorldHack SwitchWorld( FocusedWidgetPath.Window.Pin() );
-						
-						for( int32 NotifyWidgetIndex=0; NotifyWidgetIndex < NotifyUsAboutFocusChange.Num(); ++NotifyWidgetIndex )
-						{
-							NotifyUsAboutFocusChange[NotifyWidgetIndex]->OnKeyboardFocusChanging( FocusedWidgetPath, NewFocusPath );
-						}
-						
-					}
-	
-					// Empty the array for new widgets
-					NotifyUsAboutFocusChange.Empty();
-
-					// Notify widgets in the new focus path
-					{
-						for (int32 ChildIndex=0; ChildIndex < NewFocusPath.Widgets.Num(); ++ChildIndex)
-						{
-							NotifyUsAboutFocusChange.AddUnique( NewFocusPath.Widgets[ChildIndex].Widget );
-						}
-
-						FScopedSwitchWorldHack SwitchWorld( NewFocusPath );
-					
-						for( int32 NotifyWidgetIndex=0; NotifyWidgetIndex < NotifyUsAboutFocusChange.Num(); ++NotifyWidgetIndex )
-						{
-							NotifyUsAboutFocusChange[NotifyWidgetIndex]->OnKeyboardFocusChanging( FocusedWidgetPath, NewFocusPath );
-						}
-					}
+					return false;
 				}
-				
-				const FWeakWidgetPath OldFocusedWidgetPath = FocusedWidgetPath;
-				// Store a weak widget path to the widget that's taking focus
-				FocusedWidgetPath = FWeakWidgetPath( NewFocusPath );
-
-				if( OldFocusedWidget.IsValid() )
-				{
-					// Switch worlds for widgets in the old path
-					FScopedSwitchWorldHack SwitchWorld( OldFocusedWidgetPath.Window.Pin() );
+				NewFocusedWidget = WidgetToFocus.Widget;
+				NewFocusedWidgetPath = InFocusPath.GetPathDownTo(NewFocusedWidget.ToSharedRef());
+				break;
+			}
+		}
+	}
 
-					// Let previously-focused widget know that it's losing focus
-					OldFocusedWidget->OnKeyboardFocusLost( FKeyboardFocusEvent( InCause ) );
-				}
+	// Notify widgets in the old focus path that focus is changing
+	if (OldFocusedWidgetPath.IsValid())
+	{
+		FScopedSwitchWorldHack SwitchWorld(OldFocusedWidgetPath.Window.Pin());
 
-				if (bReflectorShowingFocus)
-				{
-					WidgetReflector->SetWidgetsToVisualize(NewFocusPath);
-				}
+		for (int32 ChildIndex = 0; ChildIndex < OldFocusedWidgetPath.Widgets.Num(); ++ChildIndex)
+		{
+			TSharedPtr<SWidget> SomeWidget = OldFocusedWidgetPath.Widgets[ChildIndex].Pin();
+			if (SomeWidget.IsValid())
+			{
+				SomeWidget->OnFocusChanging(OldFocusedWidgetPath, NewFocusedWidgetPath);
 
-				FocusCause = InCause;
+				//@Todo slate: Remove this backwards compatibility code as doon as the deprecated functionality in Swidget is removed
+				// Run the deprecated function to maintain backwards compatibility for now.
+				SomeWidget->OnFocusChanging(OldFocusedWidgetPath, NewFocusedWidgetPath);
+			}
+		}
+	}
 
-				// Let the new widget know that it's received keyboard focus
-				{
-					// Switch worlds for widgets in the new path
-					FScopedSwitchWorldHack SwitchWorld( NewFocusPath );
+	// Notify widgets in the new focus path that focus is changing
+	if (NewFocusedWidgetPath.IsValid())
+	{
+		FScopedSwitchWorldHack SwitchWorld(NewFocusedWidgetPath.GetWindow());
 
-					FReply Reply = WidgetToFocus.Widget->OnKeyboardFocusReceived( WidgetToFocus.Geometry, FKeyboardFocusEvent( InCause ) );
-					if (Reply.IsEventHandled())
-					{
-						ProcessReply( InFocusPath, Reply, NULL, NULL );
-					}
-				}
+		for (int32 ChildIndex = 0; ChildIndex < NewFocusedWidgetPath.Widgets.Num(); ++ChildIndex)
+		{
+			TSharedPtr<SWidget> SomeWidget = NewFocusedWidgetPath.Widgets[ChildIndex].Widget;
+			if (SomeWidget.IsValid())
+			{
+				SomeWidget->OnFocusChanging(OldFocusedWidgetPath, NewFocusedWidgetPath);
 
-				TSharedPtr<SWindow> FocusedWindow = FocusedWidgetPath.Window.Pin();
-				if ( FocusedWindow.IsValid() && FocusedWindow != WidgetToFocus.Widget )
-				{
-					FocusedWindow->SetWidgetToFocusOnActivate( WidgetToFocus.Widget );
-				}
+				//@Todo slate: Remove this backwards compatibility code as doon as the deprecated functionality in Swidget is removed
+				// Run the deprecated function to maintain backwards compatibility for now.
+				SomeWidget->OnFocusChanging(OldFocusedWidgetPath, NewFocusedWidgetPath);
 			}
-		
-			// We are about to successfully transfer focus.
-			bFocusTransferComplete = true;
 		}
 	}
 
-	return bFocusTransferComplete;
-}
+	// Store a weak widget path to the widget that's taking focus
+	UserFocusedWidgetPaths[InUserIndex] = FWeakWidgetPath(NewFocusedWidgetPath);
+	//@Todo Slate: Refactor focus cause to be stored for all users!
+	FocusCause = InCause;
 
+	// Let the old widget know that it lost keyboard focus
+	if(OldFocusedWidget.IsValid())
+	{
+		// Switch worlds for widgets in the old path
+		FScopedSwitchWorldHack SwitchWorld(OldFocusedWidgetPath.Window.Pin());
 
-FModifierKeysState FSlateApplication::GetModifierKeys() const
-{
-	return PlatformApplication->GetModifierKeys();
-}
+		// Let previously-focused widget know that it's losing focus
+		OldFocusedWidget->OnFocusLost(FFocusEvent(InCause));
 
+		//@Todo slate: Remove this backwards compatibility code as doon as the deprecated functionality in Swidget is removed
+		// Run the deprecated function to maintain backwards compatibility for now.
+		OldFocusedWidget->OnFocusLost(FFocusEvent(InCause));
+	}
 
-void FSlateApplication::ClearKeyboardFocus( const EKeyboardFocusCause::Type InCause )
-{
-	TSharedPtr< SWidget > OldFocusedWidget( GetKeyboardFocusedWidget() );
-	const FWeakWidgetPath OldFocusedWidgetPath = FocusedWidgetPath;
-	FocusedWidgetPath = FWeakWidgetPath();
+	if (bReflectorShowingFocus)
+	{
+		WidgetReflector->SetWidgetsToVisualize(NewFocusedWidgetPath);
+	}
 
-	// Let previously-focused widget know that it's losing focus
-	if( OldFocusedWidget.IsValid() )
+	// Let the new widget know that it's received keyboard focus
+	if (NewFocusedWidget.IsValid())
 	{
-		if ( OldFocusedWidgetPath.Window.IsValid() )
-		{
-			// Switch worlds for widgets in the current path
-			FScopedSwitchWorldHack SwitchWorld( OldFocusedWidgetPath.Window.Pin().ToSharedRef() );
+		TSharedPtr<SWindow> FocusedWindow = NewFocusedWidgetPath.GetWindow();
+		
+		// Switch worlds for widgets in the new path
+		FScopedSwitchWorldHack SwitchWorld(FocusedWindow);
+		
+		const FArrangedWidget& WidgetToFocus = NewFocusedWidgetPath.Widgets.Last();
 
-			OldFocusedWidget->OnKeyboardFocusLost( FKeyboardFocusEvent( InCause ) );
+		FReply Reply = NewFocusedWidget->OnFocusReceived(WidgetToFocus.Geometry, FFocusEvent(InCause));
+		if (Reply.IsEventHandled())
+		{
+			ProcessReply(InFocusPath, Reply, NULL, NULL);
 		}
+		//@Todo slate: Remove this backwards compatibility code as soon as the deprecated functionality in Swidget is removed
+		// Run the deprecated function to maintain backwards compatibility for now.
 		else
 		{
-			OldFocusedWidget->OnKeyboardFocusLost( FKeyboardFocusEvent( InCause ) );
+			FReply Reply = NewFocusedWidget->OnFocusReceived(WidgetToFocus.Geometry, FFocusEvent(InCause));
+			if (Reply.IsEventHandled())
+			{
+				ProcessReply(InFocusPath, Reply, NULL, NULL);
+			}
+		}
+		
+		// Set the windows restore state.
+		//@Todo Slate: We need to store the full focus state 
+		//@Todo Slate: Why are we chacking FocusedWindow != NewFocusedWidget?
+		if (FocusedWindow.IsValid() && FocusedWindow != NewFocusedWidget)
+		{
+			FocusedWindow->SetWidgetToFocusOnActivate(NewFocusedWidget);
 		}
 	}
 
-	TSharedPtr<IWidgetReflector> WidgetReflector = WidgetReflectorPtr.Pin();
-	const bool bReflectorShowingFocus = WidgetReflector.IsValid() && WidgetReflector->IsShowingFocus();
+	return true;
+}
 
-	if (bReflectorShowingFocus)
-	{
-		WidgetReflector->SetWidgetsToVisualize(FWidgetPath());
-	}
+
+FModifierKeysState FSlateApplication::GetModifierKeys() const
+{
+	return PlatformApplication->GetModifierKeys();
 }
 
 
@@ -2127,24 +2151,19 @@ void FSlateApplication::ProcessReply( const FWidgetPath& CurrentEventPath, const
 		MouseCaptor.InvalidateCaptureForPointer(PointerIndex);
 	}
 
-
-	if ( TheReply.ShouldReleaseJoystick() )
+	// Clear focus is requested.
+	if (TheReply.ShouldReleaseUserFocus())
 	{
-		if (JoystickCaptorWeakPaths[UserIndex].IsValid())
-		{
-			WidgetsUnderCursorLastEvent = JoystickCaptorWeakPaths[UserIndex];
-		}
-
-		if (TheReply.AffectsAllJoysticks())
+		if (TheReply.AffectsAllUsers())
 		{
 			for (int32 SlateUserIndex = 0; SlateUserIndex < SlateApplicationDefs::MaxUsers; ++SlateUserIndex)
 			{
-				JoystickCaptorWeakPaths[SlateUserIndex] = FWeakWidgetPath();
+				SetUserFocus(SlateUserIndex, FWidgetPath(), EFocusCause::Cleared);
 			}
 		}
 		else
 		{
-			JoystickCaptorWeakPaths[UserIndex] = FWeakWidgetPath();
+			SetUserFocus(UserIndex, FWidgetPath(), EFocusCause::Cleared);
 		}
 	}
 
@@ -2201,32 +2220,6 @@ void FSlateApplication::ProcessReply( const FWidgetPath& CurrentEventPath, const
 		}
 	}
 
-	TSharedPtr<SWidget> RequestedJoystickCaptor = TheReply.GetJoystickCaptor();
-	if( CurrentEventPath.IsValid() &&  RequestedJoystickCaptor.IsValid() )
-	{
-		FWidgetPath NewJoystickCaptorPath = CurrentEventPath.GetPathDownTo( RequestedJoystickCaptor.ToSharedRef() );
-
-		if ( !NewJoystickCaptorPath.IsValid() )
-		{
-			// The requested mouse captor was not in the event path.
-			// We will attempt to find it in this window; if we don't find it, then give up.
-			NewJoystickCaptorPath = CurrentEventPath.GetPathDownTo( CurrentEventPath.Widgets[0].Widget );
-			NewJoystickCaptorPath.ExtendPathTo( FWidgetMatcher(RequestedJoystickCaptor.ToSharedRef()) );
-		}
-
-		if (TheReply.AffectsAllJoysticks())
-		{
-			for (int32 SlateUserIndex = 0; SlateUserIndex < SlateApplicationDefs::MaxUsers; ++SlateUserIndex)
-			{
-				JoystickCaptorWeakPaths[SlateUserIndex] = NewJoystickCaptorPath;
-			}
-		}
-		else
-		{
-			JoystickCaptorWeakPaths[UserIndex] = NewJoystickCaptorPath;
-		}
-	}
-
 	TOptional<FIntPoint> RequestedMousePos = TheReply.GetRequestedMousePos();
 	if( RequestedMousePos.IsSet() )
 	{
@@ -2245,6 +2238,33 @@ void FSlateApplication::ProcessReply( const FWidgetPath& CurrentEventPath, const
 		// Unlock the mouse
 		LockCursor( NULL );
 	}
+	
+	// If we have a valid Navigation request and enough time has passed since the last navigation then try to navigate.
+	if (TheReply.GetNavigationType() != EUINavigation::Invalid)
+	{
+		TSharedPtr<SWidget> FocusedWidget = GetUserFocusedWidget(UserIndex);
+		if (FocusedWidget.IsValid())
+		{
+			FNavigationEvent NavigationEvent(UserIndex, TheReply.GetNavigationType());
+
+			FWidgetPath EventPath = UserFocusedWidgetPaths[UserIndex].ToWidgetPath();
+
+			FNavigationReply NavigationReply = FNavigationReply::Escape();
+			for (int32 WidgetIndex = EventPath.Widgets.Num() - 1; WidgetIndex >= 0; --WidgetIndex)
+			{
+				FArrangedWidget& SomeWidgetGettingEvent = EventPath.Widgets[WidgetIndex];
+				if (SomeWidgetGettingEvent.Widget->IsEnabled())
+				{
+					NavigationReply = SomeWidgetGettingEvent.Widget->OnNavigation(SomeWidgetGettingEvent.Geometry, NavigationEvent).SetHandler(SomeWidgetGettingEvent.Widget);
+					if (NavigationReply.GetBoundaryRule() != EUINavigationRule::Escape || WidgetIndex == 0)
+					{
+						AttemptNavigation(NavigationEvent, NavigationReply, SomeWidgetGettingEvent);
+						break;
+					}
+				}
+			}
+		}
+	}
 
 	if ( TheReply.GetDetectDragRequest().IsValid() )
 	{
@@ -2253,19 +2273,24 @@ void FSlateApplication::ProcessReply( const FWidgetPath& CurrentEventPath, const
 		DragDetector.DetectDragStartLocation = InMouseEvent->GetScreenSpacePosition();
 	}
 
-	TSharedPtr<SWidget> RequestedKeyboardFocusRecepient = TheReply.GetFocusRecepient();
-	if ( CurrentEventPath.IsValid() && RequestedKeyboardFocusRecepient.IsValid() )
+	// Set focus if requested.
+	TSharedPtr<SWidget> RequestedFocusRecepient = TheReply.GetUserFocusRecepient();
+	if (TheReply.ShouldSetUserFocus() || RequestedFocusRecepient.IsValid())
 	{
-		// The widget to focus is probably in the path of this event (likely the handler or handler's parent).
-		FWidgetPath NewFocusedWidgetPath = CurrentEventPath.GetPathDownTo( RequestedKeyboardFocusRecepient.ToSharedRef() );
-		if ( !NewFocusedWidgetPath.IsValid() )
+		FWidgetPath NewFocusedWidgetPath;
+		GeneratePathToWidgetUnchecked(RequestedFocusRecepient.ToSharedRef(), NewFocusedWidgetPath);
+
+		if (TheReply.AffectsAllUsers())
 		{
-			// The widget we want to focus is not in the event processing path.
-			// Search all the widgets for it.
-			GeneratePathToWidgetUnchecked( RequestedKeyboardFocusRecepient.ToSharedRef(), NewFocusedWidgetPath );
-		}		
-		
-		SetKeyboardFocus( NewFocusedWidgetPath, TheReply.GetFocusCause() );
+			for (int32 SlateUserIndex = 0; SlateUserIndex < SlateApplicationDefs::MaxUsers; ++SlateUserIndex)
+			{
+				SetUserFocus(SlateUserIndex, NewFocusedWidgetPath, TheReply.GetFocusCause());
+			}
+		}
+		else
+		{
+			SetUserFocus(UserIndex, NewFocusedWidgetPath, TheReply.GetFocusCause());
+		}
 	}
 }
 
@@ -2709,33 +2734,29 @@ void FSlateApplication::UpdateToolTip( bool AllowSpawningOfNewToolTips )
 	}
 }
 
-int32 FSlateApplication::DrawKeyboardFocus( const FWidgetPath& FocusPath, FSlateWindowElementList& WindowElementList, int32 InLayerId ) const
+int32 FSlateApplication::DrawFocus( const FWidgetPath& FocusPath, FSlateWindowElementList& WindowElementList, int32 InLayerId ) const
 {
-	if (FocusCause == EKeyboardFocusCause::Keyboard)
-	{
-		// Widgets where being focused matters draw themselves differently when focused.
-		// When the user navigates keyboard focus, we draw keyboard focus for everything, 
-		// so the user can see what they are doing.
-		const FArrangedWidget& FocusedWidgetGeomertry = FocusPath.Widgets.Last();
+	// Widgets where being focused matters draw themselves differently when focused.
+	// When the user navigates focus, we draw focus for everything, so the user can see what they are doing.
+	const FArrangedWidget& FocusedWidgetGeomertry = FocusPath.Widgets.Last();
 
-		// The FGeometry we get is from a WidgetPath, so it's rooted in desktop space.
-		// We need to APPEND a transform to the Geometry to essentially undo this root transform
-		// and get us back into Window Space.
-		// This is nonstandard so we have to go through some hoops and a specially exposed method 
-		// in FPaintGeometry to allow appending layout transforms.
-		FPaintGeometry WindowSpaceGeometry = FocusedWidgetGeomertry.Geometry.ToPaintGeometry();
-		WindowSpaceGeometry.AppendTransform(TransformCast<FSlateLayoutTransform>(Inverse(FocusPath.GetWindow()->GetPositionInScreen())));
+	// The FGeometry we get is from a WidgetPath, so it's rooted in desktop space.
+	// We need to APPEND a transform to the Geometry to essentially undo this root transform
+	// and get us back into Window Space.
+	// This is nonstandard so we have to go through some hoops and a specially exposed method 
+	// in FPaintGeometry to allow appending layout transforms.
+	FPaintGeometry WindowSpaceGeometry = FocusedWidgetGeomertry.Geometry.ToPaintGeometry();
+	WindowSpaceGeometry.AppendTransform(TransformCast<FSlateLayoutTransform>(Inverse(FocusPath.GetWindow()->GetPositionInScreen())));
 
-		FSlateDrawElement::MakeBox(
-			WindowElementList,
-			InLayerId++,
-			WindowSpaceGeometry,
-			FCoreStyle::Get().GetBrush("FocusRectangle"),
-			FocusPath.GetWindow()->GetClippingRectangleInWindow(),
-			ESlateDrawEffect::None,
-			FColor(255,255,255,128)
-		);
-	}
+	FSlateDrawElement::MakeBox(
+		WindowElementList,
+		InLayerId++,
+		WindowSpaceGeometry,
+		FCoreStyle::Get().GetBrush("FocusRectangle"),
+		FocusPath.GetWindow()->GetClippingRectangleInWindow(),
+		ESlateDrawEffect::None,
+		FColor(255,255,255,128)
+	);
 
 	return InLayerId;
 }
@@ -2927,7 +2948,7 @@ void FSlateApplication::SetSlateUILogger(TSharedPtr<IEventLogger> InEventLogger)
 #endif
 }
 
-void FSlateApplication::SetUnhandledKeyDownEventHandler( const FOnKeyboardEvent& NewHandler )
+void FSlateApplication::SetUnhandledKeyDownEventHandler( const FOnKeyEvent& NewHandler )
 {
 	UnhandledKeyDownEventHandler = NewHandler;
 }
@@ -3105,6 +3126,8 @@ void FSlateApplication::ShowVirtualKeyboard( bool bShow, TSharedPtr<IVirtualKeyb
 
 FSlateRect FSlateApplication::GetPreferredWorkArea() const
 {
+	const FWeakWidgetPath & FocusedWidgetPath = UserFocusedWidgetPaths[GetUserIndexForKeyboard()];
+
 	// First see if we have a focused widget
 	if( FocusedWidgetPath.IsValid() && FocusedWidgetPath.Window.IsValid() )
 	{
@@ -3244,9 +3267,9 @@ FVector2D FSlateApplication::GetCursorSize( ) const
 
 TSharedPtr< SWidget > FSlateApplication::GetKeyboardFocusedWidget() const
 {
-	if( FocusedWidgetPath.IsValid() )
+	if( UserFocusedWidgetPaths[GetUserIndexForKeyboard()].IsValid() )
 	{
-		return FocusedWidgetPath.GetLastWidget().Pin();
+		return UserFocusedWidgetPaths[GetUserIndexForKeyboard()].GetLastWidget().Pin();
 	}
 
 	return TSharedPtr< SWidget >();
@@ -3276,7 +3299,15 @@ bool FSlateApplication::HasMouseCapture(const TSharedPtr<const SWidget> Widget)
 
 bool FSlateApplication::HasFocusedDescendants( const TSharedRef< const SWidget >& Widget ) const
 {
-	return FocusedWidgetPath.IsValid() && FocusedWidgetPath.GetLastWidget().Pin() != Widget && FocusedWidgetPath.ContainsWidget( Widget );
+	for (int32 UserIndex = 0; UserIndex < SlateApplicationDefs::MaxUsers; ++UserIndex)
+	{
+		const FWeakWidgetPath & FocusedWidgetPath = UserFocusedWidgetPaths[UserIndex];
+		if (FocusedWidgetPath.IsValid() && FocusedWidgetPath.GetLastWidget().Pin() != Widget && FocusedWidgetPath.ContainsWidget(Widget))
+		{
+			return true;
+		}
+	}
+	return false;
 }
 
 
@@ -3343,7 +3374,7 @@ bool FSlateApplication::ShouldProcessUserInputMessages( const TSharedPtr< FGener
 
 bool FSlateApplication::OnKeyChar( const TCHAR Character, const bool IsRepeat )
 {
-	FCharacterEvent CharacterEvent( Character, PlatformApplication->GetModifierKeys(), IsRepeat );
+	FCharacterEvent CharacterEvent( Character, PlatformApplication->GetModifierKeys(), 0, IsRepeat );
 	return ProcessKeyCharEvent( CharacterEvent );
 }
 
@@ -3353,8 +3384,8 @@ bool FSlateApplication::ProcessKeyCharEvent( FCharacterEvent& InCharacterEvent )
 
 	// NOTE: We intentionally don't reset LastUserInteractionTimeForThrottling here so that the UI can be responsive while typing
 
-	// Bubble the keyboard event
-	FWidgetPath EventPath = FocusedWidgetPath.ToWidgetPath();
+	// Bubble the key event
+	FWidgetPath EventPath = UserFocusedWidgetPaths[InCharacterEvent.GetUserIndex()].ToWidgetPath();
 	
 	// Switch worlds for widgets in the current path
 	FScopedSwitchWorldHack SwitchWorld( EventPath );
@@ -3374,18 +3405,18 @@ bool FSlateApplication::ProcessKeyCharEvent( FCharacterEvent& InCharacterEvent )
 bool FSlateApplication::OnKeyDown( const int32 KeyCode, const uint32 CharacterCode, const bool IsRepeat ) 
 {
 	FKey const Key = FInputKeyManager::Get().GetKeyFromCodes( KeyCode, CharacterCode );
-	FKeyboardEvent KeyboardEvent( Key, PlatformApplication->GetModifierKeys(), IsRepeat, CharacterCode, KeyCode );
+	FKeyEvent KeyEvent(Key, PlatformApplication->GetModifierKeys(), GetUserIndexForKeyboard(), IsRepeat, CharacterCode, KeyCode);
 
-	return ProcessKeyDownEvent( KeyboardEvent );
+	return ProcessKeyDownEvent( KeyEvent );
 }
 
-bool FSlateApplication::ProcessKeyDownEvent( FKeyboardEvent& InKeyboardEvent )
+bool FSlateApplication::ProcessKeyDownEvent( FKeyEvent& InKeyEvent )
 {
 	FReply Reply = FReply::Unhandled();
 
 	LastUserInteractionTime = this->GetCurrentTime();
 	
-	if (IsDragDropping() && InKeyboardEvent.GetKey() == EKeys::Escape)
+	if (IsDragDropping() && InKeyEvent.GetKey() == EKeys::Escape)
 	{
 		// Pressing ESC while drag and dropping terminates the drag drop.
 		DragDropContent.Reset();
@@ -3396,7 +3427,7 @@ bool FSlateApplication::ProcessKeyDownEvent( FKeyboardEvent& InKeyboardEvent )
 		LastUserInteractionTimeForThrottling = LastUserInteractionTime;
 
 		// If we are inspecting, pressing ESC exits inspection mode.
-		if ( InKeyboardEvent.GetKey() == EKeys::Escape )
+		if ( InKeyEvent.GetKey() == EKeys::Escape )
 		{
 			TSharedPtr<IWidgetReflector> WidgetReflector = WidgetReflectorPtr.Pin();
 			const bool bIsWidgetReflectorPicking = WidgetReflector.IsValid() && WidgetReflector->IsInPickingMode();
@@ -3410,7 +3441,7 @@ bool FSlateApplication::ProcessKeyDownEvent( FKeyboardEvent& InKeyboardEvent )
 			}
 
 		// Bubble the keyboard event
-		FWidgetPath EventPath = FocusedWidgetPath.ToWidgetPath();
+		FWidgetPath EventPath = UserFocusedWidgetPaths[InKeyEvent.GetUserIndex()].ToWidgetPath();
 
 		// Switch worlds for widgets inOnPreviewMouseButtonDown the current path
 		FScopedSwitchWorldHack SwitchWorld( EventPath );
@@ -3420,7 +3451,7 @@ bool FSlateApplication::ProcessKeyDownEvent( FKeyboardEvent& InKeyboardEvent )
 		Reply = FReply::Unhandled();
 
 		// Tunnel the keyboard event
-		Reply = FEventRouter::RouteAlongFocusPath( this, FEventRouter::FTunnelPolicy(EventPath), InKeyboardEvent, []( const FArrangedWidget& CurrentWidget, const FKeyboardEvent& Event )
+		Reply = FEventRouter::RouteAlongFocusPath( this, FEventRouter::FTunnelPolicy(EventPath), InKeyEvent, []( const FArrangedWidget& CurrentWidget, const FKeyEvent& Event )
 		{
 			return (CurrentWidget.Widget->IsEnabled())
 				? CurrentWidget.Widget->OnPreviewKeyDown( CurrentWidget.Geometry, Event )
@@ -3430,7 +3461,7 @@ bool FSlateApplication::ProcessKeyDownEvent( FKeyboardEvent& InKeyboardEvent )
 		// Send out key down events.
 		if (!Reply.IsEventHandled())
 		{
-			Reply = FEventRouter::RouteAlongFocusPath(this, FEventRouter::FBubblePolicy(EventPath), InKeyboardEvent, [](const FArrangedWidget& SomeWidgetGettingEvent, const FKeyboardEvent& Event)
+			Reply = FEventRouter::RouteAlongFocusPath(this, FEventRouter::FBubblePolicy(EventPath), InKeyEvent, [](const FArrangedWidget& SomeWidgetGettingEvent, const FKeyEvent& Event)
 			{
 				return (SomeWidgetGettingEvent.Widget->IsEnabled())
 					? SomeWidgetGettingEvent.Widget->OnKeyDown( SomeWidgetGettingEvent.Geometry, Event )
@@ -3438,27 +3469,12 @@ bool FSlateApplication::ProcessKeyDownEvent( FKeyboardEvent& InKeyboardEvent )
 			});
 		}
 
-		LOG_EVENT_CONTENT( EEventLog::KeyDown, GetKeyName(InKeyboardEvent.GetKey()).ToString(), Reply );
+		LOG_EVENT_CONTENT( EEventLog::KeyDown, GetKeyName(InKeyEvent.GetKey()).ToString(), Reply );
 
-		// If the keyboard event was not processed by any widget...
-		if ( !Reply.IsEventHandled() )
+		// If the key event was not processed by any widget...
+		if (!Reply.IsEventHandled())
 		{
-			// If the key was Tab, interpret as an attempt to move focus.
-			if ( InKeyboardEvent.GetKey() == EKeys::Tab )
-			{
-				if (FocusedWidgetPath.IsValid())
-				{
-					EFocusMoveDirection::Type MoveDirection = ( InKeyboardEvent.IsShiftDown() )
-						? EFocusMoveDirection::Previous
-						: EFocusMoveDirection::Next;
-					SetKeyboardFocus( FocusedWidgetPath.ToNextFocusedPath( MoveDirection ), EKeyboardFocusCause::Keyboard );
-				}
-			}
-			else if ( UnhandledKeyDownEventHandler.IsBound() )
-			{
-				// Nothing else handled this event, give external code a chance to handle it.
-				Reply = UnhandledKeyDownEventHandler.Execute( InKeyboardEvent );
-			}
+			Reply = UnhandledKeyDownEventHandler.Execute( InKeyEvent );
 		}
 	}
 
@@ -3468,12 +3484,12 @@ bool FSlateApplication::ProcessKeyDownEvent( FKeyboardEvent& InKeyboardEvent )
 bool FSlateApplication::OnKeyUp( const int32 KeyCode, const uint32 CharacterCode, const bool IsRepeat )
 {
 	FKey const Key = FInputKeyManager::Get().GetKeyFromCodes( KeyCode, CharacterCode );
-	FKeyboardEvent KeyboardEvent( Key, PlatformApplication->GetModifierKeys(), IsRepeat, CharacterCode, KeyCode );
+	FKeyEvent KeyEvent(Key, PlatformApplication->GetModifierKeys(), GetUserIndexForKeyboard(), IsRepeat, CharacterCode, KeyCode);
 
-	return ProcessKeyUpEvent( KeyboardEvent );
+	return ProcessKeyUpEvent( KeyEvent );
 }
 
-bool FSlateApplication::ProcessKeyUpEvent( FKeyboardEvent& InKeyboardEvent )
+bool FSlateApplication::ProcessKeyUpEvent( FKeyEvent& InKeyEvent )
 {
 	FReply Reply = FReply::Unhandled();
 
@@ -3481,20 +3497,47 @@ bool FSlateApplication::ProcessKeyUpEvent( FKeyboardEvent& InKeyboardEvent )
 	
 	LastUserInteractionTimeForThrottling = LastUserInteractionTime;
 
-	// Bubble the keyboard event
-	FWidgetPath EventPath = FocusedWidgetPath.ToWidgetPath();
+	// Bubble the key event
+	FWidgetPath EventPath = UserFocusedWidgetPaths[InKeyEvent.GetUserIndex()].ToWidgetPath();
 
 	// Switch worlds for widgets in the current path
 	FScopedSwitchWorldHack SwitchWorld( EventPath );
 
-	Reply = FEventRouter::RouteAlongFocusPath(this, FEventRouter::FBubblePolicy(EventPath), InKeyboardEvent, [](const FArrangedWidget& SomeWidgetGettingEvent, const FKeyboardEvent& Event)
+	Reply = FEventRouter::RouteAlongFocusPath(this, FEventRouter::FBubblePolicy(EventPath), InKeyEvent, [](const FArrangedWidget& SomeWidgetGettingEvent, const FKeyEvent& Event)
 	{
 		return (SomeWidgetGettingEvent.Widget->IsEnabled())
 			? SomeWidgetGettingEvent.Widget->OnKeyUp( SomeWidgetGettingEvent.Geometry, Event )
 			: FReply::Unhandled();
 	});
 
-	LOG_EVENT_CONTENT( EEventLog::KeyUp, GetKeyName(InKeyboardEvent.GetKey()).ToString(), Reply );
+	LOG_EVENT_CONTENT( EEventLog::KeyUp, GetKeyName(InKeyEvent.GetKey()).ToString(), Reply );
+
+	return Reply.IsEventHandled();
+}
+
+bool FSlateApplication::ProcessAnalogInputEvent(FAnalogInputEvent& InAnalogInputEvent)
+{
+	FReply Reply = FReply::Unhandled();
+
+	LastUserInteractionTime = this->GetCurrentTime();
+
+	LastUserInteractionTimeForThrottling = LastUserInteractionTime;
+
+	// Bubble the key event
+	FWidgetPath EventPath = UserFocusedWidgetPaths[InAnalogInputEvent.GetUserIndex()].ToWidgetPath();
+	InAnalogInputEvent.SetEventPath(EventPath);
+
+	// Switch worlds for widgets in the current path
+	FScopedSwitchWorldHack SwitchWorld(EventPath);
+
+	Reply = FEventRouter::RouteAlongFocusPath(this, FEventRouter::FBubblePolicy(EventPath), InAnalogInputEvent, [](const FArrangedWidget& SomeWidgetGettingEvent, const FAnalogInputEvent& Event)
+	{
+		return (SomeWidgetGettingEvent.Widget->IsEnabled())
+			? SomeWidgetGettingEvent.Widget->OnAnalogValueChanged(SomeWidgetGettingEvent.Geometry, Event)
+			: FReply::Unhandled();
+	});
+
+	LOG_EVENT_CONTENT(EEventLog::AnalogInput, GetKeyName(InAnalogInputEvent.GetKey()).ToString(), Reply);
 
 	return Reply.IsEventHandled();
 }
@@ -3655,7 +3698,7 @@ bool FSlateApplication::ProcessMouseButtonDownEvent( const TSharedPtr< FGenericW
 			// If none of the widgets requested keyboard focus to be set (or set the keyboard focus explicitly), set it to the leaf-most widget under the mouse.
 			// On Mac we prevent the OS from activating the window on mouse down, so we have full control and can activate only if there's nothing draggable under the mouse cursor.
 			const bool bFocusChangedByEventHandler = PreviouslyFocusedWidget != GetKeyboardFocusedWidget();
-			if ((!Reply.GetFocusRecepient().IsValid() || (PLATFORM_MAC && MouseEvent.GetEffectingButton() == EKeys::LeftMouseButton && !DragDetector.DetectDragForWidget.IsValid())) && !bFocusChangedByEventHandler)
+			if ((!Reply.GetUserFocusRecepient().IsValid() || (PLATFORM_MAC && MouseEvent.GetEffectingButton() == EKeys::LeftMouseButton && !DragDetector.DetectDragForWidget.IsValid())) && !bFocusChangedByEventHandler)
 			{
 				// The event handler for OnMouseButton down may have altered the widget hierarchy.
 				// Refresh the previously-cached widget path.
@@ -3669,7 +3712,7 @@ bool FSlateApplication::ProcessMouseButtonDownEvent( const TSharedPtr< FGenericW
 					{
 						bFocusCandidateFound = true;
 						FWidgetPath NewFocusedWidgetPath = WidgetsUnderCursor.GetPathDownTo( CurWidget.Widget );
-						SetKeyboardFocus( NewFocusedWidgetPath, EKeyboardFocusCause::Mouse );
+						SetKeyboardFocus( NewFocusedWidgetPath, EFocusCause::Mouse );
 					}
 				}
 
@@ -4288,107 +4331,99 @@ FKey TranslateControllerButtonToKey( EControllerButtons::Type Button )
 	return Key;
 }
 
-#define CALL_WIDGET_FUNCTION(Event, Function) \
-	if (Event.GetUserIndex() < ARRAY_COUNT(JoystickCaptorWeakPaths) && JoystickCaptorWeakPaths[Event.GetUserIndex()].IsValid()) \
-{ \
-	/* Get the joystick capture target for this user */ \
-	FWidgetPath PathToWidget = JoystickCaptorWeakPaths[Event.GetUserIndex()].ToWidgetPath(); \
-	FArrangedWidget& ArrangedWidget = PathToWidget.Widgets.Last(); \
-	\
-	/* Switch worlds for widgets in the current path */ \
-	FScopedSwitchWorldHack SwitchWorld( PathToWidget ); \
-	\
-	/* Send the message to the widget */ \
-	FReply Reply = ArrangedWidget.Widget->Function( ArrangedWidget.Geometry, Event ).SetHandler( ArrangedWidget.Widget ); \
-	ProcessReply( PathToWidget, Reply, NULL, NULL, Event.GetUserIndex() );\
-}
-
-bool FSlateApplication::OnControllerAnalog( FKey Button, int32 ControllerId, float AnalogValue )
+bool FSlateApplication::AttemptNavigation(const FNavigationEvent& NavigationEvent, const FNavigationReply& NavigationReply, const FArrangedWidget& BoundaryWidget)
 {
-	if ( GetJoystickCaptor( ControllerId ).IsValid() )
+	TSharedPtr<SWidget> FocusedWidget = GetUserFocusedWidget(NavigationEvent.GetUserIndex());
+	TSharedPtr<SWidget> NewFocusedWidget = TSharedPtr<SWidget>();
+	if (FocusedWidget.IsValid())
 	{
-		FControllerEvent ControllerEvent( 
-			Button, 
-			ControllerId,
-			AnalogValue,
-			false
-		);
+		EUINavigation NavigationType = NavigationEvent.GetNavigationType();
 
-		ProcessControllerAnalogValueChangedEvent( ControllerEvent );
-	}
+		// If the FocusedWidget is the Boundary widget we don't need to do any complex work. 
+		if (FocusedWidget == NavigationReply.GetHandler())
+		{
+			if (NavigationReply.GetBoundaryRule() == EUINavigationRule::Explicit)
+			{
+				NewFocusedWidget = NavigationReply.GetFocusRecipient();
+			}
+			else if (NavigationReply.GetBoundaryRule() == EUINavigationRule::Custom)
+			{
+				const FNavigationDelegate& FocusDelegate = NavigationReply.GetFocusDelegate();
+				if (FocusDelegate.IsBound())
+				{
+					NewFocusedWidget = FocusDelegate.Execute(NavigationType);
+				}
+			}
+			// If it's not explicit or custom we do noting.
+			// This is an optimization because wrapping and stopping on the widget that is currently focused can't go anywhere.
+		}
+		else
+		{
+			// Get the controller focus target for this user
+			FWeakWidgetPath FocusedWeakWidgetPath = UserFocusedWidgetPaths[NavigationEvent.GetUserIndex()];
 
-	return true;
-}
+			// Find the next widget
+			if (NavigationType == EUINavigation::Next || NavigationType == EUINavigation::Previous)
+			{
+				// Fond the next widget
+				FWidgetPath NewFocusedWidgetPath = FocusedWeakWidgetPath.ToNextFocusedPath(NavigationType);
 
-bool FSlateApplication::OnControllerAnalog( EControllerButtons::Type Button, int32 ControllerId, float AnalogValue )
-{
-	return OnControllerAnalog( TranslateControllerButtonToKey( Button ), ControllerId, AnalogValue );
-}
+				// Resolve the Widget Path
+				FArrangedWidget& NewFocusedArrangedWidget = NewFocusedWidgetPath.Widgets.Last();
+				NewFocusedWidget = NewFocusedArrangedWidget.Widget;
+			}
+			else
+			{
+				// Resolve the Widget Path
+				FWidgetPath FocusedWidgetPath = FocusedWeakWidgetPath.ToWidgetPath();
+				FArrangedWidget& FocusedArrangedWidget = FocusedWidgetPath.Widgets.Last();
 
-void FSlateApplication::ProcessControllerAnalogValueChangedEvent( FControllerEvent& ControllerEvent )
-{
-	LastUserInteractionTime = this->GetCurrentTime();
-	
-	CALL_WIDGET_FUNCTION(ControllerEvent, OnControllerAnalogValueChanged);
-}
+				// Switch worlds for widgets in the current path 
+				FScopedSwitchWorldHack SwitchWorld(FocusedWidgetPath);
 
-bool FSlateApplication::OnControllerButtonPressed( FKey Button, int32 ControllerId, bool IsRepeat )
-{
-	if ( GetJoystickCaptor( ControllerId ).IsValid() )
+				NewFocusedWidget = HittestGrid->FindNextFocusableWidget(FocusedArrangedWidget, NavigationType, NavigationReply, BoundaryWidget);
+			}
+		}
+	}
+	
+	// Set controller focus if we have a valid widget
+	if (NewFocusedWidget.IsValid())
 	{
-		FControllerEvent ControllerEvent( 
-			Button, 
-			ControllerId,
-			1.0f,
-			IsRepeat
-		);
-
-		ProcessControllerButtonPressedEvent( ControllerEvent );
+		SetUserFocus(NavigationEvent.GetUserIndex(), NewFocusedWidget, EFocusCause::Navigation);
+		return true;
 	}
 
-	return true;
+	return false;
 }
 
-bool FSlateApplication::OnControllerButtonPressed( EControllerButtons::Type Button, int32 ControllerId, bool IsRepeat )
+bool FSlateApplication::OnControllerAnalog( EControllerButtons::Type Button, int32 ControllerId, float AnalogValue )
 {
-	return OnControllerButtonPressed( TranslateControllerButtonToKey( Button ), ControllerId, IsRepeat );
-}
+	FKey Key = TranslateControllerButtonToKey(Button);
+	int32 UserIndex = GetUserIndexForController(ControllerId);
 
-void FSlateApplication::ProcessControllerButtonPressedEvent( FControllerEvent& ControllerEvent )
-{
-	LastUserInteractionTime = this->GetCurrentTime();
-	
-	CALL_WIDGET_FUNCTION(ControllerEvent, OnControllerButtonPressed);
+	FAnalogInputEvent AnalogInputEvent(Key, PlatformApplication->GetModifierKeys(), UserIndex, false, 0, 0, AnalogValue);
+
+	return ProcessAnalogInputEvent(AnalogInputEvent);
 }
 
-bool FSlateApplication::OnControllerButtonReleased( FKey Button, int32 ControllerId, bool IsRepeat )
+bool FSlateApplication::OnControllerButtonPressed( EControllerButtons::Type Button, int32 ControllerId, bool IsRepeat )
 {
-	if ( GetJoystickCaptor( ControllerId ).IsValid() )
-	{
-		FControllerEvent ControllerEvent( 
-			Button, 
-			ControllerId,
-			1.0f,
-			IsRepeat
-		);
+	FKey Key = TranslateControllerButtonToKey(Button);
+	int32 UserIndex = GetUserIndexForController(ControllerId);
 
-		ProcessControllerButtonReleasedEvent( ControllerEvent );
-	}
+	FKeyEvent KeyEvent(Key, PlatformApplication->GetModifierKeys(), UserIndex, IsRepeat, 0, 0);
 
-	return false;
+	return ProcessKeyDownEvent(KeyEvent);
 }
 
 bool FSlateApplication::OnControllerButtonReleased( EControllerButtons::Type Button, int32 ControllerId, bool IsRepeat )
 {
-	return OnControllerButtonReleased( TranslateControllerButtonToKey( Button ), ControllerId, IsRepeat );
-}
-
-void FSlateApplication::ProcessControllerButtonReleasedEvent( FControllerEvent& ControllerEvent )
-{
-	LastUserInteractionTime = this->GetCurrentTime();
+	FKey Key = TranslateControllerButtonToKey(Button);
+	int32 UserIndex = GetUserIndexForController(ControllerId);
 	
-	CALL_WIDGET_FUNCTION(ControllerEvent, OnControllerButtonReleased);
+	FKeyEvent KeyEvent(Key, PlatformApplication->GetModifierKeys(),UserIndex, IsRepeat,  0, 0);
 
+	return ProcessKeyUpEvent(KeyEvent);
 }
 
 bool FSlateApplication::OnTouchGesture( EGestureEvent::Type GestureType, const FVector2D &Delta, const float MouseWheelDelta )
@@ -4491,7 +4526,19 @@ void FSlateApplication::ProcessMotionDetectedEvent( FMotionEvent& MotionEvent )
 {
 	LastUserInteractionTime = this->GetCurrentTime();
 	
-	CALL_WIDGET_FUNCTION(MotionEvent, OnMotionDetected);
+	if (MotionEvent.GetUserIndex() < ARRAY_COUNT(UserFocusedWidgetPaths) && UserFocusedWidgetPaths[MotionEvent.GetUserIndex()].IsValid())
+	{
+		/* Get the controller focus target for this user */
+		FWidgetPath PathToWidget = UserFocusedWidgetPaths[MotionEvent.GetUserIndex()].ToWidgetPath();
+		FArrangedWidget& ArrangedWidget = PathToWidget.Widgets.Last();
+
+		/* Switch worlds for widgets in the current path */
+		FScopedSwitchWorldHack SwitchWorld(PathToWidget);
+
+		/* Send the message to the widget */
+		FReply Reply = ArrangedWidget.Widget->OnMotionDetected(ArrangedWidget.Geometry, MotionEvent).SetHandler(ArrangedWidget.Widget);
+		ProcessReply(PathToWidget, Reply, NULL, NULL, MotionEvent.GetUserIndex());
+	}
 }
 
 bool FSlateApplication::OnSizeChanged( const TSharedRef< FGenericWindow >& PlatformWindow, const int32 Width, const int32 Height, bool bWasMinimized )
@@ -4638,7 +4685,7 @@ bool FSlateApplication::ProcessWindowActivatedEvent( const FWindowActivateEvent&
 			// A Slate window was activated
 			bSlateWindowActive = true;
 
-			if ( ActivateEvent.GetAffectedWindow()->IsFocusedInitially() && ActivateEvent.GetAffectedWindow()->SupportsKeyboardFocus() )
+			if (ActivateEvent.GetAffectedWindow()->IsFocusedInitially() && ActivateEvent.GetAffectedWindow()->SupportsKeyboardFocus() )
 			{
 				// Set keyboard focus on the window being activated.
 				{
@@ -4647,11 +4694,11 @@ bool FSlateApplication::ProcessWindowActivatedEvent( const FWindowActivateEvent&
 
 					if( ActivateEvent.GetActivationType() == FWindowActivateEvent::EA_ActivateByMouse )
 					{
-						SetKeyboardFocus(PathToWindowBeingActivated, EKeyboardFocusCause::Mouse);
+						SetKeyboardFocus(PathToWindowBeingActivated, EFocusCause::Mouse);
 					}
 					else
 					{
-						SetKeyboardFocus(PathToWindowBeingActivated, EKeyboardFocusCause::WindowActivate);
+						SetKeyboardFocus(PathToWindowBeingActivated, EFocusCause::WindowActivate);
 					}
 				}
 
@@ -4732,7 +4779,7 @@ void FSlateApplication::ProcessApplicationActivationEvent( bool InAppActivated )
 		bSlateWindowActive = false;
 
 		// Clear keyboard focus when the app is deactivated
-		ClearKeyboardFocus(EKeyboardFocusCause::OtherWidgetLostFocus);
+		ClearKeyboardFocus(EFocusCause::OtherWidgetLostFocus);
 
 		// If we have a slate-only drag-drop occurring, stop the drag drop.
 		if ( IsDragDropping() && !DragDropContent->IsExternalOperation() )
diff --git a/Engine/Source/Runtime/Slate/Private/Framework/Commands/UICommandList.cpp b/Engine/Source/Runtime/Slate/Private/Framework/Commands/UICommandList.cpp
--- a/Engine/Source/Runtime/Slate/Private/Framework/Commands/UICommandList.cpp
+++ b/Engine/Source/Runtime/Slate/Private/Framework/Commands/UICommandList.cpp
@@ -130,9 +130,9 @@ bool FUICommandList::IsChecked( const TSharedRef< const FUICommandInfo > InUICom
 }
 
 
-bool FUICommandList::ProcessCommandBindings( const FKeyboardEvent& InKeyboardEvent ) const 
+bool FUICommandList::ProcessCommandBindings( const FKeyEvent& InKeyEvent ) const 
 {
-	return ConditionalProcessCommandBindings( InKeyboardEvent.GetKey(), InKeyboardEvent.IsControlDown(), InKeyboardEvent.IsAltDown(), InKeyboardEvent.IsShiftDown(), InKeyboardEvent.IsCommandDown(), InKeyboardEvent.IsRepeat() );
+	return ConditionalProcessCommandBindings( InKeyEvent.GetKey(), InKeyEvent.IsControlDown(), InKeyEvent.IsAltDown(), InKeyEvent.IsShiftDown(), InKeyEvent.IsCommandDown(), InKeyEvent.IsRepeat() );
 }
 
 
diff --git a/Engine/Source/Runtime/Slate/Private/Framework/Docking/SDockingTabStack.cpp b/Engine/Source/Runtime/Slate/Private/Framework/Docking/SDockingTabStack.cpp
--- a/Engine/Source/Runtime/Slate/Private/Framework/Docking/SDockingTabStack.cpp
+++ b/Engine/Source/Runtime/Slate/Private/Framework/Docking/SDockingTabStack.cpp
@@ -376,7 +376,7 @@ FReply SDockingTabStack::OnDrop( const FGeometry& MyGeometry, const FDragDropEve
 	return FReply::Unhandled();
 }
 
-void SDockingTabStack::OnKeyboardFocusChanging( const FWeakWidgetPath& PreviousFocusPath, const FWidgetPath& NewWidgetPath )
+void SDockingTabStack::OnFocusChanging( const FWeakWidgetPath& PreviousFocusPath, const FWidgetPath& NewWidgetPath )
 {
 	
 	const TSharedPtr<SDockTab> ForegroundTab = TabWell->GetForegroundTab();
diff --git a/Engine/Source/Runtime/Slate/Private/Framework/Docking/SDockingTabStack.h b/Engine/Source/Runtime/Slate/Private/Framework/Docking/SDockingTabStack.h
--- a/Engine/Source/Runtime/Slate/Private/Framework/Docking/SDockingTabStack.h
+++ b/Engine/Source/Runtime/Slate/Private/Framework/Docking/SDockingTabStack.h
@@ -95,7 +95,7 @@ public:
 	
 	virtual FReply OnDrop( const FGeometry& MyGeometry, const FDragDropEvent& DragDropEvent ) override;
 
-	virtual void OnKeyboardFocusChanging( const FWeakWidgetPath& PreviousFocusPath, const FWidgetPath& NewWidgetPath ) override;
+	virtual void OnFocusChanging( const FWeakWidgetPath& PreviousFocusPath, const FWidgetPath& NewWidgetPath ) override;
 
 	virtual FReply OnMouseButtonDown( const FGeometry& MyGeometry, const FPointerEvent& MouseEvent ) override;
 
diff --git a/Engine/Source/Runtime/Slate/Private/Framework/MultiBox/Mac/MacMenu.cpp b/Engine/Source/Runtime/Slate/Private/Framework/MultiBox/Mac/MacMenu.cpp
--- a/Engine/Source/Runtime/Slate/Private/Framework/MultiBox/Mac/MacMenu.cpp
+++ b/Engine/Source/Runtime/Slate/Private/Framework/MultiBox/Mac/MacMenu.cpp
@@ -74,7 +74,7 @@ static TMap<FMacMenu*, TSharedPtr<TArray<FMacMenuItemState>>> GCachedMenuState;
 	FPlatformMisc::bChachedMacMenuStateNeedsUpdate = true;
 	
 	GameThreadCall(^{
-		FSlateApplication::Get().ClearKeyboardFocus( EKeyboardFocusCause::WindowActivate );
+		FSlateApplication::Get().ClearKeyboardFocus( EFocusCause::WindowActivate );
 	}, @[ NSDefaultRunLoopMode ], false);
 
 }
diff --git a/Engine/Source/Runtime/Slate/Private/Framework/MultiBox/MultiBox.cpp b/Engine/Source/Runtime/Slate/Private/Framework/MultiBox/MultiBox.cpp
--- a/Engine/Source/Runtime/Slate/Private/Framework/MultiBox/MultiBox.cpp
+++ b/Engine/Source/Runtime/Slate/Private/Framework/MultiBox/MultiBox.cpp
@@ -998,7 +998,7 @@ bool SMultiBoxWidget::SupportsKeyboardFocus() const
 	return true;
 }
 
-FReply SMultiBoxWidget::FocusNextWidget( EFocusMoveDirection::Type MoveDirection )
+FReply SMultiBoxWidget::FocusNextWidget(EUINavigation NavigationType)
 {
 	TSharedPtr<SWidget> FocusWidget = FSlateApplication::Get().GetKeyboardFocusedWidget();
 	if(FocusWidget.IsValid())
@@ -1006,35 +1006,35 @@ FReply SMultiBoxWidget::FocusNextWidget( EFocusMoveDirection::Type MoveDirection
 		FWidgetPath FocusPath;
 		FSlateApplication::Get().GeneratePathToWidgetUnchecked( FocusWidget.ToSharedRef(), FocusPath );
 		FWeakWidgetPath WeakFocusPath = FocusPath;
-		FWidgetPath NextFocusPath = WeakFocusPath.ToNextFocusedPath( MoveDirection );
+		FWidgetPath NextFocusPath = WeakFocusPath.ToNextFocusedPath(NavigationType);
 		if ( NextFocusPath.Widgets.Num() > 0 )
 		{
-			return FReply::Handled().SetKeyboardFocus( NextFocusPath.Widgets.Last().Widget, EKeyboardFocusCause::Keyboard );
+			return FReply::Handled().SetUserFocus(NextFocusPath.Widgets.Last().Widget, EFocusCause::Navigation);
 		}
 	}
 
 	return FReply::Unhandled();
 }
 
-FReply SMultiBoxWidget::OnKeyboardFocusReceived( const FGeometry& MyGeometry, const FKeyboardFocusEvent& InKeyboardFocusEvent )
+FReply SMultiBoxWidget::OnFocusReceived( const FGeometry& MyGeometry, const FFocusEvent& InFocusEvent )
 {
-	if(InKeyboardFocusEvent.GetCause() == EKeyboardFocusCause::Keyboard)
+	if (InFocusEvent.GetCause() == EFocusCause::Navigation)
 	{
 		// forward focus to children
-		return FocusNextWidget( EFocusMoveDirection::Next );
+		return FocusNextWidget( EUINavigation::Next );
 	}
 
 	return FReply::Unhandled();
 }
 
-FReply SMultiBoxWidget::OnKeyDown( const FGeometry& MyGeometry, const FKeyboardEvent& KeyboardEvent )
+FReply SMultiBoxWidget::OnKeyDown( const FGeometry& MyGeometry, const FKeyEvent& KeyEvent )
 {
-	SCompoundWidget::OnKeyDown(MyGeometry, KeyboardEvent);
+	SCompoundWidget::OnKeyDown(MyGeometry, KeyEvent);
 
 	// allow use of up and down keys to transfer focus/hover state
-	if(KeyboardEvent.GetKey() == EKeys::Up || KeyboardEvent.GetKey() == EKeys::Down)
+	if(KeyEvent.GetKey() == EKeys::Up || KeyEvent.GetKey() == EKeys::Down)
 	{
-		return FocusNextWidget( EFocusMoveDirection::Next );
+		return FocusNextWidget(EUINavigation::Next);
 	}
 
 	return FReply::Unhandled();
diff --git a/Engine/Source/Runtime/Slate/Private/Framework/MultiBox/SMenuEntryBlock.cpp b/Engine/Source/Runtime/Slate/Private/Framework/MultiBox/SMenuEntryBlock.cpp
--- a/Engine/Source/Runtime/Slate/Private/Framework/MultiBox/SMenuEntryBlock.cpp
+++ b/Engine/Source/Runtime/Slate/Private/Framework/MultiBox/SMenuEntryBlock.cpp
@@ -1034,17 +1034,17 @@ void SMenuEntryBlock::OnMouseLeave( const FPointerEvent& MouseEvent )
 	CancelPendingSubMenu();
 }
 
-FReply SMenuEntryBlock::OnKeyDown( const FGeometry& MyGeometry, const FKeyboardEvent& KeyboardEvent )
+FReply SMenuEntryBlock::OnKeyDown( const FGeometry& MyGeometry, const FKeyEvent& KeyEvent )
 {
-	SMultiBlockBaseWidget::OnKeyDown(MyGeometry, KeyboardEvent);
+	SMultiBlockBaseWidget::OnKeyDown(MyGeometry, KeyEvent);
 
 	// allow use of up and down keys to transfer focus
-	if(KeyboardEvent.GetKey() == EKeys::Up || KeyboardEvent.GetKey() == EKeys::Down)
+	if(KeyEvent.GetKey() == EKeys::Up || KeyEvent.GetKey() == EKeys::Down)
 	{
 		// find the next widget to focus
-		EFocusMoveDirection::Type MoveDirection = ( KeyboardEvent.GetKey() == EKeys::Up )
-			? EFocusMoveDirection::Previous
-			: EFocusMoveDirection::Next;
+		EUINavigation MoveDirection = (KeyEvent.GetKey() == EKeys::Up)
+			? EUINavigation::Previous
+			: EUINavigation::Next;
 		 
 		return SMultiBoxWidget::FocusNextWidget(MoveDirection);
 	}
diff --git a/Engine/Source/Runtime/Slate/Private/Framework/MultiBox/SMenuEntryBlock.h b/Engine/Source/Runtime/Slate/Private/Framework/MultiBox/SMenuEntryBlock.h
--- a/Engine/Source/Runtime/Slate/Private/Framework/MultiBox/SMenuEntryBlock.h
+++ b/Engine/Source/Runtime/Slate/Private/Framework/MultiBox/SMenuEntryBlock.h
@@ -232,7 +232,7 @@ protected:
 	// SWidget interface
 	virtual void OnMouseEnter( const FGeometry& MyGeometry, const FPointerEvent& MouseEvent ) override;
 	virtual void OnMouseLeave( const FPointerEvent& MouseEvent ) override;
-	virtual FReply OnKeyDown( const FGeometry& MyGeometry, const FKeyboardEvent& KeyboardEvent ) override;
+	virtual FReply OnKeyDown( const FGeometry& MyGeometry, const FKeyEvent& KeyEvent ) override;
 	virtual void Tick( const FGeometry& AllottedGeometry, const double InCurrentTime, const float InDeltaTime ) override;
 	// End of SWidget interface
 
diff --git a/Engine/Source/Runtime/Slate/Private/Framework/Text/TextEditHelper.cpp b/Engine/Source/Runtime/Slate/Private/Framework/Text/TextEditHelper.cpp
--- a/Engine/Source/Runtime/Slate/Private/Framework/Text/TextEditHelper.cpp
+++ b/Engine/Source/Runtime/Slate/Private/Framework/Text/TextEditHelper.cpp
@@ -92,58 +92,58 @@ FReply FTextEditHelper::OnKeyChar( const FCharacterEvent& InCharacterEvent, cons
 	return Reply;
 }
 
-FReply FTextEditHelper::OnKeyDown( const FKeyboardEvent& InKeyboardEvent, const TSharedRef< ITextEditorWidget >& TextEditor )
+FReply FTextEditHelper::OnKeyDown( const FKeyEvent& InKeyEvent, const TSharedRef< ITextEditorWidget >& TextEditor )
 {
-	const FKey Key = InKeyboardEvent.GetKey();
+	const FKey Key = InKeyEvent.GetKey();
 
 	if( Key == EKeys::Left )
 	{
 		return TextEditor->MoveCursor( FMoveCursor::Cardinal(
 			// Ctrl moves a whole word instead of one character.	
-			InKeyboardEvent.IsControlDown( ) ? ECursorMoveGranularity::Word : ECursorMoveGranularity::Character,
+			InKeyEvent.IsControlDown( ) ? ECursorMoveGranularity::Word : ECursorMoveGranularity::Character,
 			// Move left
 			FIntPoint( -1, 0 ),
 			// Shift selects text.	
-			InKeyboardEvent.IsShiftDown() ? ECursorAction::SelectText : ECursorAction::MoveCursor
+			InKeyEvent.IsShiftDown() ? ECursorAction::SelectText : ECursorAction::MoveCursor
 		));
 	}
 	else if( Key == EKeys::Right )
 	{
 		return TextEditor->MoveCursor( FMoveCursor::Cardinal(
 			// Ctrl moves a whole word instead of one character.	
-			InKeyboardEvent.IsControlDown( ) ? ECursorMoveGranularity::Word : ECursorMoveGranularity::Character,
+			InKeyEvent.IsControlDown( ) ? ECursorMoveGranularity::Word : ECursorMoveGranularity::Character,
 			// Move right
 			FIntPoint( +1, 0 ),
 			// Shift selects text.	
-			InKeyboardEvent.IsShiftDown() ? ECursorAction::SelectText : ECursorAction::MoveCursor
+			InKeyEvent.IsShiftDown() ? ECursorAction::SelectText : ECursorAction::MoveCursor
 		));
 	}
 	else if ( Key == EKeys::Up )
 	{
 		return TextEditor->MoveCursor( FMoveCursor::Cardinal(
-			InKeyboardEvent.IsControlDown( ) ? ECursorMoveGranularity::Word : ECursorMoveGranularity::Character,
+			InKeyEvent.IsControlDown( ) ? ECursorMoveGranularity::Word : ECursorMoveGranularity::Character,
 			// Move up
 			FIntPoint( 0, -1 ),
 			// Shift selects text.	
-			InKeyboardEvent.IsShiftDown() ? ECursorAction::SelectText : ECursorAction::MoveCursor
+			InKeyEvent.IsShiftDown() ? ECursorAction::SelectText : ECursorAction::MoveCursor
 		));
 	}
 	else if ( Key == EKeys::Down )
 	{
 		return TextEditor->MoveCursor( FMoveCursor::Cardinal(
-			InKeyboardEvent.IsControlDown( ) ? ECursorMoveGranularity::Word : ECursorMoveGranularity::Character,
+			InKeyEvent.IsControlDown( ) ? ECursorMoveGranularity::Word : ECursorMoveGranularity::Character,
 			// Move down
 			FIntPoint( 0, +1 ),
 			// Shift selects text.	
-			InKeyboardEvent.IsShiftDown() ? ECursorAction::SelectText : ECursorAction::MoveCursor
+			InKeyEvent.IsShiftDown() ? ECursorAction::SelectText : ECursorAction::MoveCursor
 		));
 	}
 	else if( Key == EKeys::Home )
 	{
 		// Go to the beginning of the document; select text if Shift is down.
 		TextEditor->JumpTo(
-			(InKeyboardEvent.IsControlDown() ) ? ETextLocation::BeginningOfDocument : ETextLocation::BeginningOfLine,
-			(InKeyboardEvent.IsShiftDown()) ? ECursorAction::SelectText : ECursorAction::MoveCursor );
+			(InKeyEvent.IsControlDown() ) ? ETextLocation::BeginningOfDocument : ETextLocation::BeginningOfLine,
+			(InKeyEvent.IsShiftDown()) ? ECursorAction::SelectText : ECursorAction::MoveCursor );
 
 		return FReply::Handled();
 	}
@@ -151,8 +151,8 @@ FReply FTextEditHelper::OnKeyDown( const FKeyboardEvent& InKeyboardEvent, const
 	{
 		// Go to the end of the document; select text if Shift is down.
 		TextEditor->JumpTo(
-			(InKeyboardEvent.IsControlDown() ) ? ETextLocation::EndOfDocument : ETextLocation::EndOfLine,
-			(InKeyboardEvent.IsShiftDown()) ? ECursorAction::SelectText : ECursorAction::MoveCursor );
+			(InKeyEvent.IsControlDown() ) ? ETextLocation::EndOfDocument : ETextLocation::EndOfLine,
+			(InKeyEvent.IsShiftDown()) ? ECursorAction::SelectText : ECursorAction::MoveCursor );
 
 		return FReply::Handled();
 	}
@@ -168,7 +168,7 @@ FReply FTextEditHelper::OnKeyDown( const FKeyboardEvent& InKeyboardEvent, const
 	{
 		// @Todo: Slate keybindings support more than one set of keys. 
 		// Delete to next word boundary (Ctrl+Delete)
-		if (InKeyboardEvent.IsControlDown() && !InKeyboardEvent.IsAltDown() && !InKeyboardEvent.IsShiftDown())
+		if (InKeyEvent.IsControlDown() && !InKeyEvent.IsAltDown() && !InKeyEvent.IsShiftDown())
 		{
 			TextEditor->MoveCursor( FMoveCursor::Cardinal(
 				ECursorMoveGranularity::Word, 
@@ -193,7 +193,7 @@ FReply FTextEditHelper::OnKeyDown( const FKeyboardEvent& InKeyboardEvent, const
 
 	// @Todo: Slate keybindings support more than one set of keys. 
 	//Alternate key for cut (Shift+Delete)
-	else if( Key == EKeys::Delete && InKeyboardEvent.IsShiftDown() && TextEditor->CanExecuteCut() )
+	else if( Key == EKeys::Delete && InKeyEvent.IsShiftDown() && TextEditor->CanExecuteCut() )
 	{
 		// Cut text to clipboard
 		TextEditor->CutSelectedTextToClipboard();
@@ -203,7 +203,7 @@ FReply FTextEditHelper::OnKeyDown( const FKeyboardEvent& InKeyboardEvent, const
 
 	// @Todo: Slate keybindings support more than one set of keys. 
 	// Alternate key for copy (Ctrl+Insert) 
-	else if( Key == EKeys::Insert && InKeyboardEvent.IsControlDown() && TextEditor->CanExecuteCopy() ) 
+	else if( Key == EKeys::Insert && InKeyEvent.IsControlDown() && TextEditor->CanExecuteCopy() ) 
 	{
 		// Copy text to clipboard
 		TextEditor->CopySelectedTextToClipboard();
@@ -214,7 +214,7 @@ FReply FTextEditHelper::OnKeyDown( const FKeyboardEvent& InKeyboardEvent, const
 
 	// @Todo: Slate keybindings support more than one set of keys. 
 	// Alternate key for paste (Shift+Insert) 
-	else if( Key == EKeys::Insert && InKeyboardEvent.IsShiftDown() && TextEditor->CanExecutePaste() )
+	else if( Key == EKeys::Insert && InKeyEvent.IsShiftDown() && TextEditor->CanExecutePaste() )
 	{
 		// Paste text from clipboard
 		TextEditor->PasteTextFromClipboard();
@@ -224,7 +224,7 @@ FReply FTextEditHelper::OnKeyDown( const FKeyboardEvent& InKeyboardEvent, const
 
 	// @Todo: Slate keybindings support more than one set of keys. 
 	//Alternate key for undo (Alt+Backspace)
-	else if( Key == EKeys::BackSpace && InKeyboardEvent.IsAltDown() && !InKeyboardEvent.IsShiftDown() && TextEditor->CanExecuteUndo() )
+	else if( Key == EKeys::BackSpace && InKeyEvent.IsAltDown() && !InKeyEvent.IsShiftDown() && TextEditor->CanExecuteUndo() )
 	{
 		// Undo
 		TextEditor->Undo();
@@ -234,7 +234,7 @@ FReply FTextEditHelper::OnKeyDown( const FKeyboardEvent& InKeyboardEvent, const
 
 	// @Todo: Slate keybindings support more than one set of keys. 
 	// Delete to previous word boundary (Ctrl+Backspace)
-	else if( Key == EKeys::BackSpace && InKeyboardEvent.IsControlDown() && !InKeyboardEvent.IsAltDown() && !InKeyboardEvent.IsShiftDown() && !TextEditor->GetIsReadOnly() )
+	else if( Key == EKeys::BackSpace && InKeyEvent.IsControlDown() && !InKeyEvent.IsAltDown() && !InKeyEvent.IsShiftDown() && !TextEditor->GetIsReadOnly() )
 	{
 		FScopedTextTransaction TextTransaction(TextEditor);
 
@@ -251,16 +251,16 @@ FReply FTextEditHelper::OnKeyDown( const FKeyboardEvent& InKeyboardEvent, const
 
 
 	// Ctrl+Y (or Ctrl+Shift+Z, or Alt+Shift+Backspace) to redo
-	else if( !TextEditor->GetIsReadOnly() && ( ( Key == EKeys::Y && InKeyboardEvent.IsControlDown() ) ||
-		( Key == EKeys::Z && InKeyboardEvent.IsControlDown() && InKeyboardEvent.IsShiftDown() ) ||
-		( Key == EKeys::BackSpace && InKeyboardEvent.IsAltDown() && InKeyboardEvent.IsShiftDown() ) ) )
+	else if( !TextEditor->GetIsReadOnly() && ( ( Key == EKeys::Y && InKeyEvent.IsControlDown() ) ||
+		( Key == EKeys::Z && InKeyEvent.IsControlDown() && InKeyEvent.IsShiftDown() ) ||
+		( Key == EKeys::BackSpace && InKeyEvent.IsAltDown() && InKeyEvent.IsShiftDown() ) ) )
 	{
 		// Redo
 		TextEditor->Redo();
 		
 		return FReply::Handled();
 	}
-	else if( !InKeyboardEvent.IsAltDown() && !InKeyboardEvent.IsControlDown() && InKeyboardEvent.GetKey() != EKeys::Tab && InKeyboardEvent.GetCharacter() != 0 )
+	else if( !InKeyEvent.IsAltDown() && !InKeyEvent.IsControlDown() && InKeyEvent.GetKey() != EKeys::Tab && InKeyEvent.GetCharacter() != 0 )
 	{
 		// Shift and a character was pressed or a single character was pressed.  We will type something in an upcoming OnKeyChar event.  
 		// Absorb this event so it is not bubbled and handled by other widgets that could have something bound to the key press.
@@ -328,7 +328,7 @@ FReply FTextEditHelper::OnMouseButtonDown( const FGeometry& MyGeometry, const FP
 			// Right clicking to summon context menu, but we'll do that on mouse-up.
 			Reply = FReply::Handled();
 			Reply.CaptureMouse( TextEditor->GetWidget() );
-			Reply.SetKeyboardFocus( TextEditor->GetWidget(), EKeyboardFocusCause::Mouse );
+			Reply.SetUserFocus(TextEditor->GetWidget(), EFocusCause::Mouse);
 		}
 	}
 
