Refactor to unify keyboard and controller input and focus.

diff --git a/Engine/Source/Runtime/Slate/Private/Framework/Text/TextEditHelper.h b/Engine/Source/Runtime/Slate/Private/Framework/Text/TextEditHelper.h
--- a/Engine/Source/Runtime/Slate/Private/Framework/Text/TextEditHelper.h
+++ b/Engine/Source/Runtime/Slate/Private/Framework/Text/TextEditHelper.h
@@ -65,7 +65,7 @@ public:
 	static FReply OnKeyChar( const FCharacterEvent& InCharacterEvent, const TSharedRef< class ITextEditorWidget >& TextEditor );
 
 	/** Someone pressed a key, translate the input event into a call to a ITextEditorWidget action */
-	static FReply OnKeyDown( const FKeyboardEvent& InKeyboardEvent, const TSharedRef< ITextEditorWidget >& TextEditor );
+	static FReply OnKeyDown( const FKeyEvent& InKeyEvent, const TSharedRef< ITextEditorWidget >& TextEditor );
 
 	static FReply OnMouseButtonDown( const FGeometry& MyGeometry, const FPointerEvent& InMouseEvent, const TSharedRef< ITextEditorWidget >& TextEditor );
 
diff --git a/Engine/Source/Runtime/Slate/Private/Widgets/Input/SButton.cpp b/Engine/Source/Runtime/Slate/Private/Widgets/Input/SButton.cpp
--- a/Engine/Source/Runtime/Slate/Private/Widgets/Input/SButton.cpp
+++ b/Engine/Source/Runtime/Slate/Private/Widgets/Input/SButton.cpp
@@ -55,6 +55,7 @@ void SButton::Construct( const FArguments& InArgs )
 
 	ClickMethod = InArgs._ClickMethod;
 	TouchMethod = InArgs._TouchMethod;
+	PressMethod = InArgs._PressMethod;
 
 	HoveredSound = InArgs._HoveredSoundOverride.Get(Style->HoveredSlateSound);
 	PressedSound = InArgs._PressedSoundOverride.Get(Style->PressedSlateSound);
@@ -92,32 +93,73 @@ bool SButton::SupportsKeyboardFocus() const
 	return bIsFocusable;
 }
 
-FReply SButton::OnKeyDown( const FGeometry& MyGeometry, const FKeyboardEvent& InKeyboardEvent )
+void SButton::OnFocusLost( const FFocusEvent& InFocusEvent )
 {
-	//see if we pressed the Enter or Spacebar keys
-	if(InKeyboardEvent.GetKey() == EKeys::Enter || InKeyboardEvent.GetKey() == EKeys::SpaceBar)
+	bIsPressed = false;
+}
+
+FReply SButton::OnKeyDown( const FGeometry& MyGeometry, const FKeyEvent& InKeyEvent )
+{
+	FReply Reply = FReply::Unhandled();
+	if (IsEnabled() && (InKeyEvent.GetKey() == EKeys::Enter || InKeyEvent.GetKey() == EKeys::SpaceBar || InKeyEvent.GetKey() == EKeys::Gamepad_FaceButton_Bottom))
 	{
-		if(IsEnabled())
+		bIsPressed = true;
+
+		PlayPressedSound();
+
+		if (PressMethod == EButtonPressMethod::ButtonPress)
 		{
-			PlayPressedSound();
-		}		
+			//execute our "OnClicked" delegate, and get the reply
+			Reply = OnClicked.IsBound() ? OnClicked.Execute() : FReply::Handled();
 
-		//execute our "OnClicked" delegate, if we have one
-		if(OnClicked.IsBound() == true)
+			//You should ALWAYS handle the OnClicked event.
+			ensure(Reply.IsEventHandled() == true);
+		}
+		else
 		{
-			return OnClicked.Execute();
+			Reply = FReply::Handled();
 		}
 	}
+	else
+	{
+		Reply = SBorder::OnKeyDown(MyGeometry, InKeyEvent);
+	}
 
-	//if it was some other button, ignore it
-	return FReply::Unhandled();
+	//return the constructed reply
+	return Reply;
 }
 
+FReply SButton::OnKeyUp(const FGeometry& MyGeometry, const FKeyEvent& InKeyEvent)
+{
+	FReply Reply = FReply::Unhandled();
+
+	if (IsEnabled() && (InKeyEvent.GetKey() == EKeys::Enter || InKeyEvent.GetKey() == EKeys::SpaceBar || InKeyEvent.GetKey() == EKeys::Gamepad_FaceButton_Bottom))
+	{
+		bool bWasPressed = bIsPressed;
+		bIsPressed = false;
+		//@Todo Slate: This should check focus, however we don't have that API yet, will be easier when focus is unified.
+		if (PressMethod == EButtonPressMethod::ButtonRelease || (PressMethod == EButtonPressMethod::DownAndUp && bWasPressed))
+		{
+			//execute our "OnClicked" delegate, and get the reply
+			Reply = OnClicked.IsBound() ? OnClicked.Execute() : FReply::Handled();
+
+			//You should ALWAYS handle the OnClicked event.
+			ensure(Reply.IsEventHandled() == true);
+		}
+		else
+		{
+			Reply = FReply::Handled();
+		}
+	}
+
+	//return the constructed reply
+	return Reply;
+}
 
 FReply SButton::OnMouseButtonDown( const FGeometry& MyGeometry, const FPointerEvent& MouseEvent )
 {
 	FReply Reply = FReply::Unhandled();
-	if ( MouseEvent.GetEffectingButton() == EKeys::LeftMouseButton || MouseEvent.IsTouchEvent() )
+	if (IsEnabled() && (MouseEvent.GetEffectingButton() == EKeys::LeftMouseButton || MouseEvent.IsTouchEvent()))
 	{
 		bIsPressed = true;
 
@@ -159,7 +201,7 @@ FReply SButton::OnMouseButtonDoubleClick( const FGeometry& InMyGeometry, const F
 FReply SButton::OnMouseButtonUp( const FGeometry& MyGeometry, const FPointerEvent& MouseEvent )
 {
 	FReply Reply = FReply::Unhandled();
-	if ( MouseEvent.GetEffectingButton() == EKeys::LeftMouseButton || MouseEvent.IsTouchEvent() )
+	if (IsEnabled() && (MouseEvent.GetEffectingButton() == EKeys::LeftMouseButton || MouseEvent.IsTouchEvent()))
 	{
 		bIsPressed = false;
 
diff --git a/Engine/Source/Runtime/Slate/Private/Widgets/Input/SCheckBox.cpp b/Engine/Source/Runtime/Slate/Private/Widgets/Input/SCheckBox.cpp
--- a/Engine/Source/Runtime/Slate/Private/Widgets/Input/SCheckBox.cpp
+++ b/Engine/Source/Runtime/Slate/Private/Widgets/Input/SCheckBox.cpp
@@ -112,9 +112,9 @@ bool SCheckBox::SupportsKeyboardFocus() const
 	return bIsFocusable;
 }
 
-FReply SCheckBox::OnKeyUp( const FGeometry& MyGeometry, const FKeyboardEvent& InKeyboardEvent )
+FReply SCheckBox::OnKeyUp( const FGeometry& MyGeometry, const FKeyEvent& InKeyEvent )
 {
-	if ( InKeyboardEvent.GetKey() == EKeys::SpaceBar )
+	if ( InKeyEvent.GetKey() == EKeys::SpaceBar )
 	{
 		ToggleCheckedState();
 		const TAttribute<ESlateCheckBoxState::Type>& State = IsCheckboxChecked.Get();
@@ -161,12 +161,12 @@ FReply SCheckBox::OnMouseButtonDown( const FGeometry& MyGeometry, const FPointer
 			}
 
 			// Set focus to this button, but don't capture the mouse
-			return FReply::Handled().SetKeyboardFocus( AsShared(), EKeyboardFocusCause::Mouse );
+			return FReply::Handled().SetUserFocus(AsShared(), EFocusCause::Mouse);
 		}
 		else
 		{
 			// Capture the mouse, and also set focus to this button
-			return FReply::Handled().CaptureMouse( AsShared() ).SetKeyboardFocus( AsShared(), EKeyboardFocusCause::Mouse );
+			return FReply::Handled().CaptureMouse(AsShared()).SetUserFocus(AsShared(), EFocusCause::Mouse);
 		}
 	}
 	else if ( MouseEvent.GetEffectingButton() == EKeys::RightMouseButton && OnGetMenuContent.IsBound() )
diff --git a/Engine/Source/Runtime/Slate/Private/Widgets/Input/SComboButton.cpp b/Engine/Source/Runtime/Slate/Private/Widgets/Input/SComboButton.cpp
--- a/Engine/Source/Runtime/Slate/Private/Widgets/Input/SComboButton.cpp
+++ b/Engine/Source/Runtime/Slate/Private/Widgets/Input/SComboButton.cpp
@@ -109,7 +109,7 @@ FReply SComboButton::OnButtonClicked()
 
 		if (WidgetToFocus.IsValid())
 		{
-			ButtonClickedReply.SetKeyboardFocus(WidgetToFocus.ToSharedRef(), EKeyboardFocusCause::SetDirectly);
+			ButtonClickedReply.SetUserFocus(WidgetToFocus.ToSharedRef(), EFocusCause::SetDirectly);
 		}
 	}
 
diff --git a/Engine/Source/Runtime/Slate/Private/Widgets/Input/SEditableText.cpp b/Engine/Source/Runtime/Slate/Private/Widgets/Input/SEditableText.cpp
--- a/Engine/Source/Runtime/Slate/Private/Widgets/Input/SEditableText.cpp
+++ b/Engine/Source/Runtime/Slate/Private/Widgets/Input/SEditableText.cpp
@@ -979,7 +979,7 @@ void SEditableText::OnEnter()
 
 		if ( ClearKeyboardFocusOnCommit.Get() )
 		{
-			FSlateApplication::Get().ClearKeyboardFocus( EKeyboardFocusCause::Cleared );
+			FSlateApplication::Get().ClearKeyboardFocus( EFocusCause::Cleared );
 		}
 		else
 		{
@@ -1502,14 +1502,14 @@ bool SEditableText::SupportsKeyboardFocus() const
 	return true;
 }
 
-FReply SEditableText::OnKeyboardFocusReceived( const FGeometry& MyGeometry, const FKeyboardFocusEvent& InKeyboardFocusEvent )
+FReply SEditableText::OnFocusReceived( const FGeometry& MyGeometry, const FFocusEvent& InFocusEvent )
 {
 	// Skip the focus received code if it's due to the context menu closing
 	if ( !ContextMenuWindow.IsValid() )
 	{
 		// Don't unselect text when focus is set because another widget lost focus, as that may be in response to a
 		// dismissed context menu where the user clicked an option to select text.
-		if( InKeyboardFocusEvent.GetCause() != EKeyboardFocusCause::OtherWidgetLostFocus )
+		if( InFocusEvent.GetCause() != EFocusCause::OtherWidgetLostFocus )
 		{
 			// Deselect text, but override the last selection interaction time, so that the user doesn't
 			// see a selection transition animation when keyboard focus is received.
@@ -1526,11 +1526,11 @@ FReply SEditableText::OnKeyboardFocusReceived( const FGeometry& MyGeometry, cons
 		LoadText();
 
 		// For mouse-based focus, we'll handle the cursor positioning in the actual mouse down/up events
-		if( InKeyboardFocusEvent.GetCause() != EKeyboardFocusCause::Mouse )
+		if( InFocusEvent.GetCause() != EFocusCause::Mouse )
 		{
 			// Don't move the cursor if we gained focus because another widget lost focus.  It may
 			// have been a dismissed context menu where the user is expecting the selection to be retained
-			if( InKeyboardFocusEvent.GetCause() != EKeyboardFocusCause::OtherWidgetLostFocus )
+			if( InFocusEvent.GetCause() != EFocusCause::OtherWidgetLostFocus )
 			{
 				// Select all text on keyboard focus.  This mirrors mouse focus functionality.
 				if( bSelectAllTextWhenFocused.Get() )
@@ -1566,7 +1566,7 @@ FReply SEditableText::OnKeyboardFocusReceived( const FGeometry& MyGeometry, cons
 }
 
 
-void SEditableText::OnKeyboardFocusLost( const FKeyboardFocusEvent& InKeyboardFocusEvent )
+void SEditableText::OnFocusLost( const FFocusEvent& InFocusEvent )
 {
 	// Skip the focus lost code if it's due to the context menu opening
 	if ( !ContextMenuWindow.IsValid() )
@@ -1586,7 +1586,7 @@ void SEditableText::OnKeyboardFocusLost( const FKeyboardFocusEvent& InKeyboardFo
 		}
 
 		// Clear selection unless activating a new window (otherwise can't copy and past on right click)
-		if (InKeyboardFocusEvent.GetCause() != EKeyboardFocusCause::WindowActivate)
+		if (InFocusEvent.GetCause() != EFocusCause::WindowActivate)
 		{
 			ClearSelection();
 		}
@@ -1596,14 +1596,14 @@ void SEditableText::OnKeyboardFocusLost( const FKeyboardFocusEvent& InKeyboardFo
 		
 		// See if user explicitly tabbed away or moved focus
 		ETextCommit::Type TextAction;
-		switch ( InKeyboardFocusEvent.GetCause() )
+		switch ( InFocusEvent.GetCause() )
 		{
-		case EKeyboardFocusCause::Keyboard:
-		case EKeyboardFocusCause::Mouse:
+		case EFocusCause::Navigation:
+		case EFocusCause::Mouse:
 			TextAction = ETextCommit::OnUserMovedFocus;
 			break;
 
-		case EKeyboardFocusCause::Cleared:
+		case EFocusCause::Cleared:
 			TextAction = ETextCommit::OnCleared;
 			break;
 
@@ -1624,12 +1624,12 @@ FReply SEditableText::OnKeyChar( const FGeometry& MyGeometry, const FCharacterEv
 }
 
 
-FReply SEditableText::OnKeyDown( const FGeometry& MyGeometry, const FKeyboardEvent& InKeyboardEvent )
+FReply SEditableText::OnKeyDown( const FGeometry& MyGeometry, const FKeyEvent& InKeyEvent )
 {
-	FReply Reply = FTextEditHelper::OnKeyDown( InKeyboardEvent, SharedThis( this ) );
+	FReply Reply = FTextEditHelper::OnKeyDown( InKeyEvent, SharedThis( this ) );
 
 	// Process keybindings if the event wasn't already handled
-	if( !Reply.IsEventHandled() && UICommandList->ProcessCommandBindings( InKeyboardEvent ) )
+	if( !Reply.IsEventHandled() && UICommandList->ProcessCommandBindings( InKeyEvent ) )
 	{
 		Reply = FReply::Handled();
 	}
@@ -1638,7 +1638,7 @@ FReply SEditableText::OnKeyDown( const FGeometry& MyGeometry, const FKeyboardEve
 }
 
 
-FReply SEditableText::OnKeyUp( const FGeometry& MyGeometry, const FKeyboardEvent& InKeyboardEvent )
+FReply SEditableText::OnKeyUp( const FGeometry& MyGeometry, const FKeyEvent& InKeyEvent )
 {
 	FReply Reply = FReply::Unhandled();
 
@@ -2102,7 +2102,7 @@ void SEditableText::RestartSelectionTargetAnimation()
 void SEditableText::OnWindowClosed(const TSharedRef<SWindow>&)
 {
 	// Give this focus when the context menu has been dismissed
-	FSlateApplication::Get().SetKeyboardFocus( AsShared(), EKeyboardFocusCause::OtherWidgetLostFocus );
+	FSlateApplication::Get().SetKeyboardFocus( AsShared(), EFocusCause::OtherWidgetLostFocus );
 }
 
 void SEditableText::SetHintText( const TAttribute< FText >& InHintText )
diff --git a/Engine/Source/Runtime/Slate/Private/Widgets/Input/SEditableTextBox.cpp b/Engine/Source/Runtime/Slate/Private/Widgets/Input/SEditableTextBox.cpp
--- a/Engine/Source/Runtime/Slate/Private/Widgets/Input/SEditableTextBox.cpp
+++ b/Engine/Source/Runtime/Slate/Private/Widgets/Input/SEditableTextBox.cpp
@@ -156,29 +156,29 @@ bool SEditableTextBox::HasKeyboardFocus() const
 }
 
 
-FReply SEditableTextBox::OnKeyboardFocusReceived( const FGeometry& MyGeometry, const FKeyboardFocusEvent& InKeyboardFocusEvent )
+FReply SEditableTextBox::OnFocusReceived( const FGeometry& MyGeometry, const FFocusEvent& InFocusEvent )
 {
 	FReply Reply = FReply::Handled();
 
-	if ( InKeyboardFocusEvent.GetCause() != EKeyboardFocusCause::Cleared )
+	if ( InFocusEvent.GetCause() != EFocusCause::Cleared )
 	{
 		// Forward keyboard focus to our editable text widget
-		Reply.SetKeyboardFocus( EditableText.ToSharedRef(), InKeyboardFocusEvent.GetCause() );
+		Reply.SetUserFocus(EditableText.ToSharedRef(), InFocusEvent.GetCause());
 	}
 
 	return Reply;
 }
 
 
-FReply SEditableTextBox::OnKeyDown( const FGeometry& MyGeometry, const FKeyboardEvent& InKeyboardEvent )
+FReply SEditableTextBox::OnKeyDown( const FGeometry& MyGeometry, const FKeyEvent& InKeyEvent )
 {
-	FKey Key = InKeyboardEvent.GetKey();
+	FKey Key = InKeyEvent.GetKey();
 
 	if( Key == EKeys::Escape && EditableText->HasKeyboardFocus() )
 	{
 		// Clear selection
 		EditableText->ClearSelection();
-		return FReply::Handled().SetKeyboardFocus( SharedThis( this ), EKeyboardFocusCause::Cleared );
+		return FReply::Handled().SetUserFocus(SharedThis(this), EFocusCause::Cleared);
 	}
 
 	return FReply::Unhandled();
diff --git a/Engine/Source/Runtime/Slate/Private/Widgets/Input/SMenuAnchor.cpp b/Engine/Source/Runtime/Slate/Private/Widgets/Input/SMenuAnchor.cpp
--- a/Engine/Source/Runtime/Slate/Private/Widgets/Input/SMenuAnchor.cpp
+++ b/Engine/Source/Runtime/Slate/Private/Widgets/Input/SMenuAnchor.cpp
@@ -312,7 +312,7 @@ void SMenuAnchor::SetIsOpen( bool InIsOpen, const bool bFocusMenu )
 
 					if ( bFocusMenu )
 					{
-						FSlateApplication::Get().SetKeyboardFocus( MenuContentRef, EKeyboardFocusCause::SetDirectly );
+						FSlateApplication::Get().SetKeyboardFocus( MenuContentRef, EFocusCause::SetDirectly );
 					}
 				}
 			}
diff --git a/Engine/Source/Runtime/Slate/Private/Widgets/Input/SMultiLineEditableTextBox.cpp b/Engine/Source/Runtime/Slate/Private/Widgets/Input/SMultiLineEditableTextBox.cpp
--- a/Engine/Source/Runtime/Slate/Private/Widgets/Input/SMultiLineEditableTextBox.cpp
+++ b/Engine/Source/Runtime/Slate/Private/Widgets/Input/SMultiLineEditableTextBox.cpp
@@ -202,14 +202,14 @@ bool SMultiLineEditableTextBox::HasKeyboardFocus() const
 	return SBorder::HasKeyboardFocus() || EditableText->HasKeyboardFocus();
 }
 
-FReply SMultiLineEditableTextBox::OnKeyboardFocusReceived( const FGeometry& MyGeometry, const FKeyboardFocusEvent& InKeyboardFocusEvent )
+FReply SMultiLineEditableTextBox::OnFocusReceived( const FGeometry& MyGeometry, const FFocusEvent& InFocusEvent )
 {
 	FReply Reply = FReply::Handled();
 
-	if ( InKeyboardFocusEvent.GetCause() != EKeyboardFocusCause::Cleared )
+	if ( InFocusEvent.GetCause() != EFocusCause::Cleared )
 	{
 		// Forward keyboard focus to our editable text widget
-		Reply.SetKeyboardFocus( EditableText.ToSharedRef(), InKeyboardFocusEvent.GetCause() );
+		Reply.SetUserFocus(EditableText.ToSharedRef(), InFocusEvent.GetCause());
 	}
 
 	return Reply;
diff --git a/Engine/Source/Runtime/Slate/Private/Widgets/Input/SSearchBox.cpp b/Engine/Source/Runtime/Slate/Private/Widgets/Input/SSearchBox.cpp
--- a/Engine/Source/Runtime/Slate/Private/Widgets/Input/SSearchBox.cpp
+++ b/Engine/Source/Runtime/Slate/Private/Widgets/Input/SSearchBox.cpp
@@ -166,10 +166,10 @@ FReply SSearchBox::OnClearSearch()
 {
 	// We clear the focus to commit any unset values as the search box is typically used for filtering
 	// and the widget could get immediately destroyed before committing its value.
-	FSlateApplication::Get().ClearKeyboardFocus( EKeyboardFocusCause::SetDirectly );
+	FSlateApplication::Get().ClearKeyboardFocus( EFocusCause::SetDirectly );
 
 	this->SetText( FText::GetEmpty() );
-	return FReply::Handled().SetKeyboardFocus( EditableText.ToSharedRef(), EKeyboardFocusCause::SetDirectly );
+	return FReply::Handled().SetUserFocus(EditableText.ToSharedRef(), EFocusCause::SetDirectly);
 }
 
 FSlateFontInfo SSearchBox::GetWidgetFont() const
diff --git a/Engine/Source/Runtime/Slate/Private/Widgets/Input/SSuggestionTextBox.cpp b/Engine/Source/Runtime/Slate/Private/Widgets/Input/SSuggestionTextBox.cpp
--- a/Engine/Source/Runtime/Slate/Private/Widgets/Input/SSuggestionTextBox.cpp
+++ b/Engine/Source/Runtime/Slate/Private/Widgets/Input/SSuggestionTextBox.cpp
@@ -100,7 +100,7 @@ void SSuggestionTextBox::FocusTextBox( )
 	FWidgetPath TextBoxPath;
 
 	FSlateApplication::Get().GeneratePathToWidgetUnchecked(TextBox.ToSharedRef(), TextBoxPath);
-	FSlateApplication::Get().SetKeyboardFocus(TextBoxPath, EKeyboardFocusCause::SetDirectly);
+	FSlateApplication::Get().SetKeyboardFocus(TextBoxPath, EFocusCause::SetDirectly);
 }
 
 
@@ -181,19 +181,19 @@ void SSuggestionTextBox::SetSuggestions( TArray<FString>& SuggestionStrings, boo
 /* SWidget overrides
  *****************************************************************************/
 
-void SSuggestionTextBox::OnKeyboardFocusLost( const FKeyboardFocusEvent& InKeyboardFocusEvent )
+void SSuggestionTextBox::OnFocusLost( const FFocusEvent& InFocusEvent )
 {
 	//	MenuAnchor->SetIsOpen(false);
 }
 
 
-FReply SSuggestionTextBox::OnKeyDown( const FGeometry& MyGeometry, const FKeyboardEvent& KeyboardEvent )
+FReply SSuggestionTextBox::OnKeyDown( const FGeometry& MyGeometry, const FKeyEvent& KeyEvent )
 {
 	const FString& InputTextStr = TextBox->GetText().ToString();
 
 	if (MenuAnchor->IsOpen())
 	{
-		if (KeyboardEvent.GetKey() == EKeys::Up)
+		if (KeyEvent.GetKey() == EKeys::Up)
 		{
 			// backward navigate the list of suggestions
 			if (SelectedSuggestion < 0)
@@ -209,7 +209,7 @@ FReply SSuggestionTextBox::OnKeyDown( const FGeometry& MyGeometry, const FKeyboa
 
 			return FReply::Handled();
 		}
-		else if(KeyboardEvent.GetKey() == EKeys::Down)
+		else if(KeyEvent.GetKey() == EKeys::Down)
 		{
 			// forward navigate the list of suggestions
 			if (SelectedSuggestion < Suggestions.Num() - 1)
@@ -225,7 +225,7 @@ FReply SSuggestionTextBox::OnKeyDown( const FGeometry& MyGeometry, const FKeyboa
 
 			return FReply::Handled();
 		}
-		else if (KeyboardEvent.GetKey() == EKeys::Tab)
+		else if (KeyEvent.GetKey() == EKeys::Tab)
 		{
 			// auto-complete the highlighted suggestion
 			if (Suggestions.Num())
@@ -248,7 +248,7 @@ FReply SSuggestionTextBox::OnKeyDown( const FGeometry& MyGeometry, const FKeyboa
 	}
 	else
 	{
-		if ((KeyboardEvent.GetKey() == EKeys::Up) || (KeyboardEvent.GetKey() == EKeys::Down))
+		if ((KeyEvent.GetKey() == EKeys::Up) || (KeyEvent.GetKey() == EKeys::Down))
 		{
 			// show the input history
 			if (OnShowingHistory.IsBound())
@@ -261,7 +261,7 @@ FReply SSuggestionTextBox::OnKeyDown( const FGeometry& MyGeometry, const FKeyboa
 				{
 					SetSuggestions(HistoryStrings, true);
 
-					if (KeyboardEvent.GetKey() == EKeys::Up)
+					if (KeyEvent.GetKey() == EKeys::Up)
 					{
 						SelectedSuggestion = HistoryStrings.Num() - 1;
 					}
diff --git a/Engine/Source/Runtime/Slate/Private/Widgets/Input/STextComboPopup.cpp b/Engine/Source/Runtime/Slate/Private/Widgets/Input/STextComboPopup.cpp
--- a/Engine/Source/Runtime/Slate/Private/Widgets/Input/STextComboPopup.cpp
+++ b/Engine/Source/Runtime/Slate/Private/Widgets/Input/STextComboPopup.cpp
@@ -94,12 +94,12 @@ void STextComboPopup::FocusDefaultWidget()
 {
 	FWidgetPath FocusMe;
 	FSlateApplication::Get().GeneratePathToWidgetChecked( StringCombo.ToSharedRef(), FocusMe );
-	FSlateApplication::Get().SetKeyboardFocus( FocusMe, EKeyboardFocusCause::SetDirectly );
+	FSlateApplication::Get().SetKeyboardFocus( FocusMe, EFocusCause::SetDirectly );
 }
 
-FReply STextComboPopup::OnKeyDown( const FGeometry& MyGeometry, const FKeyboardEvent& InKeyboardEvent )
+FReply STextComboPopup::OnKeyDown( const FGeometry& MyGeometry, const FKeyEvent& InKeyEvent )
 {
-	if(InKeyboardEvent.GetKey() == EKeys::Enter)
+	if(InKeyEvent.GetKey() == EKeys::Enter)
 	{
 		return OnOK();
 	}
diff --git a/Engine/Source/Runtime/Slate/Private/Widgets/Input/STextEntryPopup.cpp b/Engine/Source/Runtime/Slate/Private/Widgets/Input/STextEntryPopup.cpp
--- a/Engine/Source/Runtime/Slate/Private/Widgets/Input/STextEntryPopup.cpp
+++ b/Engine/Source/Runtime/Slate/Private/Widgets/Input/STextEntryPopup.cpp
@@ -60,7 +60,7 @@ void STextEntryPopup::FocusDefaultWidget()
 	// Focus the text widget
 	FWidgetPath FocusMe;
 	FSlateApplication::Get().GeneratePathToWidgetChecked( WidgetWithDefaultFocus.ToSharedRef(), FocusMe );
-	FSlateApplication::Get().SetKeyboardFocus( FocusMe, EKeyboardFocusCause::SetDirectly );
+	FSlateApplication::Get().SetKeyboardFocus( FocusMe, EFocusCause::SetDirectly );
 }
 
 void STextEntryPopup::SetError( const FText& InError )
diff --git a/Engine/Source/Runtime/Slate/Private/Widgets/Input/SVirtualJoystick.cpp b/Engine/Source/Runtime/Slate/Private/Widgets/Input/SVirtualJoystick.cpp
--- a/Engine/Source/Runtime/Slate/Private/Widgets/Input/SVirtualJoystick.cpp
+++ b/Engine/Source/Runtime/Slate/Private/Widgets/Input/SVirtualJoystick.cpp
@@ -364,9 +364,9 @@ void SVirtualJoystick::Tick( const FGeometry& AllottedGeometry, const double InC
 			EControllerButtons::Type YAxis = ControlIndex == 0 ? EControllerButtons::LeftAnalogY : EControllerButtons::RightAnalogY;
 
 	//		UE_LOG(LogTemp, Log, TEXT("Joysticking %f,%f"), NormalizedOffset.X, -NormalizedOffset.Y);
-			FSlateApplication::Get().SetJoystickCaptorToGameViewport();
-			FSlateApplication::Get().OnControllerAnalog(Control.MainInputKey, 0, NormalizedOffset.X);
-			FSlateApplication::Get().OnControllerAnalog(Control.AltInputKey, 0, -NormalizedOffset.Y);
+			FSlateApplication::Get().SetAllUserFocusToGameViewport();
+			FSlateApplication::Get().OnControllerAnalog(XAxis, 0, NormalizedOffset.X);
+			FSlateApplication::Get().OnControllerAnalog(YAxis, 0, -NormalizedOffset.Y);
 		}
 
 		// is this active?
diff --git a/Engine/Source/Runtime/Slate/Private/Widgets/Input/SVirtualKeyboardEntry.cpp b/Engine/Source/Runtime/Slate/Private/Widgets/Input/SVirtualKeyboardEntry.cpp
--- a/Engine/Source/Runtime/Slate/Private/Widgets/Input/SVirtualKeyboardEntry.cpp
+++ b/Engine/Source/Runtime/Slate/Private/Widgets/Input/SVirtualKeyboardEntry.cpp
@@ -248,7 +248,7 @@ int32 SVirtualKeyboardEntry::OnPaint( const FPaintArgs& Args, const FGeometry& A
 	return LayerId + TextLayer;
 }
 
-FReply SVirtualKeyboardEntry::OnKeyboardFocusReceived( const FGeometry& MyGeometry, const FKeyboardFocusEvent& InKeyboardFocusEvent )
+FReply SVirtualKeyboardEntry::OnFocusReceived( const FGeometry& MyGeometry, const FFocusEvent& InFocusEvent )
 {
 	// The user wants to edit text. Make a copy of the observed text for the user to edit.
 	EditedText = Text.Get();
@@ -263,20 +263,20 @@ FReply SVirtualKeyboardEntry::OnKeyboardFocusReceived( const FGeometry& MyGeomet
 /**
  * Called when this widget loses the keyboard focus.  This event does not bubble.
  *
- * @param  InKeyboardFocusEvent  KeyboardFocusEvent
+ * @param  InFocusEvent  FocusEvent
  */
-void SVirtualKeyboardEntry::OnKeyboardFocusLost( const FKeyboardFocusEvent& InKeyboardFocusEvent )
+void SVirtualKeyboardEntry::OnFocusLost( const FFocusEvent& InFocusEvent )
 {
 	// See if user explicitly tabbed away or moved focus
 	ETextCommit::Type TextAction;
-	switch ( InKeyboardFocusEvent.GetCause() )
+	switch ( InFocusEvent.GetCause() )
 	{
-		case EKeyboardFocusCause::Keyboard:
-		case EKeyboardFocusCause::Mouse:
+		case EFocusCause::Navigation:
+		case EFocusCause::Mouse:
 			TextAction = ETextCommit::OnUserMovedFocus;
 			break;
 
-		case EKeyboardFocusCause::Cleared:
+		case EFocusCause::Cleared:
 			TextAction = ETextCommit::OnCleared;
 			break;
 
diff --git a/Engine/Source/Runtime/Slate/Private/Widgets/Layout/SExpandableArea.cpp b/Engine/Source/Runtime/Slate/Private/Widgets/Layout/SExpandableArea.cpp
--- a/Engine/Source/Runtime/Slate/Private/Widgets/Layout/SExpandableArea.cpp
+++ b/Engine/Source/Runtime/Slate/Private/Widgets/Layout/SExpandableArea.cpp
@@ -125,7 +125,7 @@ FReply SExpandableArea::OnMouseDown( const FGeometry& MyGeometry, const FPointer
 	if ( MouseEvent.GetEffectingButton() == EKeys::LeftMouseButton )
 	{
 		//we need to capture the mouse for MouseUp events
-		return FReply::Handled().CaptureMouse( HeaderBorder.ToSharedRef() ).SetKeyboardFocus( AsShared(), EKeyboardFocusCause::Mouse );
+		return FReply::Handled().CaptureMouse( HeaderBorder.ToSharedRef() ).SetUserFocus( AsShared(), EFocusCause::Mouse );
 	}
 
 	return FReply::Unhandled();
diff --git a/Engine/Source/Runtime/Slate/Private/Widgets/Layout/SScrollBar.cpp b/Engine/Source/Runtime/Slate/Private/Widgets/Layout/SScrollBar.cpp
--- a/Engine/Source/Runtime/Slate/Private/Widgets/Layout/SScrollBar.cpp
+++ b/Engine/Source/Runtime/Slate/Private/Widgets/Layout/SScrollBar.cpp
@@ -374,7 +374,7 @@ FReply SScrollBar::OnMouseButtonDown( const FGeometry& MyGeometry, const FPointe
 
 	if( bDraggingThumb )
 	{
-		return FReply::Handled().CaptureMouse(AsShared()).SetKeyboardFocus( AsShared(), EKeyboardFocusCause::Mouse );
+		return FReply::Handled().CaptureMouse(AsShared()).SetUserFocus(AsShared(), EFocusCause::Mouse);
 	}
 	else
 	{
diff --git a/Engine/Source/Runtime/Slate/Private/Widgets/SViewport.cpp b/Engine/Source/Runtime/Slate/Private/Widgets/SViewport.cpp
--- a/Engine/Source/Runtime/Slate/Private/Widgets/SViewport.cpp
+++ b/Engine/Source/Runtime/Slate/Private/Widgets/SViewport.cpp
@@ -191,40 +191,43 @@ FReply SViewport::OnMouseButtonDoubleClick( const FGeometry& MyGeometry, const F
 }
 
 
-FReply SViewport::OnKeyDown( const FGeometry& MyGeometry, const FKeyboardEvent& KeyboardEvent )
+FReply SViewport::OnKeyDown( const FGeometry& MyGeometry, const FKeyEvent& KeyEvent )
 {
-	return ViewportInterface.IsValid() ? ViewportInterface.Pin()->OnKeyDown(MyGeometry, KeyboardEvent) : FReply::Unhandled();
+	return ViewportInterface.IsValid() ? ViewportInterface.Pin()->OnKeyDown(MyGeometry, KeyEvent) : FReply::Unhandled();
 }
 
  
-FReply SViewport::OnKeyUp( const FGeometry& MyGeometry, const FKeyboardEvent& KeyboardEvent )
+FReply SViewport::OnKeyUp( const FGeometry& MyGeometry, const FKeyEvent& KeyEvent )
 {
-	return ViewportInterface.IsValid() ? ViewportInterface.Pin()->OnKeyUp(MyGeometry, KeyboardEvent) : FReply::Unhandled();
+	return ViewportInterface.IsValid() ? ViewportInterface.Pin()->OnKeyUp(MyGeometry, KeyEvent) : FReply::Unhandled();
 }
 
+FReply SViewport::OnAnalogValueChanged( const FGeometry& MyGeometry, const FAnalogInputEvent& InAnalogInputEvent )
+{
+	return ViewportInterface.IsValid() ? ViewportInterface.Pin()->OnAnalogValueChanged(MyGeometry, InAnalogInputEvent) : FReply::Unhandled();
+}
 
 FReply SViewport::OnKeyChar( const FGeometry& MyGeometry, const FCharacterEvent& CharacterEvent )
 {
 	return ViewportInterface.IsValid() ? ViewportInterface.Pin()->OnKeyChar(MyGeometry, CharacterEvent) : FReply::Unhandled();
 }
 
-
-FReply SViewport::OnKeyboardFocusReceived( const FGeometry& MyGeometry, const FKeyboardFocusEvent& InKeyboardFocusEvent )
+FReply SViewport::OnFocusReceived( const FGeometry& MyGeometry, const FFocusEvent& InFocusEvent )
 {
 	if (WidgetToFocusOnActivate.IsValid())
 	{
-		return FReply::Handled().SetKeyboardFocus(WidgetToFocusOnActivate.Pin().ToSharedRef(), InKeyboardFocusEvent.GetCause());
+		return FReply::Handled().SetUserFocus(WidgetToFocusOnActivate.Pin().ToSharedRef(), InFocusEvent.GetCause());
 	}
 
-	return ViewportInterface.IsValid() ? ViewportInterface.Pin()->OnKeyboardFocusReceived(InKeyboardFocusEvent) : FReply::Unhandled();
+	return ViewportInterface.IsValid() ? ViewportInterface.Pin()->OnFocusReceived(InFocusEvent) : FReply::Unhandled();
 }
 
 
-void SViewport::OnKeyboardFocusLost( const FKeyboardFocusEvent& InKeyboardFocusEvent )
+void SViewport::OnFocusLost( const FFocusEvent& InFocusEvent )
 {
 	if (ViewportInterface.IsValid())
 	{
-		ViewportInterface.Pin()->OnKeyboardFocusLost(InKeyboardFocusEvent);
+		ViewportInterface.Pin()->OnFocusLost(InFocusEvent);
 	}
 }
 
@@ -255,39 +258,21 @@ void SViewport::OnWindowClosed( const TSharedRef<SWindow>& WindowBeingClosed )
 }
 
 
-FReply SViewport::OnControllerButtonPressed( const FGeometry& MyGeometry, const FControllerEvent& ControllerEvent )
-{
-	return ViewportInterface.IsValid() ? ViewportInterface.Pin()->OnControllerButtonPressed(MyGeometry, ControllerEvent) : FReply::Unhandled();
-}
-
-
-FReply SViewport::OnControllerButtonReleased( const FGeometry& MyGeometry, const FControllerEvent& ControllerEvent )
+FReply SViewport::OnTouchStarted( const FGeometry& MyGeometry, const FPointerEvent& InTouchEvent )
 {
-	return ViewportInterface.IsValid() ? ViewportInterface.Pin()->OnControllerButtonReleased(MyGeometry, ControllerEvent) : FReply::Unhandled();
+	return ViewportInterface.IsValid() ? ViewportInterface.Pin()->OnTouchStarted(MyGeometry, InTouchEvent) : FReply::Unhandled();
 }
 
 
-FReply SViewport::OnControllerAnalogValueChanged( const FGeometry& MyGeometry, const FControllerEvent& ControllerEvent )
+FReply SViewport::OnTouchMoved( const FGeometry& MyGeometry, const FPointerEvent& InTouchEvent )
 {
-	return ViewportInterface.IsValid() ? ViewportInterface.Pin()->OnControllerAnalogValueChanged(MyGeometry, ControllerEvent) : FReply::Unhandled();
+	return ViewportInterface.IsValid() ? ViewportInterface.Pin()->OnTouchMoved(MyGeometry, InTouchEvent) : FReply::Unhandled();
 }
 
 
-FReply SViewport::OnTouchStarted( const FGeometry& MyGeometry, const FPointerEvent& TouchEvent )
+FReply SViewport::OnTouchEnded( const FGeometry& MyGeometry, const FPointerEvent& InTouchEvent )
 {
-	return ViewportInterface.IsValid() ? ViewportInterface.Pin()->OnTouchStarted(MyGeometry, TouchEvent) : FReply::Unhandled();
-}
-
-
-FReply SViewport::OnTouchMoved( const FGeometry& MyGeometry, const FPointerEvent& TouchEvent )
-{
-	return ViewportInterface.IsValid() ? ViewportInterface.Pin()->OnTouchMoved(MyGeometry, TouchEvent) : FReply::Unhandled();
-}
-
-
-FReply SViewport::OnTouchEnded( const FGeometry& MyGeometry, const FPointerEvent& TouchEvent )
-{
-	return ViewportInterface.IsValid() ? ViewportInterface.Pin()->OnTouchEnded(MyGeometry, TouchEvent) : FReply::Unhandled();
+	return ViewportInterface.IsValid() ? ViewportInterface.Pin()->OnTouchEnded(MyGeometry, InTouchEvent) : FReply::Unhandled();
 }
 
 
@@ -297,9 +282,9 @@ FReply SViewport::OnTouchGesture( const FGeometry& MyGeometry, const FPointerEve
 }
 
 
-FReply SViewport::OnMotionDetected( const FGeometry& MyGeometry, const FMotionEvent& MotionEvent )
+FReply SViewport::OnMotionDetected( const FGeometry& MyGeometry, const FMotionEvent& InMotionEvent )
 {
-	return ViewportInterface.IsValid() ? ViewportInterface.Pin()->OnMotionDetected(MyGeometry, MotionEvent) : FReply::Unhandled();
+	return ViewportInterface.IsValid() ? ViewportInterface.Pin()->OnMotionDetected(MyGeometry, InMotionEvent) : FReply::Unhandled();
 }
 
 void SViewport::OnFinishedPointerInput()
@@ -328,4 +313,9 @@ TSharedPtr<FVirtualPointerPosition> SViewport::TranslateMouseCoordinateFor3DChil
 	}
 
 	return nullptr;
-}
\ No newline at end of file
+}
+
+FNavigationReply SViewport::OnNavigation(const FGeometry& MyGeometry, const FNavigationEvent& InNavigationEvent)
+{
+	return ViewportInterface.IsValid() ? ViewportInterface.Pin()->OnNavigation(MyGeometry, InNavigationEvent) : FNavigationReply::Stop();
+}
diff --git a/Engine/Source/Runtime/Slate/Private/Widgets/Text/SInlineEditableTextBlock.cpp b/Engine/Source/Runtime/Slate/Private/Widgets/Text/SInlineEditableTextBlock.cpp
--- a/Engine/Source/Runtime/Slate/Private/Widgets/Text/SInlineEditableTextBlock.cpp
+++ b/Engine/Source/Runtime/Slate/Private/Widgets/Text/SInlineEditableTextBlock.cpp
@@ -83,7 +83,7 @@ void SInlineEditableTextBlock::EnterEditingMode()
 					TextBox.ToSharedRef()
 				];
 
-			FSlateApplication::Get().SetKeyboardFocus(TextBox, EKeyboardFocusCause::SetDirectly);
+			FSlateApplication::Get().SetKeyboardFocus(TextBox, EFocusCause::SetDirectly);
 
 			TextBlock->SetVisibility(EVisibility::Collapsed);
 
@@ -174,9 +174,9 @@ void SInlineEditableTextBlock::Tick( const FGeometry& AllottedGeometry, const do
 	}
 }
 
-FReply SInlineEditableTextBlock::OnKeyDown( const FGeometry& MyGeometry, const FKeyboardEvent& InKeyboardEvent )
+FReply SInlineEditableTextBlock::OnKeyDown( const FGeometry& MyGeometry, const FKeyEvent& InKeyEvent )
 {
-	if(InKeyboardEvent.GetKey() == EKeys::F2)
+	if(InKeyEvent.GetKey() == EKeys::F2)
 	{
 		EnterEditingMode();
 		return FReply::Handled();
diff --git a/Engine/Source/Runtime/Slate/Private/Widgets/Text/SMultiLineEditableText.cpp b/Engine/Source/Runtime/Slate/Private/Widgets/Text/SMultiLineEditableText.cpp
--- a/Engine/Source/Runtime/Slate/Private/Widgets/Text/SMultiLineEditableText.cpp
+++ b/Engine/Source/Runtime/Slate/Private/Widgets/Text/SMultiLineEditableText.cpp
@@ -464,7 +464,7 @@ void SMultiLineEditableText::OnVScrollBarMoved(const float InScrollOffsetFractio
 	ScrollOffset.Y = FMath::Clamp<float>(InScrollOffsetFraction, 0.0, 1.0) * GetDesiredSize().Y;
 }
 
-FReply SMultiLineEditableText::OnKeyboardFocusReceived( const FGeometry& MyGeometry, const FKeyboardFocusEvent& InKeyboardFocusEvent )
+FReply SMultiLineEditableText::OnFocusReceived( const FGeometry& MyGeometry, const FFocusEvent& InFocusEvent )
 {
 	// Skip the focus received code if it's due to the context menu closing
 	if ( !ContextMenuWindow.IsValid() )
@@ -493,13 +493,13 @@ FReply SMultiLineEditableText::OnKeyboardFocusReceived( const FGeometry& MyGeome
 		// of making sure that gets scrolled into view
 		PositionToScrollIntoView = TOptional<FScrollInfo>();
 
-		return SWidget::OnKeyboardFocusReceived( MyGeometry, InKeyboardFocusEvent );
+		return SWidget::OnFocusReceived( MyGeometry, InFocusEvent );
 	}
 
 	return FReply::Handled();
 }
 
-void SMultiLineEditableText::OnKeyboardFocusLost( const FKeyboardFocusEvent& InKeyboardFocusEvent )
+void SMultiLineEditableText::OnFocusLost( const FFocusEvent& InFocusEvent )
 {
 	// Skip the focus lost code if it's due to the context menu opening
 	if (!ContextMenuWindow.IsValid())
@@ -521,14 +521,14 @@ void SMultiLineEditableText::OnKeyboardFocusLost( const FKeyboardFocusEvent& InK
 		// When focus is lost let anyone who is interested that text was committed
 		// See if user explicitly tabbed away or moved focus
 		ETextCommit::Type TextAction;
-		switch (InKeyboardFocusEvent.GetCause())
+		switch (InFocusEvent.GetCause())
 		{
-		case EKeyboardFocusCause::Keyboard:
-		case EKeyboardFocusCause::Mouse:
+		case EFocusCause::Navigation:
+		case EFocusCause::Mouse:
 			TextAction = ETextCommit::OnUserMovedFocus;
 			break;
 
-		case EKeyboardFocusCause::Cleared:
+		case EFocusCause::Cleared:
 			TextAction = ETextCommit::OnCleared;
 			break;
 
@@ -2090,7 +2090,7 @@ void SMultiLineEditableText::SummonContextMenu(const FVector2D& InLocation)
 void SMultiLineEditableText::OnWindowClosed(const TSharedRef<SWindow>&)
 {
 	// Give this focus when the context menu has been dismissed
-	FSlateApplication::Get().SetKeyboardFocus(AsShared(), EKeyboardFocusCause::OtherWidgetLostFocus);
+	FSlateApplication::Get().SetKeyboardFocus(AsShared(), EFocusCause::OtherWidgetLostFocus);
 }
 
 void SMultiLineEditableText::LoadText()
@@ -2298,12 +2298,12 @@ FReply SMultiLineEditableText::OnKeyChar( const FGeometry& MyGeometry,const FCha
 	return FTextEditHelper::OnKeyChar( InCharacterEvent, SharedThis( this ) );
 }
 
-FReply SMultiLineEditableText::OnKeyDown( const FGeometry& MyGeometry, const FKeyboardEvent& InKeyboardEvent )
+FReply SMultiLineEditableText::OnKeyDown( const FGeometry& MyGeometry, const FKeyEvent& InKeyEvent )
 {
-	FReply Reply = FTextEditHelper::OnKeyDown( InKeyboardEvent, SharedThis( this ) );
+	FReply Reply = FTextEditHelper::OnKeyDown( InKeyEvent, SharedThis( this ) );
 
 	// Process keybindings if the event wasn't already handled
-	if (!Reply.IsEventHandled() && UICommandList->ProcessCommandBindings(InKeyboardEvent))
+	if (!Reply.IsEventHandled() && UICommandList->ProcessCommandBindings(InKeyEvent))
 	{
 		Reply = FReply::Handled();
 	}
@@ -2311,7 +2311,7 @@ FReply SMultiLineEditableText::OnKeyDown( const FGeometry& MyGeometry, const FKe
 	return Reply;
 }
 
-FReply SMultiLineEditableText::OnKeyUp( const FGeometry& MyGeometry, const FKeyboardEvent& InKeyboardEvent )
+FReply SMultiLineEditableText::OnKeyUp( const FGeometry& MyGeometry, const FKeyEvent& InKeyEvent )
 {
 	return FReply::Unhandled();
 }
diff --git a/Engine/Source/Runtime/Slate/Private/Widgets/Views/STableViewBase.cpp b/Engine/Source/Runtime/Slate/Private/Widgets/Views/STableViewBase.cpp
--- a/Engine/Source/Runtime/Slate/Private/Widgets/Views/STableViewBase.cpp
+++ b/Engine/Source/Runtime/Slate/Private/Widgets/Views/STableViewBase.cpp
@@ -93,7 +93,7 @@ bool STableViewBase::SupportsKeyboardFocus() const
 }
 
 
-void STableViewBase::OnKeyboardFocusLost( const FKeyboardFocusEvent& InKeyboardFocusEvent )
+void STableViewBase::OnFocusLost( const FFocusEvent& InFocusEvent )
 {
 	bShowSoftwareCursor = false;
 }
@@ -418,9 +418,9 @@ FReply STableViewBase::OnMouseWheel( const FGeometry& MyGeometry, const FPointer
 }
 
 
-FReply STableViewBase::OnKeyDown( const FGeometry& MyGeometry, const FKeyboardEvent& InKeyboardEvent )
+FReply STableViewBase::OnKeyDown( const FGeometry& MyGeometry, const FKeyEvent& InKeyEvent )
 {
-	if ( InKeyboardEvent.IsControlDown() && InKeyboardEvent.GetKey() == EKeys::End )
+	if ( InKeyEvent.IsControlDown() && InKeyEvent.GetKey() == EKeys::End )
 	{
 		ScrollOffset = GetNumItemsBeingObserved();
 		RequestListRefresh();
diff --git a/Engine/Source/Runtime/Slate/Public/Framework/Application/SlateApplication.h b/Engine/Source/Runtime/Slate/Public/Framework/Application/SlateApplication.h
--- a/Engine/Source/Runtime/Slate/Public/Framework/Application/SlateApplication.h
+++ b/Engine/Source/Runtime/Slate/Public/Framework/Application/SlateApplication.h
@@ -8,6 +8,7 @@
 class SToolTip;
 class SViewport;
 class SWindow;
+enum class EHittestDirection;
 class FHittestGrid;
 struct FPopupTransitionEffect;
 class FMenuStack;
@@ -337,6 +338,11 @@ public:
 	 */
 	void RegisterGameViewport( TSharedRef<SViewport> InViewport );
 
+	/**
+	 * Returns the game viewport registered with the slate application
+	 *
+	 * @return registered game viewport
+	 */
 	TSharedPtr<SViewport> GetGameViewport() const;
 
 	/** 
@@ -346,20 +352,54 @@ public:
 	void UnregisterGameViewport();
 
 	/**
-	 *	Sets the current focus to the SWidget representing the currently active game viewport
+	 * Sets specified user focus to the SWidget representing the currently active game viewport
 	 */
+	void SetUserFocusToGameViewport(uint32 UserIndex, EFocusCause ReasonFocusIsChanging = EFocusCause::SetDirectly);
+
+	DEPRECATED(4.6, "FSlateApplication::SetFocusToGameViewport() is deprecated, use FSlateApplication::SetUserFocusToGameViewport() instead.")
 	void SetFocusToGameViewport();
 
 	/**
-	 *	Forces the currently active game viewport to capture the joystick
+	 * Sets all users focus to the SWidget representing the currently active game viewport
 	 */
+	void SetAllUserFocusToGameViewport(EFocusCause ReasonFocusIsChanging = EFocusCause::SetDirectly);
+
+	DEPRECATED(4.6, "FSlateApplication::SetJoystickCaptorToGameViewport() is deprecated, use FSlateApplication::SetAllUserFocusToGameViewport() instead.")
 	void SetJoystickCaptorToGameViewport();
 
 	/**
-	 * 	Sets the current focus to the specified SWidget
+	 * Sets specified user focus to the SWidget passed in.
+	 *
+	 * @param UserIndex Index of the user to change focus for
+	 * @param WidgetToFocus the widget to set focus to
+	 * @param ReasonFocusIsChanging the contextual reason for the focus change
 	 */
-	void SetKeyboardFocus( const TSharedPtr< SWidget >& OptionalWidgetToFocus, EKeyboardFocusCause::Type ReasonFocusIsChanging = EKeyboardFocusCause::SetDirectly );
-	
+	void SetUserFocus(uint32 UserIndex, const TSharedPtr<SWidget>& WidgetToFocus, EFocusCause ReasonFocusIsChanging = EFocusCause::SetDirectly);
+
+	/** @return a pointer to the Widget that currently has the users focus; Empty pointer when the user has no focus. */
+	TSharedPtr< SWidget > GetUserFocusedWidget(uint32 UserIndex) const;
+
+	DEPRECATED(4.6, "FSlateApplication::GetJoystickCaptor() is deprecated, use FSlateApplication::GetUserFocusedWidget() instead.")
+	TSharedPtr< SWidget > GetJoystickCaptor(uint32 UserIndex) const;
+
+	/** Releases the users focus from whatever it currently is on. */
+	void ClearUserFocus(uint32 UserIndex, EFocusCause ReasonFocusIsChanging = EFocusCause::Cleared);
+
+	DEPRECATED(4.6, "FSlateApplication::ReleaseJoystickCapture() is deprecated, use FSlateApplication::ClearUserFocus() instead.")
+	void ReleaseJoystickCapture(uint32 UserIndex);
+
+	/**
+	 * Sets the Keyboard focus to the specified SWidget
+	 */
+	void SetKeyboardFocus(const TSharedPtr<SWidget>& OptionalWidgetToFocus, EFocusCause ReasonFocusIsChanging = EFocusCause::SetDirectly);
+
+	/**
+	 * Clears keyboard focus, if any widget is currently focused
+	 *
+	 * @param ReasonFocusIsChanging The reason that keyboard focus is changing
+	 */
+	void ClearKeyboardFocus(const EFocusCause ReasonFocusIsChanging = EFocusCause::Cleared);
+
 	/**
 	 * Returns the current modifier keys state
 	 *
@@ -385,26 +425,9 @@ public:
 	/** Releases the mouse capture from whatever it currently is on. */
 	void ReleaseMouseCapture();
 
-	/** @return a pointer to the Widget that currently captures the joystick; Empty pointer when the joystick is not captured. */
-	TSharedPtr< SWidget > GetJoystickCaptor(uint32 UserIndex) const;
-
-	/** Releases the joystick capture from whatever it currently is on. */
-	void ReleaseJoystickCapture(uint32 UserIndex);
-
 	/** @return The active modal window or NULL if there is no modal window. */
 	TSharedPtr<SWindow> GetActiveModalWindow() const;
 
-	/**
-	 * Keyboard focus
-	 */
-
-	/**
-	 * Clears keyboard focus, if any widget is currently focused
-	 *
-	 * @param  InCause  The reason that keyboard focus is changing
-	 */
-	void ClearKeyboardFocus( const EKeyboardFocusCause::Type InCause );
-
 	/**
 	 * Assign a delegate to be called when this application is requesting an exit (e.g. when the last window is closed).
 	 *
@@ -729,21 +752,30 @@ public:
 	/**
 	 * Called when a key is pressed
 	 *
-	 * @param  InKeyboardEvent  Keyboard event
+	 * @param  InKeyEvent  Keyb event
 	 *
 	 * @return  Was this event handled by the Slate application?
 	 */
-	bool ProcessKeyDownEvent( FKeyboardEvent& InKeyboardEvent );
+	bool ProcessKeyDownEvent( FKeyEvent& InKeyEvent );
 
 	/**
 	 * Called when a key is released
 	 *
-	 * @param  InKeyboardEvent  Keyboard event
+	 * @param  InKeyEvent  Key event
 	 *
 	 * @return  Was this event handled by the Slate application?
 	 */
-	bool ProcessKeyUpEvent( FKeyboardEvent& InKeyboardEvent );
+	bool ProcessKeyUpEvent( FKeyEvent& InKeyEvent );
 	
+	/**
+	 * Called when a analog input values change
+	 *
+	 * @param  InAnalogInputEvent Analog input event
+	 *
+	 * @return  Was this event handled by the Slate application?
+	 */
+	bool ProcessAnalogInputEvent(FAnalogInputEvent& InAnalogInputEvent);
+
 	/**
 	 * Called when a drag from an external (non-slate) source enters a window
 	 *
@@ -754,27 +786,6 @@ public:
 	 */
 	bool ProcessDragEnterEvent( TSharedRef<SWindow> WindowEntered, FDragDropEvent& DragDropEvent );
 
-	/**
-	 * Called when a controller button is found pressed when polling game device state
-	 * 
-	 * @param ControllerEvent	The controller event generated
-	 */
-	void ProcessControllerButtonPressedEvent( FControllerEvent& ControllerEvent );
-
-	/**
-	 * Called when a controller button is found released when polling game device state
-	 * 
-	 * @param ControllerEvent	The controller event generated
-	 */
-	void ProcessControllerButtonReleasedEvent( FControllerEvent& ControllerEvent );
-
-	/**
-	 * Called when a controller analog value has changed while polling game device state
-	 * 
-	 * @param ControllerEvent	The controller event generated
-	 */
-	void ProcessControllerAnalogValueChangedEvent( FControllerEvent& ControllerEvent );
-
 	/**
 	 * Called when a touchpad touch is started (finger down) when polling game device state
 	 * 
@@ -840,7 +851,7 @@ public:
 	 * @param InMouseEvent       Optional mouse event that caused this action.
 	 * @param UserIndex			 User index that generated the event we are replying to (defaults to 0, at least for now)
 	 */
-	void ProcessReply( const FWidgetPath& CurrentEventPath, const FReply TheReply, const FWidgetPath* WidgetsUnderMouse, const FPointerEvent* InMouseEvent, uint32 UserIndex=0 );
+	void ProcessReply(const FWidgetPath& CurrentEventPath, const FReply TheReply, const FWidgetPath* WidgetsUnderMouse, const FPointerEvent* InMouseEvent, uint32 UserIndex = 0);
 	
 	void LockCursor( const TSharedPtr<SWidget>& Widget );
 
@@ -882,7 +893,7 @@ public:
 	 *
 	 * @return New maximum layer id
 	 */
-	int32 DrawKeyboardFocus( const FWidgetPath& FocusPath, class FSlateWindowElementList& WindowElementList, int32 InLayerId ) const;
+	int32 DrawFocus( const FWidgetPath& FocusPath, class FSlateWindowElementList& WindowElementList, int32 InLayerId ) const;
 
 	/** @return an array of top-level windows that can be interacted with. e.g. when a modal window is up, only return the modal window */
 	TArray< TSharedRef<SWindow> > GetInteractiveTopLevelWindows();
@@ -906,9 +917,9 @@ public:
 	void SetGameIsFakingTouchEvents(const bool bIsFaking, FVector2D* CursorLocation = nullptr);
 
 	/** Sets the handler for otherwise unhandled key down events. This is used by the editor to provide a global action list, if the key was not consumed by any widget. */
-	void SetUnhandledKeyDownEventHandler( const FOnKeyboardEvent& NewHandler );
+	void SetUnhandledKeyDownEventHandler( const FOnKeyEvent& NewHandler );
 
-	/** @return the last time a user interacted with a keyboard, mouse, touch device, or joystick */
+	/** @return the last time a user interacted with a keyboard, mouse, touch device, or controller */
 	double GetLastUserInteractionTime() const { return LastUserInteractionTime; }
 
 	/** @return the deadzone size for dragging in screen pixels (aka virtual desktop pixels) */
@@ -986,7 +997,9 @@ public:
 
 	virtual void RequestDestroyWindow( TSharedRef<SWindow> WindowToDestroy ) override;
 
-	virtual bool SetKeyboardFocus( const FWidgetPath& InFocusPath, const EKeyboardFocusCause::Type InCause ) override;
+	virtual bool SetKeyboardFocus( const FWidgetPath& InFocusPath, const EFocusCause InCause ) override;
+
+	virtual bool SetUserFocus(const uint32 InUserIndex, const FWidgetPath& InFocusPath, const EFocusCause InCause) override;
 
 	// End FSlateApplicationBase interface
 
@@ -1016,16 +1029,10 @@ public:
 
 	virtual bool OnCursorSet() override;
 
-	virtual bool OnControllerAnalog( FKey Button, int32 ControllerId, float AnalogValue );
-
 	virtual bool OnControllerAnalog( EControllerButtons::Type Button, int32 ControllerId, float AnalogValue ) override;
 
-	virtual bool OnControllerButtonPressed( FKey Button, int32 ControllerId, bool IsRepeat );
-
 	virtual bool OnControllerButtonPressed( EControllerButtons::Type Button, int32 ControllerId, bool IsRepeat ) override;
 
-	virtual bool OnControllerButtonReleased( FKey, int32 ControllerId, bool IsRepeat );
-
 	virtual bool OnControllerButtonReleased( EControllerButtons::Type Button, int32 ControllerId, bool IsRepeat ) override;
 
 	virtual bool OnTouchGesture( EGestureEvent::Type GestureType, const FVector2D& Delta, float WheelDelta ) override;
@@ -1087,6 +1094,23 @@ private:
 	 */
 	void PrivateDestroyWindow( const TSharedRef<SWindow>& DestroyedWindow );
 
+	/**
+	 * Attempts to navigate to the next widget in the direction specified
+	 *
+	 * @return if a new widget was navigated too
+	 */
+	bool AttemptNavigation(const FNavigationEvent& NavigationEvent, const FNavigationReply& NavigationReply, const FArrangedWidget& BoundaryWidget);
+
+	/**
+	 * @return int user index that the keyboard is mapped to. -1 if the keyboard isn't mapped
+	 */
+	int32 GetUserIndexForKeyboard() const;
+
+	/** 
+	 * @return int user index that this controller is mapped to. -1 if the controller isn't mapped
+	 */
+	int32 GetUserIndexForController(int32 ControllerId) const;
+
 private:
 
 	// Hidden default constructor.
@@ -1212,8 +1236,8 @@ private:
 	/** The current mouse captor for the application, if any. */
 	MouseCaptorHelper MouseCaptor;
 
-	/** A weak path to the widget currently capturing joystick events, if any. */
-	FWeakWidgetPath JoystickCaptorWeakPaths[SlateApplicationDefs::MaxUsers];
+	/** A weak path to the widget currently focued by a user, if any. */
+	FWeakWidgetPath UserFocusedWidgetPaths[SlateApplicationDefs::MaxUsers];
 
 	/**
 	 * Application throttling
@@ -1228,10 +1252,9 @@ private:
 	/** The last real time that the user pressed a key or mouse button */
 	double LastUserInteractionTime;
 
-	/** Subset of LastUserInteractionTime that is used only when consdering when to throttle */
+	/** Subset of LastUserInteractionTime that is used only when considering when to throttle */
 	double LastUserInteractionTimeForThrottling;
 
-
 	/** Helper for detecting when a drag should begin */
 	struct FDragDetector
 	{
@@ -1251,10 +1274,8 @@ private:
 	/** Support for auto-dismissing popups */
 	FPopupSupport PopupSupport;
 
-	/** Weak path to the widget that currently has keyboard focus, if any */
-	FWeakWidgetPath FocusedWidgetPath;
 	/** The way in which the last focus was set. */
-	EKeyboardFocusCause::Type FocusCause;
+	EFocusCause FocusCause;
 	
 	/** Pointer to the currently registered game viewport widget if any */
 	TWeakPtr<SViewport> GameViewportWidget;
@@ -1400,7 +1421,7 @@ private:
 	bool bIsFakingTouched;
 
 	/** Delegate for when a key down event occurred but was not handled in any other way by ProcessKeyDownMessage */
-	FOnKeyboardEvent UnhandledKeyDownEventHandler;
+	FOnKeyEvent UnhandledKeyDownEventHandler;
 
 	/** controls whether unhandled touch events fall back to sending mouse events */
 	bool bTouchFallbackToMouse;
diff --git a/Engine/Source/Runtime/Slate/Public/Framework/Commands/UICommandList.h b/Engine/Source/Runtime/Slate/Public/Framework/Commands/UICommandList.h
--- a/Engine/Source/Runtime/Slate/Public/Framework/Commands/UICommandList.h
+++ b/Engine/Source/Runtime/Slate/Public/Framework/Commands/UICommandList.h
@@ -100,17 +100,17 @@ public:
 	bool IsChecked( const TSharedRef< const FUICommandInfo > InUICommandInfo ) const;
 
 	/**
-	 * Processes any UI commands which are activated by the specified keyboard event
+	 * Processes any UI commands which are activated by the specified key event
 	 *
-	 * @param InKeyboardEvent	The keyboard event to check
+	 * @param InKeyEvent	The key event to check
 	 * @return true if an action was processed
 	 */
-	bool ProcessCommandBindings( const FKeyboardEvent& InKeyboardEvent ) const;
+	bool ProcessCommandBindings( const FKeyEvent& InKeyEvent ) const;
 
 	/**
 	 * Processes any UI commands which are activated by the specified mouse event
 	 *
-	 * @param InKeyboardEvent	The mouse event to check
+	 * @param InKeyEvent	The mouse event to check
 	 * @return true if an action was processed
 	 */
 	bool ProcessCommandBindings( const FPointerEvent& InMouseEvent ) const;
diff --git a/Engine/Source/Runtime/Slate/Public/Framework/MultiBox/MultiBox.h b/Engine/Source/Runtime/Slate/Public/Framework/MultiBox/MultiBox.h
--- a/Engine/Source/Runtime/Slate/Public/Framework/MultiBox/MultiBox.h
+++ b/Engine/Source/Runtime/Slate/Public/Framework/MultiBox/MultiBox.h
@@ -579,14 +579,14 @@ public:
 	void OnDropExternal();
 
 	/** Helper function used to transfer focus to the next/previous widget */
-	static FReply FocusNextWidget( EFocusMoveDirection::Type MoveDirection );
+	static FReply FocusNextWidget(EUINavigation NavigationType);
 
 	/** SWidget interface */
 	virtual FReply OnDragOver( const FGeometry& MyGeometry, const FDragDropEvent& DragDropEvent ) override;
 	virtual FReply OnDrop( const FGeometry& MyGeometry, const FDragDropEvent& DragDropEvent ) override;
 	virtual bool SupportsKeyboardFocus() const override;
-	virtual FReply OnKeyboardFocusReceived( const FGeometry& MyGeometry, const FKeyboardFocusEvent& InKeyboardFocusEvent ) override;
-	virtual FReply OnKeyDown( const FGeometry& MyGeometry, const FKeyboardEvent& KeyboardEvent ) override;
+	virtual FReply OnFocusReceived( const FGeometry& MyGeometry, const FFocusEvent& InFocusEvent ) override;
+	virtual FReply OnKeyDown( const FGeometry& MyGeometry, const FKeyEvent& KeyEvent ) override;
 private:
 	/** Adds a block Widget to this widget */
 	void AddBlockWidget( const FMultiBlock& Block, TSharedPtr<SHorizontalBox> HorizontalBox, TSharedPtr<SVerticalBox> VerticalBox, EMultiBlockLocation::Type InLocation );
diff --git a/Engine/Source/Runtime/Slate/Public/Framework/SlateDelegates.h b/Engine/Source/Runtime/Slate/Public/Framework/SlateDelegates.h
--- a/Engine/Source/Runtime/Slate/Public/Framework/SlateDelegates.h
+++ b/Engine/Source/Runtime/Slate/Public/Framework/SlateDelegates.h
@@ -205,8 +205,8 @@ public:
 };
 
 
-/** Notification for when a keyboard event occurs */
-DECLARE_DELEGATE_RetVal_OneParam( FReply, FOnKeyboardEvent, const FKeyboardEvent& )
+/** Notification for when a key event occurs */
+DECLARE_DELEGATE_RetVal_OneParam( FReply, FOnKeyEvent, const FKeyEvent& )
 
 /** A Delegate for passing along a string of a source code location to access */
 DECLARE_DELEGATE_RetVal_ThreeParams(bool, FAccessSourceCode, const FString& /*FileName*/, int32 /*InLineNumber*/, int32 /*InColumnNumber*/);
diff --git a/Engine/Source/Runtime/Slate/Public/Widgets/Input/SButton.h b/Engine/Source/Runtime/Slate/Public/Widgets/Input/SButton.h
--- a/Engine/Source/Runtime/Slate/Public/Widgets/Input/SButton.h
+++ b/Engine/Source/Runtime/Slate/Public/Widgets/Input/SButton.h
@@ -21,6 +21,7 @@ public:
 		, _Text()
 		, _ClickMethod( EButtonClickMethod::DownAndUp )
 		, _TouchMethod( EButtonTouchMethod::DownAndUp )
+		, _PressMethod( EButtonPressMethod::DownAndUp )
 		, _DesiredSizeScale( FVector2D(1,1) )
 		, _ContentScale( FVector2D(1,1) )
 		, _ButtonColorAndOpacity(FLinearColor::White)
@@ -58,6 +59,9 @@ public:
 		/** How should the button be clicked with touch events? */
 		SLATE_ARGUMENT( EButtonTouchMethod::Type, TouchMethod )
 
+		/** How should the button be clicked with keyboard/controller button events? */
+		SLATE_ARGUMENT( EButtonPressMethod::Type, PressMethod )
+
 		SLATE_ATTRIBUTE( FVector2D, DesiredSizeScale )
 
 		SLATE_ATTRIBUTE( FVector2D, ContentScale )
@@ -119,7 +123,11 @@ public:
 
 	virtual bool SupportsKeyboardFocus() const override;
 
-	virtual FReply OnKeyDown( const FGeometry& MyGeometry, const FKeyboardEvent& InKeyboardEvent ) override;
+	virtual void OnFocusLost( const FFocusEvent& InFocusEvent ) override;
+
+	virtual FReply OnKeyDown( const FGeometry& MyGeometry, const FKeyEvent& InKeyEvent ) override;
+
+	virtual FReply OnKeyUp( const FGeometry& MyGeometry, const FKeyEvent& InKeyEvent ) override;
 
 	virtual FReply OnMouseButtonDown( const FGeometry& MyGeometry, const FPointerEvent& MouseEvent ) override;
 
@@ -169,6 +177,9 @@ protected:
 	/** How should the button be clicked with touch events? */
 	EButtonTouchMethod::Type TouchMethod;
 
+	/** How should the button be clicked with keyboard/controller button events? */
+	EButtonPressMethod::Type PressMethod;
+
 	/** Can this button be focused? */
 	bool bIsFocusable;
 
diff --git a/Engine/Source/Runtime/Slate/Public/Widgets/Input/SCheckBox.h b/Engine/Source/Runtime/Slate/Public/Widgets/Input/SCheckBox.h
--- a/Engine/Source/Runtime/Slate/Public/Widgets/Input/SCheckBox.h
+++ b/Engine/Source/Runtime/Slate/Public/Widgets/Input/SCheckBox.h
@@ -152,7 +152,7 @@ public:
 
 	// SWidget interface
 	virtual bool SupportsKeyboardFocus() const override;
-	virtual FReply OnKeyUp( const FGeometry& MyGeometry, const FKeyboardEvent& InKeyboardEvent ) override;
+	virtual FReply OnKeyUp( const FGeometry& MyGeometry, const FKeyEvent& InKeyEvent ) override;
 	virtual FReply OnMouseButtonDown( const FGeometry& MyGeometry, const FPointerEvent& MouseEvent ) override;
 	virtual FReply OnMouseButtonDoubleClick( const FGeometry& InMyGeometry, const FPointerEvent& InMouseEvent ) override;
 	virtual FReply OnMouseButtonUp( const FGeometry& MyGeometry, const FPointerEvent& MouseEvent ) override;
diff --git a/Engine/Source/Runtime/Slate/Public/Widgets/Input/SComboBox.h b/Engine/Source/Runtime/Slate/Public/Widgets/Input/SComboBox.h
--- a/Engine/Source/Runtime/Slate/Public/Widgets/Input/SComboBox.h
+++ b/Engine/Source/Runtime/Slate/Public/Widgets/Input/SComboBox.h
@@ -244,9 +244,9 @@ protected:
 		return FReply::Unhandled();
 	}
 
-	FReply OnKeyDown( const FGeometry& MyGeometry, const FKeyboardEvent& InKeyboardEvent ) override
+	FReply OnKeyDown( const FGeometry& MyGeometry, const FKeyEvent& InKeyEvent ) override
 	{
-		const FKey Key = InKeyboardEvent.GetKey();
+		const FKey Key = InKeyEvent.GetKey();
 
 		if(Key == EKeys::Up)
 		{
@@ -271,7 +271,7 @@ protected:
 			return FReply::Handled();
 		}
 
-		return SComboButton::OnKeyDown( MyGeometry, InKeyboardEvent );
+		return SComboButton::OnKeyDown( MyGeometry, InKeyEvent );
 	}
 
 	virtual bool SupportsKeyboardFocus() const override
diff --git a/Engine/Source/Runtime/Slate/Public/Widgets/Input/SEditableComboBox.h b/Engine/Source/Runtime/Slate/Public/Widgets/Input/SEditableComboBox.h
--- a/Engine/Source/Runtime/Slate/Public/Widgets/Input/SEditableComboBox.h
+++ b/Engine/Source/Runtime/Slate/Public/Widgets/Input/SEditableComboBox.h
@@ -273,7 +273,7 @@ private:
 	// Callback for clicking the 'Accept' button.
 	FReply HandleAcceptButtonClicked( )
 	{
-		return FReply::Handled().SetKeyboardFocus(ComboBox->AsShared(), EKeyboardFocusCause::Mouse);
+		return FReply::Handled().SetUserFocus(ComboBox->AsShared(), EFocusCause::Mouse);
 	}
 
 	// Callback for getting the visibility of the edit mode box.
@@ -304,7 +304,7 @@ private:
 			TextBox->SetText( FText::FromString( OnGetEditableText.Execute() ));
 		}
 
-		return FReply::Handled().SetKeyboardFocus(TextBox->AsShared(), EKeyboardFocusCause::Mouse);
+		return FReply::Handled().SetUserFocus(TextBox->AsShared(), EFocusCause::Mouse);
 	}
 
 	// Callback for committing the text in the text box.
diff --git a/Engine/Source/Runtime/Slate/Public/Widgets/Input/SEditableText.h b/Engine/Source/Runtime/Slate/Public/Widgets/Input/SEditableText.h
--- a/Engine/Source/Runtime/Slate/Public/Widgets/Input/SEditableText.h
+++ b/Engine/Source/Runtime/Slate/Public/Widgets/Input/SEditableText.h
@@ -349,11 +349,11 @@ protected:
 	virtual FReply OnDragOver( const FGeometry& MyGeometry, const FDragDropEvent& DragDropEvent ) override;
 	virtual FReply OnDrop( const FGeometry& MyGeometry, const FDragDropEvent& DragDropEvent ) override;
 	virtual bool SupportsKeyboardFocus() const override;
-	virtual FReply OnKeyboardFocusReceived( const FGeometry& MyGeometry, const FKeyboardFocusEvent& InKeyboardFocusEvent ) override;
-	virtual void OnKeyboardFocusLost( const FKeyboardFocusEvent& InKeyboardFocusEvent ) override;
+	virtual FReply OnFocusReceived( const FGeometry& MyGeometry, const FFocusEvent& InFocusEvent ) override;
+	virtual void OnFocusLost( const FFocusEvent& InFocusEvent ) override;
 	virtual FReply OnKeyChar( const FGeometry& MyGeometry,  const FCharacterEvent& InCharacterEvent ) override;
-	virtual FReply OnKeyDown( const FGeometry& MyGeometry, const FKeyboardEvent& InKeyboardEvent ) override;
-	virtual FReply OnKeyUp( const FGeometry& MyGeometry, const FKeyboardEvent& InKeyboardEvent ) override;
+	virtual FReply OnKeyDown( const FGeometry& MyGeometry, const FKeyEvent& InKeyEvent ) override;
+	virtual FReply OnKeyUp( const FGeometry& MyGeometry, const FKeyEvent& InKeyEvent ) override;
 	virtual FReply OnMouseButtonDown( const FGeometry& InMyGeometry, const FPointerEvent& InMouseEvent ) override;
 	virtual FReply OnMouseButtonUp( const FGeometry& InMyGeometry, const FPointerEvent& InMouseEvent ) override;
 	virtual FReply OnMouseMove( const FGeometry& InMyGeometry, const FPointerEvent& InMouseEvent ) override;
diff --git a/Engine/Source/Runtime/Slate/Public/Widgets/Input/SEditableTextBox.h b/Engine/Source/Runtime/Slate/Public/Widgets/Input/SEditableTextBox.h
--- a/Engine/Source/Runtime/Slate/Public/Widgets/Input/SEditableTextBox.h
+++ b/Engine/Source/Runtime/Slate/Public/Widgets/Input/SEditableTextBox.h
@@ -135,8 +135,8 @@ public:
 	// SWidget overrides
 	virtual bool SupportsKeyboardFocus() const override;
 	virtual bool HasKeyboardFocus() const override;
-	virtual FReply OnKeyboardFocusReceived( const FGeometry& MyGeometry, const FKeyboardFocusEvent& InKeyboardFocusEvent ) override;
-	virtual FReply OnKeyDown( const FGeometry& MyGeometry, const FKeyboardEvent& InKeyboardEvent ) override;
+	virtual FReply OnFocusReceived( const FGeometry& MyGeometry, const FFocusEvent& InFocusEvent ) override;
+	virtual FReply OnKeyDown( const FGeometry& MyGeometry, const FKeyEvent& InKeyEvent ) override;
 
 protected:
 	const FEditableTextBoxStyle* Style;
diff --git a/Engine/Source/Runtime/Slate/Public/Widgets/Input/SNumericEntryBox.h b/Engine/Source/Runtime/Slate/Public/Widgets/Input/SNumericEntryBox.h
--- a/Engine/Source/Runtime/Slate/Public/Widgets/Input/SNumericEntryBox.h
+++ b/Engine/Source/Runtime/Slate/Public/Widgets/Input/SNumericEntryBox.h
@@ -216,7 +216,7 @@ private:
 		return StaticCastSharedPtr<SWidget>(EditableText)->SupportsKeyboardFocus();
 	}
 
-	virtual FReply OnKeyboardFocusReceived( const FGeometry& MyGeometry, const FKeyboardFocusEvent& InKeyboardFocusEvent ) override
+	virtual FReply OnFocusReceived( const FGeometry& MyGeometry, const FFocusEvent& InFocusEvent ) override
 	{
 		FReply Reply = FReply::Handled();
 
@@ -231,22 +231,22 @@ private:
 			FocusWidget = EditableText;
 		}
 
-		if ( InKeyboardFocusEvent.GetCause() != EKeyboardFocusCause::Cleared )
+		if ( InFocusEvent.GetCause() != EFocusCause::Cleared )
 		{
 			// Forward keyboard focus to our chosen widget
-			Reply.SetKeyboardFocus( FocusWidget.ToSharedRef(), InKeyboardFocusEvent.GetCause() );
+			Reply.SetUserFocus(FocusWidget.ToSharedRef(), InFocusEvent.GetCause());
 		}
 
 		return Reply;
 	}
 
-	virtual FReply OnKeyDown( const FGeometry& MyGeometry, const FKeyboardEvent& InKeyboardEvent ) override
+	virtual FReply OnKeyDown( const FGeometry& MyGeometry, const FKeyEvent& InKeyEvent ) override
 	{
-		FKey Key = InKeyboardEvent.GetKey();
+		FKey Key = InKeyEvent.GetKey();
 
 		if( Key == EKeys::Escape && EditableText->HasKeyboardFocus() )
 		{
-			return FReply::Handled().SetKeyboardFocus( SharedThis( this ), EKeyboardFocusCause::Cleared );
+			return FReply::Handled().SetUserFocus(SharedThis(this), EFocusCause::Cleared);
 		}
 
 		return FReply::Unhandled();
diff --git a/Engine/Source/Runtime/Slate/Public/Widgets/Input/SSpinBox.h b/Engine/Source/Runtime/Slate/Public/Widgets/Input/SSpinBox.h
--- a/Engine/Source/Runtime/Slate/Public/Widgets/Input/SSpinBox.h
+++ b/Engine/Source/Runtime/Slate/Public/Widgets/Input/SSpinBox.h
@@ -249,7 +249,7 @@ public:
 			DistanceDragged = 0;
 			PreDragValue = InternalValue;
 			CachedMousePosition = MouseEvent.GetScreenSpacePosition().IntPoint();
-			return FReply::Handled().CaptureMouse( SharedThis(this) ).UseHighPrecisionMouseMovement( SharedThis(this) ).SetKeyboardFocus( SharedThis(this), EKeyboardFocusCause::Mouse);
+			return FReply::Handled().CaptureMouse(SharedThis(this)).UseHighPrecisionMouseMovement(SharedThis(this)).SetUserFocus(SharedThis(this), EFocusCause::Mouse);
 		}
 		else
 		{
@@ -278,7 +278,7 @@ public:
 			if ( DistanceDragged < FSlateApplication::Get().GetDragTriggerDistnace() )
 			{
 				EnterTextMode();
-				return FReply::Handled().ReleaseMouseCapture().SetKeyboardFocus( EditableText.ToSharedRef(), EKeyboardFocusCause::SetDirectly ).SetMousePos(CachedMousePosition);
+				return FReply::Handled().ReleaseMouseCapture().SetUserFocus(EditableText.ToSharedRef(), EFocusCause::SetDirectly).SetMousePos(CachedMousePosition);
 			}
 
 			return FReply::Handled().ReleaseMouseCapture().SetMousePos(CachedMousePosition);
@@ -385,12 +385,12 @@ public:
 	}
 
 
-	virtual FReply OnKeyboardFocusReceived( const FGeometry& MyGeometry, const FKeyboardFocusEvent& InKeyboardFocusEvent ) override
+	virtual FReply OnFocusReceived( const FGeometry& MyGeometry, const FFocusEvent& InFocusEvent ) override
 	{
-		if ( !bDragging && (InKeyboardFocusEvent.GetCause() == EKeyboardFocusCause::Keyboard || InKeyboardFocusEvent.GetCause() == EKeyboardFocusCause::SetDirectly) )
+		if ( !bDragging && (InFocusEvent.GetCause() == EFocusCause::Navigation || InFocusEvent.GetCause() == EFocusCause::SetDirectly) )
 		{
 			EnterTextMode();
-			return FReply::Handled().SetKeyboardFocus( EditableText.ToSharedRef(), InKeyboardFocusEvent.GetCause() );
+			return FReply::Handled().SetUserFocus(EditableText.ToSharedRef(), InFocusEvent.GetCause());
 		}
 		else
 		{
@@ -398,9 +398,9 @@ public:
 		}
 	}
 	
-	virtual FReply OnKeyDown( const FGeometry& MyGeometry, const FKeyboardEvent& InKeyboardEvent ) override
+	virtual FReply OnKeyDown( const FGeometry& MyGeometry, const FKeyEvent& InKeyEvent ) override
 	{
-		const FKey Key = InKeyboardEvent.GetKey();
+		const FKey Key = InKeyEvent.GetKey();
 		if ( Key == EKeys::Escape && HasMouseCapture())
 		{
 			bDragging = false;
@@ -424,7 +424,7 @@ public:
 		else if ( Key == EKeys::Enter )
 		{
 			EnterTextMode();
-			return FReply::Handled().SetKeyboardFocus( EditableText.ToSharedRef(), EKeyboardFocusCause::Keyboard );
+			return FReply::Handled().SetUserFocus(EditableText.ToSharedRef(), EFocusCause::Navigation);
 		}
 		else
 		{
diff --git a/Engine/Source/Runtime/Slate/Public/Widgets/Input/SSuggestionTextBox.h b/Engine/Source/Runtime/Slate/Public/Widgets/Input/SSuggestionTextBox.h
--- a/Engine/Source/Runtime/Slate/Public/Widgets/Input/SSuggestionTextBox.h
+++ b/Engine/Source/Runtime/Slate/Public/Widgets/Input/SSuggestionTextBox.h
@@ -183,8 +183,8 @@ protected:
 
 	// SWidget overrides
 
-	virtual void OnKeyboardFocusLost( const FKeyboardFocusEvent& InKeyboardFocusEvent ) override;
-	virtual FReply OnKeyDown( const FGeometry& MyGeometry, const FKeyboardEvent& KeyboardEvent ) override;
+	virtual void OnFocusLost( const FFocusEvent& InFocusEvent ) override;
+	virtual FReply OnKeyDown( const FGeometry& MyGeometry, const FKeyEvent& KeyEvent ) override;
 	virtual bool SupportsKeyboardFocus( ) const override;
 
 private:
diff --git a/Engine/Source/Runtime/Slate/Public/Widgets/Input/SVirtualKeyboardEntry.h b/Engine/Source/Runtime/Slate/Public/Widgets/Input/SVirtualKeyboardEntry.h
--- a/Engine/Source/Runtime/Slate/Public/Widgets/Input/SVirtualKeyboardEntry.h
+++ b/Engine/Source/Runtime/Slate/Public/Widgets/Input/SVirtualKeyboardEntry.h
@@ -118,8 +118,8 @@ protected:
 	virtual FVector2D ComputeDesiredSize() const override;
 	virtual int32 OnPaint( const FPaintArgs& Args, const FGeometry& AllottedGeometry, const FSlateRect& MyClippingRect, FSlateWindowElementList& OutDrawElements, int32 LayerId, const FWidgetStyle& InWidgetStyle, bool bParentEnabled ) const override;
 	virtual bool SupportsKeyboardFocus() const override;
-	virtual FReply OnKeyboardFocusReceived( const FGeometry& MyGeometry, const FKeyboardFocusEvent& InKeyboardFocusEvent ) override;
-	virtual void OnKeyboardFocusLost( const FKeyboardFocusEvent& InKeyboardFocusEvent ) override;
+	virtual FReply OnFocusReceived( const FGeometry& MyGeometry, const FFocusEvent& InFocusEvent ) override;
+	virtual void OnFocusLost( const FFocusEvent& InFocusEvent ) override;
 	// End of SWidget interface
 
 private:
diff --git a/Engine/Source/Runtime/Slate/Public/Widgets/SViewport.h b/Engine/Source/Runtime/Slate/Public/Widgets/SViewport.h
--- a/Engine/Source/Runtime/Slate/Public/Widgets/SViewport.h
+++ b/Engine/Source/Runtime/Slate/Public/Widgets/SViewport.h
@@ -132,9 +132,9 @@ public:
 
 	virtual int32 OnPaint( const FPaintArgs& Args, const FGeometry& AllottedGeometry, const FSlateRect& MyClippingRect, FSlateWindowElementList& OutDrawElements, int32 LayerId, const FWidgetStyle& InWidgetStyle, bool bParentEnabled ) const override;
 	virtual void Tick( const FGeometry& AllottedGeometry, const double InCurrentTime, const float InDeltaTime ) override;
-	virtual FReply OnTouchStarted( const FGeometry& MyGeometry, const FPointerEvent& TouchEvent ) override;
-	virtual FReply OnTouchMoved( const FGeometry& MyGeometry, const FPointerEvent& TouchEvent ) override;
-	virtual FReply OnTouchEnded( const FGeometry& MyGeometry, const FPointerEvent& TouchEvent ) override;
+	virtual FReply OnTouchStarted( const FGeometry& MyGeometry, const FPointerEvent& InTouchEvent ) override;
+	virtual FReply OnTouchMoved( const FGeometry& MyGeometry, const FPointerEvent& InTouchEvent ) override;
+	virtual FReply OnTouchEnded( const FGeometry& MyGeometry, const FPointerEvent& InTouchEvent ) override;
 	virtual FReply OnTouchGesture( const FGeometry& MyGeometry, const FPointerEvent& GestureEvent ) override;
 	virtual FCursorReply OnCursorQuery( const FGeometry& MyGeometry, const FPointerEvent& CursorEvent ) const override;
 	virtual FReply OnMouseButtonDown( const FGeometry& MyGeometry, const FPointerEvent& MouseEvent ) override;
@@ -144,18 +144,17 @@ public:
 	virtual FReply OnMouseMove( const FGeometry& MyGeometry, const FPointerEvent& MouseEvent ) override;
 	virtual FReply OnMouseWheel( const FGeometry& MyGeometry, const FPointerEvent& MouseEvent ) override;
 	virtual FReply OnMouseButtonDoubleClick( const FGeometry& MyGeometry, const FPointerEvent& MouseEvent ) override;
-	virtual FReply OnKeyDown( const FGeometry& MyGeometry, const FKeyboardEvent& KeyboardEvent ) override;
-	virtual FReply OnKeyUp( const FGeometry& MyGeometry, const FKeyboardEvent& KeyboardEvent ) override;
+	virtual FReply OnKeyDown( const FGeometry& MyGeometry, const FKeyEvent& KeyEvent ) override;
+	virtual FReply OnKeyUp( const FGeometry& MyGeometry, const FKeyEvent& KeyEvent ) override;
+	virtual FReply OnAnalogValueChanged( const FGeometry& MyGeometry, const FAnalogInputEvent& InAnalogInputEvent ) override;
 	virtual FReply OnKeyChar( const FGeometry& MyGeometry, const FCharacterEvent& CharacterEvent ) override;
-	virtual FReply OnKeyboardFocusReceived( const FGeometry& MyGeometry, const FKeyboardFocusEvent& InKeyboardFocusEvent ) override;
-	virtual void OnKeyboardFocusLost( const FKeyboardFocusEvent& InKeyboardFocusEvent ) override;
-	virtual FReply OnControllerButtonPressed( const FGeometry& MyGeometry, const FControllerEvent& ControllerEvent ) override;
-	virtual FReply OnControllerButtonReleased( const FGeometry& MyGeometry, const FControllerEvent& ControllerEvent ) override;
-	virtual FReply OnControllerAnalogValueChanged( const FGeometry& MyGeometry, const FControllerEvent& ControllerEvent ) override;
-	virtual FReply OnMotionDetected( const FGeometry& MyGeometry, const FMotionEvent& MotionEvent ) override;
+	virtual FReply OnFocusReceived( const FGeometry& MyGeometry, const FFocusEvent& InFocusEvent ) override;
+	virtual void OnFocusLost( const FFocusEvent& InFocusEvent ) override;
+	virtual FReply OnMotionDetected( const FGeometry& MyGeometry, const FMotionEvent& InMotionEvent ) override;
 	virtual void OnFinishedPointerInput() override;
 	virtual void OnArrangeChildren( const FGeometry& AllottedGeometry, FArrangedChildren& ArrangedChildren ) const override;
 	virtual TSharedPtr<struct FVirtualPointerPosition> TranslateMouseCoordinateFor3DChild(const TSharedRef<SWidget>& ChildWidget, const FGeometry& MyGeometry, const FVector2D& ScreenSpaceMouseCoordinate, const FVector2D& LastScreenSpaceMouseCoordinate) const;
+	virtual FNavigationReply OnNavigation( const FGeometry& MyGeometry, const FNavigationEvent& InNavigationEvent ) override;
 
 private:
 
diff --git a/Engine/Source/Runtime/Slate/Public/Widgets/Text/SMultiLineEditableText.h b/Engine/Source/Runtime/Slate/Public/Widgets/Text/SMultiLineEditableText.h
--- a/Engine/Source/Runtime/Slate/Public/Widgets/Text/SMultiLineEditableText.h
+++ b/Engine/Source/Runtime/Slate/Public/Widgets/Text/SMultiLineEditableText.h
@@ -429,10 +429,10 @@ private:
 	virtual void OnArrangeChildren( const FGeometry& AllottedGeometry, FArrangedChildren& ArrangedChildren ) const override;
 	virtual bool SupportsKeyboardFocus() const override;
 	virtual FReply OnKeyChar( const FGeometry& MyGeometry,const FCharacterEvent& InCharacterEvent ) override;
-	virtual FReply OnKeyDown( const FGeometry& MyGeometry, const FKeyboardEvent& InKeyboardEvent ) override;
-	virtual FReply OnKeyUp( const FGeometry& MyGeometry, const FKeyboardEvent& InKeyboardEvent ) override;
-	virtual FReply OnKeyboardFocusReceived( const FGeometry& MyGeometry, const FKeyboardFocusEvent& InKeyboardFocusEvent ) override;
-	virtual void OnKeyboardFocusLost( const FKeyboardFocusEvent& InKeyboardFocusEvent ) override;
+	virtual FReply OnKeyDown( const FGeometry& MyGeometry, const FKeyEvent& InKeyEvent ) override;
+	virtual FReply OnKeyUp( const FGeometry& MyGeometry, const FKeyEvent& InKeyEvent ) override;
+	virtual FReply OnFocusReceived( const FGeometry& MyGeometry, const FFocusEvent& InFocusEvent ) override;
+	virtual void OnFocusLost( const FFocusEvent& InFocusEvent ) override;
 	virtual FReply OnMouseButtonDown( const FGeometry& MyGeometry, const FPointerEvent& MouseEvent ) override;
 	virtual FReply OnMouseButtonUp( const FGeometry& MyGeometry, const FPointerEvent& MouseEvent ) override;
 	virtual FReply OnMouseMove( const FGeometry& MyGeometry, const FPointerEvent& MouseEvent ) override;
diff --git a/Engine/Source/Runtime/Slate/Public/Widgets/Views/SListView.h b/Engine/Source/Runtime/Slate/Public/Widgets/Views/SListView.h
--- a/Engine/Source/Runtime/Slate/Public/Widgets/Views/SListView.h
+++ b/Engine/Source/Runtime/Slate/Public/Widgets/Views/SListView.h
@@ -169,30 +169,30 @@ public:
 
 	// SWidget overrides
 
-	virtual FReply OnKeyDown( const FGeometry& MyGeometry, const FKeyboardEvent& InKeyboardEvent ) override
+	virtual FReply OnKeyDown( const FGeometry& MyGeometry, const FKeyEvent& InKeyEvent ) override
 	{
 		const TArray<ItemType>& ItemsSourceRef = (*this->ItemsSource);
 
 		// Don't respond to key-presses containing "Alt" as a modifier
-		if ( ItemsSourceRef.Num() > 0 && !InKeyboardEvent.IsAltDown() )
+		if ( ItemsSourceRef.Num() > 0 && !InKeyEvent.IsAltDown() )
 		{
 			bool bWasHandled = false;
 			NullableItemType ItemNavigatedTo( nullptr );
 
 			// Check for selection manipulation keys (Up, Down, Home, End, PageUp, PageDown)
-			if ( InKeyboardEvent.GetKey() == EKeys::Home )
+			if ( InKeyEvent.GetKey() == EKeys::Home )
 			{
 				// Select the first item
 				ItemNavigatedTo = ItemsSourceRef[0];
 				bWasHandled = true;
 			}
-			else if ( InKeyboardEvent.GetKey() == EKeys::End )
+			else if ( InKeyEvent.GetKey() == EKeys::End )
 			{
 				// Select the last item
 				ItemNavigatedTo = ItemsSourceRef.Last();
 				bWasHandled = true;
 			}
-			else if ( InKeyboardEvent.GetKey() == EKeys::PageUp )
+			else if ( InKeyEvent.GetKey() == EKeys::PageUp )
 			{
 				int32 SelectionIndex = 0;
 				if( TListTypeTraits<ItemType>::IsPtrValid(SelectorItem) )
@@ -216,7 +216,7 @@ public:
 
 				bWasHandled = true;
 			}
-			else if ( InKeyboardEvent.GetKey() == EKeys::PageDown )
+			else if ( InKeyEvent.GetKey() == EKeys::PageDown )
 			{
 				int32 SelectionIndex = 0;
 				if( TListTypeTraits<ItemType>::IsPtrValid(SelectorItem) )
@@ -240,7 +240,7 @@ public:
 
 				bWasHandled = true;
 			}
-			else if ( InKeyboardEvent.GetKey() == EKeys::Up )
+			else if ( InKeyEvent.GetKey() == EKeys::Up )
 			{
 				int32 SelectionIndex = 0;
 				if( TListTypeTraits<ItemType>::IsPtrValid(SelectorItem) )
@@ -258,7 +258,7 @@ public:
 
 				bWasHandled = true;
 			}
-			else if ( InKeyboardEvent.GetKey() == EKeys::Down )
+			else if ( InKeyEvent.GetKey() == EKeys::Down )
 			{
 				// Begin at INDEX_NONE so the first item will get selected
 				int32 SelectionIndex = INDEX_NONE;
@@ -281,17 +281,17 @@ public:
 			if( TListTypeTraits<ItemType>::IsPtrValid(ItemNavigatedTo) )
 			{
 				ItemType ItemToSelect( TListTypeTraits<ItemType>::NullableItemTypeConvertToItemType( ItemNavigatedTo ) );
-				KeyboardSelect( ItemToSelect, InKeyboardEvent );
+				KeyboardSelect( ItemToSelect, InKeyEvent );
 			}
 			else
 			{
 				// Change selected status of item.
-				if( TListTypeTraits<ItemType>::IsPtrValid(SelectorItem) && InKeyboardEvent.GetKey() == EKeys::SpaceBar )
+				if( TListTypeTraits<ItemType>::IsPtrValid(SelectorItem) && InKeyEvent.GetKey() == EKeys::SpaceBar )
 				{
 					ItemType SelectorItemDereference( TListTypeTraits<ItemType>::NullableItemTypeConvertToItemType( SelectorItem ) );
 
 					// Deselect.
-					if( InKeyboardEvent.IsControlDown() || this->SelectionMode == ESelectionMode::SingleToggle )
+					if( InKeyEvent.IsControlDown() || this->SelectionMode == ESelectionMode::SingleToggle )
 					{
 						this->Private_SetItemSelection( SelectorItemDereference, !( this->Private_IsItemSelected( SelectorItemDereference ) ), true );
 						this->Private_SignalSelectionChanged( ESelectInfo::OnKeyPress );
@@ -323,7 +323,7 @@ public:
 					}
 				}
 				// Select all items
-				else if ( (!InKeyboardEvent.IsShiftDown() && !InKeyboardEvent.IsAltDown() && InKeyboardEvent.IsControlDown() && InKeyboardEvent.GetKey() == EKeys::A) && this->SelectionMode == ESelectionMode::Multi )
+				else if ( (!InKeyEvent.IsShiftDown() && !InKeyEvent.IsAltDown() && InKeyEvent.IsControlDown() && InKeyEvent.GetKey() == EKeys::A) && this->SelectionMode == ESelectionMode::Multi )
 				{
 					this->Private_ClearSelection();
 
@@ -341,7 +341,7 @@ public:
 			return (bWasHandled ? FReply::Handled() : FReply::Unhandled());
 		}
 
-		return STableViewBase::OnKeyDown(MyGeometry, InKeyboardEvent);
+		return STableViewBase::OnKeyDown(MyGeometry, InKeyEvent);
 	}
 
 	virtual FReply OnMouseButtonDown( const FGeometry& MyGeometry, const FPointerEvent& MouseEvent ) override
@@ -1277,22 +1277,22 @@ protected:
 	 * Selects the specified item and scrolls it into view. If shift is held, it will be a range select.
 	 * 
 	 * @param ItemToSelect		The item that was selected by a keystroke
-	 * @param InKeyboardEvent	The keyboard event that caused this selection
+	 * @param InKeyEvent	The key event that caused this selection
 	 */
-	virtual void KeyboardSelect(const ItemType& ItemToSelect, const FKeyboardEvent& InKeyboardEvent, bool bCausedByNavigation=false)
+	virtual void KeyboardSelect(const ItemType& ItemToSelect, const FKeyEvent& InKeyEvent, bool bCausedByNavigation=false)
 	{
 		if ( SelectionMode != ESelectionMode::None )
 		{
 			// Must be set before signaling selection changes because sometimes new items will be selected that need to stomp this value
 			SelectorItem = ItemToSelect;
 
-			if ( SelectionMode == ESelectionMode::Multi && ( InKeyboardEvent.IsShiftDown() || InKeyboardEvent.IsControlDown() ) )
+			if ( SelectionMode == ESelectionMode::Multi && ( InKeyEvent.IsShiftDown() || InKeyEvent.IsControlDown() ) )
 			{
 				// Range select.
-				if ( InKeyboardEvent.IsShiftDown() )
+				if ( InKeyEvent.IsShiftDown() )
 				{
 					// Holding control makes the range select bidirectional, where as it is normally unidirectional.
-					if( !( InKeyboardEvent.IsControlDown() ) )
+					if( !( InKeyEvent.IsControlDown() ) )
 					{
 						this->Private_ClearSelection();
 					}
diff --git a/Engine/Source/Runtime/Slate/Public/Widgets/Views/STableRow.h b/Engine/Source/Runtime/Slate/Public/Widgets/Views/STableRow.h
--- a/Engine/Source/Runtime/Slate/Public/Widgets/Views/STableRow.h
+++ b/Engine/Source/Runtime/Slate/Public/Widgets/Views/STableRow.h
@@ -221,7 +221,7 @@ public:
 					
 					return FReply::Handled()
 						.DetectDrag( SharedThis(this), EKeys::LeftMouseButton )
-						.SetKeyboardFocus( OwnerWidget->AsWidget(), EKeyboardFocusCause::Mouse )
+						.SetUserFocus(OwnerWidget->AsWidget(), EFocusCause::Mouse)
 						.CaptureMouse( SharedThis(this) );
 				}
 				break;
@@ -240,7 +240,7 @@ public:
 
 					return FReply::Handled()
 						.DetectDrag( SharedThis(this), EKeys::LeftMouseButton )
-						.SetKeyboardFocus( OwnerWidget->AsWidget(), EKeyboardFocusCause::Mouse )
+						.SetUserFocus(OwnerWidget->AsWidget(), EFocusCause::Mouse)
 						.CaptureMouse( SharedThis(this) );
 				}
 				break;
@@ -271,7 +271,7 @@ public:
 
 					return FReply::Handled()
 						.DetectDrag( SharedThis(this), EKeys::LeftMouseButton )
-						.SetKeyboardFocus( OwnerWidget->AsWidget(), EKeyboardFocusCause::Mouse )
+						.SetUserFocus(OwnerWidget->AsWidget(), EFocusCause::Mouse)
 						.CaptureMouse( SharedThis(this) );
 				}
 				break;
diff --git a/Engine/Source/Runtime/Slate/Public/Widgets/Views/STableViewBase.h b/Engine/Source/Runtime/Slate/Public/Widgets/Views/STableViewBase.h
--- a/Engine/Source/Runtime/Slate/Public/Widgets/Views/STableViewBase.h
+++ b/Engine/Source/Runtime/Slate/Public/Widgets/Views/STableViewBase.h
@@ -84,7 +84,7 @@ public:
 
 	// SWidget interface
 
-	virtual void OnKeyboardFocusLost( const FKeyboardFocusEvent& InKeyboardFocusEvent ) override;
+	virtual void OnFocusLost( const FFocusEvent& InFocusEvent ) override;
 	virtual void OnMouseCaptureLost() override;
 	virtual bool SupportsKeyboardFocus() const override;
 	virtual void Tick( const FGeometry& AllottedGeometry, const double InCurrentTime, const float InDeltaTime ) override;
@@ -95,7 +95,7 @@ public:
 	virtual void OnMouseLeave( const FPointerEvent& MouseEvent ) override;
 	virtual FReply OnMouseMove( const FGeometry& MyGeometry, const FPointerEvent& MouseEvent ) override;
 	virtual FReply OnMouseWheel( const FGeometry& MyGeometry, const FPointerEvent& MouseEvent ) override;
-	virtual FReply OnKeyDown( const FGeometry& MyGeometry, const FKeyboardEvent& InKeyboardEvent ) override;
+	virtual FReply OnKeyDown( const FGeometry& MyGeometry, const FKeyEvent& InKeyEvent ) override;
 	virtual FCursorReply OnCursorQuery( const FGeometry& MyGeometry, const FPointerEvent& CursorEvent ) const override;
 	virtual int32 OnPaint(const FPaintArgs& Args, const FGeometry& AllottedGeometry, const FSlateRect& MyClippingRect, FSlateWindowElementList& OutDrawElements, int32 LayerId, const FWidgetStyle& InWidgetStyle, bool bParentEnabled) const override;
 	virtual FReply OnTouchStarted( const FGeometry& MyGeometry, const FPointerEvent& InTouchEvent ) override;
diff --git a/Engine/Source/Runtime/Slate/Public/Widgets/Views/STileView.h b/Engine/Source/Runtime/Slate/Public/Widgets/Views/STileView.h
--- a/Engine/Source/Runtime/Slate/Public/Widgets/Views/STileView.h
+++ b/Engine/Source/Runtime/Slate/Public/Widgets/Views/STileView.h
@@ -131,41 +131,41 @@ public:
 
 	// SWidget overrides
 
-	virtual FReply OnKeyDown( const FGeometry& MyGeometry, const FKeyboardEvent& InKeyboardEvent ) override
+	virtual FReply OnKeyDown( const FGeometry& MyGeometry, const FKeyEvent& InKeyEvent ) override
 	{
 		const TArray<ItemType>& ItemsSourceRef = (*this->ItemsSource);
 
 		// Don't respond to key-presses containing "Alt" as a modifier
-		if ( ItemsSourceRef.Num() > 0 && !InKeyboardEvent.IsAltDown() )
+		if ( ItemsSourceRef.Num() > 0 && !InKeyEvent.IsAltDown() )
 		{
 			// Check for selection manipulation keys that differ from SListView (Left, Right)
-			if ( InKeyboardEvent.GetKey() == EKeys::Left )
+			if ( InKeyEvent.GetKey() == EKeys::Left )
 			{
 				int32 SelectionIndex = ( !TListTypeTraits<ItemType>::IsPtrValid(this->SelectorItem) ) ? 0 : ItemsSourceRef.Find( TListTypeTraits<ItemType>::NullableItemTypeConvertToItemType( this->SelectorItem ) );
 
 				if ( SelectionIndex > 0 )
 				{
 					// Select the previous item
-					this->KeyboardSelect(ItemsSourceRef[SelectionIndex - 1], InKeyboardEvent);
+					this->KeyboardSelect(ItemsSourceRef[SelectionIndex - 1], InKeyEvent);
 				}
 
 				return FReply::Handled();
 			}
-			else if ( InKeyboardEvent.GetKey() == EKeys::Right )
+			else if ( InKeyEvent.GetKey() == EKeys::Right )
 			{
 				int32 SelectionIndex = ( !TListTypeTraits<ItemType>::IsPtrValid(this->SelectorItem) ) ? 0 : ItemsSourceRef.Find( TListTypeTraits<ItemType>::NullableItemTypeConvertToItemType( this->SelectorItem ) );
 
 				if ( SelectionIndex < ItemsSourceRef.Num() - 1 )
 				{
 					// Select the next item
-					this->KeyboardSelect(ItemsSourceRef[SelectionIndex + 1], InKeyboardEvent);
+					this->KeyboardSelect(ItemsSourceRef[SelectionIndex + 1], InKeyEvent);
 				}
 
 				return FReply::Handled();
 			}
 		}
 
-		return SListView<ItemType>::OnKeyDown(MyGeometry, InKeyboardEvent);
+		return SListView<ItemType>::OnKeyDown(MyGeometry, InKeyEvent);
 	}
 
 public:	
diff --git a/Engine/Source/Runtime/Slate/Public/Widgets/Views/STreeView.h b/Engine/Source/Runtime/Slate/Public/Widgets/Views/STreeView.h
--- a/Engine/Source/Runtime/Slate/Public/Widgets/Views/STreeView.h
+++ b/Engine/Source/Runtime/Slate/Public/Widgets/Views/STreeView.h
@@ -191,14 +191,14 @@ public:
 
 	// SWidget overrides
 
-	virtual FReply OnKeyDown( const FGeometry& MyGeometry, const FKeyboardEvent& InKeyboardEvent ) override
+	virtual FReply OnKeyDown( const FGeometry& MyGeometry, const FKeyEvent& InKeyEvent ) override
 	{
 		// Check for selection/expansion toggling keys (Left, Right)
 		// SelectorItem represents the keyboard selection. If it isn't valid then we don't know what to expand.
 		// Don't respond to key-presses containing "Alt" as a modifier
-		if ( TListTypeTraits<ItemType>::IsPtrValid(this->SelectorItem) && !InKeyboardEvent.IsAltDown() )
+		if ( TListTypeTraits<ItemType>::IsPtrValid(this->SelectorItem) && !InKeyEvent.IsAltDown() )
 		{
-			if ( InKeyboardEvent.GetKey() == EKeys::Left )
+			if ( InKeyEvent.GetKey() == EKeys::Left )
 			{
 				if( TListTypeTraits<ItemType>::IsPtrValid(this->SelectorItem) )
 				{
@@ -219,7 +219,7 @@ public:
 							if ( Private_GetNestingDepth(SelectionIndex) < SelectedNestingDepth )
 							{
 								// Found the parent
-								this->KeyboardSelect(this->LinearizedItems[SelectionIndex], InKeyboardEvent);
+								this->KeyboardSelect(this->LinearizedItems[SelectionIndex], InKeyEvent);
 								break;
 							}
 						}
@@ -228,7 +228,7 @@ public:
 
 				return FReply::Handled();
 			}
-			else if ( InKeyboardEvent.GetKey() == EKeys::Right )
+			else if ( InKeyEvent.GetKey() == EKeys::Right )
 			{
 				if( TListTypeTraits<ItemType>::IsPtrValid(this->SelectorItem) )
 				{
@@ -249,7 +249,7 @@ public:
 							// Make sure we aren't the last item on the list
 							if ( SelectionIndex < this->LinearizedItems.Num() - 1 )
 							{
-								this->KeyboardSelect(this->LinearizedItems[SelectionIndex + 1], InKeyboardEvent);
+								this->KeyboardSelect(this->LinearizedItems[SelectionIndex + 1], InKeyEvent);
 							}
 						}
 					}
@@ -259,7 +259,7 @@ public:
 			}
 		}
 
-		return SListView<ItemType>::OnKeyDown(MyGeometry, InKeyboardEvent);
+		return SListView<ItemType>::OnKeyDown(MyGeometry, InKeyEvent);
 	}
 	
 private:
diff --git a/Engine/Source/Runtime/SlateCore/Private/Input/Events.cpp b/Engine/Source/Runtime/SlateCore/Private/Input/Events.cpp
--- a/Engine/Source/Runtime/SlateCore/Private/Input/Events.cpp
+++ b/Engine/Source/Runtime/SlateCore/Private/Input/Events.cpp
@@ -24,11 +24,16 @@ FText FControllerEvent::ToText() const
 	return NSLOCTEXT("Events","Gamepad","Gamepad()");
 }
 
-FText FKeyboardEvent::ToText() const
+FText FKeyEvent::ToText() const
 {
 	return FText::Format( NSLOCTEXT("Events","Key","Key({0})"), Key.GetDisplayName() );
 }
 
+FText FAnalogInputEvent::ToText() const
+{
+	return FText::Format(NSLOCTEXT("Events", "AnalogInput", "AnalogInput Key({0})"), GetKey().GetDisplayName());
+}
+
 FText FPointerEvent::ToText() const
 {
 	return FText::Format( NSLOCTEXT("Events","Pointer","Pointer({0}, {1})"), EffectingButton.GetDisplayName() );
diff --git a/Engine/Source/Runtime/SlateCore/Private/Input/HittestGrid.cpp b/Engine/Source/Runtime/SlateCore/Private/Input/HittestGrid.cpp
--- a/Engine/Source/Runtime/SlateCore/Private/Input/HittestGrid.cpp
+++ b/Engine/Source/Runtime/SlateCore/Private/Input/HittestGrid.cpp
@@ -225,6 +225,241 @@ void FHittestGrid::InsertCustomHitTestPath( TSharedRef<ICustomHitTestPath> Custo
 	}
 }
 
+bool FHittestGrid::IsDecendantOf(const TSharedRef<SWidget> Parent, const FCachedWidget& Child)
+{
+	TSharedPtr<SWidget> ChildWidget = Child.WidgetPtr.Pin();
+
+	if (!ChildWidget.IsValid() || ChildWidget == Parent)
+	{
+		return false;
+	}
+
+	int32 CurWidgetIndex = Child.ParentIndex;
+	while (CurWidgetIndex != INDEX_NONE)
+	{
+		const FCachedWidget& CurCachedWidget = (*WidgetsCachedThisFrame)[CurWidgetIndex];
+		CurWidgetIndex = CurCachedWidget.ParentIndex;
+
+		if (Parent == CurCachedWidget.WidgetPtr.Pin())
+		{
+			return true;
+		}
+	}
+		
+	return false;
+}
+
+template<typename TCompareFunc, typename TSourceSideFunc, typename TDestSideFunc>
+TSharedPtr<SWidget> FHittestGrid::FindFocusableWidget(FSlateRect WidgetRect, const FSlateRect SweptRect, int32 AxisIndex, int32 Increment, const EUINavigation Direction, const FNavigationReply& NavigationReply, TCompareFunc CompareFunc, TSourceSideFunc SourceSideFunc, TDestSideFunc DestSideFunc)
+{
+	FIntPoint CurrentCellPoint = GetCellCoordinate(WidgetRect.GetCenter());
+
+	int32 StartingIndex = CurrentCellPoint[AxisIndex];
+
+	float CurrentSourceSide = SourceSideFunc(WidgetRect);
+
+	int32 StrideAxis, StrideAxisMin, StrideAxisMax;
+	if (AxisIndex == 0)
+	{
+		StrideAxis = 1;
+		StrideAxisMin = FMath::Min(FMath::Max(FMath::FloorToInt(SweptRect.Top / CellSize.Y), 0), NumCells.Y - 1);
+		StrideAxisMax = FMath::Min(FMath::Max(FMath::FloorToInt(SweptRect.Bottom / CellSize.Y), 0), NumCells.Y - 1);
+	}
+	else
+	{
+		StrideAxis = 0;
+		StrideAxisMin = FMath::Min(FMath::Max(FMath::FloorToInt(SweptRect.Left / CellSize.X), 0), NumCells.X - 1);
+		StrideAxisMax = FMath::Min(FMath::Max(FMath::FloorToInt(SweptRect.Right / CellSize.X), 0), NumCells.X - 1);
+	}
+
+	bool bWrapped = false;
+	while (CurrentCellPoint[AxisIndex] >= 0 && CurrentCellPoint[AxisIndex] < NumCells[AxisIndex])
+	{
+		FIntPoint StrideCellPoint = CurrentCellPoint;
+		int32 CurrentCellProcessed = CurrentCellPoint[AxisIndex];
+		
+		// Increment before the search as a wrap case will change our current cell.
+		CurrentCellPoint[AxisIndex] += Increment;
+		
+		FSlateRect BestWidgetRect;
+		TSharedPtr<SWidget> BestWidget = TSharedPtr<SWidget>();
+		
+		for (StrideCellPoint[StrideAxis] = StrideAxisMin; StrideCellPoint[StrideAxis] <= StrideAxisMax; ++StrideCellPoint[StrideAxis])
+		{
+			FHittestGrid::FCell& Cell = FHittestGrid::CellAt(StrideCellPoint.X, StrideCellPoint.Y);
+			const TArray<int32>& IndexesInCell = Cell.CachedWidgetIndexes;
+
+			for (int32 i = IndexesInCell.Num() - 1; i >= 0; --i)
+			{
+				int32 CurrentIndex = IndexesInCell[i];
+				check(CurrentIndex < WidgetsCachedThisFrame->Num());
+
+				const FCachedWidget& TestCandidate = (*WidgetsCachedThisFrame)[CurrentIndex];
+				FSlateRect TestCandidateRect = TestCandidate.ClippingRect;
+
+				if (CompareFunc(DestSideFunc(TestCandidateRect), CurrentSourceSide) && FSlateRect::DoRectanglesIntersect(SweptRect, TestCandidateRect))
+				{
+					// If this found widget isn't closer then the previously found widget then keep looking.
+					if (BestWidget.IsValid() && !CompareFunc(DestSideFunc(BestWidgetRect), DestSideFunc(TestCandidateRect)))
+					{
+						continue;
+					}
+
+					// If we have a non escape boundary condition and this widget isn't a descendant of our boundary condition widget then it's invalid so we keep looking.
+					if (NavigationReply.GetBoundaryRule() != EUINavigationRule::Escape 
+						&& NavigationReply.GetHandler().IsValid() 
+						&& !IsDecendantOf(NavigationReply.GetHandler().ToSharedRef(), TestCandidate))
+					{
+						continue;
+					}
+
+					TSharedPtr<SWidget> Widget = TestCandidate.WidgetPtr.Pin();
+					if (Widget.IsValid() && Widget->IsEnabled() && Widget->SupportsKeyboardFocus())
+					{
+						BestWidgetRect = TestCandidateRect;
+						BestWidget = Widget;
+					}
+				}
+			}
+		}
+
+		if (BestWidget.IsValid())
+		{
+			// Check for the need to apply our rule
+			if (CompareFunc(DestSideFunc(BestWidgetRect), SourceSideFunc(SweptRect)))
+			{
+				switch (NavigationReply.GetBoundaryRule())
+				{
+				case EUINavigationRule::Explicit:
+					return NavigationReply.GetFocusRecipient();
+				case EUINavigationRule::Custom:
+					{
+						const FNavigationDelegate& FocusDelegate = NavigationReply.GetFocusDelegate();
+						if (FocusDelegate.IsBound())
+						{
+							return FocusDelegate.Execute(Direction);
+						}
+						return TSharedPtr<SWidget>();
+					}
+				case EUINavigationRule::Stop:
+					return TSharedPtr<SWidget>();
+				case EUINavigationRule::Wrap:
+					CurrentSourceSide = DestSideFunc(SweptRect);
+					FVector2D SampleSpot = WidgetRect.GetCenter();
+					SampleSpot[AxisIndex] = CurrentSourceSide;
+					CurrentCellPoint = GetCellCoordinate(SampleSpot);
+					bWrapped = true;
+					break;
+				}
+			}
+
+			return BestWidget;
+		}
+
+		// break if we have looped back to where we started.
+		if (bWrapped && StartingIndex == CurrentCellProcessed) { break; }
+
+		// If were going to fail our bounds check and our rule is to wrap then wrap our position
+		if (!(CurrentCellPoint[AxisIndex] >= 0 && CurrentCellPoint[AxisIndex] < NumCells[AxisIndex]) && NavigationReply.GetBoundaryRule() == EUINavigationRule::Wrap)
+		{
+			CurrentSourceSide = DestSideFunc(SweptRect);
+			FVector2D SampleSpot = WidgetRect.GetCenter();
+			SampleSpot[AxisIndex] = CurrentSourceSide;
+			CurrentCellPoint = GetCellCoordinate(SampleSpot);
+			bWrapped = true;
+		}
+	}
+
+	return TSharedPtr<SWidget>();
+}
+
+TSharedPtr<SWidget> FHittestGrid::FindNextFocusableWidget(const FArrangedWidget& StartingWidget, const EUINavigation Direction, const FNavigationReply& NavigationReply, const FArrangedWidget& RuleWidget)
+{
+	FSlateRect WidgetRect = 
+		TransformRect(
+			Concatenate(
+				Inverse(StartingWidget.Geometry.GetAccumulatedLayoutTransform()),
+				StartingWidget.Geometry.GetAccumulatedRenderTransform()
+			),
+			FSlateRotatedRect(StartingWidget.Geometry.GetClippingRect())
+		)
+		.ToBoundingRect()
+		.OffsetBy(-GridOrigin);
+
+	FSlateRect BoundingRuleRect = 
+		TransformRect(
+			Concatenate(
+				Inverse(RuleWidget.Geometry.GetAccumulatedLayoutTransform()),
+				RuleWidget.Geometry.GetAccumulatedRenderTransform()
+			),
+			FSlateRotatedRect(RuleWidget.Geometry.GetClippingRect())
+		)
+		.ToBoundingRect()
+		.OffsetBy(-GridOrigin);
+
+	FSlateRect SweptWidgetRect = WidgetRect;
+
+	TSharedPtr<SWidget> Widget = TSharedPtr<SWidget>();
+
+	switch (Direction)
+	{
+	case EUINavigation::Left:
+		SweptWidgetRect.Left = BoundingRuleRect.Left;
+		SweptWidgetRect.Right = BoundingRuleRect.Right;
+		SweptWidgetRect.Top += 0.5f;
+		SweptWidgetRect.Bottom -= 0.5f;
+		Widget = FindFocusableWidget(WidgetRect, SweptWidgetRect, 0, -1, Direction, NavigationReply,
+			[](float A, float B) { return A - 0.1f < B; }, // Compare function
+			[](FSlateRect SourceRect) { return SourceRect.Left; }, // Source side function
+			[](FSlateRect DestRect) { return DestRect.Right; }); // Dest side function
+		break;
+	case EUINavigation::Right:
+		SweptWidgetRect.Left = BoundingRuleRect.Left;
+		SweptWidgetRect.Right = BoundingRuleRect.Right;
+		SweptWidgetRect.Top += 0.5f;
+		SweptWidgetRect.Bottom -= 0.5f;
+		Widget = FindFocusableWidget(WidgetRect, SweptWidgetRect, 0, 1, Direction, NavigationReply,
+			[](float A, float B) { return A + 0.1f > B; }, // Compare function
+			[](FSlateRect SourceRect) { return SourceRect.Right; }, // Source side function
+			[](FSlateRect DestRect) { return DestRect.Left; }); // Dest side function
+		break;
+	case EUINavigation::Up:
+		SweptWidgetRect.Top = BoundingRuleRect.Top;
+		SweptWidgetRect.Bottom = BoundingRuleRect.Bottom;
+		SweptWidgetRect.Left += 0.5f;
+		SweptWidgetRect.Right -= 0.5f;
+		Widget = FindFocusableWidget(WidgetRect, SweptWidgetRect, 1, -1, Direction, NavigationReply,
+			[](float A, float B) { return A - 0.1f < B; }, // Compare function
+			[](FSlateRect SourceRect) { return SourceRect.Top; }, // Source side function
+			[](FSlateRect DestRect) { return DestRect.Bottom; }); // Dest side function
+		break;
+	case EUINavigation::Down:
+		SweptWidgetRect.Top = BoundingRuleRect.Top;
+		SweptWidgetRect.Bottom = BoundingRuleRect.Bottom;
+		SweptWidgetRect.Left += 0.5f;
+		SweptWidgetRect.Right -= 0.5f;
+		Widget = FindFocusableWidget(WidgetRect, SweptWidgetRect, 1, 1, Direction, NavigationReply,
+			[](float A, float B) { return A + 0.1f > B; }, // Compare function
+			[](FSlateRect SourceRect) { return SourceRect.Bottom; }, // Source side function
+			[](FSlateRect DestRect) { return DestRect.Top; }); // Dest side function
+		break;
+
+	default:
+		break;
+	}
+
+	return Widget;
+}
+
+
+
+FIntPoint FHittestGrid::GetCellCoordinate(FVector2D Position)
+{
+	return FIntPoint(
+		FMath::Min(FMath::Max(FMath::FloorToInt(Position.X / CellSize.X), 0), NumCells.X - 1),
+		FMath::Min(FMath::Max(FMath::FloorToInt(Position.Y / CellSize.Y), 0), NumCells.Y - 1));
+}
+
 FHittestGrid::FCell& FHittestGrid::CellAt( const int32 X, const int32 Y )
 {
 	check( (Y*NumCells.X + X) < Cells.Num() );
@@ -289,4 +524,4 @@ void FHittestGrid::LogChildren(int32 Index, int32 IndentLevel, const TArray<FCac
 	{
 		LogChildren( CachedWidget.Children[i], IndentLevel+1, WidgetsCachedThisFrame );
 	}
-}
\ No newline at end of file
+}
