Reinstated behavior allowing click to select color inside color picker.

diff --git a/Engine/Source/Runtime/Slate/Private/Widgets/Colors/SColorSpectrum.cpp b/Engine/Source/Runtime/Slate/Private/Widgets/Colors/SColorSpectrum.cpp
--- a/Engine/Source/Runtime/Slate/Private/Widgets/Colors/SColorSpectrum.cpp
+++ b/Engine/Source/Runtime/Slate/Private/Widgets/Colors/SColorSpectrum.cpp
@@ -38,6 +38,7 @@ FReply SColorSpectrum::OnMouseButtonDown( const FGeometry& MyGeometry, const FPo
 {
 	if (MouseEvent.GetEffectingButton() == EKeys::LeftMouseButton)
 	{
+		ProcessMouseAction(MyGeometry, MouseEvent);
 		OnMouseCaptureBegin.ExecuteIfBound();
 
 		return FReply::Handled().CaptureMouse(SharedThis(this));
@@ -67,24 +68,7 @@ FReply SColorSpectrum::OnMouseMove( const FGeometry& MyGeometry, const FPointerE
 		return FReply::Unhandled();
 	}
 
-	FVector2D NormalizedMousePosition = MyGeometry.AbsoluteToLocal(MouseEvent.GetScreenSpacePosition()) / MyGeometry.Size;
-	NormalizedMousePosition = NormalizedMousePosition.ClampAxes(0.0f, 1.0f);
-
-	FLinearColor NewColor = SelectedColor.Get();
-	NewColor.R = 360.0f * NormalizedMousePosition.X;
-
-	if (NormalizedMousePosition.Y > 0.5f)
-	{
-		NewColor.G = 1.0f;
-		NewColor.B = 2.0f * (1.0f - NormalizedMousePosition.Y);
-	}
-	else
-	{
-		NewColor.G = 2.0f * NormalizedMousePosition.Y;
-		NewColor.B = 1.0f;
-	}
-
-	OnValueChanged.ExecuteIfBound(NewColor);
+	ProcessMouseAction(MyGeometry, MouseEvent);
 
 	return FReply::Handled();
 }
@@ -141,3 +125,26 @@ FVector2D SColorSpectrum::CalcRelativeSelectedPosition( ) const
 
 	return FVector2D(Color.R / 360.0f, 0.5f * Color.G);
 }
+
+
+void SColorSpectrum::ProcessMouseAction(const FGeometry& MyGeometry, const FPointerEvent& MouseEvent)
+{
+	FVector2D NormalizedMousePosition = MyGeometry.AbsoluteToLocal(MouseEvent.GetScreenSpacePosition()) / MyGeometry.Size;
+	NormalizedMousePosition = NormalizedMousePosition.ClampAxes(0.0f, 1.0f);
+
+	FLinearColor NewColor = SelectedColor.Get();
+	NewColor.R = 360.0f * NormalizedMousePosition.X;
+
+	if (NormalizedMousePosition.Y > 0.5f)
+	{
+		NewColor.G = 1.0f;
+		NewColor.B = 2.0f * (1.0f - NormalizedMousePosition.Y);
+	}
+	else
+	{
+		NewColor.G = 2.0f * NormalizedMousePosition.Y;
+		NewColor.B = 1.0f;
+	}
+
+	OnValueChanged.ExecuteIfBound(NewColor);
+}
diff --git a/Engine/Source/Runtime/Slate/Private/Widgets/Colors/SColorWheel.cpp b/Engine/Source/Runtime/Slate/Private/Widgets/Colors/SColorWheel.cpp
--- a/Engine/Source/Runtime/Slate/Private/Widgets/Colors/SColorWheel.cpp
+++ b/Engine/Source/Runtime/Slate/Private/Widgets/Colors/SColorWheel.cpp
@@ -38,11 +38,7 @@ FReply SColorWheel::OnMouseButtonDown( const FGeometry& MyGeometry, const FPoint
 {
 	if (MouseEvent.GetEffectingButton() == EKeys::LeftMouseButton)
 	{
-		FVector2D LocalMouseCoordinate = MyGeometry.AbsoluteToLocal( MouseEvent.GetScreenSpacePosition());
-		FVector2D Location = LocalMouseCoordinate / (MyGeometry.Size * 0.5f) - FVector2D(1.0f, 1.0f);
-		float Radius = Location.Size();
-	
-		if (Radius > 1.f)
+		if (!ProcessMouseAction(MyGeometry, MouseEvent, false))
 		{
 			return FReply::Unhandled();
 		}
@@ -80,21 +76,7 @@ FReply SColorWheel::OnMouseMove( const FGeometry& MyGeometry, const FPointerEven
 		return FReply::Unhandled();
 	}
 
-	FVector2D LocalMouseCoordinate = MyGeometry.AbsoluteToLocal(MouseEvent.GetScreenSpacePosition());
-	FVector2D Location = LocalMouseCoordinate / (MyGeometry.Size * 0.5f) - FVector2D(1.0f, 1.0f);
-	float Angle = FMath::Atan2(Location.Y, Location.X);
-	float Radius = FMath::Min(Location.Size(), 1.0f);
-	
-	if (Angle < 0.0f)
-	{
-		Angle += 2.0f * PI;
-	}
-
-	FLinearColor NewColor = SelectedColor.Get();
-	NewColor.R = Angle * 180.0f * INV_PI;
-	NewColor.G = Radius;
-
-	OnValueChanged.ExecuteIfBound(NewColor);
+	ProcessMouseAction(MyGeometry, MouseEvent, true);
 
 	return FReply::Handled();
 }
@@ -139,3 +121,29 @@ FVector2D SColorWheel::CalcRelativeSelectedPosition( ) const
 
 	return FVector2D(FMath::Cos(Angle), FMath::Sin(Angle)) * Radius + FVector2D(1.0f, 1.0f);
 }
+
+
+bool SColorWheel::ProcessMouseAction(const FGeometry& MyGeometry, const FPointerEvent& MouseEvent, bool bProcessWhenOutsideColorWheel)
+{
+	const FVector2D LocalMouseCoordinate = MyGeometry.AbsoluteToLocal(MouseEvent.GetScreenSpacePosition());
+	const FVector2D Location = LocalMouseCoordinate / (MyGeometry.Size * 0.5f) - FVector2D(1.0f, 1.0f);
+	const float Radius = Location.Size();
+
+	if (Radius <= 1.0f || bProcessWhenOutsideColorWheel)
+	{
+		float Angle = FMath::Atan2(Location.Y, Location.X);
+
+		if (Angle < 0.0f)
+		{
+			Angle += 2.0f * PI;
+		}
+
+		FLinearColor NewColor = SelectedColor.Get();
+		NewColor.R = Angle * 180.0f * INV_PI;
+		NewColor.G = FMath::Min(Radius, 1.0f);
+
+		OnValueChanged.ExecuteIfBound(NewColor);
+	}
+
+	return (Radius <= 1.0f);
+}
diff --git a/Engine/Source/Runtime/Slate/Public/Widgets/Colors/SColorSpectrum.h b/Engine/Source/Runtime/Slate/Public/Widgets/Colors/SColorSpectrum.h
--- a/Engine/Source/Runtime/Slate/Public/Widgets/Colors/SColorSpectrum.h
+++ b/Engine/Source/Runtime/Slate/Public/Widgets/Colors/SColorSpectrum.h
@@ -61,6 +61,11 @@ protected:
 	 */
 	FVector2D CalcRelativeSelectedPosition( ) const;
 
+	/**
+	 * Performs actions according to mouse click / move
+	 */
+	void ProcessMouseAction(const FGeometry& MyGeometry, const FPointerEvent& MouseEvent);
+
 private:
 
 	// The color wheel image to show.
diff --git a/Engine/Source/Runtime/Slate/Public/Widgets/Colors/SColorWheel.h b/Engine/Source/Runtime/Slate/Public/Widgets/Colors/SColorWheel.h
--- a/Engine/Source/Runtime/Slate/Public/Widgets/Colors/SColorWheel.h
+++ b/Engine/Source/Runtime/Slate/Public/Widgets/Colors/SColorWheel.h
@@ -61,6 +61,13 @@ protected:
 	 */
 	FVector2D CalcRelativeSelectedPosition( ) const;
 
+	/**
+	 * Performs actions according to mouse click / move
+	 *
+	 * @return	True if the mouse action occurred within the color wheel radius
+	 */
+	bool ProcessMouseAction(const FGeometry& MyGeometry, const FPointerEvent& MouseEvent, bool bProcessWhenOutsideColorWheel);
+
 private:
 
 	/** The color wheel image to show. */
