Make sure UnhandledKeyDownEventHandler is bound in FSlateApplication::ProcessKeyDownEvent() before trying to execute it

diff --git a/Engine/Source/Runtime/Slate/Private/Framework/Application/SlateApplication.cpp b/Engine/Source/Runtime/Slate/Private/Framework/Application/SlateApplication.cpp
--- a/Engine/Source/Runtime/Slate/Private/Framework/Application/SlateApplication.cpp
+++ b/Engine/Source/Runtime/Slate/Private/Framework/Application/SlateApplication.cpp
@@ -3472,7 +3472,7 @@ bool FSlateApplication::ProcessKeyDownEvent( FKeyEvent& InKeyEvent )
 		LOG_EVENT_CONTENT( EEventLog::KeyDown, GetKeyName(InKeyEvent.GetKey()).ToString(), Reply );
 
 		// If the key event was not processed by any widget...
-		if (!Reply.IsEventHandled())
+		if (!Reply.IsEventHandled() && UnhandledKeyDownEventHandler.IsBound())
 		{
 			Reply = UnhandledKeyDownEventHandler.Execute( InKeyEvent );
 		}
