Fixed some shadow offset issues

UE-4766 - Text actors don't render on mobile devices

Fixed places where the baseline was being converted between signed and unsigned values. Baselines are typically negative, and there was a case where arithmetic was being performed on an unsigned baseline value that was causing the text to be entirely clipped on mobile platforms.

The baseline adjustment was also the wrong way around, and were letting positive shadow offsets clip the text off the top of its clipping region. This should only be done for negative shadow offsets (to move the baseline to compensate for the shadow). This was causing noticable clipping issues in UMG.

Shadow offsets also weren't repecting scale, which was most noticable when zooming in UMG. The shadow offset now takes scale into account so that it remains consistent at different zoom levels.

diff --git a/Engine/Source/Runtime/AppFramework/Private/Widgets/Testing/STestSuite.cpp b/Engine/Source/Runtime/AppFramework/Private/Widgets/Testing/STestSuite.cpp
--- a/Engine/Source/Runtime/AppFramework/Private/Widgets/Testing/STestSuite.cpp
+++ b/Engine/Source/Runtime/AppFramework/Private/Widgets/Testing/STestSuite.cpp
@@ -1744,7 +1744,7 @@ public:
 			];
 		
 		TSharedRef< FSlateFontMeasure > FontMeasure = FSlateApplication::Get().GetRenderer()->GetFontMeasureService();
-		uint16 Baseline = FontMeasure->GetBaseline( FTestStyle::Get().GetWidgetStyle<FTextBlockStyle>( "RichText.Text" ).Font );
+		int16 Baseline = FontMeasure->GetBaseline( FTestStyle::Get().GetWidgetStyle<FTextBlockStyle>( "RichText.Text" ).Font );
 
 		return FSlateWidgetRun::FWidgetRunInfo( Widget, Baseline - 3 );
 	}
diff --git a/Engine/Source/Runtime/Slate/Private/Framework/Text/SlateHyperlinkRun.cpp b/Engine/Source/Runtime/Slate/Private/Framework/Text/SlateHyperlinkRun.cpp
--- a/Engine/Source/Runtime/Slate/Private/Framework/Text/SlateHyperlinkRun.cpp
+++ b/Engine/Source/Runtime/Slate/Private/Framework/Text/SlateHyperlinkRun.cpp
@@ -30,22 +30,22 @@ void FSlateHyperlinkRun::SetTextRange( const FTextRange& Value )
 int16 FSlateHyperlinkRun::GetBaseLine( float Scale ) const 
 {
 	const TSharedRef< FSlateFontMeasure > FontMeasure = FSlateApplication::Get().GetRenderer()->GetFontMeasureService();
-	return FontMeasure->GetBaseline( Style.TextStyle.Font, Scale ) - FMath::Max(0.0f, Style.TextStyle.ShadowOffset.Y);
+	return FontMeasure->GetBaseline( Style.TextStyle.Font, Scale ) - FMath::Min(0.0f, Style.TextStyle.ShadowOffset.Y * Scale);
 }
 
 int16 FSlateHyperlinkRun::GetMaxHeight( float Scale ) const 
 {
 	const TSharedRef< FSlateFontMeasure > FontMeasure = FSlateApplication::Get().GetRenderer()->GetFontMeasureService();
-	return FontMeasure->GetMaxCharacterHeight( Style.TextStyle.Font, Scale ) + FMath::Abs(Style.TextStyle.ShadowOffset.Y);
+	return FontMeasure->GetMaxCharacterHeight( Style.TextStyle.Font, Scale ) + FMath::Abs(Style.TextStyle.ShadowOffset.Y * Scale);
 }
 
 FVector2D FSlateHyperlinkRun::Measure( int32 StartIndex, int32 EndIndex, float Scale ) const 
 {
-	const FVector2D ShadowOffsetToApply((EndIndex == Range.EndIndex) ? FMath::Abs(Style.TextStyle.ShadowOffset.X) : 0.0f, FMath::Abs(Style.TextStyle.ShadowOffset.Y));
+	const FVector2D ShadowOffsetToApply((EndIndex == Range.EndIndex) ? FMath::Abs(Style.TextStyle.ShadowOffset.X * Scale) : 0.0f, FMath::Abs(Style.TextStyle.ShadowOffset.Y * Scale));
 
 	if ( EndIndex - StartIndex == 0 )
 	{
-		return FVector2D( ShadowOffsetToApply.X, GetMaxHeight( Scale ) );
+		return FVector2D( ShadowOffsetToApply.X * Scale, GetMaxHeight( Scale ) );
 	}
 
 	const TSharedRef< FSlateFontMeasure > FontMeasure = FSlateApplication::Get().GetRenderer()->GetFontMeasureService();
diff --git a/Engine/Source/Runtime/Slate/Private/Framework/Text/SlateTextRun.cpp b/Engine/Source/Runtime/Slate/Private/Framework/Text/SlateTextRun.cpp
--- a/Engine/Source/Runtime/Slate/Private/Framework/Text/SlateTextRun.cpp
+++ b/Engine/Source/Runtime/Slate/Private/Framework/Text/SlateTextRun.cpp
@@ -29,22 +29,27 @@ void FSlateTextRun::SetTextRange( const FTextRange& Value )
 int16 FSlateTextRun::GetBaseLine( float Scale ) const 
 {
 	const TSharedRef< FSlateFontMeasure > FontMeasure = FSlateApplication::Get().GetRenderer()->GetFontMeasureService();
-	return FontMeasure->GetBaseline( Style.Font, Scale ) - FMath::Max(0.0f, Style.ShadowOffset.Y);
+
+	const int16 BaseLine = FontMeasure->GetBaseline( Style.Font, Scale );
+	const int16 BaseLineAdjustment = FMath::Min(0.0f, Style.ShadowOffset.Y * Scale);
+	UE_LOG(LogSlate, Display, TEXT("Baseline: %d. Baseline Adjustment: %d. Shadow Offset: {%f, %f}\n"), BaseLine, BaseLineAdjustment, Style.ShadowOffset.X, Style.ShadowOffset.Y);
+
+	return FontMeasure->GetBaseline( Style.Font, Scale ) - FMath::Min(0.0f, Style.ShadowOffset.Y * Scale);
 }
 
 int16 FSlateTextRun::GetMaxHeight( float Scale ) const 
 {
 	const TSharedRef< FSlateFontMeasure > FontMeasure = FSlateApplication::Get().GetRenderer()->GetFontMeasureService();
-	return FontMeasure->GetMaxCharacterHeight( Style.Font, Scale ) + FMath::Abs(Style.ShadowOffset.Y);
+	return FontMeasure->GetMaxCharacterHeight( Style.Font, Scale ) + FMath::Abs(Style.ShadowOffset.Y * Scale);
 }
 
 FVector2D FSlateTextRun::Measure( int32 BeginIndex, int32 EndIndex, float Scale ) const 
 {
-	const FVector2D ShadowOffsetToApply((EndIndex == Range.EndIndex) ? FMath::Abs(Style.ShadowOffset.X) : 0.0f, FMath::Abs(Style.ShadowOffset.Y));
+	const FVector2D ShadowOffsetToApply((EndIndex == Range.EndIndex) ? FMath::Abs(Style.ShadowOffset.X * Scale) : 0.0f, FMath::Abs(Style.ShadowOffset.Y * Scale));
 
 	if ( EndIndex - BeginIndex == 0 )
 	{
-		return FVector2D( ShadowOffsetToApply.X, GetMaxHeight( Scale ) );
+		return FVector2D( ShadowOffsetToApply.X * Scale, GetMaxHeight( Scale ) );
 	}
 
 	const TSharedRef< FSlateFontMeasure > FontMeasure = FSlateApplication::Get().GetRenderer()->GetFontMeasureService();
@@ -82,12 +87,12 @@ int32 FSlateTextRun::OnPaint( const FPaintArgs& Args, const FTextLayout::FLineVi
 
 	// A negative shadow offset should be applied as a positive offset to the text to avoid clipping issues
 	const FVector2D DrawShadowOffset(
-		(Style.ShadowOffset.X > 0.0f) ? Style.ShadowOffset.X : 0.0f, 
-		(Style.ShadowOffset.Y > 0.0f) ? Style.ShadowOffset.Y : 0.0f
+		(Style.ShadowOffset.X > 0.0f) ? Style.ShadowOffset.X * AllottedGeometry.Scale : 0.0f, 
+		(Style.ShadowOffset.Y > 0.0f) ? Style.ShadowOffset.Y * AllottedGeometry.Scale : 0.0f
 		);
 	const FVector2D DrawTextOffset(
-		(Style.ShadowOffset.X < 0.0f) ? -Style.ShadowOffset.X : 0.0f, 
-		(Style.ShadowOffset.Y < 0.0f) ? -Style.ShadowOffset.Y : 0.0f
+		(Style.ShadowOffset.X < 0.0f) ? -Style.ShadowOffset.X * AllottedGeometry.Scale : 0.0f, 
+		(Style.ShadowOffset.Y < 0.0f) ? -Style.ShadowOffset.Y * AllottedGeometry.Scale : 0.0f
 		);
 
 	// Draw the optional shadow
diff --git a/Engine/Source/Runtime/SlateCore/Private/Fonts/FontCache.cpp b/Engine/Source/Runtime/SlateCore/Private/Fonts/FontCache.cpp
--- a/Engine/Source/Runtime/SlateCore/Private/Fonts/FontCache.cpp
+++ b/Engine/Source/Runtime/SlateCore/Private/Fonts/FontCache.cpp
@@ -1081,7 +1081,7 @@ uint16 FSlateFontCache::GetMaxCharacterHeight( const FSlateFontInfo& InFontInfo,
 	return NewRenderData.MaxHeight;
 }
 
-uint16 FSlateFontCache::GetBaseline( const FSlateFontInfo& InFontInfo, float FontScale ) const
+int16 FSlateFontCache::GetBaseline( const FSlateFontInfo& InFontInfo, float FontScale ) const
 {
 	FCharacterRenderData NewRenderData;
 
diff --git a/Engine/Source/Runtime/SlateCore/Private/Fonts/FontMeasure.cpp b/Engine/Source/Runtime/SlateCore/Private/Fonts/FontMeasure.cpp
--- a/Engine/Source/Runtime/SlateCore/Private/Fonts/FontMeasure.cpp
+++ b/Engine/Source/Runtime/SlateCore/Private/Fonts/FontMeasure.cpp
@@ -313,12 +313,10 @@ int8 FSlateFontMeasure::GetKerning(const FSlateFontInfo& InFontInfo, float FontS
 	return CharacterList.GetKerning( PreviousCharacter, CurrentCharacter );
 }
 
-uint16 FSlateFontMeasure::GetBaseline( const FSlateFontInfo& InFontInfo, float FontScale ) const
+int16 FSlateFontMeasure::GetBaseline( const FSlateFontInfo& InFontInfo, float FontScale ) const
 {
 	FCharacterList& CharacterList = FontCache->GetCharacterList( InFontInfo, FontScale );
-
-	const uint16 Baseline = CharacterList.GetBaseline();
-	return Baseline;
+	return CharacterList.GetBaseline();
 }
 
 void FSlateFontMeasure::FlushCache()
diff --git a/Engine/Source/Runtime/SlateCore/Public/Fonts/FontCache.h b/Engine/Source/Runtime/SlateCore/Public/Fonts/FontCache.h
--- a/Engine/Source/Runtime/SlateCore/Public/Fonts/FontCache.h
+++ b/Engine/Source/Runtime/SlateCore/Public/Fonts/FontCache.h
@@ -219,7 +219,7 @@ private:
 	/** The global max height for any character in this font */
 	mutable uint16 MaxHeight;
 	/** The offset from the bottom of the max character height to the baseline. */
-	mutable uint16 Baseline;
+	mutable int16 Baseline;
 };
 
 /**
@@ -330,7 +330,7 @@ public:
 	 * 
 	 * @return The offset from the bottom of the max character height to the baseline.
 	 */
-	uint16 GetBaseline( const FSlateFontInfo& InFontInfo, float FontScale ) const;
+	int16 GetBaseline( const FSlateFontInfo& InFontInfo, float FontScale ) const;
 
 	/**
 	 * Returns the font attributes for the specified font.
diff --git a/Engine/Source/Runtime/SlateCore/Public/Fonts/FontMeasure.h b/Engine/Source/Runtime/SlateCore/Public/Fonts/FontMeasure.h
--- a/Engine/Source/Runtime/SlateCore/Public/Fonts/FontMeasure.h
+++ b/Engine/Source/Runtime/SlateCore/Public/Fonts/FontMeasure.h
@@ -187,7 +187,7 @@ public:
 	 * 
 	 * @return The offset from the bottom of the max character height to the baseline.
 	 */
-	uint16 GetBaseline( const FSlateFontInfo& InFontInfo, float FontScale = 1.0f ) const;
+	int16 GetBaseline( const FSlateFontInfo& InFontInfo, float FontScale = 1.0f ) const;
 
 	void FlushCache();
 
