Removed some debug code

diff --git a/Engine/Source/Runtime/Slate/Private/Framework/Text/SlateTextRun.cpp b/Engine/Source/Runtime/Slate/Private/Framework/Text/SlateTextRun.cpp
--- a/Engine/Source/Runtime/Slate/Private/Framework/Text/SlateTextRun.cpp
+++ b/Engine/Source/Runtime/Slate/Private/Framework/Text/SlateTextRun.cpp
@@ -29,11 +29,6 @@ void FSlateTextRun::SetTextRange( const FTextRange& Value )
 int16 FSlateTextRun::GetBaseLine( float Scale ) const 
 {
 	const TSharedRef< FSlateFontMeasure > FontMeasure = FSlateApplication::Get().GetRenderer()->GetFontMeasureService();
-
-	const int16 BaseLine = FontMeasure->GetBaseline( Style.Font, Scale );
-	const int16 BaseLineAdjustment = FMath::Min(0.0f, Style.ShadowOffset.Y * Scale);
-	UE_LOG(LogSlate, Display, TEXT("Baseline: %d. Baseline Adjustment: %d. Shadow Offset: {%f, %f}\n"), BaseLine, BaseLineAdjustment, Style.ShadowOffset.X, Style.ShadowOffset.Y);
-
 	return FontMeasure->GetBaseline( Style.Font, Scale ) - FMath::Min(0.0f, Style.ShadowOffset.Y * Scale);
 }
 
