[Pull Request #553] [Fix] MultilineEditableText::ScanForWordBoundary infinite loop

FString.Len() already trims the terminating character, so IsAtEndOfLine/IsAtEndOfDocument were not correct, causing ScanForWordBoundary to get stuck in an infinite loop when using ctrl to skip towards the end of the Multiline Text control.

Credit to github user stephenwhittle for finding and fixing this issue.

#github

diff --git a/Engine/Source/Runtime/Slate/Private/Widgets/Text/SMultiLineEditableText.cpp b/Engine/Source/Runtime/Slate/Private/Widgets/Text/SMultiLineEditableText.cpp
--- a/Engine/Source/Runtime/Slate/Private/Widgets/Text/SMultiLineEditableText.cpp
+++ b/Engine/Source/Runtime/Slate/Private/Widgets/Text/SMultiLineEditableText.cpp
@@ -2458,7 +2458,7 @@ bool SMultiLineEditableText::IsAtEndOfDocument( const FTextLocation& Location )
 	const TArray< FTextLayout::FLineModel >& Lines = TextLayout->GetLineModels();
 	const int32 NumberOfLines = Lines.Num();
 
-	return NumberOfLines == 0 || ( NumberOfLines - 1 == Location.GetLineIndex() && Lines[ Location.GetLineIndex() ].Text->Len()-1 == Location.GetOffset() );
+	return NumberOfLines == 0 || ( NumberOfLines - 1 == Location.GetLineIndex() && Lines[ Location.GetLineIndex() ].Text->Len() == Location.GetOffset() );
 }
 
 bool SMultiLineEditableText::IsAtBeginningOfLine( const FTextLocation& Location ) const
@@ -2469,7 +2469,7 @@ bool SMultiLineEditableText::IsAtBeginningOfLine( const FTextLocation& Location
 bool SMultiLineEditableText::IsAtEndOfLine( const FTextLocation& Location ) const
 {
 	const TArray< FTextLayout::FLineModel >& Lines = TextLayout->GetLineModels();
-	return Lines[ Location.GetLineIndex() ].Text->Len()-1 == Location.GetOffset();
+	return Lines[ Location.GetLineIndex() ].Text->Len() == Location.GetOffset();
 }
 
 bool SMultiLineEditableText::IsAtWordStart( const FTextLocation& Location ) const
