Added experimental Slate.FoldTick variable. Defaults to 0. Results in significant performance savings of up to 35%.

diff --git a/Engine/Source/Runtime/Slate/Private/Framework/Application/SlateApplication.cpp b/Engine/Source/Runtime/Slate/Private/Framework/Application/SlateApplication.cpp
--- a/Engine/Source/Runtime/Slate/Private/Framework/Application/SlateApplication.cpp
+++ b/Engine/Source/Runtime/Slate/Private/Framework/Application/SlateApplication.cpp
@@ -907,7 +907,7 @@ void FSlateApplication::DrawWindowAndChildren( const TSharedRef<SWindow>& Window
 		int32 MaxLayerId = 0;
 		{
 			MaxLayerId = WindowToDraw->PaintWindow(
-				FPaintArgs(WindowToDraw, *HittestGrid, WindowToDraw->GetPositionInScreen()),
+				FPaintArgs(WindowToDraw, *HittestGrid, WindowToDraw->GetPositionInScreen(), GetCurrentTime(), GetDeltaTime()),
 				WindowGeometry, WindowToDraw->GetClippingRectangleInWindow(),
 				WindowElementList,
 				0,
@@ -1137,6 +1137,8 @@ void FSlateApplication::FinishedInputThisFrame()
 	}
 }
 
+extern SLATECORE_API TAutoConsoleVariable<int32> FoldTick;
+
 /**
  * Ticks this application
  */
@@ -1221,39 +1223,43 @@ void FSlateApplication::Tick()
 	// Update auto-throttling based on elapsed time since user interaction
 	ThrottleApplicationBasedOnMouseMovement();
 
-	TSharedPtr<SWindow> ActiveModalWindow = GetActiveModalWindow();
+	const bool bFoldTick = (FoldTick.GetValueOnGameThread() != 0);
 
+	if ( !bFoldTick )
 	{
-		SCOPE_CYCLE_COUNTER( STAT_SlateTickWindowAndChildren );
-
-		if ( ActiveModalWindow.IsValid() )
+		TSharedPtr<SWindow> ActiveModalWindow = GetActiveModalWindow();
 		{
-			// There is a modal window, and we just need to tick it.
-			TickWindowAndChildren( ActiveModalWindow.ToSharedRef() );
-			// And also tick any top-level windows.
-			for( TArray< TSharedRef<SWindow> >::TIterator CurrentWindowIt( SlateWindows ); CurrentWindowIt; ++CurrentWindowIt )
+			SCOPE_CYCLE_COUNTER( STAT_SlateTickWindowAndChildren );
+
+			if ( ActiveModalWindow.IsValid() )
 			{
-				TSharedRef<SWindow>& CurrentWindow = *CurrentWindowIt;
-				if (CurrentWindow->IsTopmostWindow())
+				// There is a modal window, and we just need to tick it.
+				TickWindowAndChildren( ActiveModalWindow.ToSharedRef() );
+				// And also tick any top-level windows.
+				for( TArray< TSharedRef<SWindow> >::TIterator CurrentWindowIt( SlateWindows ); CurrentWindowIt; ++CurrentWindowIt )
 				{
-					TickWindowAndChildren(CurrentWindow);
+					TSharedRef<SWindow>& CurrentWindow = *CurrentWindowIt;
+					if (CurrentWindow->IsTopmostWindow())
+					{
+						TickWindowAndChildren(CurrentWindow);
+					}
 				}
+				// also tick the notification manager's windows
+				TArray< TSharedRef<SWindow> > NotificationWindows;
+				FSlateNotificationManager::Get().GetWindows(NotificationWindows);
+				for( auto CurrentWindowIt( NotificationWindows.CreateIterator() ); CurrentWindowIt; ++CurrentWindowIt )
+				{
+					TickWindowAndChildren(*CurrentWindowIt);
+				}		
 			}
-			// also tick the notification manager's windows
-			TArray< TSharedRef<SWindow> > NotificationWindows;
-			FSlateNotificationManager::Get().GetWindows(NotificationWindows);
-			for( auto CurrentWindowIt( NotificationWindows.CreateIterator() ); CurrentWindowIt; ++CurrentWindowIt )
-			{
-				TickWindowAndChildren(*CurrentWindowIt);
-			}		
-		}
-		else
-		{
-			// No modal window; tick all slate windows.
-			for( TArray< TSharedRef<SWindow> >::TIterator CurrentWindowIt( SlateWindows ); CurrentWindowIt; ++CurrentWindowIt )
+			else
 			{
-				TSharedRef<SWindow>& CurrentWindow = *CurrentWindowIt;
-				TickWindowAndChildren( CurrentWindow );
+				// No modal window; tick all slate windows.
+				for( TArray< TSharedRef<SWindow> >::TIterator CurrentWindowIt( SlateWindows ); CurrentWindowIt; ++CurrentWindowIt )
+				{
+					TSharedRef<SWindow>& CurrentWindow = *CurrentWindowIt;
+					TickWindowAndChildren( CurrentWindow );
+				}
 			}
 		}
 	}
diff --git a/Engine/Source/Runtime/SlateCore/Private/Types/PaintArgs.cpp b/Engine/Source/Runtime/SlateCore/Private/Types/PaintArgs.cpp
--- a/Engine/Source/Runtime/SlateCore/Private/Types/PaintArgs.cpp
+++ b/Engine/Source/Runtime/SlateCore/Private/Types/PaintArgs.cpp
@@ -3,19 +3,21 @@
 #include "SlateCorePrivatePCH.h"
 #include "HittestGrid.h"
 
-FPaintArgs::FPaintArgs( const TSharedRef<SWidget>& Parent, FHittestGrid& InHittestGrid, FVector2D InWindowOffset )
+FPaintArgs::FPaintArgs( const TSharedRef<SWidget>& Parent, FHittestGrid& InHittestGrid, FVector2D InWindowOffset, double InCurrentTime, float InDeltaTime )
 : ParentPtr(Parent)
 , Grid(InHittestGrid)
 , LastHittestIndex(INDEX_NONE)
 , LastRecordedVisibility( EVisibility::Visible )
 , WindowOffset(InWindowOffset)
+, CurrentTime(InCurrentTime)
+, DeltaTime(InDeltaTime)
 {
 }
 
 
 FPaintArgs FPaintArgs::WithNewParent( const SWidget* Parent ) const
 {
-	FPaintArgs Args = FPaintArgs( const_cast<SWidget*>(Parent)->AsShared(), this->Grid, this->WindowOffset );
+	FPaintArgs Args = FPaintArgs( const_cast<SWidget*>(Parent)->AsShared(), this->Grid, this->WindowOffset, this->CurrentTime, this->DeltaTime );
 	Args.LastHittestIndex = this->LastHittestIndex;
 	Args.LastRecordedVisibility = this->LastRecordedVisibility;
 	return Args;
diff --git a/Engine/Source/Runtime/SlateCore/Private/Widgets/SWidget.cpp b/Engine/Source/Runtime/SlateCore/Private/Widgets/SWidget.cpp
--- a/Engine/Source/Runtime/SlateCore/Private/Widgets/SWidget.cpp
+++ b/Engine/Source/Runtime/SlateCore/Private/Widgets/SWidget.cpp
@@ -584,13 +584,26 @@ void SWidget::SetDebugInfo( const ANSICHAR* InType, const ANSICHAR* InFile, int3
 	this->CreatedOnLine = OnLine;
 }
 
+SLATECORE_API TAutoConsoleVariable<int32> FoldTick(TEXT("Slate.FoldTick"), 0, TEXT("When folding, call Tick as part of the paint pass instead of a separate tick pass."));
+
 int32 SWidget::Paint(const FPaintArgs& Args, const FGeometry& AllottedGeometry, const FSlateRect& MyClippingRect, FSlateWindowElementList& OutDrawElements, int32 LayerId, const FWidgetStyle& InWidgetStyle, bool bParentEnabled) const
 {
 	SCOPE_CYCLE_COUNTER(STAT_SlateOnPaint);
 	INC_DWORD_STAT(STAT_SlateNumPaintedWidgets);
 
-	FPaintArgs UpdatedArgs = Args.RecordHittestGeometry( this, AllottedGeometry, MyClippingRect );
+	const bool bFoldTick = (FoldTick.GetValueOnGameThread() != 0);
+	if ( bFoldTick )
+	{
+		FGeometry TickGeometry = AllottedGeometry;
+		TickGeometry.AppendTransform( FSlateLayoutTransform(Args.GetWindowToDesktopTransform()) );
+
+		SWidget* MutableThis = const_cast<SWidget*>(this);
+		MutableThis->Tick( TickGeometry, Args.GetCurrentTime(), Args.GetDeltaTime() );
+	}
+
+	const FPaintArgs UpdatedArgs = Args.RecordHittestGeometry( this, AllottedGeometry, MyClippingRect );
 	return OnPaint(UpdatedArgs, AllottedGeometry, MyClippingRect, OutDrawElements, LayerId, InWidgetStyle, bParentEnabled);
+	//return OnPaint(Args, AllottedGeometry, MyClippingRect, OutDrawElements, LayerId, InWidgetStyle, bParentEnabled);
 }
 
 void SWidget::ArrangeChildren(const FGeometry& AllottedGeometry, FArrangedChildren& ArrangedChildren) const
@@ -598,3 +611,4 @@ void SWidget::ArrangeChildren(const FGeometry& AllottedGeometry, FArrangedChildr
 	SCOPE_CYCLE_COUNTER(STAT_SlateArrangeChildren);
 	OnArrangeChildren(AllottedGeometry, ArrangedChildren);
 }
+
diff --git a/Engine/Source/Runtime/SlateCore/Public/Types/PaintArgs.h b/Engine/Source/Runtime/SlateCore/Public/Types/PaintArgs.h
--- a/Engine/Source/Runtime/SlateCore/Public/Types/PaintArgs.h
+++ b/Engine/Source/Runtime/SlateCore/Public/Types/PaintArgs.h
@@ -18,15 +18,21 @@ class ICustomHitTestPath;
 class SLATECORE_API FPaintArgs
 {
 public:
-	FPaintArgs( const TSharedRef<SWidget>& Parent, FHittestGrid& InHittestGrid, FVector2D InWindowOffset );
+	FPaintArgs( const TSharedRef<SWidget>& Parent, FHittestGrid& InHittestGrid, FVector2D InWindowOffset, double InCurrentTime, float InDeltaTime );
 	FPaintArgs WithNewParent( const SWidget* Parent ) const;
 	FPaintArgs RecordHittestGeometry(const SWidget* Widget, const FGeometry& WidgetGeometry, const FSlateRect& InClippingRect) const;
 	FPaintArgs InsertCustomHitTestPath( TSharedRef<ICustomHitTestPath> CustomHitTestPath, int32 HitTestIndex ) const;
 	int32 GetLastHitTestIndex() const { return LastHittestIndex; }
+	FVector2D GetWindowToDesktopTransform() const { return WindowOffset; }
+	double GetCurrentTime() const { return CurrentTime; }
+	float GetDeltaTime() const { return DeltaTime; }
+
 private:
 	const TSharedRef<SWidget>& ParentPtr;
 	FHittestGrid& Grid;
 	int32 LastHittestIndex;
 	EVisibility LastRecordedVisibility;
 	FVector2D WindowOffset;
+	double CurrentTime;
+	float DeltaTime;
 };
\ No newline at end of file
diff --git a/Engine/Source/Runtime/UMG/Private/Components/WidgetComponent.cpp b/Engine/Source/Runtime/UMG/Private/Components/WidgetComponent.cpp
--- a/Engine/Source/Runtime/UMG/Private/Components/WidgetComponent.cpp
+++ b/Engine/Source/Runtime/UMG/Private/Components/WidgetComponent.cpp
@@ -473,7 +473,7 @@ void UWidgetComponent::TickComponent(float DeltaTime, enum ELevelTick TickType,
 		int32 MaxLayerId = 0;
 		{
 			MaxLayerId = SlateWidget->Paint(
-				FPaintArgs(SlateWidget.ToSharedRef(), *HitTestGrid, FVector2D::ZeroVector),
+				FPaintArgs(SlateWidget.ToSharedRef(), *HitTestGrid, FVector2D::ZeroVector, FSlateApplication::Get().GetCurrentTime(), FSlateApplication::Get().GetDeltaTime()),
 				WindowGeometry, WindowGeometry.GetClippingRect(),
 				WindowElementList,
 				0,
@@ -723,4 +723,4 @@ TArray<FWidgetAndPointer> UWidgetComponent::GetHitWidgetPath( const FHitResult&
 TSharedPtr<SWidget> UWidgetComponent::GetSlateWidget() const
 {
 	return SlateWidget;
-}
+}
\ No newline at end of file
