Removed uneeded clear of user focus.
Added ensure on SetUserFocus() that notifies user that the Widget isn't visible and in the tree.

diff --git a/Engine/Source/Runtime/Slate/Private/Framework/Application/SlateApplication.cpp b/Engine/Source/Runtime/Slate/Private/Framework/Application/SlateApplication.cpp
--- a/Engine/Source/Runtime/Slate/Private/Framework/Application/SlateApplication.cpp
+++ b/Engine/Source/Runtime/Slate/Private/Framework/Application/SlateApplication.cpp
@@ -1817,12 +1817,21 @@ void FSlateApplication::SetJoystickCaptorToGameViewport()
 
 void FSlateApplication::SetUserFocus(uint32 UserIndex, const TSharedPtr<SWidget>& WidgetToFocus, EFocusCause ReasonFocusIsChanging /* = EFocusCause::SetDirectly*/)
 {
-	FWidgetPath PathToWidget;
-	if (WidgetToFocus.IsValid())
+	const bool bValidWidget = WidgetToFocus.IsValid();
+	ensureMsgf(bValidWidget, TEXT("Attempting to focus an invalid widget. If your intent is to clear focus use ClearUserFocus()"));
+	if (bValidWidget)
 	{
-		FSlateWindowHelper::FindPathToWidget(SlateWindows, WidgetToFocus.ToSharedRef(), /*OUT*/ PathToWidget);
+		FWidgetPath PathToWidget;
+		const bool bFound = FSlateWindowHelper::FindPathToWidget(SlateWindows, WidgetToFocus.ToSharedRef(), /*OUT*/ PathToWidget);
+		if (bFound)
+		{
+			SetUserFocus(UserIndex, PathToWidget, ReasonFocusIsChanging);
+		}
+		else
+		{
+			//ensureMsgf(bFound, TEXT("Attempting to focus a widget that isn't in the tree and visible: %s. If your intent is to clear focus use ClearUserFocus()"), WidgetToFocus->ToString());
+		}
 	}
-	SetUserFocus(UserIndex, PathToWidget, ReasonFocusIsChanging);
 }
 
 TSharedPtr<SWidget> FSlateApplication::GetUserFocusedWidget(uint32 UserIndex) const
@@ -1974,6 +1983,8 @@ bool FSlateApplication::SetUserFocus(const uint32 InUserIndex, const FWidgetPath
 		}
 	}
 
+	//UE_LOG(LogSlate, Warning, TEXT("Focus for user %i set to %s."), InUserIndex, NewFocusedWidget.IsValid() ? *NewFocusedWidget->ToString() : TEXT("Invalid"));
+
 	// Store a weak widget path to the widget that's taking focus
 	UserFocusedWidgetPaths[InUserIndex] = FWeakWidgetPath(NewFocusedWidgetPath);
 	//@Todo Slate: Refactor focus cause to be stored for all users!
