Fixed issue in menu creation, in which menus with bFocusImmediately set to false will still result in the previous window losing focus.

diff --git a/Engine/Source/Runtime/Slate/Private/Framework/Application/MenuStack.cpp b/Engine/Source/Runtime/Slate/Private/Framework/Application/MenuStack.cpp
--- a/Engine/Source/Runtime/Slate/Private/Framework/Application/MenuStack.cpp
+++ b/Engine/Source/Runtime/Slate/Private/Framework/Application/MenuStack.cpp
@@ -171,7 +171,7 @@ TSharedRef<SWindow> FMenuStack::PushMenu( const TSharedRef<SWindow>& ParentWindo
 		.InitialOpacity(InitialWindowOpacity)
 		.SupportsTransparency( bUseTransparency )
 		.FocusWhenFirstShown( bFocusImmediately )
-		.ActivateWhenFirstShown( true )
+		.ActivateWhenFirstShown( bFocusImmediately )
 		[
 			WrappedContent
 		];
