ScrollBar: Scrollbars now span their track their full length instead of splitting it between elements above and below the thumb widget. This brings scrollbar implementation more inline with it's styling options and allows for track bars with rounded endpoints.

diff --git a/Engine/Source/Runtime/Slate/Private/Widgets/Layout/SScrollBar.cpp b/Engine/Source/Runtime/Slate/Private/Widgets/Layout/SScrollBar.cpp
--- a/Engine/Source/Runtime/Slate/Private/Widgets/Layout/SScrollBar.cpp
+++ b/Engine/Source/Runtime/Slate/Private/Widgets/Layout/SScrollBar.cpp
@@ -281,37 +281,24 @@ void SScrollBar::Construct(const FArguments& InArgs)
 			+ SVerticalBox::Slot()
 			.FillHeight( 1 )
 			[
-				SAssignNew(Track, SScrollBarTrack)
-				.Orientation(InArgs._Orientation)
-				.TopSlot()
+				SNew(SBorder)
+				.BorderImage(Brush)
+				.HAlign(HorizontalAlignment)
+				.VAlign(VerticalAlignment)
+				.Padding(0)
 				[
-					SNew(SBox)
-					.HAlign(HorizontalAlignment)
-					.VAlign(VerticalAlignment)
+					SAssignNew(Track, SScrollBarTrack)
+					.Orientation(InArgs._Orientation)
+					.ThumbSlot()
 					[
-						SNew(SImage)
-						.Image(Brush)
-					]
-				]
-				.ThumbSlot()
-				[
-					SAssignNew(DragThumb, SBorder)
-					.BorderImage( this, &SScrollBar::GetDragThumbImage )
-					.HAlign(HAlign_Center)
-					.VAlign(VAlign_Center)
-					[
-						SAssignNew(ThicknessSpacer, SSpacer)
-						.Size(InArgs._Thickness)
-					]
-				]
-				.BottomSlot()
-				[
-					SNew(SBox)
-					.HAlign(HorizontalAlignment)
-					.VAlign(VerticalAlignment)
-					[
-						SNew(SImage)
-						.Image(Brush)
+						SAssignNew(DragThumb, SBorder)
+						.BorderImage( this, &SScrollBar::GetDragThumbImage )
+						.HAlign(HAlign_Center)
+						.VAlign(VAlign_Center)
+						[
+							SAssignNew(ThicknessSpacer, SSpacer)
+							.Size(InArgs._Thickness)
+						]
 					]
 				]
 			]
