Fixed visual regression with scrollbars that prevented certain styling methods.

diff --git a/Engine/Source/Editor/EditorStyle/Private/SlateEditorStyle.cpp b/Engine/Source/Editor/EditorStyle/Private/SlateEditorStyle.cpp
--- a/Engine/Source/Editor/EditorStyle/Private/SlateEditorStyle.cpp
+++ b/Engine/Source/Editor/EditorStyle/Private/SlateEditorStyle.cpp
@@ -442,8 +442,10 @@ void FSlateEditorStyle::FStyle::SetupGeneralStyles()
 
 	// Scrollbar
 	const FScrollBarStyle ScrollBar = FScrollBarStyle()
-		.SetVerticalBackgroundImage( IMAGE_BRUSH( "Common/Scrollbar_Background_Vertical", FVector2D(8,8) ) )
-		.SetHorizontalBackgroundImage( IMAGE_BRUSH( "Common/Scrollbar_Background_Horizontal", FVector2D(8,8) ) )
+		.SetVerticalTopSlotImage(IMAGE_BRUSH("Common/Scrollbar_Background_Vertical", FVector2D(8, 8)))
+		.SetVerticalBottomSlotImage(IMAGE_BRUSH("Common/Scrollbar_Background_Vertical", FVector2D(8, 8)))
+		.SetHorizontalTopSlotImage(IMAGE_BRUSH("Common/Scrollbar_Background_Horizontal", FVector2D(8, 8)))
+		.SetHorizontalBottomSlotImage(IMAGE_BRUSH("Common/Scrollbar_Background_Horizontal", FVector2D(8, 8)))
 		.SetNormalThumbImage( BOX_BRUSH( "Common/Scrollbar_Thumb", FMargin(4.f/16.f) ) )
 		.SetDraggedThumbImage( BOX_BRUSH( "Common/Scrollbar_Thumb", FMargin(4.f/16.f) ) )
 		.SetHoveredThumbImage( BOX_BRUSH( "Common/Scrollbar_Thumb", FMargin(4.f/16.f) ) );
diff --git a/Engine/Source/Runtime/Slate/Private/Widgets/Layout/SScrollBar.cpp b/Engine/Source/Runtime/Slate/Private/Widgets/Layout/SScrollBar.cpp
--- a/Engine/Source/Runtime/Slate/Private/Widgets/Layout/SScrollBar.cpp
+++ b/Engine/Source/Runtime/Slate/Private/Widgets/Layout/SScrollBar.cpp
@@ -256,19 +256,25 @@ void SScrollBar::Construct(const FArguments& InArgs)
 
 	EHorizontalAlignment HorizontalAlignment;
 	EVerticalAlignment VerticalAlignment;
-	const FSlateBrush* Brush;
+	const FSlateBrush* BackgroundBrush;
+	const FSlateBrush* TopBrush;
+	const FSlateBrush* BottomBrush;
 
 	if(Orientation == Orient_Vertical)
 	{
 		HorizontalAlignment = HAlign_Center;
 		VerticalAlignment = VAlign_Fill;
-		Brush = &InArgs._Style->VerticalBackgroundImage;
+		BackgroundBrush = &InArgs._Style->VerticalBackgroundImage;
+		TopBrush = &InArgs._Style->VerticalTopSlotImage;
+		BottomBrush = &InArgs._Style->VerticalBottomSlotImage;
 	}
 	else
 	{
 		HorizontalAlignment = HAlign_Fill;
 		VerticalAlignment = VAlign_Center;
-		Brush = &InArgs._Style->HorizontalBackgroundImage;
+		BackgroundBrush = &InArgs._Style->HorizontalBackgroundImage;
+		TopBrush = &InArgs._Style->HorizontalTopSlotImage;
+		BottomBrush = &InArgs._Style->HorizontalBottomSlotImage;
 	}
 
 	SBorder::Construct( SBorder::FArguments()
@@ -282,13 +288,23 @@ void SScrollBar::Construct(const FArguments& InArgs)
 			.FillHeight( 1 )
 			[
 				SNew(SBorder)
-				.BorderImage(Brush)
+				.BorderImage(BackgroundBrush)
 				.HAlign(HorizontalAlignment)
 				.VAlign(VerticalAlignment)
 				.Padding(0)
 				[
 					SAssignNew(Track, SScrollBarTrack)
 					.Orientation(InArgs._Orientation)
+					.TopSlot()
+					[
+						SNew(SBox)
+						.HAlign(HorizontalAlignment)
+						.VAlign(VerticalAlignment)
+						[
+							SNew(SImage)
+							.Image(TopBrush)
+						]
+					]
 					.ThumbSlot()
 					[
 						SAssignNew(DragThumb, SBorder)
@@ -300,6 +316,16 @@ void SScrollBar::Construct(const FArguments& InArgs)
 							.Size(InArgs._Thickness)
 						]
 					]
+					.BottomSlot()
+					[
+						SNew(SBox)
+						.HAlign(HorizontalAlignment)
+						.VAlign(VerticalAlignment)
+						[
+							SNew(SImage)
+							.Image(BottomBrush)
+						]
+					]
 				]
 			]
 		]
diff --git a/Engine/Source/Runtime/SlateCore/Private/Styling/CoreStyle.cpp b/Engine/Source/Runtime/SlateCore/Private/Styling/CoreStyle.cpp
--- a/Engine/Source/Runtime/SlateCore/Private/Styling/CoreStyle.cpp
+++ b/Engine/Source/Runtime/SlateCore/Private/Styling/CoreStyle.cpp
@@ -268,8 +268,10 @@ TSharedRef<ISlateStyle> FCoreStyle::Create( const FName& InStyleSetName )
 
 	// SScrollBar defaults...
 	const FScrollBarStyle ScrollBar = FScrollBarStyle()
-		.SetVerticalBackgroundImage( IMAGE_BRUSH( "Common/Scrollbar_Background_Vertical", FVector2D(8,8) ) )
-		.SetHorizontalBackgroundImage( IMAGE_BRUSH( "Common/Scrollbar_Background_Horizontal", FVector2D(8,8) ) )
+		.SetVerticalTopSlotImage(IMAGE_BRUSH("Common/Scrollbar_Background_Vertical", FVector2D(8, 8)))
+		.SetVerticalBottomSlotImage(IMAGE_BRUSH("Common/Scrollbar_Background_Vertical", FVector2D(8, 8)))
+		.SetHorizontalTopSlotImage(IMAGE_BRUSH("Common/Scrollbar_Background_Horizontal", FVector2D(8, 8)))
+		.SetHorizontalBottomSlotImage(IMAGE_BRUSH("Common/Scrollbar_Background_Horizontal", FVector2D(8, 8)))
 		.SetNormalThumbImage( BOX_BRUSH( "Common/Scrollbar_Thumb", FMargin(4.f/16.f) ) )
 		.SetDraggedThumbImage( BOX_BRUSH( "Common/Scrollbar_Thumb", FMargin(4.f/16.f) ) )
 		.SetHoveredThumbImage( BOX_BRUSH( "Common/Scrollbar_Thumb", FMargin(4.f/16.f) ) );
diff --git a/Engine/Source/Runtime/SlateCore/Private/Styling/SlateTypes.cpp b/Engine/Source/Runtime/SlateCore/Private/Styling/SlateTypes.cpp
--- a/Engine/Source/Runtime/SlateCore/Private/Styling/SlateTypes.cpp
+++ b/Engine/Source/Runtime/SlateCore/Private/Styling/SlateTypes.cpp
@@ -287,6 +287,15 @@ const FProgressBarStyle& FProgressBarStyle::GetDefault()
 
 
 FScrollBarStyle::FScrollBarStyle()
+	: HorizontalBackgroundImage(FSlateNoResource())
+	, VerticalBackgroundImage(FSlateNoResource())
+	, VerticalTopSlotImage(FSlateNoResource())
+	, HorizontalTopSlotImage(FSlateNoResource())
+	, VerticalBottomSlotImage(FSlateNoResource())
+	, HorizontalBottomSlotImage(FSlateNoResource())
+	, NormalThumbImage(FSlateNoResource())
+	, HoveredThumbImage(FSlateNoResource())
+	, DraggedThumbImage(FSlateNoResource())
 {
 }
 
@@ -294,6 +303,10 @@ void FScrollBarStyle::GetResources( TArray< const FSlateBrush* >& OutBrushes ) c
 {
 	OutBrushes.Add( &HorizontalBackgroundImage );
 	OutBrushes.Add( &VerticalBackgroundImage );
+	OutBrushes.Add( &VerticalTopSlotImage);
+	OutBrushes.Add( &HorizontalTopSlotImage);
+	OutBrushes.Add( &VerticalBottomSlotImage);
+	OutBrushes.Add( &HorizontalBottomSlotImage);
 	OutBrushes.Add( &NormalThumbImage );
 	OutBrushes.Add( &HoveredThumbImage );
 	OutBrushes.Add( &DraggedThumbImage );
diff --git a/Engine/Source/Runtime/SlateCore/Public/Styling/SlateTypes.h b/Engine/Source/Runtime/SlateCore/Public/Styling/SlateTypes.h
--- a/Engine/Source/Runtime/SlateCore/Public/Styling/SlateTypes.h
+++ b/Engine/Source/Runtime/SlateCore/Public/Styling/SlateTypes.h
@@ -530,6 +530,26 @@ struct SLATECORE_API FScrollBarStyle : public FSlateWidgetStyle
 	FSlateBrush VerticalBackgroundImage;
 	FScrollBarStyle& SetVerticalBackgroundImage( const FSlateBrush& InVerticalBackgroundImage ){ VerticalBackgroundImage = InVerticalBackgroundImage; return *this; }
 
+	/** The image to use to represent the track above the thumb when the scrollbar is oriented vertically */
+	UPROPERTY(EditAnywhere, BlueprintReadWrite, Category = Appearance)
+	FSlateBrush VerticalTopSlotImage;
+	FScrollBarStyle& SetVerticalTopSlotImage(const FSlateBrush& Value){ VerticalTopSlotImage = Value; return *this; }
+
+	/** The image to use to represent the track above the thumb when the scrollbar is oriented horizontally */
+	UPROPERTY(EditAnywhere, BlueprintReadWrite, Category = Appearance)
+	FSlateBrush HorizontalTopSlotImage;
+	FScrollBarStyle& SetHorizontalTopSlotImage(const FSlateBrush& Value){ HorizontalTopSlotImage = Value; return *this; }
+
+	/** The image to use to represent the track below the thumb when the scrollbar is oriented vertically */
+	UPROPERTY(EditAnywhere, BlueprintReadWrite, Category = Appearance)
+	FSlateBrush VerticalBottomSlotImage;
+	FScrollBarStyle& SetVerticalBottomSlotImage(const FSlateBrush& Value){ VerticalBottomSlotImage = Value; return *this; }
+
+	/** The image to use to represent the track below the thumb when the scrollbar is oriented horizontally */
+	UPROPERTY(EditAnywhere, BlueprintReadWrite, Category = Appearance)
+	FSlateBrush HorizontalBottomSlotImage;
+	FScrollBarStyle& SetHorizontalBottomSlotImage(const FSlateBrush& Value){ HorizontalBottomSlotImage = Value; return *this; }
+
 	/** Image to use when the scrollbar thumb is in its normal state */
 	UPROPERTY(EditAnywhere, BlueprintReadWrite, Category=Appearance)
 	FSlateBrush NormalThumbImage;
