Slate: Added some comments to SMultiColumnTableRow

diff --git a/Engine/Source/Runtime/Slate/Public/Widgets/Views/STableRow.h b/Engine/Source/Runtime/Slate/Public/Widgets/Views/STableRow.h
--- a/Engine/Source/Runtime/Slate/Public/Widgets/Views/STableRow.h
+++ b/Engine/Source/Runtime/Slate/Public/Widgets/Views/STableRow.h
@@ -778,7 +778,10 @@ public:
 	 */
 	virtual TSharedRef<SWidget> GenerateWidgetForColumn( const FName& ColumnName ) = 0;
 
+	/** Use this to construct the superclass; e.g. FSuperRowType::Construct( FTableRowArgs(), OwnerTableView ) */
 	typedef SMultiColumnTableRow< ItemType > FSuperRowType;
+
+	/** Use this to construct the superclass; e.g. FSuperRowType::Construct( FTableRowArgs(), OwnerTableView ) */
 	typedef typename STableRow<ItemType>::FArguments FTableRowArgs;
 
 protected:
