Slate : removed erroneous Args.WithNewParent(this) while painting STextBlock.

diff --git a/Engine/Source/Runtime/Slate/Private/Widgets/Text/STextBlock.cpp b/Engine/Source/Runtime/Slate/Private/Widgets/Text/STextBlock.cpp
--- a/Engine/Source/Runtime/Slate/Private/Widgets/Text/STextBlock.cpp
+++ b/Engine/Source/Runtime/Slate/Private/Widgets/Text/STextBlock.cpp
@@ -107,7 +107,7 @@ int32 STextBlock::OnPaint( const FPaintArgs& Args, const FGeometry& AllottedGeom
 	// OnPaint will also update the text layout cache if required
 	LayerId = TextLayoutCache->OnPaint(
 		FTextBlockLayout::FWidgetArgs(BoundText, HighlightText, WrapTextAt, AutoWrapText, Margin, LineHeightPercentage, Justification), 
-		Args.WithNewParent(this), AllottedGeometry, MyClippingRect, OutDrawElements, LayerId, InWidgetStyle, ShouldBeEnabled(bParentEnabled)
+		Args, AllottedGeometry, MyClippingRect, OutDrawElements, LayerId, InWidgetStyle, ShouldBeEnabled(bParentEnabled)
 		);
 
 #else//WITH_FANCY_TEXT
