Removed some incorrect usages of Args.WithNewParent

Follow up to CL# 2352204, this time dealing with SRichTextBlock and SMultiLineEditableText.

This also removes the unused FWidgetArgs parameter from FTextBlockLayout::OnPaint.

diff --git a/Engine/Source/Runtime/Slate/Private/Widgets/Text/SMultiLineEditableText.cpp b/Engine/Source/Runtime/Slate/Private/Widgets/Text/SMultiLineEditableText.cpp
--- a/Engine/Source/Runtime/Slate/Private/Widgets/Text/SMultiLineEditableText.cpp
+++ b/Engine/Source/Runtime/Slate/Private/Widgets/Text/SMultiLineEditableText.cpp
@@ -2228,7 +2228,7 @@ int32 SMultiLineEditableText::OnPaint( const FPaintArgs& Args, const FGeometry&
 	// Note: This is done here rather than in Tick(), because Tick() doesn't get called while resizing windows, but OnPaint() does
 	CachedSize = AllottedGeometry.Size;
 
-	LayerId = TextLayout->OnPaint( Args.WithNewParent(this), AllottedGeometry, MyClippingRect, OutDrawElements, LayerId, InWidgetStyle, ShouldBeEnabled( bParentEnabled ) );
+	LayerId = TextLayout->OnPaint( Args, AllottedGeometry, MyClippingRect, OutDrawElements, LayerId, InWidgetStyle, ShouldBeEnabled( bParentEnabled ) );
 
 	return LayerId;
 }
diff --git a/Engine/Source/Runtime/Slate/Private/Widgets/Text/SRichTextBlock.cpp b/Engine/Source/Runtime/Slate/Private/Widgets/Text/SRichTextBlock.cpp
--- a/Engine/Source/Runtime/Slate/Private/Widgets/Text/SRichTextBlock.cpp
+++ b/Engine/Source/Runtime/Slate/Private/Widgets/Text/SRichTextBlock.cpp
@@ -46,10 +46,7 @@ int32 SRichTextBlock::OnPaint( const FPaintArgs& Args, const FGeometry& Allotted
 	SCOPE_CYCLE_COUNTER( STAT_SlateOnPaint_SRichTextBlock );
 
 	// OnPaint will also update the text layout cache if required
-	LayerId = TextLayoutCache->OnPaint(
-		FTextBlockLayout::FWidgetArgs(BoundText, HighlightText, WrapTextAt, AutoWrapText, Margin, LineHeightPercentage, Justification), 
-		Args.WithNewParent(this), AllottedGeometry, MyClippingRect, OutDrawElements, LayerId, InWidgetStyle, ShouldBeEnabled(bParentEnabled)
-		);
+	LayerId = TextLayoutCache->OnPaint(Args, AllottedGeometry, MyClippingRect, OutDrawElements, LayerId, InWidgetStyle, ShouldBeEnabled(bParentEnabled));
 
 	return LayerId;
 }
diff --git a/Engine/Source/Runtime/Slate/Private/Widgets/Text/STextBlock.cpp b/Engine/Source/Runtime/Slate/Private/Widgets/Text/STextBlock.cpp
--- a/Engine/Source/Runtime/Slate/Private/Widgets/Text/STextBlock.cpp
+++ b/Engine/Source/Runtime/Slate/Private/Widgets/Text/STextBlock.cpp
@@ -105,10 +105,7 @@ int32 STextBlock::OnPaint( const FPaintArgs& Args, const FGeometry& AllottedGeom
 #if WITH_FANCY_TEXT
 
 	// OnPaint will also update the text layout cache if required
-	LayerId = TextLayoutCache->OnPaint(
-		FTextBlockLayout::FWidgetArgs(BoundText, HighlightText, WrapTextAt, AutoWrapText, Margin, LineHeightPercentage, Justification), 
-		Args, AllottedGeometry, MyClippingRect, OutDrawElements, LayerId, InWidgetStyle, ShouldBeEnabled(bParentEnabled)
-		);
+	LayerId = TextLayoutCache->OnPaint(Args, AllottedGeometry, MyClippingRect, OutDrawElements, LayerId, InWidgetStyle, ShouldBeEnabled(bParentEnabled));
 
 #else//WITH_FANCY_TEXT
 
diff --git a/Engine/Source/Runtime/Slate/Private/Widgets/Text/TextBlockLayout.cpp b/Engine/Source/Runtime/Slate/Private/Widgets/Text/TextBlockLayout.cpp
--- a/Engine/Source/Runtime/Slate/Private/Widgets/Text/TextBlockLayout.cpp
+++ b/Engine/Source/Runtime/Slate/Private/Widgets/Text/TextBlockLayout.cpp
@@ -94,7 +94,7 @@ FVector2D FTextBlockLayout::ComputeDesiredSize(const FWidgetArgs& InWidgetArgs/*
 	return CachedDesiredSize;
 }
 
-int32 FTextBlockLayout::OnPaint(const FWidgetArgs& InWidgetArgs, const FPaintArgs& InPaintArgs, const FGeometry& InAllottedGeometry, const FSlateRect& InClippingRect, FSlateWindowElementList& OutDrawElements, int32 LayerId, const FWidgetStyle& InWidgetStyle, bool bParentEnabled)
+int32 FTextBlockLayout::OnPaint(const FPaintArgs& InPaintArgs, const FGeometry& InAllottedGeometry, const FSlateRect& InClippingRect, FSlateWindowElementList& OutDrawElements, int32 LayerId, const FWidgetStyle& InWidgetStyle, bool bParentEnabled)
 {
 	CachedSize = InAllottedGeometry.Size;
 
diff --git a/Engine/Source/Runtime/Slate/Public/Widgets/Text/TextBlockLayout.h b/Engine/Source/Runtime/Slate/Public/Widgets/Text/TextBlockLayout.h
--- a/Engine/Source/Runtime/Slate/Public/Widgets/Text/TextBlockLayout.h
+++ b/Engine/Source/Runtime/Slate/Public/Widgets/Text/TextBlockLayout.h
@@ -52,7 +52,7 @@ public:
 	/**
 	 * Paint this layout, updating the internal cache as required
 	 */
-	int32 OnPaint(const FWidgetArgs& InWidgetArgs, const FPaintArgs& InPaintArgs, const FGeometry& InAllottedGeometry, const FSlateRect& InClippingRect, FSlateWindowElementList& OutDrawElements, int32 LayerId, const FWidgetStyle& InWidgetStyle, bool bParentEnabled);
+	int32 OnPaint(const FPaintArgs& InPaintArgs, const FGeometry& InAllottedGeometry, const FSlateRect& InClippingRect, FSlateWindowElementList& OutDrawElements, int32 LayerId, const FWidgetStyle& InWidgetStyle, bool bParentEnabled);
 
 	/**
 	 * Force dirty the layout due to an external change that can't be picked up automatically by this cache
