Added an option to have the anchored menu be below and centered to the MenuAnchor.

diff --git a/Engine/Source/Runtime/Slate/Private/Widgets/Input/SMenuAnchor.cpp b/Engine/Source/Runtime/Slate/Private/Widgets/Input/SMenuAnchor.cpp
--- a/Engine/Source/Runtime/Slate/Private/Widgets/Input/SMenuAnchor.cpp
+++ b/Engine/Source/Runtime/Slate/Private/Widgets/Input/SMenuAnchor.cpp
@@ -158,6 +158,8 @@ static FVector2D GetMenuOffsetForPlacement(const FGeometry& AllottedGeometry, EM
 		return FVector2D(AllottedGeometry.GetLocalSize().X - PopupSizeLocalSpace.X, 0.0f);
 	case MenuPlacement_MenuLeft:
 		return FVector2D(-AllottedGeometry.GetLocalSize().X - PopupSizeLocalSpace.X, 0.0f);
+	case MenuPlacement_CenteredBelowAnchor:
+		return FVector2D(-((PopupSizeLocalSpace.X / 2 ) - (AllottedGeometry.GetLocalSize().X / 2)), 0.0f);
 	default:
 		return FVector2D::ZeroVector;
 	}
diff --git a/Engine/Source/Runtime/SlateCore/Public/Types/SlateEnums.h b/Engine/Source/Runtime/SlateCore/Public/Types/SlateEnums.h
--- a/Engine/Source/Runtime/SlateCore/Public/Types/SlateEnums.h
+++ b/Engine/Source/Runtime/SlateCore/Public/Types/SlateEnums.h
@@ -154,6 +154,9 @@ enum EMenuPlacement
 	/** Place the menu immediately below the anchor */
 	MenuPlacement_BelowAnchor UMETA(DisplayName="Below"),
 
+	/** Place the menu immediately centered below the anchor */
+	MenuPlacement_CenteredBelowAnchor UMETA(DisplayName = "Centered Below"),
+
 	/** Place the menu immediately below the anchor and match is width to the anchor's content */
 	MenuPlacement_ComboBox UMETA(DisplayName="Combo Box"),
 
