Fixed issue in which Output Log context menu doesn't close when it loses focus.

diff --git a/Engine/Source/Runtime/Slate/Private/Widgets/Input/SEditableText.cpp b/Engine/Source/Runtime/Slate/Private/Widgets/Input/SEditableText.cpp
--- a/Engine/Source/Runtime/Slate/Private/Widgets/Input/SEditableText.cpp
+++ b/Engine/Source/Runtime/Slate/Private/Widgets/Input/SEditableText.cpp
@@ -1973,8 +1973,8 @@ void SEditableText::SummonContextMenu( const FVector2D& InLocation )
 
 #undef LOCTEXT_NAMESPACE
 
-	const bool bFocusImmediatley = false;
-	TSharedPtr< SWindow > ContextMenuWindowPinned = FSlateApplication::Get().PushMenu( SharedThis( this ), MenuBuilder.MakeWidget(), InLocation, FPopupTransitionEffect( FPopupTransitionEffect::ContextMenu ), bFocusImmediatley );
+	const bool bFocusImmediately = true;
+	TSharedPtr< SWindow > ContextMenuWindowPinned = FSlateApplication::Get().PushMenu( SharedThis( this ), MenuBuilder.MakeWidget(), InLocation, FPopupTransitionEffect( FPopupTransitionEffect::ContextMenu ), bFocusImmediately );
 	ContextMenuWindow = ContextMenuWindowPinned;
 
 	// Make sure the window is valid.  It's possible for the parent to already be in the destroy queue, for example if the editable text was configured to dismiss it's window during OnTextCommitted.
diff --git a/Engine/Source/Runtime/Slate/Private/Widgets/Text/SMultiLineEditableText.cpp b/Engine/Source/Runtime/Slate/Private/Widgets/Text/SMultiLineEditableText.cpp
--- a/Engine/Source/Runtime/Slate/Private/Widgets/Text/SMultiLineEditableText.cpp
+++ b/Engine/Source/Runtime/Slate/Private/Widgets/Text/SMultiLineEditableText.cpp
@@ -2076,8 +2076,8 @@ void SMultiLineEditableText::SummonContextMenu(const FVector2D& InLocation)
 
 #undef LOCTEXT_NAMESPACE
 
-	const bool bFocusImmediatley = false;
-	TSharedPtr< SWindow > ContextMenuWindowPinned = FSlateApplication::Get().PushMenu(SharedThis(this), MenuBuilder.MakeWidget(), InLocation, FPopupTransitionEffect(FPopupTransitionEffect::ContextMenu), bFocusImmediatley);
+	const bool bFocusImmediately = true;
+	TSharedPtr< SWindow > ContextMenuWindowPinned = FSlateApplication::Get().PushMenu(SharedThis(this), MenuBuilder.MakeWidget(), InLocation, FPopupTransitionEffect(FPopupTransitionEffect::ContextMenu), bFocusImmediately);
 	ContextMenuWindow = ContextMenuWindowPinned;
 
 	// Make sure the window is valid.  It's possible for the parent to already be in the destroy queue, for example if the editable text was configured to dismiss it's window during OnTextCommitted.
