Fixed some jitter when resizing multiline editable text widgets

UE-5236 - Text Jitters in Detail Panel When Creating a Comment in the Material Editor

We have some code which tries to prevent initial layout jitter when auto-wrapping text, and this was checking that the auto-wrap cached size matched the current paint size. This was what was causing the jitter when resizing, as it was constantly hiding the scrollbar when the size had changed. This has now been changed to check that the cached size has been set, as once that's done, we should have also done a valid wrap.

diff --git a/Engine/Source/Runtime/Slate/Private/Widgets/Text/SMultiLineEditableText.cpp b/Engine/Source/Runtime/Slate/Private/Widgets/Text/SMultiLineEditableText.cpp
--- a/Engine/Source/Runtime/Slate/Private/Widgets/Text/SMultiLineEditableText.cpp
+++ b/Engine/Source/Runtime/Slate/Private/Widgets/Text/SMultiLineEditableText.cpp
@@ -2133,7 +2133,7 @@ void SMultiLineEditableText::Tick( const FGeometry& AllottedGeometry, const doub
 
 	// If we're auto-wrapping, we need to hide the scrollbars until the first valid auto-wrap has been performed
 	// If we don't do this, then we can get some nasty layout shuffling as the scrollbars appear for one frame and then vanish again
-	const EVisibility ScrollBarVisiblityOverride = (AutoWrapText.Get() && CachedSize != AllottedGeometry.Size) ? EVisibility::Collapsed : EVisibility::Visible;
+	const EVisibility ScrollBarVisiblityOverride = (AutoWrapText.Get() && CachedSize.IsZero()) ? EVisibility::Collapsed : EVisibility::Visible;
 
 	// Try and make sure that the line containing the cursor is in view
 	if (PositionToScrollIntoView.IsSet())
