Added very simple analog cursor.
FSlateApplication::SetAnalogCursorEnable(true); //to turn on the analog cursor

diff --git a/Engine/Source/Runtime/Slate/Private/Framework/Application/AnalogCursor.cpp b/Engine/Source/Runtime/Slate/Private/Framework/Application/AnalogCursor.cpp
--- a/Engine/Source/Runtime/Slate/Private/Framework/Application/AnalogCursor.cpp
+++ b/Engine/Source/Runtime/Slate/Private/Framework/Application/AnalogCursor.cpp
@@ -3,62 +3,160 @@
 #include "SlatePrivatePCH.h"
 #include "AnalogCursor.h"
 
-const float FAnalogCursor::Acceleration     = 50.0f;
-const float FAnalogCursor::Decceleration    = 100.0f;
-const float FAnalogCursor::MaxSpeed         = 100.0f;
+#include "HittestGrid.h"
+
+const float FAnalogCursor::Acceleration          = 150.0f;
+const float FAnalogCursor::Decceleration         = 250.0f;
+const float FAnalogCursor::MaxSpeed			     = 250.0f;
+const float FAnalogCursor::OverWidgetMultiplier  = 0.6f;
 
 FAnalogCursor::FAnalogCursor()
 {
 
 }
 
-void FAnalogCursor::Tick(const float DeltaTime, TSharedPtr<ICursor> Cursor)
+void FAnalogCursor::Tick(const float DeltaTime, FSlateApplication& SlateApp, TSharedRef<ICursor> Cursor)
 {
-	float CurrentMaxSpeedX = FMath::Abs(AnalogValues.X) * MaxSpeed;
-	float CurrentMaxSpeedY = FMath::Abs(AnalogValues.Y) * MaxSpeed;
+	FVector2D OldPos = Cursor->GetPosition();
+
+	float SpeedMult = 1.0f; // Used to do a speed multiplication before adding the delta to the position to make widgets sticky
+	FVector2D AdjAnalogVals = AnalogValues; // A copy of the analog values so I can modify them based being over a widget, not currently doing this
+
+	FWidgetPath WidgetPath = SlateApp.LocateWindowUnderMouse(OldPos, SlateApp.GetInteractiveTopLevelWindows());
+	if (WidgetPath.IsValid())
+	{
+		FArrangedWidget& ArrangedWidget = WidgetPath.Widgets.Last();
+		TSharedRef<SWidget> Widget = ArrangedWidget.Widget;
+		if (Widget->SupportsKeyboardFocus() && Widget->GetType() != FName("SViewport"))
+		{
+			SpeedMult = 0.5f;
+			//FVector2D Adjustment = WidgetsAndCursors.Last().Geometry.Position - CurrentPos; // example of calculating distance from cursor to widget center
+		}
+	}
 
-	CurrentSpeed += AnalogValues * Acceleration * DeltaTime;
-	CurrentSpeed.X = FMath::Clamp(CurrentSpeed.X, -CurrentMaxSpeedX, CurrentMaxSpeedX);
-	CurrentSpeed.Y = FMath::Clamp(CurrentSpeed.Y, -CurrentMaxSpeedY, CurrentMaxSpeedY);
+	// Generate Min and Max for X to clamp the speed, this gives us instant direction change when crossing the axis
+	float CurrentMinSpeedX = 0.0f;
+	float CurrentMaxSpeedX = 0.0f;
+	if (AdjAnalogVals.X > 0.0f)
+	{ 
+		CurrentMaxSpeedX = AdjAnalogVals.X * MaxSpeed;
+	}
+	else
+	{
+		CurrentMinSpeedX = AdjAnalogVals.X * MaxSpeed;
+	}
 
-	FVector2D Pos = Cursor->GetPosition();
-	Pos += CurrentSpeed;
-	Cursor->SetPosition(Pos.X, Pos.Y);
+	// Generate Min and Max for Y to clamp the speed, this gives us instant direction change when crossing the axis
+	float CurrentMinSpeedY = 0.0f;
+	float CurrentMaxSpeedY = 0.0f;
+	if (AdjAnalogVals.Y > 0.0f)
+	{
+		CurrentMaxSpeedY = AdjAnalogVals.Y * MaxSpeed;
+	}
+	else
+	{
+		CurrentMinSpeedY = AdjAnalogVals.Y * MaxSpeed;
+	}
+
+	// Cubic acceleration curve
+	FVector2D ExpAcceleration = AdjAnalogVals * AdjAnalogVals * AdjAnalogVals * Acceleration;
+	// Preserve direction (if we use a squared equation above)
+	//ExpAcceleration.X *= FMath::Sign(AnalogValues.X);
+	//ExpAcceleration.Y *= FMath::Sign(AnalogValues.Y);
+
+	CurrentSpeed += ExpAcceleration * DeltaTime;
+
+	CurrentSpeed.X = FMath::Clamp(CurrentSpeed.X, CurrentMinSpeedX, CurrentMaxSpeedX);
+	CurrentSpeed.Y = FMath::Clamp(CurrentSpeed.Y, CurrentMinSpeedY, CurrentMaxSpeedY);
+
+	CurrentPos = OldPos + (CurrentSpeed * SpeedMult);
+	Cursor->SetPosition(CurrentPos.X, CurrentPos.Y);
 }
 
-bool FAnalogCursor::HandleReleased(const FKeyEvent& KeyEvent)
+bool FAnalogCursor::HandleKeyDownEvent(FSlateApplication& SlateApp, const FKeyEvent& InKeyEvent)
 {
-	FKey Key = KeyEvent.GetKey();
-	if (Key == EKeys::Gamepad_LeftX)
+	FKey Key = InKeyEvent.GetKey();
+
+	// Consume the sticks input so it doesn't effect other things
+	if (Key == EKeys::Gamepad_LeftStick_Right || 
+		Key == EKeys::Gamepad_LeftStick_Left ||
+		Key == EKeys::Gamepad_LeftStick_Up ||
+		Key == EKeys::Gamepad_LeftStick_Down)
 	{
-		AnalogValues.X = 0.0f;
+		return true;
 	}
-	else if (Key == EKeys::Gamepad_LeftY)
+
+	// Bottom face button is a click
+	if (Key == EKeys::Gamepad_FaceButton_Bottom)
 	{
-		AnalogValues.Y = 0.0f;
+		FPointerEvent MouseEvent(
+			0,
+			CurrentPos,
+			CurrentPos,
+			FTouchKeySet::EmptySet,
+			EKeys::LeftMouseButton,
+			0,
+			SlateApp.GetPlatformApplication()->GetModifierKeys()
+			);
+
+		TSharedPtr<FGenericWindow> GenWindow;
+		SlateApp.ProcessMouseButtonDownEvent(GenWindow, MouseEvent);
 	}
-	else
+
+	return false;
+}
+
+bool FAnalogCursor::HandleKeyUpEvent(FSlateApplication& SlateApp, const FKeyEvent& InKeyEvent)
+{
+	FKey Key = InKeyEvent.GetKey();
+
+	// Consume the sticks input so it doesn't effect other things
+	if (Key == EKeys::Gamepad_LeftStick_Right ||
+		Key == EKeys::Gamepad_LeftStick_Left ||
+		Key == EKeys::Gamepad_LeftStick_Up ||
+		Key == EKeys::Gamepad_LeftStick_Down)
 	{
-		return false;
+		return true;
 	}
-	return true;
+
+	// Bottom face button is a click
+	if (Key == EKeys::Gamepad_FaceButton_Bottom)
+	{
+		FPointerEvent MouseEvent(
+			0,
+			CurrentPos,
+			CurrentPos,
+			FTouchKeySet::EmptySet,
+			EKeys::LeftMouseButton,
+			0,
+			SlateApp.GetPlatformApplication()->GetModifierKeys()
+			);
+
+		SlateApp.ProcessMouseButtonUpEvent(MouseEvent);
+	}
+	
+	return false;
 }
 
-bool FAnalogCursor::HandleAnalog(const FAnalogInputEvent& InputEvent)
+bool FAnalogCursor::HandleAnalogInputEvent(FSlateApplication& SlateApp, const FAnalogInputEvent& InAnalogInputEvent)
 {
-	FKey Key = InputEvent.GetKey();
-	float AnalogValue = InputEvent.GetAnalogValue();
+	FKey Key = InAnalogInputEvent.GetKey();
+	float AnalogValue = InAnalogInputEvent.GetAnalogValue();
+
+	//@Todo Slate: Investigate this more, doesn't seem right that analog doesn't zero
+	// Analog doesn't zero out
 	if (FMath::Abs(AnalogValue) < 0.1f)
 	{
 		AnalogValue = 0.0f;
 	}
+
 	if (Key == EKeys::Gamepad_LeftX)
 	{
 		AnalogValues.X = AnalogValue;
 	}
 	else if (Key == EKeys::Gamepad_LeftY)
 	{
-		AnalogValues.Y = AnalogValue;
+		AnalogValues.Y = -AnalogValue;
 	}
 	else
 	{
diff --git a/Engine/Source/Runtime/Slate/Private/Framework/Application/AnalogCursor.h b/Engine/Source/Runtime/Slate/Private/Framework/Application/AnalogCursor.h
--- a/Engine/Source/Runtime/Slate/Private/Framework/Application/AnalogCursor.h
+++ b/Engine/Source/Runtime/Slate/Private/Framework/Application/AnalogCursor.h
@@ -2,6 +2,8 @@
 
 #pragma once
 
+class FSlateApplication;
+
 /**
  * A class that simulates a cursor driven by an analog stick.
  */
@@ -10,18 +12,22 @@ class FAnalogCursor
 public:
 	FAnalogCursor();
 
-	void Tick(const float DeltaTime, TSharedPtr<ICursor> Cursor);
+	void Tick(const float DeltaTime, FSlateApplication& SlateApp, TSharedRef<ICursor> Cursor);
 
-	bool HandleReleased(const FKeyEvent& KeyEvent);
-	bool HandleAnalog(const FAnalogInputEvent& InputEvent);
+	bool HandleKeyDownEvent(FSlateApplication& SlateApp, const FKeyEvent& InKeyEvent);
+	bool HandleKeyUpEvent(FSlateApplication& SlateApp, const FKeyEvent& InKeyEvent);
+	bool HandleAnalogInputEvent(FSlateApplication& SlateApp, const FAnalogInputEvent& InAnalogInputEvent);
 
 private:
 
+	FVector2D CurrentPos;
+
 	FVector2D AnalogValues;
 	FVector2D CurrentSpeed;
 
 	static const float Acceleration;
 	static const float Decceleration;
 	static const float MaxSpeed;
+	static const float OverWidgetMultiplier;
 };
 
diff --git a/Engine/Source/Runtime/Slate/Private/Framework/Application/SlateApplication.cpp b/Engine/Source/Runtime/Slate/Private/Framework/Application/SlateApplication.cpp
--- a/Engine/Source/Runtime/Slate/Private/Framework/Application/SlateApplication.cpp
+++ b/Engine/Source/Runtime/Slate/Private/Framework/Application/SlateApplication.cpp
@@ -711,8 +711,6 @@ FSlateApplication::FSlateApplication()
 
 	NormalExecutionGetter.BindRaw( this, &FSlateApplication::IsNormalExecution );
 	PointerIndexLastPositionMap.Add(CursorPointerIndex, FVector2D::ZeroVector);
-
-	//AnalogCursor = MakeShareable(new FAnalogCursor);
 }
 
 FSlateApplication::~FSlateApplication()
@@ -1090,9 +1088,9 @@ void FSlateApplication::FinishedInputThisFrame()
 {
 	const float DeltaTime = GetDeltaTime();
 
-	if (AnalogCursor.IsValid())
+	if (AnalogCursor.IsValid() && PlatformApplication->Cursor.IsValid())
 	{
-		AnalogCursor->Tick(DeltaTime, PlatformApplication->Cursor);
+		AnalogCursor->Tick(DeltaTime, *this, PlatformApplication->Cursor.ToSharedRef());
 	}
 
 	// All the input events have been processed.
@@ -2983,6 +2981,18 @@ void FSlateApplication::SetDragTriggerDistnace( float ScreenPixels )
 	DragTriggerDistnace = ScreenPixels;
 }
 
+void FSlateApplication::SetAnalogCursorEnable(bool bEnable)
+{
+	if (bEnable && !AnalogCursor.IsValid())
+	{
+		AnalogCursor = MakeShareable(new FAnalogCursor);
+	}
+	else
+	{
+		AnalogCursor.Reset();
+	}
+}
+
 FVector2D FSlateApplication::CalculatePopupWindowPosition( const FSlateRect& InAnchor, const FVector2D& InSize, const EOrientation Orientation ) const
 {
 	// Do nothing if this window has no size
@@ -3441,6 +3451,12 @@ bool FSlateApplication::OnKeyDown( const int32 KeyCode, const uint32 CharacterCo
 
 bool FSlateApplication::ProcessKeyDownEvent( FKeyEvent& InKeyEvent )
 {
+	// Analog cursor gets first chance at the input
+	if (AnalogCursor.IsValid() && AnalogCursor->HandleKeyDownEvent(*this, InKeyEvent))
+	{
+		return true;
+	}
+
 	FReply Reply = FReply::Unhandled();
 
 	LastUserInteractionTime = this->GetCurrentTime();
@@ -3520,6 +3536,12 @@ bool FSlateApplication::OnKeyUp( const int32 KeyCode, const uint32 CharacterCode
 
 bool FSlateApplication::ProcessKeyUpEvent( FKeyEvent& InKeyEvent )
 {
+	// Analog cursor gets first chance at the input
+	if (AnalogCursor.IsValid() && AnalogCursor->HandleKeyUpEvent(*this, InKeyEvent))
+	{
+		return true;
+	}
+
 	FReply Reply = FReply::Unhandled();
 
 	LastUserInteractionTime = this->GetCurrentTime();
@@ -3546,6 +3568,12 @@ bool FSlateApplication::ProcessKeyUpEvent( FKeyEvent& InKeyEvent )
 
 bool FSlateApplication::ProcessAnalogInputEvent(FAnalogInputEvent& InAnalogInputEvent)
 {
+	// Analog cursor gets first chance at the input
+	if (AnalogCursor.IsValid() && AnalogCursor->HandleAnalogInputEvent(*this, InAnalogInputEvent))
+	{
+		return true;
+	}
+
 	FReply Reply = FReply::Unhandled();
 
 	LastUserInteractionTime = this->GetCurrentTime();
@@ -3641,7 +3669,10 @@ bool FSlateApplication::ProcessMouseButtonDownEvent( const TSharedPtr< FGenericW
 	LastUserInteractionTime = this->GetCurrentTime();
 	LastUserInteractionTimeForThrottling = LastUserInteractionTime;
 	
-	PlatformApplication->SetCapture( PlatformWindow );
+	if (PlatformWindow.IsValid())
+	{
+		PlatformApplication->SetCapture(PlatformWindow);
+	}
 	PressedMouseButtons.Add( MouseEvent.GetEffectingButton() );
 
 	bool bInGame = false;
@@ -4431,14 +4462,6 @@ bool FSlateApplication::OnControllerAnalog( EControllerButtons::Type Button, int
 	int32 UserIndex = GetUserIndexForController(ControllerId);
 	
 	FAnalogInputEvent AnalogInputEvent(Key, PlatformApplication->GetModifierKeys(), UserIndex, false, 0, 0, AnalogValue);
-	
-	if (AnalogCursor.IsValid())
-	{
-		if (AnalogCursor->HandleAnalog(AnalogInputEvent))
-		{
-			return true;
-		}
-	}
 
 	return ProcessAnalogInputEvent(AnalogInputEvent);
 }
@@ -4460,14 +4483,6 @@ bool FSlateApplication::OnControllerButtonReleased( EControllerButtons::Type But
 	
 	FKeyEvent KeyEvent(Key, PlatformApplication->GetModifierKeys(),UserIndex, IsRepeat,  0, 0);
 
-	if (AnalogCursor.IsValid())
-	{
-		if (AnalogCursor->HandleReleased(KeyEvent))
-		{
-			return true;
-		}
-	}
-
 	return ProcessKeyUpEvent(KeyEvent);
 }
 
diff --git a/Engine/Source/Runtime/Slate/Public/Framework/Application/SlateApplication.h b/Engine/Source/Runtime/Slate/Public/Framework/Application/SlateApplication.h
--- a/Engine/Source/Runtime/Slate/Public/Framework/Application/SlateApplication.h
+++ b/Engine/Source/Runtime/Slate/Public/Framework/Application/SlateApplication.h
@@ -721,11 +721,13 @@ public:
 	/**
 	 * Called by the native application in response to a mouse button press. Routs the event to Slate Widgets.
 	 *
-	 * @param  InMouseEvent  Mouse event
+	 * @param  PlatformWindow  The platform window the event originated from, used to set focus at the platform level. 
+	 *                         If Invalid the Mouse event will work but there will be no effect on the platform.
+	 * @param  InMouseEvent    Mouse event
 	 *
 	 * @return  Was this event handled by the Slate application?
 	 */
-	bool ProcessMouseButtonDownEvent( const TSharedPtr< FGenericWindow >& PlatformWindow, FPointerEvent& MouseEvent );
+	bool ProcessMouseButtonDownEvent(const TSharedPtr< FGenericWindow >& PlatformWindow, FPointerEvent& InMouseEvent);
 
 	/**
 	 * Called by the native application in response to a mouse button release. Routs the event to Slate Widgets.
@@ -943,6 +945,9 @@ public:
 	/** Set the size of the deadzone for dragging in screen pixels */
 	void SetDragTriggerDistnace( float ScreenPixels );
 	
+	/** Set the analog cursor to be enabled or disabled. */
+	void SetAnalogCursorEnable(bool bEnable);
+
 public:
 
 	// Begin FSlateApplicationBase interface
