Generated events from Analog Cursor are now more consistent with standard mouse input.

diff --git a/Engine/Source/Runtime/Slate/Private/Framework/Application/AnalogCursor.cpp b/Engine/Source/Runtime/Slate/Private/Framework/Application/AnalogCursor.cpp
--- a/Engine/Source/Runtime/Slate/Private/Framework/Application/AnalogCursor.cpp
+++ b/Engine/Source/Runtime/Slate/Private/Framework/Application/AnalogCursor.cpp
@@ -71,6 +71,22 @@ void FAnalogCursor::Tick(const float DeltaTime, FSlateApplication& SlateApp, TSh
 
 	CurrentPos = OldPos + (CurrentSpeed * SpeedMult);
 	Cursor->SetPosition(CurrentPos.X, CurrentPos.Y);
+
+	// Send out a move event
+	if (OldPos != CurrentPos)
+	{
+		FPointerEvent MouseEvent(
+			0,
+			CurrentPos,
+			OldPos,
+			SlateApp.PressedMouseButtons,
+			EKeys::Invalid,
+			0,
+			SlateApp.GetPlatformApplication()->GetModifierKeys()
+			);
+
+		SlateApp.ProcessMouseMoveEvent(MouseEvent);
+	}
 }
 
 bool FAnalogCursor::HandleKeyDownEvent(FSlateApplication& SlateApp, const FKeyEvent& InKeyEvent)
@@ -91,9 +107,9 @@ bool FAnalogCursor::HandleKeyDownEvent(FSlateApplication& SlateApp, const FKeyEv
 	{
 		FPointerEvent MouseEvent(
 			0,
-			CurrentPos,
-			CurrentPos,
-			FTouchKeySet::EmptySet,
+			SlateApp.GetCursorPos(),
+			SlateApp.GetLastCursorPos(),
+			SlateApp.PressedMouseButtons,
 			EKeys::LeftMouseButton,
 			0,
 			SlateApp.GetPlatformApplication()->GetModifierKeys()
@@ -124,9 +140,9 @@ bool FAnalogCursor::HandleKeyUpEvent(FSlateApplication& SlateApp, const FKeyEven
 	{
 		FPointerEvent MouseEvent(
 			0,
-			CurrentPos,
-			CurrentPos,
-			FTouchKeySet::EmptySet,
+			SlateApp.GetCursorPos(),
+			SlateApp.GetLastCursorPos(),
+			SlateApp.PressedMouseButtons,
 			EKeys::LeftMouseButton,
 			0,
 			SlateApp.GetPlatformApplication()->GetModifierKeys()
diff --git a/Engine/Source/Runtime/Slate/Public/Framework/Application/SlateApplication.h b/Engine/Source/Runtime/Slate/Public/Framework/Application/SlateApplication.h
--- a/Engine/Source/Runtime/Slate/Public/Framework/Application/SlateApplication.h
+++ b/Engine/Source/Runtime/Slate/Public/Framework/Application/SlateApplication.h
@@ -666,6 +666,7 @@ public:
 
 protected:
 
+	friend class FAnalogCursor;
 	friend class FEventRouter;
 
 	virtual bool HasMouseCapture(const TSharedPtr<const SWidget> Widget) const override;
