Made ClearUserFocus default to using EFocusCause::SetDirectly rather than EFocusCause::Cleared

UE-4862 - Content browser search bar is functioning inconsistently.

ClearUserFocus is supposed to be used as a replacement for calling SetUserFocus with a null widget, however it was defaulting to EFocusCause::Cleared rather than EFocusCause::SetDirectly. EFocusCause::Cleared should only be used when you want to unset focus *and* empty the contents of a widget (this is used by text controls when you press escape), so this change meant that some text controls were being cleared out at inappropriate times.

diff --git a/Engine/Source/Runtime/Slate/Private/Framework/Application/SlateApplication.cpp b/Engine/Source/Runtime/Slate/Private/Framework/Application/SlateApplication.cpp
--- a/Engine/Source/Runtime/Slate/Private/Framework/Application/SlateApplication.cpp
+++ b/Engine/Source/Runtime/Slate/Private/Framework/Application/SlateApplication.cpp
@@ -1842,7 +1842,7 @@ TSharedPtr<SWidget> FSlateApplication::GetJoystickCaptor(uint32 UserIndex) const
 	return GetUserFocusedWidget(UserIndex);
 }
 
-void FSlateApplication::ClearUserFocus(uint32 UserIndex, EFocusCause ReasonFocusIsChanging /* = EFocusCause::Cleared*/)
+void FSlateApplication::ClearUserFocus(uint32 UserIndex, EFocusCause ReasonFocusIsChanging /* = EFocusCause::SetDirectly*/)
 {
 	SetUserFocus(UserIndex, FWidgetPath(), ReasonFocusIsChanging);
 }
@@ -2179,12 +2179,12 @@ void FSlateApplication::ProcessReply( const FWidgetPath& CurrentEventPath, const
 		{
 			for (int32 SlateUserIndex = 0; SlateUserIndex < SlateApplicationDefs::MaxUsers; ++SlateUserIndex)
 			{
-				SetUserFocus(SlateUserIndex, FWidgetPath(), EFocusCause::Cleared);
+				SetUserFocus(SlateUserIndex, FWidgetPath(), TheReply.GetFocusCause());
 			}
 		}
 		else
 		{
-			SetUserFocus(UserIndex, FWidgetPath(), EFocusCause::Cleared);
+			SetUserFocus(UserIndex, FWidgetPath(), TheReply.GetFocusCause());
 		}
 	}
 
diff --git a/Engine/Source/Runtime/Slate/Public/Framework/Application/SlateApplication.h b/Engine/Source/Runtime/Slate/Public/Framework/Application/SlateApplication.h
--- a/Engine/Source/Runtime/Slate/Public/Framework/Application/SlateApplication.h
+++ b/Engine/Source/Runtime/Slate/Public/Framework/Application/SlateApplication.h
@@ -386,7 +386,7 @@ public:
 	TSharedPtr< SWidget > GetJoystickCaptor(uint32 UserIndex) const;
 
 	/** Releases the users focus from whatever it currently is on. */
-	void ClearUserFocus(uint32 UserIndex, EFocusCause ReasonFocusIsChanging = EFocusCause::Cleared);
+	void ClearUserFocus(uint32 UserIndex, EFocusCause ReasonFocusIsChanging = EFocusCause::SetDirectly);
 
 	DEPRECATED(4.6, "FSlateApplication::ReleaseJoystickCapture() is deprecated, use FSlateApplication::ClearUserFocus() instead.")
 	void ReleaseJoystickCapture(uint32 UserIndex);
@@ -401,7 +401,7 @@ public:
 	 *
 	 * @param ReasonFocusIsChanging The reason that keyboard focus is changing
 	 */
-	void ClearKeyboardFocus(const EFocusCause ReasonFocusIsChanging = EFocusCause::Cleared);
+	void ClearKeyboardFocus(const EFocusCause ReasonFocusIsChanging = EFocusCause::SetDirectly);
 
 	/**
 	 * Returns the current modifier keys state
diff --git a/Engine/Source/Runtime/SlateCore/Public/Input/Reply.h b/Engine/Source/Runtime/SlateCore/Public/Input/Reply.h
--- a/Engine/Source/Runtime/SlateCore/Public/Input/Reply.h
+++ b/Engine/Source/Runtime/SlateCore/Public/Input/Reply.h
@@ -69,9 +69,15 @@ public:
 
 	/** An event should return a FReply::Handled().ClearUserFocus() to ask the system to clear user focus*/
 	FReply& ClearUserFocus(bool bInAllUsers = false)
+	{
+		return ClearUserFocus(EFocusCause::SetDirectly, bInAllUsers);
+	}
+
+	/** An event should return a FReply::Handled().ClearUserFocus() to ask the system to clear user focus*/
+	FReply& ClearUserFocus(EFocusCause ReasonFocusIsChanging, bool bInAllUsers = false)
 	{
 		this->FocusRecipient = nullptr;
-		this->FocusChangeReason = EFocusCause::Cleared;
+		this->FocusChangeReason = ReasonFocusIsChanging;
 		this->bReleaseUserFocus = true;
 		this->bSetUserFocus = false;
 		this->bAllUsers = bInAllUsers;
