#UE4

- Initilizing values.

diff --git a/Engine/Source/Runtime/Slate/Private/Framework/Application/AnalogCursor.cpp b/Engine/Source/Runtime/Slate/Private/Framework/Application/AnalogCursor.cpp
--- a/Engine/Source/Runtime/Slate/Private/Framework/Application/AnalogCursor.cpp
+++ b/Engine/Source/Runtime/Slate/Private/Framework/Application/AnalogCursor.cpp
@@ -11,6 +11,9 @@ const float FAnalogCursor::MaxSpeed			     = 250.0f;
 const float FAnalogCursor::OverWidgetMultiplier  = 0.6f;
 
 FAnalogCursor::FAnalogCursor()
+: CurrentPos(FVector2D::ZeroVector)
+, AnalogValues(FVector2D::ZeroVector)
+, CurrentSpeed(FVector2D::ZeroVector)
 {
 
 }
