Fixed crash closing a window on mac after all other windows have been minimized.  The mac menu contains a stale entry that must be cleared.

diff --git a/Engine/Source/Editor/UnrealEd/Private/Toolkits/SStandaloneAssetEditorToolkitHost.cpp b/Engine/Source/Editor/UnrealEd/Private/Toolkits/SStandaloneAssetEditorToolkitHost.cpp
--- a/Engine/Source/Editor/UnrealEd/Private/Toolkits/SStandaloneAssetEditorToolkitHost.cpp
+++ b/Engine/Source/Editor/UnrealEd/Private/Toolkits/SStandaloneAssetEditorToolkitHost.cpp
@@ -349,7 +349,17 @@ void SStandaloneAssetEditorToolkitHost::OnTabClosed(TSharedRef<SDockTab> TabClos
 {
 	check(TabClosed == HostTabPtr.Pin());
 
+>>>> ORIGINAL //depot/UE4/Engine/Source/Editor/UnrealEd/Private/Toolkits/SStandaloneAssetEditorToolkitHost.cpp#40
+	const TSharedPtr<FAssetEditorToolkit> AssetEditorToolkit = StaticCastSharedPtr<FAssetEditorToolkit>(HostedToolkit);
+	if(AssetEditorToolkit.IsValid())
+==== THEIRS //depot/UE4/Engine/Source/Editor/UnrealEd/Private/Toolkits/SStandaloneAssetEditorToolkitHost.cpp#41
 	if (HostedAssetEditorToolkit.IsValid())
+==== YOURS //matt_kuhlenschmidt_mac/UE4/Engine/Source/Editor/UnrealEd/Private/Toolkits/SStandaloneAssetEditorToolkitHost.cpp
+	MyTabManager->SetMenuMultiBox(nullptr);
+	
+	const TSharedPtr<FAssetEditorToolkit> AssetEditorToolkit = StaticCastSharedPtr<FAssetEditorToolkit>(HostedToolkit);
+	if(AssetEditorToolkit.IsValid())
+<<<<
 	{
 		const TArray<UObject*>* const ObjectsBeingEdited = HostedAssetEditorToolkit->GetObjectsCurrentlyBeingEdited();
 		if(ObjectsBeingEdited)
diff --git a/Engine/Source/Runtime/Slate/Private/Framework/Docking/TabManager.cpp b/Engine/Source/Runtime/Slate/Private/Framework/Docking/TabManager.cpp
--- a/Engine/Source/Runtime/Slate/Private/Framework/Docking/TabManager.cpp
+++ b/Engine/Source/Runtime/Slate/Private/Framework/Docking/TabManager.cpp
@@ -468,6 +468,10 @@ void FTabManager::SetMenuMultiBox(const TSharedPtr< FMultiBox >& NewMenuMutliBox
 	{
 		UpdateMainMenu(false);
 	}
+	else
+	{
+		FSlateMacMenu::UpdateWithMultiBox(nullptr);
+	}
 #endif
 }
 
diff --git a/Engine/Source/Runtime/Slate/Private/Framework/MultiBox/Mac/MacMenu.cpp b/Engine/Source/Runtime/Slate/Private/Framework/MultiBox/Mac/MacMenu.cpp
--- a/Engine/Source/Runtime/Slate/Private/Framework/MultiBox/Mac/MacMenu.cpp
+++ b/Engine/Source/Runtime/Slate/Private/Framework/MultiBox/Mac/MacMenu.cpp
@@ -81,7 +81,7 @@ static TMap<FMacMenu*, TSharedPtr<TArray<FMacMenuItemState>>> GCachedMenuState;
 
 @end
 
-void FSlateMacMenu::UpdateWithMultiBox(const TSharedRef< FMultiBox >& MultiBox)
+void FSlateMacMenu::UpdateWithMultiBox(const TSharedPtr< FMultiBox > MultiBox)
 {
 	MainThreadCall(^{
 		if (!FPlatformMisc::UpdateCachedMacMenuState)
@@ -98,54 +98,57 @@ void FSlateMacMenu::UpdateWithMultiBox(const TSharedRef< FMultiBox >& MultiBox)
 			[[NSApp mainMenu] removeItemAtIndex:Index];
 		}
 		GCachedMenuState.Reset();
-
-		const TArray<TSharedRef<const FMultiBlock> >& MenuBlocks = MultiBox->GetBlocks();
-
-		for (int32 Index = 0; Index < MenuBlocks.Num(); Index++)
+		
+		if( MultiBox.IsValid() )
 		{
-			TSharedRef<const FMenuEntryBlock> Block = StaticCastSharedRef<const FMenuEntryBlock>(MenuBlocks[Index]);
-			FMacMenu* Menu = [[FMacMenu alloc] initWithMenuEntryBlock:Block];
-			NSString* Title = FSlateMacMenu::GetMenuItemTitle(Block);
-			[Menu setTitle:Title];
-
-			NSMenuItem* MenuItem = [[NSMenuItem new] autorelease];
-			[MenuItem setTitle:Title];
-			[[NSApp mainMenu] addItem:MenuItem];
-			[MenuItem setSubmenu:Menu];
-
-			const bool bIsWindowMenu = (WindowLabel.ToString().Compare(FString(Title)) == 0);
-			if (bIsWindowMenu)
+			const TArray<TSharedRef<const FMultiBlock> >& MenuBlocks = MultiBox->GetBlocks();
+
+			for (int32 Index = 0; Index < MenuBlocks.Num(); Index++)
 			{
-				[NSApp setWindowsMenu:nil];
+				TSharedRef<const FMenuEntryBlock> Block = StaticCastSharedRef<const FMenuEntryBlock>(MenuBlocks[Index]);
+				FMacMenu* Menu = [[FMacMenu alloc] initWithMenuEntryBlock:Block];
+				NSString* Title = FSlateMacMenu::GetMenuItemTitle(Block);
+				[Menu setTitle:Title];
 
-				[Menu removeAllItems];
+				NSMenuItem* MenuItem = [[NSMenuItem new] autorelease];
+				[MenuItem setTitle:Title];
+				[[NSApp mainMenu] addItem:MenuItem];
+				[MenuItem setSubmenu:Menu];
 
-				NSMenuItem* MinimizeItem = [[[NSMenuItem alloc] initWithTitle:@"Minimize" action:@selector(miniaturize:) keyEquivalent:@"m"] autorelease];
-				NSMenuItem* ZoomItem = [[[NSMenuItem alloc] initWithTitle:@"Zoom" action:@selector(performZoom:) keyEquivalent:@""] autorelease];
-				NSMenuItem* CloseItem = [[[NSMenuItem alloc] initWithTitle:@"Close" action:@selector(performClose:) keyEquivalent:@"w"] autorelease];
-				NSMenuItem* BringAllToFrontItem = [[[NSMenuItem alloc] initWithTitle:@"Bring All to Front" action:@selector(arrangeInFront:) keyEquivalent:@""] autorelease];
+				const bool bIsWindowMenu = (WindowLabel.ToString().Compare(FString(Title)) == 0);
+				if (bIsWindowMenu)
+				{
+					[NSApp setWindowsMenu:nil];
+
+					[Menu removeAllItems];
 
-				[Menu addItem:MinimizeItem];
-				[Menu addItem:ZoomItem];
-				[Menu addItem:CloseItem];
-				[Menu addItem:[NSMenuItem separatorItem]];
-				[Menu addItem:BringAllToFrontItem];
-				[Menu addItem:[NSMenuItem separatorItem]];
+					NSMenuItem* MinimizeItem = [[[NSMenuItem alloc] initWithTitle:@"Minimize" action:@selector(miniaturize:) keyEquivalent:@"m"] autorelease];
+					NSMenuItem* ZoomItem = [[[NSMenuItem alloc] initWithTitle:@"Zoom" action:@selector(performZoom:) keyEquivalent:@""] autorelease];
+					NSMenuItem* CloseItem = [[[NSMenuItem alloc] initWithTitle:@"Close" action:@selector(performClose:) keyEquivalent:@"w"] autorelease];
+					NSMenuItem* BringAllToFrontItem = [[[NSMenuItem alloc] initWithTitle:@"Bring All to Front" action:@selector(arrangeInFront:) keyEquivalent:@""] autorelease];
 
-				[NSApp setWindowsMenu:Menu];
-				[Menu addItem:[NSMenuItem separatorItem]];
+					[Menu addItem:MinimizeItem];
+					[Menu addItem:ZoomItem];
+					[Menu addItem:CloseItem];
+					[Menu addItem:[NSMenuItem separatorItem]];
+					[Menu addItem:BringAllToFrontItem];
+					[Menu addItem:[NSMenuItem separatorItem]];
+
+					[NSApp setWindowsMenu:Menu];
+					[Menu addItem:[NSMenuItem separatorItem]];
+				}
 			}
-		}
 
-		if ([[NSApp mainMenu] itemWithTitle:HelpTitle.ToString().GetNSString()] == nil && HelpMenu)
-		{
-			[[NSApp mainMenu] addItem:HelpMenu];
-			GCachedMenuState.Add((FMacMenu*)HelpMenu, TSharedPtr<TArray<FMacMenuItemState>>(new TArray<FMacMenuItemState>()));
-		}
+			if ([[NSApp mainMenu] itemWithTitle:HelpTitle.ToString().GetNSString()] == nil && HelpMenu)
+			{
+				[[NSApp mainMenu] addItem:HelpMenu];
+				GCachedMenuState.Add((FMacMenu*)HelpMenu, TSharedPtr<TArray<FMacMenuItemState>>(new TArray<FMacMenuItemState>()));
+			}
 
-		if (HelpMenu)
-		{
-			[HelpMenu release];
+			if (HelpMenu)
+			{
+				[HelpMenu release];
+			}
 		}
 
 		FPlatformMisc::bChachedMacMenuStateNeedsUpdate = true;
diff --git a/Engine/Source/Runtime/Slate/Private/Framework/MultiBox/Mac/MacMenu.h b/Engine/Source/Runtime/Slate/Private/Framework/MultiBox/Mac/MacMenu.h
--- a/Engine/Source/Runtime/Slate/Private/Framework/MultiBox/Mac/MacMenu.h
+++ b/Engine/Source/Runtime/Slate/Private/Framework/MultiBox/Mac/MacMenu.h
@@ -15,7 +15,7 @@ class SLATE_API FSlateMacMenu
 {
 public:
 
-	static void UpdateWithMultiBox(const TSharedRef<FMultiBox>& MultiBox);
+	static void UpdateWithMultiBox(const TSharedPtr<FMultiBox> MultiBox);
 	static void UpdateMenu(FMacMenu* Menu);
 	static void UpdateCachedState();
 	static void ExecuteMenuItemAction(const TSharedRef<const FMenuEntryBlock>& Block);
