Slate : ending a drag and drop operation sends a DragLeave to all widgets under the cursor.

diff --git a/Engine/Source/Runtime/Slate/Private/Framework/Application/SlateApplication.cpp b/Engine/Source/Runtime/Slate/Private/Framework/Application/SlateApplication.cpp
--- a/Engine/Source/Runtime/Slate/Private/Framework/Application/SlateApplication.cpp
+++ b/Engine/Source/Runtime/Slate/Private/Framework/Application/SlateApplication.cpp
@@ -2190,7 +2190,7 @@ void FSlateApplication::ProcessReply( const FWidgetPath& CurrentEventPath, const
 
 	if (TheReply.ShouldEndDragDrop())
 	{
-		DragDropContent.Reset();
+		EndDragDrop();
 	}
 
 	if ( bStartingDragDrop )
@@ -2833,6 +2833,15 @@ TSharedPtr<FDragDropOperation> FSlateApplication::GetDragDroppingContent() const
 
 void FSlateApplication::EndDragDrop()
 {
+	FWidgetPath WidgetsToDragLeave = WidgetsUnderCursorLastEvent.ToWidgetPath(FWeakWidgetPath::EInterruptedPathHandling::Truncate);
+	if (WidgetsToDragLeave.IsValid())
+	{
+		const FDragDropEvent DragDropEvent(FPointerEvent(), DragDropContent);
+		for (int32 WidgetIndex = WidgetsToDragLeave.Widgets.Num() - 1; WidgetIndex >= 0; --WidgetIndex)
+		{
+			WidgetsToDragLeave.Widgets[WidgetIndex].Widget->OnDragLeave(DragDropEvent);
+		}
+	}
 	DragDropContent.Reset();
 }
 
@@ -3464,7 +3473,7 @@ bool FSlateApplication::ProcessKeyDownEvent( FKeyEvent& InKeyEvent )
 	if (IsDragDropping() && InKeyEvent.GetKey() == EKeys::Escape)
 	{
 		// Pressing ESC while drag and dropping terminates the drag drop.
-		DragDropContent.Reset();
+		EndDragDrop();
 		Reply = FReply::Handled();
 	}
 	else
