#UE4

- Small changes to Analog Cursor in prep for bigger changes

diff --git a/Engine/Source/Runtime/Slate/Private/Framework/Application/AnalogCursor.cpp b/Engine/Source/Runtime/Slate/Private/Framework/Application/AnalogCursor.cpp
--- a/Engine/Source/Runtime/Slate/Private/Framework/Application/AnalogCursor.cpp
+++ b/Engine/Source/Runtime/Slate/Private/Framework/Application/AnalogCursor.cpp
@@ -11,11 +11,9 @@ const float FAnalogCursor::MaxSpeed			     = 250.0f;
 const float FAnalogCursor::OverWidgetMultiplier  = 0.6f;
 
 FAnalogCursor::FAnalogCursor()
-: CurrentPos(FVector2D::ZeroVector)
-, AnalogValues(FVector2D::ZeroVector)
+: AnalogValues(FVector2D::ZeroVector)
 , CurrentSpeed(FVector2D::ZeroVector)
 {
-
 }
 
 void FAnalogCursor::Tick(const float DeltaTime, FSlateApplication& SlateApp, TSharedRef<ICursor> Cursor)
@@ -33,7 +31,7 @@ void FAnalogCursor::Tick(const float DeltaTime, FSlateApplication& SlateApp, TSh
 		if (Widget->SupportsKeyboardFocus() && Widget->GetType() != FName("SViewport"))
 		{
 			SpeedMult = 0.5f;
-			//FVector2D Adjustment = WidgetsAndCursors.Last().Geometry.Position - CurrentPos; // example of calculating distance from cursor to widget center
+			//FVector2D Adjustment = WidgetsAndCursors.Last().Geometry.Position - OldPos; // example of calculating distance from cursor to widget center
 		}
 	}
 
@@ -72,24 +70,10 @@ void FAnalogCursor::Tick(const float DeltaTime, FSlateApplication& SlateApp, TSh
 	CurrentSpeed.X = FMath::Clamp(CurrentSpeed.X, CurrentMinSpeedX, CurrentMaxSpeedX);
 	CurrentSpeed.Y = FMath::Clamp(CurrentSpeed.Y, CurrentMinSpeedY, CurrentMaxSpeedY);
 
-	CurrentPos = OldPos + (CurrentSpeed * SpeedMult);
-	Cursor->SetPosition(CurrentPos.X, CurrentPos.Y);
+	const FVector2D NewPosition = OldPos + (CurrentSpeed * SpeedMult);
 
-	// Send out a move event
-	if (OldPos != CurrentPos)
-	{
-		FPointerEvent MouseEvent(
-			0,
-			CurrentPos,
-			OldPos,
-			SlateApp.PressedMouseButtons,
-			EKeys::Invalid,
-			0,
-			SlateApp.GetPlatformApplication()->GetModifierKeys()
-			);
-
-		SlateApp.ProcessMouseMoveEvent(MouseEvent);
-	}
+	//update the cursor position
+	UpdateCursorPosition(SlateApp, Cursor, NewPosition);
 }
 
 bool FAnalogCursor::HandleKeyDownEvent(FSlateApplication& SlateApp, const FKeyEvent& InKeyEvent)
@@ -182,4 +166,32 @@ bool FAnalogCursor::HandleAnalogInputEvent(FSlateApplication& SlateApp, const FA
 		return false;
 	}
 	return true;
-}
\ No newline at end of file
+}
+
+void FAnalogCursor::UpdateCursorPosition(FSlateApplication& SlateApp, TSharedRef<ICursor> Cursor, const FVector2D& NewPosition)
+{
+	//grab the old position
+	const FVector2D OldPosition = Cursor->GetPosition();
+
+	//make sure we are actually moving
+	if (OldPosition != NewPosition)
+	{
+		//put the cursor in the correct spot
+		Cursor->SetPosition(NewPosition.X, NewPosition.Y);
+
+		//create a new mouse event
+		FPointerEvent MouseEvent(
+			0,
+			NewPosition,
+			OldPosition,
+			SlateApp.PressedMouseButtons,
+			EKeys::Invalid,
+			0,
+			SlateApp.GetPlatformApplication()->GetModifierKeys()
+			);
+
+		
+		//process the event
+		SlateApp.ProcessMouseMoveEvent(MouseEvent);
+	}
+}
diff --git a/Engine/Source/Runtime/Slate/Private/Framework/Application/AnalogCursor.h b/Engine/Source/Runtime/Slate/Private/Framework/Application/AnalogCursor.h
--- a/Engine/Source/Runtime/Slate/Private/Framework/Application/AnalogCursor.h
+++ b/Engine/Source/Runtime/Slate/Private/Framework/Application/AnalogCursor.h
@@ -12,15 +12,19 @@ class FAnalogCursor
 public:
 	FAnalogCursor();
 
-	void Tick(const float DeltaTime, FSlateApplication& SlateApp, TSharedRef<ICursor> Cursor);
+	virtual void Tick(const float DeltaTime, FSlateApplication& SlateApp, TSharedRef<ICursor> Cursor);
 
 	bool HandleKeyDownEvent(FSlateApplication& SlateApp, const FKeyEvent& InKeyEvent);
 	bool HandleKeyUpEvent(FSlateApplication& SlateApp, const FKeyEvent& InKeyEvent);
 	bool HandleAnalogInputEvent(FSlateApplication& SlateApp, const FAnalogInputEvent& InAnalogInputEvent);
 
-private:
 
-	FVector2D CurrentPos;
+protected:
+
+	/** Handles updating the cursor position and processing a Mouse Move Event */
+	void UpdateCursorPosition(FSlateApplication& SlateApp, TSharedRef<ICursor> Cursor, const FVector2D& NewPosition);
+
+private:
 
 	FVector2D AnalogValues;
 	FVector2D CurrentSpeed;
