#Orion

- Adding in OrionAnalogCursor.

diff --git a/Engine/Source/Runtime/Slate/Private/Framework/Application/AnalogCursor.h b/Engine/Source/Runtime/Slate/Private/Framework/Application/AnalogCursor.h
deleted file mode 100644
diff --git a/Engine/Source/Runtime/Slate/Private/Framework/Application/SlateApplication.cpp b/Engine/Source/Runtime/Slate/Private/Framework/Application/SlateApplication.cpp
--- a/Engine/Source/Runtime/Slate/Private/Framework/Application/SlateApplication.cpp
+++ b/Engine/Source/Runtime/Slate/Private/Framework/Application/SlateApplication.cpp
@@ -2990,11 +2990,11 @@ void FSlateApplication::SetDragTriggerDistnace( float ScreenPixels )
 	DragTriggerDistnace = ScreenPixels;
 }
 
-void FSlateApplication::SetAnalogCursorEnable(bool bEnable)
+void FSlateApplication::SetAnalogCursorEnable(bool bEnable, TSharedPtr<class FAnalogCursor> OptionalNewAnalogCursor)
 {
 	if (bEnable && !AnalogCursor.IsValid())
 	{
-		AnalogCursor = MakeShareable(new FAnalogCursor);
+		AnalogCursor = OptionalNewAnalogCursor.IsValid() ? OptionalNewAnalogCursor : MakeShareable(new FAnalogCursor);
 	}
 	else
 	{
diff --git a/Engine/Source/Runtime/Slate/Public/Framework/Application/AnalogCursor.h b/Engine/Source/Runtime/Slate/Public/Framework/Application/AnalogCursor.h
new file mode 100644
--- /dev/null
+++ b/Engine/Source/Runtime/Slate/Public/Framework/Application/AnalogCursor.h
@@ -0,0 +1,37 @@
+// Copyright 1998-2014 Epic Games, Inc. All Rights Reserved.
+
+#pragma once
+
+class FSlateApplication;
+
+/**
+ * A class that simulates a cursor driven by an analog stick.
+ */
+class SLATE_API FAnalogCursor
+{
+public:
+	FAnalogCursor();
+
+	virtual void Tick(const float DeltaTime, FSlateApplication& SlateApp, TSharedRef<ICursor> Cursor);
+
+	bool HandleKeyDownEvent(FSlateApplication& SlateApp, const FKeyEvent& InKeyEvent);
+	bool HandleKeyUpEvent(FSlateApplication& SlateApp, const FKeyEvent& InKeyEvent);
+	bool HandleAnalogInputEvent(FSlateApplication& SlateApp, const FAnalogInputEvent& InAnalogInputEvent);
+
+
+protected:
+
+	/** Handles updating the cursor position and processing a Mouse Move Event */
+	void UpdateCursorPosition(FSlateApplication& SlateApp, TSharedRef<ICursor> Cursor, const FVector2D& NewPosition);
+
+private:
+
+	FVector2D AnalogValues;
+	FVector2D CurrentSpeed;
+
+	static const float Acceleration;
+	static const float Decceleration;
+	static const float MaxSpeed;
+	static const float OverWidgetMultiplier;
+};
+
diff --git a/Engine/Source/Runtime/Slate/Public/Framework/Application/SlateApplication.h b/Engine/Source/Runtime/Slate/Public/Framework/Application/SlateApplication.h
--- a/Engine/Source/Runtime/Slate/Public/Framework/Application/SlateApplication.h
+++ b/Engine/Source/Runtime/Slate/Public/Framework/Application/SlateApplication.h
@@ -949,7 +949,7 @@ public:
 	void SetDragTriggerDistnace( float ScreenPixels );
 	
 	/** Set the analog cursor to be enabled or disabled. */
-	void SetAnalogCursorEnable(bool bEnable);
+	void SetAnalogCursorEnable(bool bEnable, TSharedPtr<class FAnalogCursor> OptionalNewAnalogCursor = nullptr);
 
 public:
 
