#UE4

- Fixing Android compile error

diff --git a/Engine/Source/Runtime/Slate/Public/Framework/Application/AnalogCursor.h b/Engine/Source/Runtime/Slate/Public/Framework/Application/AnalogCursor.h
--- a/Engine/Source/Runtime/Slate/Public/Framework/Application/AnalogCursor.h
+++ b/Engine/Source/Runtime/Slate/Public/Framework/Application/AnalogCursor.h
@@ -12,6 +12,10 @@ class SLATE_API FAnalogCursor
 public:
 	FAnalogCursor();
 
+	/** Dtor */
+	virtual ~FAnalogCursor()
+	{}
+
 	virtual void Tick(const float DeltaTime, FSlateApplication& SlateApp, TSharedRef<ICursor> Cursor);
 
 	bool HandleKeyDownEvent(FSlateApplication& SlateApp, const FKeyEvent& InKeyEvent);
