Back out changelists 2362012 & 2362003

diff --git a/Engine/Source/Developer/FriendsAndChat/Private/UI/Widgets/SChatWindow.cpp b/Engine/Source/Developer/FriendsAndChat/Private/UI/Widgets/SChatWindow.cpp
--- a/Engine/Source/Developer/FriendsAndChat/Private/UI/Widgets/SChatWindow.cpp
+++ b/Engine/Source/Developer/FriendsAndChat/Private/UI/Widgets/SChatWindow.cpp
@@ -149,7 +149,6 @@ public:
 					.OnTextCommitted(this, &SChatWindowImpl::HandleChatEntered)
 					.HintText(LOCTEXT("FriendsListSearch", "Enter to chat"))
 					.Font(FriendStyle.FriendsFontStyle)
-					.MaxCharacterLength(200)
 				]
 			]
 			]
diff --git a/Engine/Source/Runtime/Slate/Private/Widgets/Input/SEditableText.cpp b/Engine/Source/Runtime/Slate/Private/Widgets/Input/SEditableText.cpp
--- a/Engine/Source/Runtime/Slate/Private/Widgets/Input/SEditableText.cpp
+++ b/Engine/Source/Runtime/Slate/Private/Widgets/Input/SEditableText.cpp
@@ -84,7 +84,6 @@ void SEditableText::Construct( const FArguments& InArgs )
 	MinDesiredWidth = InArgs._MinDesiredWidth;
 	SelectAllTextOnCommit = InArgs._SelectAllTextOnCommit;
 	VirtualKeyboardType = IsPassword.Get() ? EKeyboardType::Keyboard_Password : InArgs._VirtualKeyboardType;
-	MaxCharacterLength = InArgs._MaxCharacterLength;
 
 	// Map UI commands to delegates which are called when the command should be executed
 	UICommandList->MapAction( FGenericCommands::Get().Undo,
@@ -658,14 +657,6 @@ void SEditableText::TypeChar( const int32 InChar )
 		{
 			bIsCharAllowed = false;
 		}
-		else
-		{
-			int32 MaxLen = MaxCharacterLength.Get();
-			if (MaxLen >= 0 && EditedText.ToString().Len() >= MaxLen)
-			{
-				bIsCharAllowed = false;
-			}
-		}
 	}
 
 	if( bIsCharAllowed )
@@ -1110,11 +1101,7 @@ void SEditableText::PasteTextFromClipboard()
 
 		for( int32 CurCharIndex = 0; CurCharIndex < PastedText.Len(); ++CurCharIndex )
 		{
-			TCHAR CharToType = PastedText.GetCharArray()[ CurCharIndex ];
-			if (CanTypeCharacter(CharToType))
-			{
-				TypeChar(CharToType);
-			}
+			TypeChar( PastedText.GetCharArray()[ CurCharIndex ] );
 		}
 
 		FinishChangingText();
diff --git a/Engine/Source/Runtime/Slate/Private/Widgets/Input/SEditableTextBox.cpp b/Engine/Source/Runtime/Slate/Private/Widgets/Input/SEditableTextBox.cpp
--- a/Engine/Source/Runtime/Slate/Private/Widgets/Input/SEditableTextBox.cpp
+++ b/Engine/Source/Runtime/Slate/Private/Widgets/Input/SEditableTextBox.cpp
@@ -49,7 +49,6 @@ void SEditableTextBox::Construct( const FArguments& InArgs )
 					.OnTextCommitted( InArgs._OnTextCommitted )
 					.MinDesiredWidth( InArgs._MinDesiredWidth )
 					.SelectAllTextOnCommit( InArgs._SelectAllTextOnCommit )
-					.MaxCharacterLength( InArgs._MaxCharacterLength )
 				]
 			]
 		]
diff --git a/Engine/Source/Runtime/Slate/Public/Widgets/Input/SEditableText.h b/Engine/Source/Runtime/Slate/Public/Widgets/Input/SEditableText.h
--- a/Engine/Source/Runtime/Slate/Public/Widgets/Input/SEditableText.h
+++ b/Engine/Source/Runtime/Slate/Public/Widgets/Input/SEditableText.h
@@ -34,7 +34,6 @@ public:
 		, _MinDesiredWidth(0.0f)
 		, _SelectAllTextOnCommit( false )
 		, _VirtualKeyboardType(EKeyboardType::Keyboard_Default)
-		, _MaxCharacterLength(-1)
 		{}
 
 		/** Sets the text content for this editable text widget */
@@ -108,9 +107,6 @@ public:
 		/** The type of virtual keyboard to use on mobile devices */
 		SLATE_ATTRIBUTE( EKeyboardType, VirtualKeyboardType)
 
-		/** The maximum number of characters that can be entered by the user into this text box */
-		SLATE_ATTRIBUTE( int32, MaxCharacterLength)
-
 	SLATE_END_ARGS()
 
 	/** Constructor */
@@ -527,9 +523,6 @@ private:
 	/** Whether to select all text when pressing enter to commit changes */
 	TAttribute< bool > SelectAllTextOnCommit;
 
-	/** Maximum number of characters that can be entered into this text box */
-	TAttribute< int32 > MaxCharacterLength;
-
 	/** Called when a character is typed and we want to know if the text field supports typing this character. */
 	FOnIsTypedCharValid OnIsTypedCharValid;
 
diff --git a/Engine/Source/Runtime/Slate/Public/Widgets/Input/SEditableTextBox.h b/Engine/Source/Runtime/Slate/Public/Widgets/Input/SEditableTextBox.h
--- a/Engine/Source/Runtime/Slate/Public/Widgets/Input/SEditableTextBox.h
+++ b/Engine/Source/Runtime/Slate/Public/Widgets/Input/SEditableTextBox.h
@@ -28,7 +28,6 @@ public:
 		, _BackgroundColor()		
 		, _Padding()
 		, _ErrorReporting()
-		, _MaxCharacterLength(-1)
 		{}
 
 		/** The styling of the textbox */
@@ -88,9 +87,6 @@ public:
 		/** Provide a alternative mechanism for error reporting. */
 		SLATE_ARGUMENT( TSharedPtr<class IErrorReportingWidget>, ErrorReporting )
 
-		/** The maximum number of characters that can be entered by the user into this text box */
-		SLATE_ATTRIBUTE(int32, MaxCharacterLength)
-
 	SLATE_END_ARGS()
 	
 	/**
