Fixed output log performance degradation as new text was added

UETOOL-129 - Investigate SOutputLog performance degradation

FTextLayout::AddLine was causing a re-flow of the entire document when a new line was added. This change makes it only flow the line that is being added (providing the rest of the layout is already up-to-date).

diff --git a/Engine/Source/Developer/OutputLog/Private/SOutputLog.cpp b/Engine/Source/Developer/OutputLog/Private/SOutputLog.cpp
--- a/Engine/Source/Developer/OutputLog/Private/SOutputLog.cpp
+++ b/Engine/Source/Developer/OutputLog/Private/SOutputLog.cpp
@@ -772,9 +772,6 @@ void SOutputLog::Serialize( const TCHAR* V, ELogVerbosity::Type Verbosity, const
 		const float DistanceFromBottom = MessagesTextBox->GetVScrollBar()->DistanceFromBottom();
 		if( DistanceFromBottom <= KINDA_SMALL_NUMBER )
 		{
-			// Force a refresh so that the message has been added before we try and jump to it
-			//MessagesTextBox->Refresh();
-
 			MessagesTextBox->ScrollTo(FTextLocation(MessagesTextMarshaller->GetNumMessages() - 1));
 		}
 	}
diff --git a/Engine/Source/Runtime/Slate/Private/Framework/Text/TextLayout.cpp b/Engine/Source/Runtime/Slate/Private/Framework/Text/TextLayout.cpp
--- a/Engine/Source/Runtime/Slate/Private/Framework/Text/TextLayout.cpp
+++ b/Engine/Source/Runtime/Slate/Private/Framework/Text/TextLayout.cpp
@@ -290,80 +290,90 @@ void FTextLayout::JustifyLayout()
 	}
 }
 
-void FTextLayout::FlowLayout()
+float FTextLayout::GetWrappingDrawWidth() const
 {
 	check( WrappingWidth >= 0 );
-	const bool IsWrapping = WrappingWidth > 0;
-	const float WrappingDrawWidth = FMath::Max( 0.01f, ( WrappingWidth - Margin.GetTotalSpaceAlong<Orient_Horizontal>() ) * Scale );
+	return FMath::Max( 0.01f, ( WrappingWidth - Margin.GetTotalSpaceAlong<Orient_Horizontal>() ) * Scale );
+}
+
+void FTextLayout::FlowLayout()
+{
+	const float WrappingDrawWidth = GetWrappingDrawWidth();
 
 	CreateWrappingCache();
 
 	TArray< TSharedRef< ILayoutBlock > > SoftLine;
-
 	for (int32 LineModelIndex = 0; LineModelIndex < LineModels.Num(); LineModelIndex++)
 	{
-		const FLineModel& LineModel = LineModels[ LineModelIndex ];
+		FlowLineLayout(LineModelIndex, WrappingDrawWidth, SoftLine);
+	}
 
-		float CurrentWidth = 0.0f;
-		int32 CurrentRunIndex = 0;
-		int32 PreviousBlockEnd = 0;
+	//Add on the margins to the DrawSize
+	DrawSize.X += Margin.GetTotalSpaceAlong<Orient_Horizontal>() * Scale;
+	DrawSize.Y += Margin.GetTotalSpaceAlong<Orient_Vertical>() * Scale;
+}
 
-		int32 CurrentRendererIndex = 0;
-		if ( CurrentRendererIndex >= LineModel.RunRenderers.Num() )
-		{
-			CurrentRendererIndex = INDEX_NONE;
-		}
+void FTextLayout::FlowLineLayout(const int32 LineModelIndex, const float WrappingDrawWidth, TArray<TSharedRef<ILayoutBlock>>& SoftLine)
+{
+	const FLineModel& LineModel = LineModels[ LineModelIndex ];
 
-		// if the Line doesn't have any BreakCandidates
-		if ( LineModel.BreakCandidates.Num() == 0 )
-		{
-			//Then iterate over all of it's runs
-			CreateLineViewBlocks( LineModelIndex, INDEX_NONE, /*OUT*/CurrentRunIndex, /*OUT*/CurrentRendererIndex, /*OUT*/PreviousBlockEnd, SoftLine );
-			check( CurrentRunIndex == LineModel.Runs.Num() );
-			CurrentWidth = 0;
-			SoftLine.Empty();
-		}
-		else
-		{
-			for (int32 BreakIndex = 0; BreakIndex < LineModel.BreakCandidates.Num(); BreakIndex++)
-			{
-				const FBreakCandidate& Break = LineModel.BreakCandidates[ BreakIndex ];
+	float CurrentWidth = 0.0f;
+	int32 CurrentRunIndex = 0;
+	int32 PreviousBlockEnd = 0;
+
+	int32 CurrentRendererIndex = 0;
+	if ( CurrentRendererIndex >= LineModel.RunRenderers.Num() )
+	{
+		CurrentRendererIndex = INDEX_NONE;
+	}
 
-				const bool IsLastBreak = BreakIndex + 1 == LineModel.BreakCandidates.Num();
-				const bool IsFirstBreakOnSoftLine = CurrentWidth == 0;
-				const uint8 Kerning = ( IsFirstBreakOnSoftLine ) ? Break.Kerning : 0;
-				const bool BreakDoesFit = !IsWrapping || CurrentWidth + Break.ActualSize.X + Kerning <= WrappingDrawWidth;
+	const bool IsWrapping = WrappingWidth > 0.0f;
 
-				if ( !BreakDoesFit || IsLastBreak )
-				{
-					const bool BreakWithoutTrailingWhitespaceDoesFit = !IsWrapping || CurrentWidth + Break.TrimmedSize.X + Kerning <= WrappingDrawWidth;
-					const bool IsFirstBreak = BreakIndex == 0;
+	// if the Line doesn't have any BreakCandidates, or we're not wrapping text
+	if ( LineModel.BreakCandidates.Num() == 0 || !IsWrapping )
+	{
+		//Then iterate over all of it's runs
+		CreateLineViewBlocks( LineModelIndex, INDEX_NONE, /*OUT*/CurrentRunIndex, /*OUT*/CurrentRendererIndex, /*OUT*/PreviousBlockEnd, SoftLine );
+		check( CurrentRunIndex == LineModel.Runs.Num() );
+		CurrentWidth = 0;
+		SoftLine.Empty();
+	}
+	else
+	{
+		for (int32 BreakIndex = 0; BreakIndex < LineModel.BreakCandidates.Num(); BreakIndex++)
+		{
+			const FBreakCandidate& Break = LineModel.BreakCandidates[ BreakIndex ];
 
-					const FBreakCandidate& FinalBreakOnSoftLine = ( !IsFirstBreak && !IsFirstBreakOnSoftLine && !BreakWithoutTrailingWhitespaceDoesFit ) ? LineModel.BreakCandidates[ --BreakIndex ] : Break;
+			const bool IsLastBreak = BreakIndex + 1 == LineModel.BreakCandidates.Num();
+			const bool IsFirstBreakOnSoftLine = CurrentWidth == 0;
+			const uint8 Kerning = ( IsFirstBreakOnSoftLine ) ? Break.Kerning : 0;
+			const bool BreakDoesFit = !IsWrapping || CurrentWidth + Break.ActualSize.X + Kerning <= WrappingDrawWidth;
 
-					CreateLineViewBlocks( LineModelIndex, FinalBreakOnSoftLine.ActualRange.EndIndex, /*OUT*/CurrentRunIndex, /*OUT*/CurrentRendererIndex, /*OUT*/PreviousBlockEnd, SoftLine );
+			if ( !BreakDoesFit || IsLastBreak )
+			{
+				const bool BreakWithoutTrailingWhitespaceDoesFit = !IsWrapping || CurrentWidth + Break.TrimmedSize.X + Kerning <= WrappingDrawWidth;
+				const bool IsFirstBreak = BreakIndex == 0;
 
-					if ( CurrentRunIndex < LineModel.Runs.Num() && FinalBreakOnSoftLine.ActualRange.EndIndex == LineModel.Runs[ CurrentRunIndex ].GetTextRange().EndIndex )
-					{
-						++CurrentRunIndex;
-					}
+				const FBreakCandidate& FinalBreakOnSoftLine = ( !IsFirstBreak && !IsFirstBreakOnSoftLine && !BreakWithoutTrailingWhitespaceDoesFit ) ? LineModel.BreakCandidates[ --BreakIndex ] : Break;
 
-					PreviousBlockEnd = FinalBreakOnSoftLine.ActualRange.EndIndex;
+				CreateLineViewBlocks( LineModelIndex, FinalBreakOnSoftLine.ActualRange.EndIndex, /*OUT*/CurrentRunIndex, /*OUT*/CurrentRendererIndex, /*OUT*/PreviousBlockEnd, SoftLine );
 
-					CurrentWidth = 0;
-					SoftLine.Empty();
-				} 
-				else
+				if ( CurrentRunIndex < LineModel.Runs.Num() && FinalBreakOnSoftLine.ActualRange.EndIndex == LineModel.Runs[ CurrentRunIndex ].GetTextRange().EndIndex )
 				{
-					CurrentWidth += Break.ActualSize.X;
+					++CurrentRunIndex;
 				}
+
+				PreviousBlockEnd = FinalBreakOnSoftLine.ActualRange.EndIndex;
+
+				CurrentWidth = 0;
+				SoftLine.Empty();
+			} 
+			else
+			{
+				CurrentWidth += Break.ActualSize.X;
 			}
 		}
 	}
-
-	//Add on the margins to the DrawSize
-	DrawSize.X += Margin.GetTotalSpaceAlong<Orient_Horizontal>() * Scale;
-	DrawSize.Y += Margin.GetTotalSpaceAlong<Orient_Vertical>() * Scale;
 }
 
 void FTextLayout::FlowHighlights()
@@ -453,25 +463,33 @@ void FTextLayout::FlowHighlights()
 
 void FTextLayout::EndLayout()
 {
-	for (int32 LineModelIndex = 0; LineModelIndex < LineModels.Num(); LineModelIndex++)
+	for (FLineModel& LineModel : LineModels)
 	{
-		FLineModel& LineModel = LineModels[ LineModelIndex ];
-		for (int32 RunIndex = 0; RunIndex < LineModel.Runs.Num(); RunIndex++)
-		{
-			LineModel.Runs[ RunIndex ].EndLayout();
-		}
+		EndLineLayout(LineModel);
+	}
+}
+
+void FTextLayout::EndLineLayout(FLineModel& LineModel)
+{
+	for (FRunModel& RunModel : LineModel.Runs)
+	{
+		RunModel.EndLayout();
 	}
 }
 
 void FTextLayout::BeginLayout()
 {
-	for (int32 LineModelIndex = 0; LineModelIndex < LineModels.Num(); LineModelIndex++)
+	for (FLineModel& LineModel : LineModels)
 	{
-		FLineModel& LineModel = LineModels[ LineModelIndex ];
-		for (int32 RunIndex = 0; RunIndex < LineModel.Runs.Num(); RunIndex++)
-		{
-			LineModel.Runs[ RunIndex ].BeginLayout();
-		}
+		BeginLineLayout(LineModel);
+	}
+}
+
+void FTextLayout::BeginLineLayout(FLineModel& LineModel)
+{
+	for (FRunModel& RunModel : LineModel.Runs)
+	{
+		RunModel.BeginLayout();
 	}
 }
 
@@ -483,38 +501,43 @@ void FTextLayout::ClearView()
 
 void FTextLayout::CreateWrappingCache()
 {
+	for (FLineModel& LineModel : LineModels)
+	{
+		CreateLineWrappingCache(LineModel);
+	}
+}
+
+void FTextLayout::CreateLineWrappingCache(FLineModel& LineModel)
+{
+	if (LineModel.HasWrappingInformation)
+	{
+		return;
+	}
+
 	// If we've not yet been provided with a custom line break iterator, then just use the default one
 	if (!LineBreakIterator.IsValid())
 	{
 		LineBreakIterator = FBreakIterator::CreateLineBreakIterator();
 	}
 
-	for (int32 LineModelIndex = 0; LineModelIndex < LineModels.Num(); LineModelIndex++)
-	{
-		FLineModel& LineModel = LineModels[ LineModelIndex ];
-
-		if ( !LineModel.HasWrappingInformation )
-		{
-			LineModel.BreakCandidates.Empty();
-			LineModel.HasWrappingInformation = true;
+	LineModel.BreakCandidates.Empty();
+	LineModel.HasWrappingInformation = true;
 
-			for (int32 RunIndex = 0; RunIndex < LineModel.Runs.Num(); RunIndex++)
-			{
-				LineModel.Runs[ RunIndex ].ClearCache();
-			}
+	for (int32 RunIndex = 0; RunIndex < LineModel.Runs.Num(); RunIndex++)
+	{
+		LineModel.Runs[ RunIndex ].ClearCache();
+	}
 
-			LineBreakIterator->SetString( **LineModel.Text );
+	LineBreakIterator->SetString( **LineModel.Text );
 
-			int32 PreviousBreak = 0;
-			int32 CurrentBreak = 0;
-			int32 CurrentRunIndex = 0;
+	int32 PreviousBreak = 0;
+	int32 CurrentBreak = 0;
+	int32 CurrentRunIndex = 0;
 
-			while( ( CurrentBreak = LineBreakIterator->MoveToNext() ) != INDEX_NONE )
-			{
-				LineModel.BreakCandidates.Add( CreateBreakCandidate(/*OUT*/CurrentRunIndex, LineModel, PreviousBreak, CurrentBreak) );
-				PreviousBreak = CurrentBreak;
-			}
-		}
+	while( ( CurrentBreak = LineBreakIterator->MoveToNext() ) != INDEX_NONE )
+	{
+		LineModel.BreakCandidates.Add( CreateBreakCandidate(/*OUT*/CurrentRunIndex, LineModel, PreviousBreak, CurrentBreak) );
+		PreviousBreak = CurrentBreak;
 	}
 
 	LineBreakIterator->ClearString();
@@ -522,9 +545,8 @@ void FTextLayout::CreateWrappingCache()
 
 void FTextLayout::ClearWrappingCache()
 {
-	for (int32 LineModelIndex = 0; LineModelIndex < LineModels.Num(); LineModelIndex++)
+	for (FLineModel& LineModel : LineModels)
 	{
-		FLineModel& LineModel = LineModels[ LineModelIndex ];
 		LineModel.HasWrappingInformation = false;
 	}
 }
@@ -1347,15 +1369,34 @@ bool FTextLayout::RemoveLine(int32 LineIndex)
 
 void FTextLayout::AddLine( const TSharedRef< FString >& Text, const TArray< TSharedRef< IRun > >& Runs )
 {
-	FLineModel LineModel( Text );
-
-	for (int32 Index = 0; Index < Runs.Num(); Index++)
 	{
-		LineModel.Runs.Add( FRunModel( Runs[ Index ] ) );
+		FLineModel LineModel( Text );
+
+		for (int32 Index = 0; Index < Runs.Num(); Index++)
+		{
+			LineModel.Runs.Add( FRunModel( Runs[ Index ] ) );
+		}
+
+		LineModels.Add( LineModel );
 	}
 
-	LineModels.Add( LineModel );
-	DirtyFlags |= EDirtyState::Layout;
+	// If our layout is clean, then we can add this new line immediately (and efficiently)
+	// If our layout is dirty, then we might as well wait as the next UpdateLayout call will add it
+	if (!(DirtyFlags & EDirtyState::Layout))
+	{
+		const int32 LineModelIndex = LineModels.Num() - 1;
+		FLineModel& LineModel = LineModels[LineModelIndex];
+
+		BeginLineLayout(LineModel);
+
+		TArray<TSharedRef<ILayoutBlock>> SoftLine;
+		FlowLineLayout(LineModelIndex, GetWrappingDrawWidth(), SoftLine);
+
+		// We need to re-justify all lines, as the new line view(s) added by this line model may have affected everything
+		JustifyLayout();
+
+		EndLineLayout(LineModel);
+	}
 }
 
 void FTextLayout::ClearLines()
diff --git a/Engine/Source/Runtime/Slate/Private/Widgets/Input/SMultiLineEditableTextBox.cpp b/Engine/Source/Runtime/Slate/Private/Widgets/Input/SMultiLineEditableTextBox.cpp
--- a/Engine/Source/Runtime/Slate/Private/Widgets/Input/SMultiLineEditableTextBox.cpp
+++ b/Engine/Source/Runtime/Slate/Private/Widgets/Input/SMultiLineEditableTextBox.cpp
@@ -237,12 +237,12 @@ void SMultiLineEditableTextBox::InsertRunAtCursor(TSharedRef<IRun> InRun)
 
 void SMultiLineEditableTextBox::GoTo(const FTextLocation& NewLocation)
 {
-	return EditableText->GoTo(NewLocation);
+	EditableText->GoTo(NewLocation);
 }
 
 void SMultiLineEditableTextBox::ScrollTo(const FTextLocation& NewLocation)
 {
-	return EditableText->ScrollTo(NewLocation);
+	EditableText->ScrollTo(NewLocation);
 }
 
 void SMultiLineEditableTextBox::ApplyToSelection(const FRunInfo& InRunInfo, const FTextBlockStyle& InStyle)
diff --git a/Engine/Source/Runtime/Slate/Public/Framework/Text/TextLayout.h b/Engine/Source/Runtime/Slate/Public/Framework/Text/TextLayout.h
--- a/Engine/Source/Runtime/Slate/Public/Framework/Text/TextLayout.h
+++ b/Engine/Source/Runtime/Slate/Public/Framework/Text/TextLayout.h
@@ -399,6 +399,11 @@ protected:
 	*/
 	void CreateWrappingCache();
 
+	/**
+	* Create the wrapping cache for the given line based upon the current scale
+	*/
+	void CreateLineWrappingCache(FLineModel& LineModel);
+
 	/**
 	 * Clears the current wrapping cache for all lines
 	 */
@@ -414,11 +419,21 @@ protected:
 	*/
 	virtual void BeginLayout();
 
+	/**
+	* Notifies all Runs on the given line is beginning to have a new layout generated.
+	*/
+	void BeginLineLayout(FLineModel& LineModel);
+
 	/**
 	* Notifies all Runs that the layout has finished generating.
 	*/
 	virtual void EndLayout();
 
+	/**
+	* Notifies all Runs on the given line has finished having a new layout generated.
+	*/
+	void EndLineLayout(FLineModel& LineModel);
+
 	/**
 	 * Called to generate a new empty text run for this text layout
 	 */
@@ -426,8 +441,12 @@ protected:
 
 private:
 
+	float GetWrappingDrawWidth() const;
+
 	void FlowLayout();
 
+	void FlowLineLayout(const int32 LineModelIndex, const float WrappingDrawWidth, TArray<TSharedRef<ILayoutBlock>>& SoftLine);
+
 	void FlowHighlights();
 
 	void JustifyLayout();
