Slate : when tick folding is enabled, make sure that we do not attempt to draw any windows destroyed during the paint pass.
The paint pass is no longer side-effect free because we folded Tick() into it.

https://jira.ol.epicgames.net/browse/UE-5571?filter=12008

diff --git a/Engine/Source/Runtime/Slate/Private/Framework/Application/SlateApplication.cpp b/Engine/Source/Runtime/Slate/Private/Framework/Application/SlateApplication.cpp
--- a/Engine/Source/Runtime/Slate/Private/Framework/Application/SlateApplication.cpp
+++ b/Engine/Source/Runtime/Slate/Private/Framework/Application/SlateApplication.cpp
@@ -1011,6 +1011,19 @@ void FSlateApplication::DrawPrepass( TSharedPtr<SWindow> DrawOnlyThisWindow )
 	}
 }
 
+
+TArray<TSharedRef<SWindow>> GatherAllDescendants(const TArray< TSharedRef<SWindow> >& InWindowList)
+{
+	TArray<TSharedRef<SWindow>> GatheredDescendants(InWindowList);
+
+	for (const TSharedRef<SWindow>& SomeWindow : InWindowList)
+	{
+		GatheredDescendants.Append( GatherAllDescendants( SomeWindow->GetChildWindows() ) );
+	}
+	
+	return GatheredDescendants;
+}
+
 void FSlateApplication::PrivateDrawWindows( TSharedPtr<SWindow> DrawOnlyThisWindow )
 {
 	check(Renderer.IsValid());
@@ -1072,6 +1085,16 @@ void FSlateApplication::PrivateDrawWindows( TSharedPtr<SWindow> DrawOnlyThisWind
 		}
 	}
 
+
+	// Some windows may have been destroyed/removed.
+	// Do not attempt to draw any windows that have been removed.
+	TArray<TSharedRef<SWindow>> AllWindows = GatherAllDescendants(SlateWindows);
+	DrawWindowArgs.OutDrawBuffer.GetWindowElementLists().RemoveAll([&]( FSlateWindowElementList& Candidate )
+	{
+		TSharedPtr<SWindow> CandidateWindow = Candidate.GetWindow();
+		return !CandidateWindow.IsValid() || !AllWindows.Contains(CandidateWindow.ToSharedRef());
+	});
+	
 	Renderer->DrawWindows( DrawWindowArgs.OutDrawBuffer );
 }
 
