#Orion

Analog Cursor work

- Small clean-up of FAnalogCursor
- Removed clamping in HandleAnalogInputEvent. Now handle "Deadzone" in Tick.
- New "Tick" function for OrionAnalogCursor
- Added lots of tweakables to Orion Global

diff --git a/Engine/Source/Runtime/Slate/Private/Framework/Application/AnalogCursor.cpp b/Engine/Source/Runtime/Slate/Private/Framework/Application/AnalogCursor.cpp
--- a/Engine/Source/Runtime/Slate/Private/Framework/Application/AnalogCursor.cpp
+++ b/Engine/Source/Runtime/Slate/Private/Framework/Application/AnalogCursor.cpp
@@ -5,14 +5,12 @@
 
 #include "HittestGrid.h"
 
-const float FAnalogCursor::Acceleration          = 150.0f;
-const float FAnalogCursor::Decceleration         = 250.0f;
-const float FAnalogCursor::MaxSpeed			     = 250.0f;
-const float FAnalogCursor::OverWidgetMultiplier  = 0.6f;
+const float FAnalogCursor::DefaultAcceleration = 150.0f;
+const float FAnalogCursor::DefaultMaxSpeed = 250.0f;
 
 FAnalogCursor::FAnalogCursor()
-: AnalogValues(FVector2D::ZeroVector)
-, CurrentSpeed(FVector2D::ZeroVector)
+: CurrentSpeed(FVector2D::ZeroVector)
+, AnalogValues(FVector2D::ZeroVector)
 {
 }
 
@@ -23,6 +21,13 @@ void FAnalogCursor::Tick(const float DeltaTime, FSlateApplication& SlateApp, TSh
 	float SpeedMult = 1.0f; // Used to do a speed multiplication before adding the delta to the position to make widgets sticky
 	FVector2D AdjAnalogVals = AnalogValues; // A copy of the analog values so I can modify them based being over a widget, not currently doing this
 
+	//if the size of the analog vals vector is too small, then set it to zero
+	static const float DeadZoneSizeSq = (0.1f * 0.1f);
+	if (AdjAnalogVals.SizeSquared() < DeadZoneSizeSq)
+	{
+		AdjAnalogVals = FVector2D::ZeroVector;
+	}
+
 	FWidgetPath WidgetPath = SlateApp.LocateWindowUnderMouse(OldPos, SlateApp.GetInteractiveTopLevelWindows());
 	if (WidgetPath.IsValid())
 	{
@@ -40,11 +45,11 @@ void FAnalogCursor::Tick(const float DeltaTime, FSlateApplication& SlateApp, TSh
 	float CurrentMaxSpeedX = 0.0f;
 	if (AdjAnalogVals.X > 0.0f)
 	{ 
-		CurrentMaxSpeedX = AdjAnalogVals.X * MaxSpeed;
+		CurrentMaxSpeedX = AdjAnalogVals.X * DefaultMaxSpeed;
 	}
 	else
 	{
-		CurrentMinSpeedX = AdjAnalogVals.X * MaxSpeed;
+		CurrentMinSpeedX = AdjAnalogVals.X * DefaultMaxSpeed;
 	}
 
 	// Generate Min and Max for Y to clamp the speed, this gives us instant direction change when crossing the axis
@@ -52,15 +57,15 @@ void FAnalogCursor::Tick(const float DeltaTime, FSlateApplication& SlateApp, TSh
 	float CurrentMaxSpeedY = 0.0f;
 	if (AdjAnalogVals.Y > 0.0f)
 	{
-		CurrentMaxSpeedY = AdjAnalogVals.Y * MaxSpeed;
+		CurrentMaxSpeedY = AdjAnalogVals.Y * DefaultMaxSpeed;
 	}
 	else
 	{
-		CurrentMinSpeedY = AdjAnalogVals.Y * MaxSpeed;
+		CurrentMinSpeedY = AdjAnalogVals.Y * DefaultMaxSpeed;
 	}
 
 	// Cubic acceleration curve
-	FVector2D ExpAcceleration = AdjAnalogVals * AdjAnalogVals * AdjAnalogVals * Acceleration;
+	FVector2D ExpAcceleration = AdjAnalogVals * AdjAnalogVals * AdjAnalogVals * DefaultAcceleration;
 	// Preserve direction (if we use a squared equation above)
 	//ExpAcceleration.X *= FMath::Sign(AnalogValues.X);
 	//ExpAcceleration.Y *= FMath::Sign(AnalogValues.Y);
@@ -146,13 +151,6 @@ bool FAnalogCursor::HandleAnalogInputEvent(FSlateApplication& SlateApp, const FA
 	FKey Key = InAnalogInputEvent.GetKey();
 	float AnalogValue = InAnalogInputEvent.GetAnalogValue();
 
-	//@Todo Slate: Investigate this more, doesn't seem right that analog doesn't zero
-	// Analog doesn't zero out
-	if (FMath::Abs(AnalogValue) < 0.1f)
-	{
-		AnalogValue = 0.0f;
-	}
-
 	if (Key == EKeys::Gamepad_LeftX)
 	{
 		AnalogValues.X = AnalogValue;
diff --git a/Engine/Source/Runtime/Slate/Public/Framework/Application/AnalogCursor.h b/Engine/Source/Runtime/Slate/Public/Framework/Application/AnalogCursor.h
--- a/Engine/Source/Runtime/Slate/Public/Framework/Application/AnalogCursor.h
+++ b/Engine/Source/Runtime/Slate/Public/Framework/Application/AnalogCursor.h
@@ -22,20 +22,27 @@ public:
 	bool HandleKeyUpEvent(FSlateApplication& SlateApp, const FKeyEvent& InKeyEvent);
 	bool HandleAnalogInputEvent(FSlateApplication& SlateApp, const FAnalogInputEvent& InAnalogInputEvent);
 
-
 protected:
 
+	/** Getter */
+	FORCEINLINE FVector2D GetAnalogValues() const
+	{
+		return AnalogValues;
+	}
+
 	/** Handles updating the cursor position and processing a Mouse Move Event */
 	void UpdateCursorPosition(FSlateApplication& SlateApp, TSharedRef<ICursor> Cursor, const FVector2D& NewPosition);
 
+	/** Current speed of the cursor */
+	FVector2D CurrentSpeed;
+
+	/** Helpful statics */
+	static const float DefaultAcceleration;
+	static const float DefaultMaxSpeed;
+
 private:
 
+	/** Input from the gamepad */
 	FVector2D AnalogValues;
-	FVector2D CurrentSpeed;
-
-	static const float Acceleration;
-	static const float Decceleration;
-	static const float MaxSpeed;
-	static const float OverWidgetMultiplier;
 };
 
