Don't activate miniaturized windows on mouse up on Mac

diff --git a/Engine/Source/Runtime/Slate/Private/Framework/Application/SlateApplication.cpp b/Engine/Source/Runtime/Slate/Private/Framework/Application/SlateApplication.cpp
--- a/Engine/Source/Runtime/Slate/Private/Framework/Application/SlateApplication.cpp
+++ b/Engine/Source/Runtime/Slate/Private/Framework/Application/SlateApplication.cpp
@@ -3954,7 +3954,8 @@ bool FSlateApplication::ProcessMouseButtonUpEvent( FPointerEvent& MouseEvent )
 #if PLATFORM_MAC
 			// Activate a window under the mouse if it's inactive and mouse up didn't bring any window to front
 			TSharedPtr<SWindow> ActiveWindow = GetActiveTopLevelWindow();
-			if (MouseEvent.GetEffectingButton() == EKeys::LeftMouseButton && MouseCaptorPath.TopLevelWindow.IsValid() && ActiveWindow != MouseCaptorPath.TopLevelWindow && ActiveNativeWindow == [NSApp keyWindow])
+			if (MouseEvent.GetEffectingButton() == EKeys::LeftMouseButton && MouseCaptorPath.TopLevelWindow.IsValid() && ActiveWindow != MouseCaptorPath.TopLevelWindow
+				&& ActiveNativeWindow == [NSApp keyWindow] && ![(NSWindow*)MouseCaptorPath.TopLevelWindow->GetNativeWindow()->GetOSWindowHandle() isMiniaturized])
 			{
 				MouseCaptorPath.TopLevelWindow->BringToFront(true);
 			}
