UMG - The slate checkbox checked state was not marked as a blueprint type.

diff --git a/Engine/Source/Runtime/Slate/Public/Widgets/Input/SCheckBox.h b/Engine/Source/Runtime/Slate/Public/Widgets/Input/SCheckBox.h
--- a/Engine/Source/Runtime/Slate/Public/Widgets/Input/SCheckBox.h
+++ b/Engine/Source/Runtime/Slate/Public/Widgets/Input/SCheckBox.h
@@ -7,7 +7,7 @@
 #include "SCheckBox.generated.h"
 
 /** Current state of the check box */
-UENUM()
+UENUM(BlueprintType)
 namespace ESlateCheckBoxState
 {
 	enum Type
