(Incorrect?) fix for a weird compile error.

diff --git a/Engine/Source/Runtime/Slate/Private/Widgets/Input/SMenuAnchor.cpp b/Engine/Source/Runtime/Slate/Private/Widgets/Input/SMenuAnchor.cpp
--- a/Engine/Source/Runtime/Slate/Private/Widgets/Input/SMenuAnchor.cpp
+++ b/Engine/Source/Runtime/Slate/Private/Widgets/Input/SMenuAnchor.cpp
@@ -408,7 +408,7 @@ bool SMenuAnchor::HasOpenSubMenus() const
 SMenuAnchor::SMenuAnchor()
 	: MenuContent( SNullWidget::NullWidget )
 	, bDismissedThisTick( false )
-	, Method(CreateNewWindow)
+	, Method(EPopupMethod::CreateNewWindow)
 {
 }
 
