Slate: Add support to multibox for identifying when a menu section contains one or more icons
- Currently this is used to display a magenta warning square with informative tooltip when 'Display multibox hooks' is enabled

Misc. Multibox cleanup: Moved some code from header to cpp and added override where appropriate

diff --git a/Engine/Source/Runtime/Slate/Private/Framework/MultiBox/MultiBox.cpp b/Engine/Source/Runtime/Slate/Private/Framework/MultiBox/MultiBox.cpp
--- a/Engine/Source/Runtime/Slate/Private/Framework/MultiBox/MultiBox.cpp
+++ b/Engine/Source/Runtime/Slate/Private/Framework/MultiBox/MultiBox.cpp
@@ -56,6 +56,37 @@ void FMultiBoxSettings::ToggleToolbarEditing()
 const FMultiBoxCustomization FMultiBoxCustomization::None( NAME_None );
 
 
+TSharedRef< SWidget > SMultiBlockBaseWidget::AsWidget()
+{
+	return this->AsShared();
+}
+
+TSharedRef< const SWidget > SMultiBlockBaseWidget::AsWidget() const
+{
+	return this->AsShared();
+}
+
+void SMultiBlockBaseWidget::SetOwnerMultiBoxWidget(TSharedRef< SMultiBoxWidget > InOwnerMultiBoxWidget)
+{
+	OwnerMultiBoxWidget = InOwnerMultiBoxWidget;
+}
+
+void SMultiBlockBaseWidget::SetMultiBlock(TSharedRef< const FMultiBlock > InMultiBlock)
+{
+	MultiBlock = InMultiBlock;
+}
+
+void SMultiBlockBaseWidget::SetMultiBlockLocation(EMultiBlockLocation::Type InLocation, bool bInSectionContainsIcons)
+{
+	Location = InLocation;
+	bSectionContainsIcons = bInSectionContainsIcons;
+}
+
+EMultiBlockLocation::Type SMultiBlockBaseWidget::GetMultiBlockLocation()
+{
+	return Location;
+}
+
 void SMultiBlockBaseWidget::OnDragEnter( const FGeometry& MyGeometry, const FDragDropEvent& DragDropEvent )
 {
 	if ( DragDropEvent.GetOperationAs<FUICommandDragDropOp>().IsValid() )
@@ -94,7 +125,7 @@ FReply SMultiBlockBaseWidget::OnDrop( const FGeometry& MyGeometry, const FDragDr
  *
  * @return  MultiBlock widget object
  */
-TSharedRef< IMultiBlockBaseWidget > FMultiBlock::MakeWidget( TSharedRef< SMultiBoxWidget > InOwnerMultiBoxWidget, EMultiBlockLocation::Type InLocation ) const
+TSharedRef< IMultiBlockBaseWidget > FMultiBlock::MakeWidget( TSharedRef< SMultiBoxWidget > InOwnerMultiBoxWidget, EMultiBlockLocation::Type InLocation, bool bSectionContainsIcons ) const
 {
 	TSharedRef< IMultiBlockBaseWidget > NewMultiBlockWidget = ConstructWidget();
 
@@ -105,7 +136,7 @@ TSharedRef< IMultiBlockBaseWidget > FMultiBlock::MakeWidget( TSharedRef< SMultiB
 	NewMultiBlockWidget->SetMultiBlock( AsShared() );
 
 	// Pass location information to widget.
-	NewMultiBlockWidget->SetMultiBlockLocation(InLocation);
+	NewMultiBlockWidget->SetMultiBlockLocation(InLocation, bSectionContainsIcons);
 
 	// Work out what style the widget should be using
 	const ISlateStyle* const StyleSet = InOwnerMultiBoxWidget->GetStyleSet();
@@ -404,13 +435,13 @@ bool SMultiBoxWidget::IsBlockBeingDragged( TSharedPtr<const FMultiBlock> Block )
 	return false;
 }
 
-void SMultiBoxWidget::AddBlockWidget( const FMultiBlock& Block, TSharedPtr<SHorizontalBox> HorizontalBox, TSharedPtr<SVerticalBox> VerticalBox, EMultiBlockLocation::Type InLocation )
+void SMultiBoxWidget::AddBlockWidget( const FMultiBlock& Block, TSharedPtr<SHorizontalBox> HorizontalBox, TSharedPtr<SVerticalBox> VerticalBox, EMultiBlockLocation::Type InLocation, bool bSectionContainsIcons )
 {
 	check( MultiBox.IsValid() );
 
 	bool bDisplayExtensionHooks = FMultiBoxSettings::DisplayMultiboxHooks.Get() && Block.GetExtensionHook() != NAME_None;
 
-	TSharedRef<SWidget> BlockWidget = Block.MakeWidget( SharedThis( this ), InLocation)->AsWidget();
+	TSharedRef<SWidget> BlockWidget = Block.MakeWidget(SharedThis(this), InLocation, bSectionContainsIcons)->AsWidget();
 
 	TWeakPtr<SWidget> BlockWidgetWeakPtr = BlockWidget;
 	TWeakPtr<const FMultiBlock> BlockWeakPtr = Block.AsShared();
@@ -596,10 +627,40 @@ void SMultiBoxWidget::BuildMultiBoxWidget()
 	bool bInsideGroup = false;
 
 	// Start building up the actual UI from each block in this MultiBox
-	bool bIsNextToStartBlock = false;
-	bool bIsNextToEndBlock = false;
+	bool bSectionContainsIcons = false;
+	int32 NextMenuSeparator = INDEX_NONE;
+
 	for( int32 Index = 0; Index < Blocks.Num(); Index++ )
 	{
+		// If we've passed the last menu separator, scan for the next one (the end of the list is also considered a menu separator for the purposes of this index)
+		if (NextMenuSeparator < Index)
+		{
+			bSectionContainsIcons = false;
+			for (++NextMenuSeparator; NextMenuSeparator < Blocks.Num(); ++NextMenuSeparator)
+			{
+				const FMultiBlock& TestBlock = *Blocks[NextMenuSeparator];
+				if (!bSectionContainsIcons && TestBlock.HasIcon())
+				{
+					bSectionContainsIcons = true;
+				}
+
+				UE_LOG(LogInit, Warning, TEXT("%d %s %s %s icon=%d type=%d"),
+					NextMenuSeparator,
+					TestBlock.GetAction().IsValid() ? *TestBlock.GetAction()->GetCommandName().ToString() : TEXT("(none)"),
+					TestBlock.IsGroupStartBlock() ? TEXT("start") : TEXT(""),
+					TestBlock.IsGroupEndBlock() ? TEXT("end") : TEXT(""),
+					TestBlock.HasIcon() ? 1 : 0,
+					(int32)TestBlock.GetType());
+					
+				if (TestBlock.GetType() == EMultiBlockType::MenuSeparator)
+				{
+					UE_LOG(LogInit, Warning, TEXT("--- %s"),
+						bSectionContainsIcons ? TEXT("icons") : TEXT("no icons"));
+					break;
+				}
+			}
+		}
+
 		const FMultiBlock& Block = *Blocks[Index];
 		EMultiBlockLocation::Type Location = EMultiBlockLocation::None;
 
@@ -616,8 +677,8 @@ void SMultiBoxWidget::BuildMultiBoxWidget()
 			}
 
 			// Check if we are next to a start or end block
-			bIsNextToStartBlock = false;
-			bIsNextToEndBlock = false;
+			bool bIsNextToStartBlock = false;
+			bool bIsNextToEndBlock = false;
 			if (Index + 1 < Blocks.Num())
 			{
 				const FMultiBlock& NextBlock = *Blocks[Index + 1];
@@ -664,13 +725,13 @@ void SMultiBoxWidget::BuildMultiBoxWidget()
 		{
 			// Add the drag preview before if we have it. This block shows where the custom block will be 
 			// added if the user drops it
-			AddBlockWidget( *DragPreview.PreviewBlock, HorizontalBox, VerticalBox, EMultiBlockLocation::None );
+			AddBlockWidget( *DragPreview.PreviewBlock, HorizontalBox, VerticalBox, EMultiBlockLocation::None, bSectionContainsIcons );
 		}
 		
 		// Do not add a block if it is being dragged
 		if( !IsBlockBeingDragged( Blocks[Index] ) )
 		{
-			AddBlockWidget( Block, HorizontalBox, VerticalBox, Location );
+			AddBlockWidget( Block, HorizontalBox, VerticalBox, Location, bSectionContainsIcons );
 		}
 	}
 
@@ -811,7 +872,7 @@ void SMultiBoxWidget::UpdateDropAreaPreviewBlock( TSharedRef<const FMultiBlock>
 					MakeShareable(
 						new FDropPreviewBlock( 
 							NewBlock.ToSharedRef(), 
-							NewBlock->MakeWidget( SharedThis(this), EMultiBlockLocation::None ) )
+							NewBlock->MakeWidget( SharedThis(this), EMultiBlockLocation::None, NewBlock->HasIcon() ) )
 					);
 
 				bAddedNewBlock = true;
diff --git a/Engine/Source/Runtime/Slate/Private/Framework/MultiBox/MultiBoxCustomization.cpp b/Engine/Source/Runtime/Slate/Private/Framework/MultiBox/MultiBoxCustomization.cpp
--- a/Engine/Source/Runtime/Slate/Private/Framework/MultiBox/MultiBoxCustomization.cpp
+++ b/Engine/Source/Runtime/Slate/Private/Framework/MultiBox/MultiBoxCustomization.cpp
@@ -79,7 +79,7 @@ FReply SMultiBlockDragHandle::OnDragDetected( const FGeometry& MyGeometry, const
 	TSharedRef<FUICommandDragDropOp> NewOp = FUICommandDragDropOp::New(
 			Block->GetAction().ToSharedRef(), 
 			MultiBoxCustomizationName, 
-			Block->MakeWidget( BaseWidget.Pin().ToSharedRef(), EMultiBlockLocation::None )->AsWidget(),
+			Block->MakeWidget( BaseWidget.Pin().ToSharedRef(), EMultiBlockLocation::None, Block->HasIcon() )->AsWidget(),
 			MyGeometry.AbsolutePosition-MouseEvent.GetScreenSpacePosition()
 		);
 			
diff --git a/Engine/Source/Runtime/Slate/Private/Framework/MultiBox/MultiBoxCustomization.h b/Engine/Source/Runtime/Slate/Private/Framework/MultiBox/MultiBoxCustomization.h
--- a/Engine/Source/Runtime/Slate/Private/Framework/MultiBox/MultiBoxCustomization.h
+++ b/Engine/Source/Runtime/Slate/Private/Framework/MultiBox/MultiBoxCustomization.h
@@ -57,7 +57,10 @@ public:
 
 	/** FMultiBlock interface */
 	virtual TSharedRef< class IMultiBlockBaseWidget > ConstructWidget() const override;
+	virtual bool HasIcon() const override { return GetActualBlock()->HasIcon(); }
+
 	TSharedRef<const FMultiBlock> GetActualBlock() const { return ActualBlock.ToSharedRef(); }
+
 private:
 	TSharedPtr<const FMultiBlock> ActualBlock;
 	TSharedPtr<IMultiBlockBaseWidget> ActualWidget;
diff --git a/Engine/Source/Runtime/Slate/Private/Framework/MultiBox/SButtonRowBlock.cpp b/Engine/Source/Runtime/Slate/Private/Framework/MultiBox/SButtonRowBlock.cpp
--- a/Engine/Source/Runtime/Slate/Private/Framework/MultiBox/SButtonRowBlock.cpp
+++ b/Engine/Source/Runtime/Slate/Private/Framework/MultiBox/SButtonRowBlock.cpp
@@ -24,6 +24,20 @@ FButtonRowBlock::FButtonRowBlock( const TAttribute<FText>& InLabel, const TAttri
 
 }
 
+bool FButtonRowBlock::HasIcon() const
+{
+	const FSlateIcon& ActualIcon = !IconOverride.IsSet() && GetAction().IsValid() ? GetAction()->GetIcon() : IconOverride;
+
+	if (ActualIcon.IsSet())
+	{
+		const FSlateBrush* IconBrush = ActualIcon.GetIcon();
+
+		return IconBrush->GetResourceName() != NAME_None;
+	}
+
+	return false;
+}
+
 /**
  * Allocates a widget for this type of MultiBlock.  Override this in derived classes.
  *
diff --git a/Engine/Source/Runtime/Slate/Private/Framework/MultiBox/SButtonRowBlock.h b/Engine/Source/Runtime/Slate/Private/Framework/MultiBox/SButtonRowBlock.h
--- a/Engine/Source/Runtime/Slate/Private/Framework/MultiBox/SButtonRowBlock.h
+++ b/Engine/Source/Runtime/Slate/Private/Framework/MultiBox/SButtonRowBlock.h
@@ -34,6 +34,9 @@ public:
 	 */
 	FButtonRowBlock( const TAttribute<FText>& InLabel, const TAttribute<FText>& InToolTip, const FSlateIcon& InIcon, const FUIAction& UIAction, const EUserInterfaceActionType::Type InUserInterfaceActionType );
 
+	/** FMultiBlock interface */
+	virtual bool HasIcon() const override;
+
 private:
 
 	/**
@@ -41,7 +44,7 @@ private:
 	 *
 	 * @return  MultiBlock widget object
 	 */
-	virtual TSharedRef< class IMultiBlockBaseWidget > ConstructWidget() const;
+	virtual TSharedRef< class IMultiBlockBaseWidget > ConstructWidget() const override;
 private:
 
 	// Friend our corresponding widget class
diff --git a/Engine/Source/Runtime/Slate/Private/Framework/MultiBox/SEditableTextBlock.cpp b/Engine/Source/Runtime/Slate/Private/Framework/MultiBox/SEditableTextBlock.cpp
--- a/Engine/Source/Runtime/Slate/Private/Framework/MultiBox/SEditableTextBlock.cpp
+++ b/Engine/Source/Runtime/Slate/Private/Framework/MultiBox/SEditableTextBlock.cpp
@@ -16,6 +16,23 @@ FEditableTextBlock::FEditableTextBlock( const FText& InLabel, const FText& InToo
 { }
 
 
+bool FEditableTextBlock::HasIcon() const
+{
+	const TSharedPtr< const FUICommandInfo >& UICommand = GetAction();
+	const FSlateIcon ActionIcon = UICommand.IsValid() ? UICommand->GetIcon() : FSlateIcon();
+
+	const FSlateIcon& ActualIcon = !IconOverride.IsSet() ? ActionIcon : IconOverride;
+
+	if (ActualIcon.IsSet())
+	{
+		const FSlateBrush* IconBrush = ActualIcon.GetIcon();
+		return (IconBrush->GetResourceName() != NAME_None);
+	}
+
+	return false;
+}
+
+
 TSharedRef<class IMultiBlockBaseWidget> FEditableTextBlock::ConstructWidget() const
 {
 	return SNew( SEditableTextBlock )
diff --git a/Engine/Source/Runtime/Slate/Private/Framework/MultiBox/SEditableTextBlock.h b/Engine/Source/Runtime/Slate/Private/Framework/MultiBox/SEditableTextBlock.h
--- a/Engine/Source/Runtime/Slate/Private/Framework/MultiBox/SEditableTextBlock.h
+++ b/Engine/Source/Runtime/Slate/Private/Framework/MultiBox/SEditableTextBlock.h
@@ -24,6 +24,9 @@ public:
 	 */
 	FEditableTextBlock( const FText& InLabel, const FText& InToolTip, const FSlateIcon& InIcon, const TAttribute< FText >& InTextAttribute, bool bInReadOnly, const FOnTextCommitted& InOnTextCommitted, const FOnTextChanged& InOnTextChanged );
 
+	/** FMultiBlock interface */
+	virtual bool HasIcon() const override;
+
 private:
 
 	/**
@@ -31,7 +34,7 @@ private:
 	 *
 	 * @return  MultiBlock widget object
 	 */
-	virtual TSharedRef<class IMultiBlockBaseWidget> ConstructWidget() const;
+	virtual TSharedRef<class IMultiBlockBaseWidget> ConstructWidget() const override;
 
 private:
 
diff --git a/Engine/Source/Runtime/Slate/Private/Framework/MultiBox/SGroupMarkerBlock.h b/Engine/Source/Runtime/Slate/Private/Framework/MultiBox/SGroupMarkerBlock.h
--- a/Engine/Source/Runtime/Slate/Private/Framework/MultiBox/SGroupMarkerBlock.h
+++ b/Engine/Source/Runtime/Slate/Private/Framework/MultiBox/SGroupMarkerBlock.h
@@ -11,15 +11,12 @@ class FGroupStartBlock
 public:
 	FGroupStartBlock();
 
-	virtual bool IsGroupStartBlock()	const	{ return true; };
+	/** FMultiBlock interface */
+	virtual bool IsGroupStartBlock() const override { return true; }
 private:
 
-	/**
-	 * Allocates a widget for this type of MultiBlock.  Override this in derived classes.
-	 *
-	 * @return  MultiBlock widget object
-	 */
-	virtual TSharedRef< class IMultiBlockBaseWidget > ConstructWidget() const;
+	/** FMultiBlock private interface */
+	virtual TSharedRef< class IMultiBlockBaseWidget > ConstructWidget() const override;
 };
 
 /**
@@ -32,15 +29,12 @@ class FGroupEndBlock
 public:
 	FGroupEndBlock();
 
-	virtual bool IsGroupEndBlock()	const	{ return true; };
+	/** FMultiBlock interface */
+	virtual bool IsGroupEndBlock() const override { return true; };
 private:
 
-	/**
-	 * Allocates a widget for this type of MultiBlock.  Override this in derived classes.
-	 *
-	 * @return  MultiBlock widget object
-	 */
-	virtual TSharedRef< class IMultiBlockBaseWidget > ConstructWidget() const;
+	/** FMultiBlock private interface */
+	virtual TSharedRef< class IMultiBlockBaseWidget > ConstructWidget() const override;
 };
 
 /**
@@ -55,9 +49,7 @@ public:
 	SLATE_BEGIN_ARGS( SGroupMarkerBlock ){}
 	SLATE_END_ARGS()
 
-	/**
-	 * Builds this MultiBlock widget up from the MultiBlock associated with it
-	 */
+	/** FMultiBlock interface */
 	virtual void BuildMultiBlockWidget(const ISlateStyle* StyleSet, const FName& StyleName) override;
 
 	/**
diff --git a/Engine/Source/Runtime/Slate/Private/Framework/MultiBox/SHeadingBlock.h b/Engine/Source/Runtime/Slate/Private/Framework/MultiBox/SHeadingBlock.h
--- a/Engine/Source/Runtime/Slate/Private/Framework/MultiBox/SHeadingBlock.h
+++ b/Engine/Source/Runtime/Slate/Private/Framework/MultiBox/SHeadingBlock.h
@@ -21,12 +21,8 @@ public:
 
 private:
 
-	/**
-	 * Allocates a widget for this type of MultiBlock.  Override this in derived classes.
-	 *
-	 * @return  MultiBlock widget object
-	 */
-	virtual TSharedRef< class IMultiBlockBaseWidget > ConstructWidget() const;
+	/** FMultiBlock private interface */
+	virtual TSharedRef< class IMultiBlockBaseWidget > ConstructWidget() const override;
 
 
 private:
diff --git a/Engine/Source/Runtime/Slate/Private/Framework/MultiBox/SMenuEntryBlock.cpp b/Engine/Source/Runtime/Slate/Private/Framework/MultiBox/SMenuEntryBlock.cpp
--- a/Engine/Source/Runtime/Slate/Private/Framework/MultiBox/SMenuEntryBlock.cpp
+++ b/Engine/Source/Runtime/Slate/Private/Framework/MultiBox/SMenuEntryBlock.cpp
@@ -134,6 +134,21 @@ void FMenuEntryBlock::CreateMenuEntry(FMenuBuilder& InMenuBuilder) const
 }
 
 
+bool FMenuEntryBlock::HasIcon() const
+{
+	const FSlateIcon ActionIcon = GetAction().IsValid() ? GetAction()->GetIcon() : FSlateIcon();
+	const FSlateIcon& ActualIcon = !IconOverride.IsSet() ? ActionIcon : IconOverride;
+
+	if (ActualIcon.IsSet())
+	{
+		const FSlateBrush* IconBrush = ActualIcon.GetIcon();
+		return IconBrush->GetResourceName() != NAME_None;
+	}
+
+	return false;
+}
+
+
 /**
  * Allocates a widget for this type of MultiBlock.  Override this in derived classes.
  *
@@ -307,6 +322,24 @@ TSharedRef< SWidget > SMenuEntryBlock::BuildMenuEntryWidget( const FMenuEntryBui
 		}
 	}
 
+	if (bSectionContainsIcons && (IconWidget == SNullWidget::NullWidget))
+	{
+		// Section should have icons but this entry does not, which is inconsistent with our menu policy (either all or none of menu items in a section should have an icon)
+		if (FMultiBoxSettings::DisplayMultiboxHooks.Get())
+		{
+			IconWidget = SNew(SColorBlock)
+				.Color(FColorList::Magenta)
+				.Size(FVector2D(MultiBoxConstants::MenuIconSize, MultiBoxConstants::MenuIconSize))
+				.ToolTipText(NSLOCTEXT("SMenuEntryBlock", "MissingIconInMenu", "This menu entry is missing an icon and should be fixed (consistency within each section is required, either every entry in the section has an icon or no entries have an icon)"));
+		}
+		else
+		{
+			// This will silently pad the offending items
+			// 	IconWidget = SNew(SSpacer)
+			// 		.Size(FVector2D(MultiBoxConstants::MenuIconSize, MultiBoxConstants::MenuIconSize));
+		}
+	}
+
 	// What type of UI should we create for this block?
 	EUserInterfaceActionType::Type UserInterfaceType = MenuEntryBlock->UserInterfaceActionType;
 	if ( UICommand.IsValid() )
diff --git a/Engine/Source/Runtime/Slate/Private/Framework/MultiBox/SMenuEntryBlock.h b/Engine/Source/Runtime/Slate/Private/Framework/MultiBox/SMenuEntryBlock.h
--- a/Engine/Source/Runtime/Slate/Private/Framework/MultiBox/SMenuEntryBlock.h
+++ b/Engine/Source/Runtime/Slate/Private/Framework/MultiBox/SMenuEntryBlock.h
@@ -88,17 +88,13 @@ public:
 
 	/** FMultiBlock interface */
 	virtual void CreateMenuEntry(class FMenuBuilder& MenuBuilder) const override;
-
+	virtual bool HasIcon() const override;
 
 
 private:
 
-	/**
-	 * Allocates a widget for this type of MultiBlock.  Override this in derived classes.
-	 *
-	 * @return  MultiBlock widget object
-	 */
-	virtual TSharedRef< class IMultiBlockBaseWidget > ConstructWidget() const;
+	/** FMultiBlock private interface */
+	virtual TSharedRef< class IMultiBlockBaseWidget > ConstructWidget() const override;
 
 
 private:
diff --git a/Engine/Source/Runtime/Slate/Private/Framework/MultiBox/SMenuSeparatorBlock.h b/Engine/Source/Runtime/Slate/Private/Framework/MultiBox/SMenuSeparatorBlock.h
--- a/Engine/Source/Runtime/Slate/Private/Framework/MultiBox/SMenuSeparatorBlock.h
+++ b/Engine/Source/Runtime/Slate/Private/Framework/MultiBox/SMenuSeparatorBlock.h
@@ -19,12 +19,8 @@ public:
 
 private:
 
-	/**
-	 * Allocates a widget for this type of MultiBlock.  Override this in derived classes.
-	 *
-	 * @return  MultiBlock widget object
-	 */
-	virtual TSharedRef< class IMultiBlockBaseWidget > ConstructWidget() const;
+	/** FMultiBlock private interface */
+	virtual TSharedRef< class IMultiBlockBaseWidget > ConstructWidget() const override;
 
 
 private:
diff --git a/Engine/Source/Runtime/Slate/Private/Framework/MultiBox/SToolBarButtonBlock.cpp b/Engine/Source/Runtime/Slate/Private/Framework/MultiBox/SToolBarButtonBlock.cpp
--- a/Engine/Source/Runtime/Slate/Private/Framework/MultiBox/SToolBarButtonBlock.cpp
+++ b/Engine/Source/Runtime/Slate/Private/Framework/MultiBox/SToolBarButtonBlock.cpp
@@ -46,6 +46,19 @@ void FToolBarButtonBlock::CreateMenuEntry(FMenuBuilder& MenuBuilder) const
 	}
 }
 
+bool FToolBarButtonBlock::HasIcon() const
+{
+	const FSlateIcon ActionIcon = GetAction().IsValid() ? GetAction()->GetIcon() : FSlateIcon();
+	const FSlateIcon& ActualIcon = IconOverride.IsSet() ? IconOverride.Get() : ActionIcon;
+
+	if (ActualIcon.IsSet())
+	{
+		return ActualIcon.GetIcon()->GetResourceName() != NAME_None;
+	}
+
+	return false;
+}
+
 /**
  * Allocates a widget for this type of MultiBlock.  Override this in derived classes.
  *
diff --git a/Engine/Source/Runtime/Slate/Private/Framework/MultiBox/SToolBarComboButtonBlock.cpp b/Engine/Source/Runtime/Slate/Private/Framework/MultiBox/SToolBarComboButtonBlock.cpp
--- a/Engine/Source/Runtime/Slate/Private/Framework/MultiBox/SToolBarComboButtonBlock.cpp
+++ b/Engine/Source/Runtime/Slate/Private/Framework/MultiBox/SToolBarComboButtonBlock.cpp
@@ -29,6 +29,12 @@ void FToolBarComboButtonBlock::CreateMenuEntry(FMenuBuilder& MenuBuilder) const
 	MenuBuilder.AddWrapperSubMenu(EntryLabel, FText::GetEmpty(), MenuContentGenerator, Icon.Get());
 }
 
+bool FToolBarComboButtonBlock::HasIcon() const
+{
+	const FSlateIcon& ActualIcon = Icon.Get();
+	return ActualIcon.GetIcon()->GetResourceName() != NAME_None;
+}
+
 TSharedRef< class IMultiBlockBaseWidget > FToolBarComboButtonBlock::ConstructWidget() const
 {
 	return SNew( SToolBarComboButtonBlock )
diff --git a/Engine/Source/Runtime/Slate/Private/Framework/MultiBox/SToolBarSeparatorBlock.h b/Engine/Source/Runtime/Slate/Private/Framework/MultiBox/SToolBarSeparatorBlock.h
--- a/Engine/Source/Runtime/Slate/Private/Framework/MultiBox/SToolBarSeparatorBlock.h
+++ b/Engine/Source/Runtime/Slate/Private/Framework/MultiBox/SToolBarSeparatorBlock.h
@@ -22,12 +22,8 @@ public:
 
 private:
 
-	/**
-	 * Allocates a widget for this type of MultiBlock.  Override this in derived classes.
-	 *
-	 * @return  MultiBlock widget object
-	 */
-	virtual TSharedRef< class IMultiBlockBaseWidget > ConstructWidget() const;
+	/** FMultiBlock private interface */
+	virtual TSharedRef< class IMultiBlockBaseWidget > ConstructWidget() const override;
 
 
 private:
diff --git a/Engine/Source/Runtime/Slate/Private/Framework/MultiBox/SWidgetBlock.h b/Engine/Source/Runtime/Slate/Private/Framework/MultiBox/SWidgetBlock.h
--- a/Engine/Source/Runtime/Slate/Private/Framework/MultiBox/SWidgetBlock.h
+++ b/Engine/Source/Runtime/Slate/Private/Framework/MultiBox/SWidgetBlock.h
@@ -26,12 +26,8 @@ public:
 
 private:
 
-	/**
-	 * Allocates a widget for this type of MultiBlock.  Override this in derived classes.
-	 *
-	 * @return  MultiBlock widget object
-	 */
-	virtual TSharedRef< class IMultiBlockBaseWidget > ConstructWidget() const;
+	/** FMultiBlock private interface */
+	virtual TSharedRef< class IMultiBlockBaseWidget > ConstructWidget() const override;
 
 
 private:
diff --git a/Engine/Source/Runtime/Slate/Public/Framework/MultiBox/MultiBox.h b/Engine/Source/Runtime/Slate/Public/Framework/MultiBox/MultiBox.h
--- a/Engine/Source/Runtime/Slate/Public/Framework/MultiBox/MultiBox.h
+++ b/Engine/Source/Runtime/Slate/Public/Framework/MultiBox/MultiBox.h
@@ -95,8 +95,9 @@ public:
 	virtual void CreateMenuEntry(class FMenuBuilder& MenuBuilder) const { }
 
 	/** Group blocks interface */
-	virtual bool IsGroupStartBlock()	const	{ return false; };
-	virtual bool IsGroupEndBlock()		const	{ return false; };
+	virtual bool IsGroupStartBlock() const { return false; }
+	virtual bool IsGroupEndBlock() const { return false; }
+	virtual bool HasIcon() const { return false; }
 
 	/** Set the tutorial highlight name for this menu entry */
 	void SetTutorialHighlightName(FName InTutorialName)
@@ -118,7 +119,7 @@ public:
 	 *
 	 * @return  MultiBlock widget object
 	 */
-	TSharedRef< class IMultiBlockBaseWidget > MakeWidget( TSharedRef< class SMultiBoxWidget > InOwnerMultiBoxWidget, EMultiBlockLocation::Type InLocation ) const;
+	TSharedRef< class IMultiBlockBaseWidget > MakeWidget( TSharedRef< class SMultiBoxWidget > InOwnerMultiBoxWidget, EMultiBlockLocation::Type InLocation, bool bSectionContainsIcons ) const;
 
 	/**
 	 * Gets the type of this MultiBox
@@ -369,8 +370,9 @@ public:
 	 * Sets the blocks location relative to the other blocks
 	 * 
 	 * @param InLocation	The MultiBlocks location
+	 * @param bSectionContainsIcons Does the section contain icons?
 	 */
-	virtual void SetMultiBlockLocation(EMultiBlockLocation::Type InLocation) = 0;
+	virtual void SetMultiBlockLocation(EMultiBlockLocation::Type InLocation, bool bSectionContainsIcons) = 0;
 
 	/**
 	 * Returns this MultiBlocks location
@@ -391,71 +393,13 @@ class SLATE_API SMultiBlockBaseWidget
 
 public:
 
-	/**
-	 * Interprets this object as a SWidget
-	 *
-	 * @return  Widget reference
-	 */
-	virtual TSharedRef< SWidget > AsWidget()
-	{
-		return this->AsShared();
-	}
-
-
-	/**
-	 * Interprets this object as a SWidget
-	 *
-	 * @return  Widget reference
-	 */
-	virtual TSharedRef< const SWidget > AsWidget() const
-	{
-		return this->AsShared();
-	}
-
-
-	/**
-	 * Associates the owner MultiBox widget with this widget
-	 *
-	 * @param	InOwnerMultiBoxWidget		The MultiBox widget that owns us
-	 */
-	void SetOwnerMultiBoxWidget( TSharedRef< SMultiBoxWidget > InOwnerMultiBoxWidget )
-	{
-		OwnerMultiBoxWidget = InOwnerMultiBoxWidget;
-	}
-
-	
-	/**
-	 * Associates this widget with a MultiBlock
-	 *
-	 * @param	InMultiBlock	The MultiBlock we'll be associated with
-	 */
-	void SetMultiBlock( TSharedRef< const FMultiBlock > InMultiBlock )
-	{
-		MultiBlock = InMultiBlock;
-	}
-
-	/**
-	 * Sets the blocks location relative to the other blocks
-	 * 
-	 * @param InLocation	The MultiBlocks location
-	 */
-	virtual void SetMultiBlockLocation(EMultiBlockLocation::Type InLocation)
-	{
-		Location = InLocation;
-	}
-
-	/**
-	 * Returns this MultiBlocks location
-	 */
-	virtual EMultiBlockLocation::Type GetMultiBlockLocation()
-	{
-		return Location;
-	}
-
-	/**
-	 * Builds this MultiBlock widget up from the MultiBlock associated with it
-	 */
-	virtual void BuildMultiBlockWidget(const ISlateStyle* StyleSet, const FName& StyleName) = 0;
+	/** IMultiBlockBaseWidget interface */
+	virtual TSharedRef< SWidget > AsWidget() override;
+	virtual TSharedRef< const SWidget > AsWidget() const override;
+	virtual void SetOwnerMultiBoxWidget( TSharedRef< SMultiBoxWidget > InOwnerMultiBoxWidget ) override;
+	virtual void SetMultiBlock( TSharedRef< const FMultiBlock > InMultiBlock ) override;
+	virtual void SetMultiBlockLocation( EMultiBlockLocation::Type InLocation, bool bInSectionContainsIcons ) override;
+	virtual EMultiBlockLocation::Type GetMultiBlockLocation() override;
 
 	/** SWidget Interface */
 	virtual void OnDragEnter( const FGeometry& MyGeometry, const FDragDropEvent& DragDropEvent ) override;
@@ -471,6 +415,9 @@ protected:
 
 	/** The MultiBlocks location relative to the other blocks in the set */
 	EMultiBlockLocation::Type Location;
+
+	/** Does the section this block resides in contain blocks with icons? */
+	bool bSectionContainsIcons;
 };
 
 /**
@@ -589,7 +536,7 @@ public:
 	virtual FReply OnKeyDown( const FGeometry& MyGeometry, const FKeyEvent& KeyEvent ) override;
 private:
 	/** Adds a block Widget to this widget */
-	void AddBlockWidget( const FMultiBlock& Block, TSharedPtr<SHorizontalBox> HorizontalBox, TSharedPtr<SVerticalBox> VerticalBox, EMultiBlockLocation::Type InLocation );
+	void AddBlockWidget( const FMultiBlock& Block, TSharedPtr<SHorizontalBox> HorizontalBox, TSharedPtr<SVerticalBox> VerticalBox, EMultiBlockLocation::Type InLocation, bool bSectionContainsIcons );
 
 	/**
 	 * @return True if the passed in block is being dragged
diff --git a/Engine/Source/Runtime/Slate/Public/Framework/MultiBox/SToolBarButtonBlock.h b/Engine/Source/Runtime/Slate/Public/Framework/MultiBox/SToolBarButtonBlock.h
--- a/Engine/Source/Runtime/Slate/Public/Framework/MultiBox/SToolBarButtonBlock.h
+++ b/Engine/Source/Runtime/Slate/Public/Framework/MultiBox/SToolBarButtonBlock.h
@@ -42,15 +42,12 @@ public:
 
 	/** FMultiBlock interface */
 	virtual void CreateMenuEntry(class FMenuBuilder& MenuBuilder) const override;
+	virtual bool HasIcon() const override;
 
 private:
 
-	/**
-	 * Allocates a widget for this type of MultiBlock.  Override this in derived classes.
-	 *
-	 * @return  MultiBlock widget object
-	 */
-	virtual TSharedRef< class IMultiBlockBaseWidget > ConstructWidget() const;
+	/** FMultiBlock private interface */
+	virtual TSharedRef< class IMultiBlockBaseWidget > ConstructWidget() const override;
 
 private:
 
diff --git a/Engine/Source/Runtime/Slate/Public/Framework/MultiBox/SToolBarComboButtonBlock.h b/Engine/Source/Runtime/Slate/Public/Framework/MultiBox/SToolBarComboButtonBlock.h
--- a/Engine/Source/Runtime/Slate/Public/Framework/MultiBox/SToolBarComboButtonBlock.h
+++ b/Engine/Source/Runtime/Slate/Public/Framework/MultiBox/SToolBarComboButtonBlock.h
@@ -27,6 +27,7 @@ public:
 
 	/** FMultiBlock interface */
 	virtual void CreateMenuEntry(class FMenuBuilder& MenuBuilder) const override;
+	virtual bool HasIcon() const override;
 
 	/** 
 	 * Sets the visibility of the blocks label
@@ -38,6 +39,7 @@ public:
 	/** Set whether this toolbar should always use small icons, regardless of the current settings */
 	void SetForceSmallIcons( const bool InForceSmallIcons ) { bForceSmallIcons = InForceSmallIcons; }
 
+
 private:
 
 	/**
@@ -45,7 +47,7 @@ private:
 	 *
 	 * @return  MultiBlock widget object
 	 */
-	virtual TSharedRef< class IMultiBlockBaseWidget > ConstructWidget() const;
+	virtual TSharedRef< class IMultiBlockBaseWidget > ConstructWidget() const override;
 
 
 private:
