Editor: Removed some multibox log spam

diff --git a/Engine/Source/Runtime/Slate/Private/Framework/MultiBox/MultiBox.cpp b/Engine/Source/Runtime/Slate/Private/Framework/MultiBox/MultiBox.cpp
--- a/Engine/Source/Runtime/Slate/Private/Framework/MultiBox/MultiBox.cpp
+++ b/Engine/Source/Runtime/Slate/Private/Framework/MultiBox/MultiBox.cpp
@@ -644,18 +644,8 @@ void SMultiBoxWidget::BuildMultiBoxWidget()
 					bSectionContainsIcons = true;
 				}
 
-				UE_LOG(LogInit, Warning, TEXT("%d %s %s %s icon=%d type=%d"),
-					NextMenuSeparator,
-					TestBlock.GetAction().IsValid() ? *TestBlock.GetAction()->GetCommandName().ToString() : TEXT("(none)"),
-					TestBlock.IsGroupStartBlock() ? TEXT("start") : TEXT(""),
-					TestBlock.IsGroupEndBlock() ? TEXT("end") : TEXT(""),
-					TestBlock.HasIcon() ? 1 : 0,
-					(int32)TestBlock.GetType());
-					
 				if (TestBlock.GetType() == EMultiBlockType::MenuSeparator)
 				{
-					UE_LOG(LogInit, Warning, TEXT("--- %s"),
-						bSectionContainsIcons ? TEXT("icons") : TEXT("no icons"));
 					break;
 				}
 			}
