Temp fix for menus in the level editor viewport not opening.

diff --git a/Engine/Source/Runtime/Slate/Private/Widgets/SViewport.cpp b/Engine/Source/Runtime/Slate/Private/Widgets/SViewport.cpp
--- a/Engine/Source/Runtime/Slate/Private/Widgets/SViewport.cpp
+++ b/Engine/Source/Runtime/Slate/Private/Widgets/SViewport.cpp
@@ -289,7 +289,7 @@ FReply SViewport::OnMotionDetected( const FGeometry& MyGeometry, const FMotionEv
 
 TOptional<EPopupMethod> SViewport::OnQueryPopupMethod() const
 {
-	return EPopupMethod::UseCurrentWindow;
+	return EPopupMethod::CreateNewWindow;
 }
 
 void SViewport::OnFinishedPointerInput()
