Focus highlighting refactored to be customizableon widget.

diff --git a/Engine/Source/Runtime/Slate/Private/Framework/Application/SlateApplication.cpp b/Engine/Source/Runtime/Slate/Private/Framework/Application/SlateApplication.cpp
--- a/Engine/Source/Runtime/Slate/Private/Framework/Application/SlateApplication.cpp
+++ b/Engine/Source/Runtime/Slate/Private/Framework/Application/SlateApplication.cpp
@@ -915,17 +915,6 @@ void FSlateApplication::DrawWindowAndChildren( const TSharedRef<SWindow>& Window
 				WindowToDraw->IsEnabled() );
 		}
 
-
-		for (int32 UserIndex = 0; UserIndex < ARRAY_COUNT(UserFocusedWidgetPaths); UserIndex++)
-		{
-			FWidgetPath ControllerFocusPath = UserFocusedWidgetPaths[UserIndex].ToWidgetPath();
-
-			if (ControllerFocusPath.IsValid() && ControllerFocusPath.GetWindow() == WindowToDraw && FocusCause == EFocusCause::Navigation)
-			{
-				MaxLayerId = DrawFocus(ControllerFocusPath, WindowElementList, MaxLayerId);
-			}
-		}
-
 		// The widget reflector may want to paint some additional stuff as part of the Widget introspection that it performs.
 		// For example: it may draw layout rectangles for hovered widgets.
 		const bool bVisualizeLayoutUnderCursor = DrawWindowArgs.WidgetsUnderCursor.IsValid();
@@ -2011,8 +2000,8 @@ bool FSlateApplication::SetUserFocus(const uint32 InUserIndex, const FWidgetPath
 
 	// Store a weak widget path to the widget that's taking focus
 	UserFocusedWidgetPaths[InUserIndex] = FWeakWidgetPath(NewFocusedWidgetPath);
-	//@Todo Slate: Refactor focus cause to be stored for all users!
-	FocusCause = InCause;
+	// Store the cause of the focus
+	UserFocusCause[InUserIndex] = InCause;
 
 	// Let the old widget know that it lost keyboard focus
 	if(OldFocusedWidget.IsValid())
@@ -2775,33 +2764,6 @@ void FSlateApplication::UpdateToolTip( bool AllowSpawningOfNewToolTips )
 	}
 }
 
-int32 FSlateApplication::DrawFocus( const FWidgetPath& FocusPath, FSlateWindowElementList& WindowElementList, int32 InLayerId ) const
-{
-	// Widgets where being focused matters draw themselves differently when focused.
-	// When the user navigates focus, we draw focus for everything, so the user can see what they are doing.
-	const FArrangedWidget& FocusedWidgetGeomertry = FocusPath.Widgets.Last();
-
-	// The FGeometry we get is from a WidgetPath, so it's rooted in desktop space.
-	// We need to APPEND a transform to the Geometry to essentially undo this root transform
-	// and get us back into Window Space.
-	// This is nonstandard so we have to go through some hoops and a specially exposed method 
-	// in FPaintGeometry to allow appending layout transforms.
-	FPaintGeometry WindowSpaceGeometry = FocusedWidgetGeomertry.Geometry.ToPaintGeometry();
-	WindowSpaceGeometry.AppendTransform(TransformCast<FSlateLayoutTransform>(Inverse(FocusPath.GetWindow()->GetPositionInScreen())));
-
-	FSlateDrawElement::MakeBox(
-		WindowElementList,
-		InLayerId++,
-		WindowSpaceGeometry,
-		FCoreStyle::Get().GetBrush("FocusRectangle"),
-		FocusPath.GetWindow()->GetClippingRectangleInWindow(),
-		ESlateDrawEffect::None,
-		FColor(255,255,255,128)
-	);
-
-	return InLayerId;
-}
-
 TArray< TSharedRef<SWindow> > FSlateApplication::GetInteractiveTopLevelWindows()
 {
 	if (ActiveModalWindows.Num() > 0)
@@ -3368,6 +3330,31 @@ bool FSlateApplication::HasMouseCapture(const TSharedPtr<const SWidget> Widget)
 	return false;
 }
 
+TOptional<EFocusCause> FSlateApplication::HasUserFocus(const TSharedPtr<const SWidget> Widget, int32 UserIndex) const
+{
+	const FWeakWidgetPath & FocusedWidgetPath = UserFocusedWidgetPaths[UserIndex];
+	if (FocusedWidgetPath.IsValid() && FocusedWidgetPath.GetLastWidget().Pin() == Widget)
+	{
+		TOptional<EFocusCause> FocusReason(UserFocusCause[UserIndex]);
+		return FocusReason;
+	}
+	return TOptional<EFocusCause>();
+}
+
+TOptional<EFocusCause> FSlateApplication::HasAnyUserFocus(const TSharedPtr<const SWidget> Widget) const
+{
+	for (int32 UserIndex = 0; UserIndex < SlateApplicationDefs::MaxUsers; ++UserIndex)
+	{
+		const FWeakWidgetPath & FocusedWidgetPath = UserFocusedWidgetPaths[UserIndex];
+		if (FocusedWidgetPath.IsValid() && FocusedWidgetPath.GetLastWidget().Pin() == Widget)
+		{
+			TOptional<EFocusCause> FocusReason(UserFocusCause[UserIndex]);
+			return FocusReason;
+		}
+	}
+	return TOptional<EFocusCause>();
+}
+
 bool FSlateApplication::HasFocusedDescendants( const TSharedRef< const SWidget >& Widget ) const
 {
 	for (int32 UserIndex = 0; UserIndex < SlateApplicationDefs::MaxUsers; ++UserIndex)
diff --git a/Engine/Source/Runtime/Slate/Public/Framework/Application/SlateApplication.h b/Engine/Source/Runtime/Slate/Public/Framework/Application/SlateApplication.h
--- a/Engine/Source/Runtime/Slate/Public/Framework/Application/SlateApplication.h
+++ b/Engine/Source/Runtime/Slate/Public/Framework/Application/SlateApplication.h
@@ -673,6 +673,10 @@ protected:
 
 	virtual bool HasMouseCapture(const TSharedPtr<const SWidget> Widget) const override;
 
+	virtual TOptional<EFocusCause> HasUserFocus(const TSharedPtr<const SWidget> Widget, int32 UserIndex) const override;
+
+	virtual TOptional<EFocusCause> HasAnyUserFocus(const TSharedPtr<const SWidget> Widget) const override;
+
 	/** 
 	 * Ticks a slate window and all of its children
 	 *
@@ -904,17 +908,6 @@ public:
 	 */
 	virtual void SynthesizeMouseMove();
 
-	/**
-	 * Draw the keyboard focus rectangle over the currently focused widget
-	 *
-	 * @param FocusPath             Path to the focused widget.
-	 * @param WindowElementList     List of draw elements to which we will add our focus rectangle
-	 * @param InLayerId             Layer onto which we should paint the rectangle
-	 *
-	 * @return New maximum layer id
-	 */
-	int32 DrawFocus( const FWidgetPath& FocusPath, class FSlateWindowElementList& WindowElementList, int32 InLayerId ) const;
-
 	/** @return an array of top-level windows that can be interacted with. e.g. when a modal window is up, only return the modal window */
 	TArray< TSharedRef<SWindow> > GetInteractiveTopLevelWindows();
 
@@ -1268,8 +1261,11 @@ private:
 
 	TSharedPtr<FAnalogCursor> AnalogCursor;
 
-	/** A weak path to the widget currently focued by a user, if any. */
+	/** A weak path to the widget currently focused by a user, if any. */
 	FWeakWidgetPath UserFocusedWidgetPaths[SlateApplicationDefs::MaxUsers];
+	
+	/** Reason a widget was focused by a user, if any. */
+	EFocusCause UserFocusCause[SlateApplicationDefs::MaxUsers];
 
 	/**
 	 * Application throttling
@@ -1305,9 +1301,6 @@ private:
 
 	/** Support for auto-dismissing popups */
 	FPopupSupport PopupSupport;
-
-	/** The way in which the last focus was set. */
-	EFocusCause FocusCause;
 	
 	/** Pointer to the currently registered game viewport widget if any */
 	TWeakPtr<SViewport> GameViewportWidget;
diff --git a/Engine/Source/Runtime/SlateCore/Private/Widgets/SWidget.cpp b/Engine/Source/Runtime/SlateCore/Private/Widgets/SWidget.cpp
--- a/Engine/Source/Runtime/SlateCore/Private/Widgets/SWidget.cpp
+++ b/Engine/Source/Runtime/SlateCore/Private/Widgets/SWidget.cpp
@@ -394,18 +394,30 @@ bool SWidget::SupportsKeyboardFocus() const
 	return false;
 }
 
-
 bool SWidget::HasKeyboardFocus() const
 {
 	return (FSlateApplicationBase::Get().GetKeyboardFocusedWidget().Get() == this);
 }
 
+TOptional<EFocusCause> SWidget::HasUserFocus(int32 UserIndex) const
+{
+	return FSlateApplicationBase::Get().HasUserFocus(SharedThis(this), UserIndex);
+}
+
+TOptional<EFocusCause> SWidget::HasAnyUserFocus() const
+{
+	return FSlateApplicationBase::Get().HasAnyUserFocus(SharedThis(this));
+}
 
 bool SWidget::HasFocusedDescendants() const
 {
 	return FSlateApplicationBase::Get().HasFocusedDescendants(SharedThis(this));
 }
 
+void SWidget::SetFocusBrush(TOptional<FSlateBrush*> InFocusBrush)
+{
+	FocusBrush = InFocusBrush;
+}
 
 bool SWidget::HasMouseCapture() const
 {
@@ -614,7 +626,31 @@ int32 SWidget::Paint(const FPaintArgs& Args, const FGeometry& AllottedGeometry,
 	}
 
 	const FPaintArgs UpdatedArgs = Args.RecordHittestGeometry( this, AllottedGeometry, MyClippingRect );
-	return OnPaint(UpdatedArgs, AllottedGeometry, MyClippingRect, OutDrawElements, LayerId, InWidgetStyle, bParentEnabled);
+	int32 NewLayerID = OnPaint(UpdatedArgs, AllottedGeometry, MyClippingRect, OutDrawElements, LayerId, InWidgetStyle, bParentEnabled);
+
+	if (SupportsKeyboardFocus())
+	{
+		TOptional<EFocusCause> FocusCause = HasAnyUserFocus();
+		if (FocusCause.IsSet() && FocusCause.GetValue() != EFocusCause::Mouse)
+		{
+			const FSlateBrush* BrushResource = FocusBrush.IsSet() ? FocusBrush.GetValue() : FCoreStyle::Get().GetBrush("FocusRectangle");
+
+			if (BrushResource != nullptr)
+			{
+				FSlateDrawElement::MakeBox(
+					OutDrawElements,
+					NewLayerID,
+					AllottedGeometry.ToPaintGeometry(),
+					BrushResource,
+					MyClippingRect,
+					ESlateDrawEffect::None,
+					FColor(255, 255, 255, 128)
+					);
+			}
+		}
+	}
+
+	return NewLayerID;
 }
 
 void SWidget::ArrangeChildren(const FGeometry& AllottedGeometry, FArrangedChildren& ArrangedChildren) const
diff --git a/Engine/Source/Runtime/SlateCore/Public/Application/SlateApplicationBase.h b/Engine/Source/Runtime/SlateCore/Public/Application/SlateApplicationBase.h
--- a/Engine/Source/Runtime/SlateCore/Public/Application/SlateApplicationBase.h
+++ b/Engine/Source/Runtime/SlateCore/Public/Application/SlateApplicationBase.h
@@ -320,12 +320,26 @@ public:
 protected:
 
 	/**
-	* Gets whether or not a particular widget has mouse capture.
-	*
-	* @return True if the widget has mouse capture, otherwise false.
-	*/
+	 * Gets whether or not a particular widget has mouse capture.
+	 *
+	 * @return True if the widget has mouse capture, otherwise false.
+	 */
 	virtual bool HasMouseCapture(const TSharedPtr<const SWidget> Widget) const = 0;
 
+	/**
+	 * Gets whether or not a particular widget has the specified users focus, and if so the type of focus.
+	 *
+	 * @return The optional will be set with the focus cause, if unset this widget doesn't have focus.
+	 */
+	virtual TOptional<EFocusCause> HasUserFocus(const TSharedPtr<const SWidget> Widget, int32 UserIndex) const = 0;
+
+	/**
+	 * Gets whether or not a particular widget has any users focus, and if so the type of focus (first one found).
+	 *
+	 * @return The optional will be set with the focus cause, if unset this widget doesn't have focus.
+	 */
+	virtual TOptional<EFocusCause> HasAnyUserFocus(const TSharedPtr<const SWidget> Widget) const = 0;
+
 protected:
 
 	// Holds the Slate renderer used to render this application.
diff --git a/Engine/Source/Runtime/SlateCore/Public/Widgets/SWidget.h b/Engine/Source/Runtime/SlateCore/Public/Widgets/SWidget.h
--- a/Engine/Source/Runtime/SlateCore/Public/Widgets/SWidget.h
+++ b/Engine/Source/Runtime/SlateCore/Public/Widgets/SWidget.h
@@ -532,11 +532,28 @@ public:
 	 */
 	virtual bool HasKeyboardFocus() const;
 
+	/**
+	 * Gets whether or not the specified users has this widget focused, and if so the type of focus.
+	 *
+	 * @return The optional will be set with the focus cause, if unset this widget doesn't have focus.
+	 */
+	TOptional<EFocusCause> HasUserFocus(int32 UserIndex) const;
+
+	/**
+	 * Gets whether or not any users users has this widget focused, and if so the type of focus (first one found).
+	 *
+	 * @return The optional will be set with the focus cause, if unset this widget doesn't have focus.
+	 */
+	TOptional<EFocusCause> HasAnyUserFocus() const;
+
 	/**
 	 * @return Whether this widget has any descendants with keyboard focus
 	 */
 	bool HasFocusedDescendants() const;
 
+	/** @param InFocusBrush Sets the brush used to render focus. */
+	void SetFocusBrush(TOptional<FSlateBrush*> InFocusBrush);
+
 	/**
 	 * Checks to see if this widget is the current mouse captor
 	 *
@@ -853,6 +870,9 @@ protected:
 	/** Render transform pivot of this widget (in normalized local space) */
 	TAttribute< FVector2D > RenderTransformPivot;
 
+	/** The brush used to draw focus, if unset it will use system default. */
+	TOptional<FSlateBrush*> FocusBrush;
+
 	/** Is this widget hovered? */
 	bool bIsHovered;
 
