Ensure that no focus hilight will display on viewports.

diff --git a/Engine/Source/Runtime/Slate/Private/Widgets/SViewport.cpp b/Engine/Source/Runtime/Slate/Private/Widgets/SViewport.cpp
--- a/Engine/Source/Runtime/Slate/Private/Widgets/SViewport.cpp
+++ b/Engine/Source/Runtime/Slate/Private/Widgets/SViewport.cpp
@@ -27,6 +27,9 @@ void SViewport::Construct( const FArguments& InArgs )
 	ViewportInterface = InArgs._ViewportInterface;
 	ViewportSize = InArgs._ViewportSize;
 
+	// Viewports should show focus
+	SetFocusBrush(nullptr);
+
 	this->ChildSlot
 	[
 		InArgs._Content.Widget
