AnchorIDs are now supported in hyperlink URLs when navigating to another point on the same page

#branch UE4

#add ScrollBox has a new function to scroll a descendant into view, if the widget should turn out to not be a descendant and ensure will fire but the code can continue safely.

#add Added new class INavigableWidget, which broadcasts events whenever a widget with an anchor id is created or a link to one is clicked on.

#add Added handling for anchor navigator, which broadcasts the the clicked event to fetch the widget to navigate to (relies on OnNavigatedDelegate handling to actually navigate to the widget, as it's parent-widget specific).

#change SBasicTileImpl, SBasicTextBlockImpl and SLayoutPanelImpl all register themselved with the INavigableWidget class incase they have an valid anchor id. SLayoutPanelImpl also handles the create and clicked events to concat a list of any children of theirs that has valid anchor id, and perform a lookup for one should a link be clicked on.

#change FMarketplaceHomeViewModelImple now has an OnAnchorClicked delegate which is fired by OnNavigatedDelegate incase anything wants to handle scrolling to the navigable widget. SMarketplaceHomeImpl registers with this delegate and informs their scrollbox to scroll to it.

#change changed FNavigableData from a struct to a class, and made members private.

diff --git a/Engine/Source/Runtime/Slate/Private/Widgets/Layout/SScrollBox.cpp b/Engine/Source/Runtime/Slate/Private/Widgets/Layout/SScrollBox.cpp
--- a/Engine/Source/Runtime/Slate/Private/Widgets/Layout/SScrollBox.cpp
+++ b/Engine/Source/Runtime/Slate/Private/Widgets/Layout/SScrollBox.cpp
@@ -361,6 +361,28 @@ void SScrollBox::ScrollToEnd()
 	bScrollToEnd = true;
 }
 
+bool SScrollBox::ScrollDescendantIntoView(const FGeometry& MyGeometry, const TSharedPtr<SWidget>& WidgetToFind, bool InAnimateScroll)
+{
+	// We need to safely find the one WidgetToFind among our descendants.
+	TSet< TSharedRef<SWidget> > WidgetsToFind;
+	{
+		WidgetsToFind.Add( WidgetToFind.ToSharedRef() );
+	}
+	TMap<TSharedRef<SWidget>, FArrangedWidget> Result;
+
+	FindChildGeometries( MyGeometry, WidgetsToFind, Result );
+
+	FArrangedWidget* WidgetGeometry = Result.Find( WidgetToFind.ToSharedRef() );
+	if ( ensureMsg( WidgetGeometry, TEXT("Unable to scroll to descendant as it's not a child of the scrollbox") ) )
+	{
+		// Calculate how much we would need to scroll to bring this to the top of the scroll box
+		const float ScrollOffset = WidgetGeometry->Geometry.AbsolutePosition.Y - MyGeometry.AbsolutePosition.Y;
+		ScrollBy(MyGeometry, ScrollOffset, InAnimateScroll);
+		return true;
+	}
+	return false;
+}
+
 EOrientation SScrollBox::GetOrientation()
 {
 	return Orientation;
diff --git a/Engine/Source/Runtime/Slate/Public/Widgets/Layout/SScrollBox.h b/Engine/Source/Runtime/Slate/Public/Widgets/Layout/SScrollBox.h
--- a/Engine/Source/Runtime/Slate/Public/Widgets/Layout/SScrollBox.h
+++ b/Engine/Source/Runtime/Slate/Public/Widgets/Layout/SScrollBox.h
@@ -101,6 +101,16 @@ public:
 
 	void ScrollToEnd();
 
+	/** 
+	 * Attempt to scroll a widget into view, will safely handle non-descendant widgets 
+	 *
+	 * @param MyGeometry   The geometry of this widget.
+	 * @param WidgetToFind The widget whose geometry we wish to discover.
+	 * @param InAnimateScroll	Whether or not to animate the scroll
+	 * @return true if descendant
+	 */
+	bool ScrollDescendantIntoView(const FGeometry& MyGeometry, const TSharedPtr<SWidget>& WidgetToFind, bool InAnimateScroll = true);
+
 	/** Get the current orientation of the scrollbox. */
 	EOrientation GetOrientation();
 
@@ -165,6 +175,7 @@ private:
 	 *
 	 * @param AllottedGeometry  The geometry allotted for this SScrollBox by the parent
 	 * @param ScrollAmount      
+	 * @param InAnimateScroll	Whether or not to animate the scroll
 	 * @return Whether or not the scroll was fully handled
 	 */
 	bool ScrollBy( const FGeometry& AllottedGeometry, float ScrollAmount, bool InAnimateScroll = true );
