Changed SLATE_TEXT_ARGUMENT to store the data as FText rather than FString

This change had been made to SLATE_TEXT_ATTRIBUTE, but missed for SLATE_TEXT_ARGUMENT

diff --git a/Engine/Source/Runtime/AppFramework/Private/Widgets/Colors/SColorThemes.cpp b/Engine/Source/Runtime/AppFramework/Private/Widgets/Colors/SColorThemes.cpp
--- a/Engine/Source/Runtime/AppFramework/Private/Widgets/Colors/SColorThemes.cpp
+++ b/Engine/Source/Runtime/AppFramework/Private/Widgets/Colors/SColorThemes.cpp
@@ -668,7 +668,7 @@ void SThemeColorBlocksBar::Construct(const FArguments& InArgs)
 
 	DestroyPlaceholders();
 
-	if (InArgs._EmptyText.Len() > 0)
+	if (!InArgs._EmptyText.IsEmpty())
 	{
 		EmptyHintTextBlock = SNew(SBorder)
 			.Padding(1.0f)
diff --git a/Engine/Source/Runtime/Slate/Public/Widgets/Input/SNumericEntryBox.h b/Engine/Source/Runtime/Slate/Public/Widgets/Input/SNumericEntryBox.h
--- a/Engine/Source/Runtime/Slate/Public/Widgets/Input/SNumericEntryBox.h
+++ b/Engine/Source/Runtime/Slate/Public/Widgets/Input/SNumericEntryBox.h
@@ -14,6 +14,7 @@ public:
 	static const FLinearColor RedLabelBackgroundColor;
 	static const FLinearColor GreenLabelBackgroundColor;
 	static const FLinearColor BlueLabelBackgroundColor;
+	static const FText DefaultUndeterminedString;
 
 	/** Notification for numeric value change */
 	DECLARE_DELEGATE_OneParam( FOnValueChanged, NumericType );
@@ -29,7 +30,7 @@ public:
 		, _LabelPadding( FMargin(3,0) )
 		, _BorderForegroundColor(FCoreStyle::Get().GetSlateColor("InvertedForeground"))
 		, _BorderBackgroundColor(FLinearColor::White)
-		, _UndeterminedString( TEXT("---") )
+		, _UndeterminedString( SNumericEntryBox<NumericType>::DefaultUndeterminedString )
 		, _Font( FCoreStyle::Get().GetFontStyle( TEXT("NormalFont") ) )
 		, _AllowSpin(false)
 		, _Delta(0)
@@ -285,7 +286,7 @@ private:
 	 */
 	FText OnGetValueForTextBox() const
 	{
-		FString NewString = TEXT("");
+		FText NewText = FText::GetEmpty();
 
 		if( EditableText->GetVisibility() == EVisibility::Visible )
 		{
@@ -295,16 +296,16 @@ private:
 			if( Value.IsSet() == true )
 			{
 				auto CurrentValue = Value.GetValue();
-				NewString = TTypeToString<NumericType>::ToSanitizedString(CurrentValue);
+				NewText = FText::FromString(TTypeToString<NumericType>::ToSanitizedString(CurrentValue));
 			}
 			else
 			{
-				NewString = UndeterminedString;
+				NewText = UndeterminedString;
 			}
 		}
 
 		// The box isnt visible, just return an empty string
-		return FText::FromString(NewString);
+		return NewText;
 	}
 
 
@@ -451,7 +452,7 @@ private:
 	/** Delegate to call when the value is committed */
 	FOnValueCommitted OnValueCommitted;
 	/** The undetermined string to display when needed */
-	FString UndeterminedString;
+	FText UndeterminedString;
 	/** Styling: border image to draw when not hovered or focused */
 	const FSlateBrush* BorderImageNormal;
 	/** Styling: border image to draw when hovered */
@@ -471,3 +472,6 @@ const FLinearColor SNumericEntryBox<NumericType>::GreenLabelBackgroundColor(0.13
 
 template <typename NumericType>
 const FLinearColor SNumericEntryBox<NumericType>::BlueLabelBackgroundColor(0.0251f,0.207f,0.85f);
+
+template <typename NumericType>
+const FText SNumericEntryBox<NumericType>::DefaultUndeterminedString = FText::FromString(TEXT("---"));
diff --git a/Engine/Source/Runtime/SlateCore/Public/Widgets/DeclarativeSyntaxSupport.h b/Engine/Source/Runtime/SlateCore/Public/Widgets/DeclarativeSyntaxSupport.h
--- a/Engine/Source/Runtime/SlateCore/Public/Widgets/DeclarativeSyntaxSupport.h
+++ b/Engine/Source/Runtime/SlateCore/Public/Widgets/DeclarativeSyntaxSupport.h
@@ -561,16 +561,16 @@
 		}
 
 #define SLATE_TEXT_ARGUMENT( ArgName ) \
-	FString _##ArgName; \
+	FText _##ArgName; \
 	WidgetArgsType& ArgName( FString InArg ) \
 	{ \
-		_##ArgName = InArg; \
+		_##ArgName = FText::FromString(InArg); \
 		return this->Me(); \
 	}\
 	\
 	WidgetArgsType& ArgName( FText InArg ) \
 	{ \
-		_##ArgName = InArg.ToString(); \
+		_##ArgName = InArg; \
 		return this->Me(); \
 	}
 
