Made SErrorText::SetError favor the FText version

diff --git a/Engine/Source/Runtime/Slate/Private/Widgets/Notifications/SErrorText.cpp b/Engine/Source/Runtime/Slate/Private/Widgets/Notifications/SErrorText.cpp
--- a/Engine/Source/Runtime/Slate/Private/Widgets/Notifications/SErrorText.cpp
+++ b/Engine/Source/Runtime/Slate/Private/Widgets/Notifications/SErrorText.cpp
@@ -29,11 +29,6 @@ void SErrorText::Construct(const FArguments& InArgs)
 }
 
 void SErrorText::SetError( const FText& InErrorText )
-{
-	SetError( InErrorText.ToString() );
-}
-
-void SErrorText::SetError( const FString& InErrorText )
 {
 	if ( TextBlock->GetText().IsEmpty() && !InErrorText.IsEmpty() )
 	{
@@ -43,6 +38,11 @@ void SErrorText::SetError( const FString& InErrorText )
 	TextBlock->SetText( InErrorText );
 }
 
+void SErrorText::SetError( const FString& InErrorText )
+{
+	SetError( FText::FromString(InErrorText) );
+}
+
 bool SErrorText::HasError() const
 {
 	return ! TextBlock->GetText().IsEmpty();
