Added the ability for viewport clients to specify the rule for showing focus.
GameViewportClient uses this to always show focus unless specifically set by the mouse.
Default behavior is to show when the focus reason is navgation.

diff --git a/Engine/Source/Runtime/Engine/Classes/Engine/GameViewportClient.h b/Engine/Source/Runtime/Engine/Classes/Engine/GameViewportClient.h
--- a/Engine/Source/Runtime/Engine/Classes/Engine/GameViewportClient.h
+++ b/Engine/Source/Runtime/Engine/Classes/Engine/GameViewportClient.h
@@ -147,6 +147,7 @@ public:
 	virtual void Precache() override;
 	virtual void Draw(FViewport* Viewport,FCanvas* SceneCanvas) override;
 	virtual void ProcessScreenShots(FViewport* Viewport) override;
+	virtual TOptional<bool> QueryShowFocus(const EFocusCause InFocusCause) const override;
 	virtual void LostFocus(FViewport* Viewport) override;
 	virtual void ReceivedFocus(FViewport* Viewport) override;
 	virtual bool IsFocused(FViewport* Viewport) override;
diff --git a/Engine/Source/Runtime/Engine/Private/GameViewportClient.cpp b/Engine/Source/Runtime/Engine/Private/GameViewportClient.cpp
--- a/Engine/Source/Runtime/Engine/Private/GameViewportClient.cpp
+++ b/Engine/Source/Runtime/Engine/Private/GameViewportClient.cpp
@@ -1242,6 +1242,15 @@ void UGameViewportClient::Precache()
 	}
 }
 
+TOptional<bool> UGameViewportClient::QueryShowFocus(const EFocusCause InFocusCause) const
+{
+	if (InFocusCause == EFocusCause::Mouse)
+	{
+		return false;
+	}
+	return true;
+}
+
 void UGameViewportClient::LostFocus(FViewport* InViewport)
 {
 	// We need to reset some key inputs, since keyup events will sometimes not be processed (such as going into immersive/maximized mode).  
diff --git a/Engine/Source/Runtime/Engine/Private/Slate/SceneViewport.cpp b/Engine/Source/Runtime/Engine/Private/Slate/SceneViewport.cpp
--- a/Engine/Source/Runtime/Engine/Private/Slate/SceneViewport.cpp
+++ b/Engine/Source/Runtime/Engine/Private/Slate/SceneViewport.cpp
@@ -675,6 +675,19 @@ TOptional<EPopupMethod> FSceneViewport::OnQueryPopupMethod() const
 	}
 }
 
+TOptional<bool> FSceneViewport::OnQueryShowFocus(const EFocusCause InFocusCause) const
+{
+	if (ViewportClient)
+	{
+		// Switch to the viewport clients world before processing input
+		FScopedConditionalWorldSwitcher WorldSwitcher(ViewportClient);
+
+		return ViewportClient->QueryShowFocus(InFocusCause);
+	}
+
+	return TOptional<bool>();
+}
+
 void FSceneViewport::OnFinishedPointerInput()
 {
 	ProcessAccumulatedPointerInput();
diff --git a/Engine/Source/Runtime/Engine/Public/Slate/SceneViewport.h b/Engine/Source/Runtime/Engine/Public/Slate/SceneViewport.h
--- a/Engine/Source/Runtime/Engine/Public/Slate/SceneViewport.h
+++ b/Engine/Source/Runtime/Engine/Public/Slate/SceneViewport.h
@@ -198,6 +198,7 @@ public:
 	virtual FReply OnTouchGesture( const FGeometry& MyGeometry, const FPointerEvent& InGestureEvent ) override;
 	virtual FReply OnMotionDetected( const FGeometry& MyGeometry, const FMotionEvent& InMotionEvent ) override;
 	virtual TOptional<EPopupMethod> OnQueryPopupMethod() const override;
+	virtual TOptional<bool> OnQueryShowFocus(const EFocusCause InFocusCause) const override;
 	virtual void OnFinishedPointerInput() override;
 	virtual FReply OnKeyDown( const FGeometry& InGeometry, const FKeyEvent& InKeyEvent ) override;
 	virtual FReply OnKeyUp( const FGeometry& InGeometry, const FKeyEvent& InKeyEvent ) override;
diff --git a/Engine/Source/Runtime/Engine/Public/UnrealClient.h b/Engine/Source/Runtime/Engine/Public/UnrealClient.h
--- a/Engine/Source/Runtime/Engine/Public/UnrealClient.h
+++ b/Engine/Source/Runtime/Engine/Public/UnrealClient.h
@@ -15,6 +15,8 @@ class FCanvas;
 class FViewportClient;
 enum class EPopupMethod : uint8;
 
+enum class EFocusCause;
+
 /**
  * A render target.
  */
@@ -784,6 +786,13 @@ public:
 	 */
 	virtual EMouseCursor::Type GetCursor(FViewport* Viewport,int32 X,int32 Y) { return EMouseCursor::Default; }
 
+	/**
+	* Called to determine if we should render the focus brush.
+	*
+	* @param InFocusCause	The cause of focus
+	*/
+	virtual TOptional<bool> QueryShowFocus(const EFocusCause InFocusCause) const { return TOptional<bool>(); }
+
 	virtual void LostFocus(FViewport* Viewport) {}
 	virtual void ReceivedFocus(FViewport* Viewport) {}
 	virtual bool IsFocused(FViewport* Viewport) { return true; }
diff --git a/Engine/Source/Runtime/Slate/Private/Framework/Application/SlateApplication.cpp b/Engine/Source/Runtime/Slate/Private/Framework/Application/SlateApplication.cpp
--- a/Engine/Source/Runtime/Slate/Private/Framework/Application/SlateApplication.cpp
+++ b/Engine/Source/Runtime/Slate/Private/Framework/Application/SlateApplication.cpp
@@ -1137,7 +1137,8 @@ void FSlateApplication::FinishedInputThisFrame()
 	// are given a chance to process accumulated values.
 	for (int32 SlateUserIndex = 0; SlateUserIndex < SlateApplicationDefs::MaxUsers; ++SlateUserIndex)
 	{
-		for (const TWeakPtr<SWidget>& WidgetPtr : UserFocusedWidgetPaths[SlateUserIndex].Widgets)
+		const FUserFocusEntry& UserFocusEntry = UserFocusEntries[SlateUserIndex];
+		for (const TWeakPtr<SWidget>& WidgetPtr : UserFocusEntry.WidgetPath.Widgets)
 		{
 			const TSharedPtr<SWidget>& Widget = WidgetPtr.Pin();
 			if (Widget.IsValid())
@@ -1846,7 +1847,8 @@ void FSlateApplication::SetUserFocus(uint32 UserIndex, const TSharedPtr<SWidget>
 
 TSharedPtr<SWidget> FSlateApplication::GetUserFocusedWidget(uint32 UserIndex) const
 {
-	return UserFocusedWidgetPaths[UserIndex].IsValid() ? UserFocusedWidgetPaths[UserIndex].GetLastWidget().Pin() : TSharedPtr<SWidget>();
+	const FUserFocusEntry& UserFocusEntry = UserFocusEntries[UserIndex];
+	return UserFocusEntry.WidgetPath.IsValid() ? UserFocusEntry.WidgetPath.GetLastWidget().Pin() : TSharedPtr<SWidget>();
 }
 
 TSharedPtr<SWidget> FSlateApplication::GetJoystickCaptor(uint32 UserIndex) const
@@ -1924,13 +1926,15 @@ bool FSlateApplication::SetKeyboardFocus(const FWidgetPath& InFocusPath, const E
 
 bool FSlateApplication::SetUserFocus(const uint32 InUserIndex, const FWidgetPath& InFocusPath, const EFocusCause InCause)
 {
+	FUserFocusEntry& UserFocusEntry = UserFocusEntries[InUserIndex];
+
 	check(InUserIndex >= 0 && InUserIndex < SlateApplicationDefs::MaxUsers);
 
 	TSharedPtr<IWidgetReflector> WidgetReflector = WidgetReflectorPtr.Pin();
 	const bool bReflectorShowingFocus = WidgetReflector.IsValid() && WidgetReflector->IsShowingFocus();
 
 	// Get the old Widget information
-	const FWeakWidgetPath OldFocusedWidgetPath = UserFocusedWidgetPaths[InUserIndex];
+	const FWeakWidgetPath OldFocusedWidgetPath = UserFocusEntry.WidgetPath;
 	TSharedPtr<SWidget> OldFocusedWidget = OldFocusedWidgetPath.IsValid() ? OldFocusedWidgetPath.GetLastWidget().Pin() : TSharedPtr< SWidget >();
 
 	// Get the new widget information by finding the first widget in the path that supports focus
@@ -1999,9 +2003,26 @@ bool FSlateApplication::SetUserFocus(const uint32 InUserIndex, const FWidgetPath
 	//UE_LOG(LogSlate, Warning, TEXT("Focus for user %i set to %s."), InUserIndex, NewFocusedWidget.IsValid() ? *NewFocusedWidget->ToString() : TEXT("Invalid"));
 
 	// Store a weak widget path to the widget that's taking focus
-	UserFocusedWidgetPaths[InUserIndex] = FWeakWidgetPath(NewFocusedWidgetPath);
+	UserFocusEntry.WidgetPath = FWeakWidgetPath(NewFocusedWidgetPath);
+
 	// Store the cause of the focus
-	UserFocusCause[InUserIndex] = InCause;
+	UserFocusEntry.FocusCause = InCause;
+
+	// Figure out if we should show focus for this focus entry
+	UserFocusEntry.ShowFocus = false;
+	if (NewFocusedWidgetPath.IsValid())
+	{
+		UserFocusEntry.ShowFocus = InCause == EFocusCause::Navigation;
+		for (int32 WidgetIndex = NewFocusedWidgetPath.Widgets.Num() - 1; WidgetIndex >= 0; --WidgetIndex)
+		{
+			TOptional<bool> ShowFocus = NewFocusedWidgetPath.Widgets[WidgetIndex].Widget->OnQueryShowFocus(InCause);
+			if (ShowFocus.IsSet())
+			{
+				UserFocusEntry.ShowFocus = ShowFocus.GetValue();
+				break;
+			}
+		}
+	}
 
 	// Let the old widget know that it lost keyboard focus
 	if(OldFocusedWidget.IsValid())
@@ -2280,7 +2301,8 @@ void FSlateApplication::ProcessReply( const FWidgetPath& CurrentEventPath, const
 		{
 			FNavigationEvent NavigationEvent(UserIndex, TheReply.GetNavigationType());
 
-			FWidgetPath EventPath = UserFocusedWidgetPaths[UserIndex].ToWidgetPath();
+			const FUserFocusEntry& UserFocusEntry = UserFocusEntries[UserIndex];
+			FWidgetPath EventPath = UserFocusEntry.WidgetPath.ToWidgetPath();
 
 			FNavigationReply NavigationReply = FNavigationReply::Escape();
 			for (int32 WidgetIndex = EventPath.Widgets.Num() - 1; WidgetIndex >= 0; --WidgetIndex)
@@ -3150,7 +3172,7 @@ void FSlateApplication::ShowVirtualKeyboard( bool bShow, TSharedPtr<IVirtualKeyb
 
 FSlateRect FSlateApplication::GetPreferredWorkArea() const
 {
-	const FWeakWidgetPath & FocusedWidgetPath = UserFocusedWidgetPaths[GetUserIndexForKeyboard()];
+	const FWeakWidgetPath & FocusedWidgetPath = UserFocusEntries[GetUserIndexForKeyboard()].WidgetPath;
 
 	// First see if we have a focused widget
 	if( FocusedWidgetPath.IsValid() && FocusedWidgetPath.Window.IsValid() )
@@ -3300,9 +3322,10 @@ EVisibility FSlateApplication::GetSoftwareCursorVis( ) const
 
 TSharedPtr< SWidget > FSlateApplication::GetKeyboardFocusedWidget() const
 {
-	if( UserFocusedWidgetPaths[GetUserIndexForKeyboard()].IsValid() )
+	const FUserFocusEntry& UserFocusEntry = UserFocusEntries[GetUserIndexForKeyboard()];
+	if (UserFocusEntry.WidgetPath.IsValid())
 	{
-		return UserFocusedWidgetPaths[GetUserIndexForKeyboard()].GetLastWidget().Pin();
+		return UserFocusEntry.WidgetPath.GetLastWidget().Pin();
 	}
 
 	return TSharedPtr< SWidget >();
@@ -3332,10 +3355,11 @@ bool FSlateApplication::HasMouseCapture(const TSharedPtr<const SWidget> Widget)
 
 TOptional<EFocusCause> FSlateApplication::HasUserFocus(const TSharedPtr<const SWidget> Widget, int32 UserIndex) const
 {
-	const FWeakWidgetPath & FocusedWidgetPath = UserFocusedWidgetPaths[UserIndex];
+	const FUserFocusEntry& UserFocusEntry = UserFocusEntries[UserIndex];
+	const FWeakWidgetPath & FocusedWidgetPath = UserFocusEntry.WidgetPath;
 	if (FocusedWidgetPath.IsValid() && FocusedWidgetPath.GetLastWidget().Pin() == Widget)
 	{
-		TOptional<EFocusCause> FocusReason(UserFocusCause[UserIndex]);
+		TOptional<EFocusCause> FocusReason(UserFocusEntry.FocusCause);
 		return FocusReason;
 	}
 	return TOptional<EFocusCause>();
@@ -3345,21 +3369,38 @@ TOptional<EFocusCause> FSlateApplication::HasAnyUserFocus(const TSharedPtr<const
 {
 	for (int32 UserIndex = 0; UserIndex < SlateApplicationDefs::MaxUsers; ++UserIndex)
 	{
-		const FWeakWidgetPath & FocusedWidgetPath = UserFocusedWidgetPaths[UserIndex];
+		const FUserFocusEntry& UserFocusEntry = UserFocusEntries[UserIndex];
+		const FWeakWidgetPath & FocusedWidgetPath = UserFocusEntry.WidgetPath;
 		if (FocusedWidgetPath.IsValid() && FocusedWidgetPath.GetLastWidget().Pin() == Widget)
 		{
-			TOptional<EFocusCause> FocusReason(UserFocusCause[UserIndex]);
+			TOptional<EFocusCause> FocusReason(UserFocusEntry.FocusCause);
 			return FocusReason;
 		}
 	}
 	return TOptional<EFocusCause>();
 }
 
+bool FSlateApplication::ShowUserFocus(const TSharedPtr<const SWidget> Widget) const
+{
+	for (int32 UserIndex = 0; UserIndex < SlateApplicationDefs::MaxUsers; ++UserIndex)
+	{
+		const FUserFocusEntry& UserFocusEntry = UserFocusEntries[UserIndex];
+		const FWeakWidgetPath & FocusedWidgetPath = UserFocusEntry.WidgetPath;
+		if (FocusedWidgetPath.IsValid() && FocusedWidgetPath.GetLastWidget().Pin() == Widget)
+		{
+			return UserFocusEntry.ShowFocus;
+		}
+	}
+
+	return false;
+}
+
 bool FSlateApplication::HasFocusedDescendants( const TSharedRef< const SWidget >& Widget ) const
 {
 	for (int32 UserIndex = 0; UserIndex < SlateApplicationDefs::MaxUsers; ++UserIndex)
 	{
-		const FWeakWidgetPath & FocusedWidgetPath = UserFocusedWidgetPaths[UserIndex];
+		const FUserFocusEntry& UserFocusEntry = UserFocusEntries[UserIndex];
+		const FWeakWidgetPath & FocusedWidgetPath = UserFocusEntry.WidgetPath;
 		if (FocusedWidgetPath.IsValid() && FocusedWidgetPath.GetLastWidget().Pin() != Widget && FocusedWidgetPath.ContainsWidget(Widget))
 		{
 			return true;
@@ -3443,7 +3484,7 @@ bool FSlateApplication::ProcessKeyCharEvent( FCharacterEvent& InCharacterEvent )
 	// NOTE: We intentionally don't reset LastUserInteractionTimeForThrottling here so that the UI can be responsive while typing
 
 	// Bubble the key event
-	FWidgetPath EventPath = UserFocusedWidgetPaths[InCharacterEvent.GetUserIndex()].ToWidgetPath();
+	FWidgetPath EventPath = UserFocusEntries[InCharacterEvent.GetUserIndex()].WidgetPath.ToWidgetPath();
 	
 	// Switch worlds for widgets in the current path
 	FScopedSwitchWorldHack SwitchWorld( EventPath );
@@ -3505,7 +3546,7 @@ bool FSlateApplication::ProcessKeyDownEvent( FKeyEvent& InKeyEvent )
 			}
 
 		// Bubble the keyboard event
-		FWidgetPath EventPath = UserFocusedWidgetPaths[InKeyEvent.GetUserIndex()].ToWidgetPath();
+		FWidgetPath EventPath = UserFocusEntries[InKeyEvent.GetUserIndex()].WidgetPath.ToWidgetPath();
 
 		// Switch worlds for widgets inOnPreviewMouseButtonDown the current path
 		FScopedSwitchWorldHack SwitchWorld( EventPath );
@@ -3568,7 +3609,7 @@ bool FSlateApplication::ProcessKeyUpEvent( FKeyEvent& InKeyEvent )
 	LastUserInteractionTimeForThrottling = LastUserInteractionTime;
 
 	// Bubble the key event
-	FWidgetPath EventPath = UserFocusedWidgetPaths[InKeyEvent.GetUserIndex()].ToWidgetPath();
+	FWidgetPath EventPath = UserFocusEntries[InKeyEvent.GetUserIndex()].WidgetPath.ToWidgetPath();
 
 	// Switch worlds for widgets in the current path
 	FScopedSwitchWorldHack SwitchWorld( EventPath );
@@ -3600,7 +3641,7 @@ bool FSlateApplication::ProcessAnalogInputEvent(FAnalogInputEvent& InAnalogInput
 	LastUserInteractionTimeForThrottling = LastUserInteractionTime;
 
 	// Bubble the key event
-	FWidgetPath EventPath = UserFocusedWidgetPaths[InAnalogInputEvent.GetUserIndex()].ToWidgetPath();
+	FWidgetPath EventPath = UserFocusEntries[InAnalogInputEvent.GetUserIndex()].WidgetPath.ToWidgetPath();
 	InAnalogInputEvent.SetEventPath(EventPath);
 
 	// Switch worlds for widgets in the current path
@@ -4440,7 +4481,7 @@ bool FSlateApplication::AttemptNavigation(const FNavigationEvent& NavigationEven
 		else
 		{
 			// Get the controller focus target for this user
-			FWeakWidgetPath FocusedWeakWidgetPath = UserFocusedWidgetPaths[NavigationEvent.GetUserIndex()];
+			FWeakWidgetPath FocusedWeakWidgetPath = UserFocusEntries[NavigationEvent.GetUserIndex()].WidgetPath;
 
 			// Find the next widget
 			if (NavigationType == EUINavigation::Next || NavigationType == EUINavigation::Previous)
@@ -4606,10 +4647,10 @@ void FSlateApplication::ProcessMotionDetectedEvent( FMotionEvent& MotionEvent )
 {
 	LastUserInteractionTime = this->GetCurrentTime();
 	
-	if (MotionEvent.GetUserIndex() < ARRAY_COUNT(UserFocusedWidgetPaths) && UserFocusedWidgetPaths[MotionEvent.GetUserIndex()].IsValid())
+	if (MotionEvent.GetUserIndex() < ARRAY_COUNT(UserFocusEntries) && UserFocusEntries[MotionEvent.GetUserIndex()].WidgetPath.IsValid())
 	{
 		/* Get the controller focus target for this user */
-		FWidgetPath PathToWidget = UserFocusedWidgetPaths[MotionEvent.GetUserIndex()].ToWidgetPath();
+		FWidgetPath PathToWidget = UserFocusEntries[MotionEvent.GetUserIndex()].WidgetPath.ToWidgetPath();
 		FArrangedWidget& ArrangedWidget = PathToWidget.Widgets.Last();
 
 		/* Switch worlds for widgets in the current path */
diff --git a/Engine/Source/Runtime/Slate/Private/Widgets/SViewport.cpp b/Engine/Source/Runtime/Slate/Private/Widgets/SViewport.cpp
--- a/Engine/Source/Runtime/Slate/Private/Widgets/SViewport.cpp
+++ b/Engine/Source/Runtime/Slate/Private/Widgets/SViewport.cpp
@@ -287,6 +287,11 @@ FReply SViewport::OnMotionDetected( const FGeometry& MyGeometry, const FMotionEv
 	return ViewportInterface.IsValid() ? ViewportInterface.Pin()->OnMotionDetected(MyGeometry, InMotionEvent) : FReply::Unhandled();
 }
 
+TOptional<bool> SViewport::OnQueryShowFocus(const EFocusCause InFocusCause) const
+{
+	return ViewportInterface.IsValid() ? ViewportInterface.Pin()->OnQueryShowFocus(InFocusCause) : TOptional<bool>();
+}
+
 TOptional<EPopupMethod> SViewport::OnQueryPopupMethod() const
 {
 	TSharedPtr<ISlateViewport> PinnedInterface = ViewportInterface.Pin();
diff --git a/Engine/Source/Runtime/Slate/Public/Framework/Application/SlateApplication.h b/Engine/Source/Runtime/Slate/Public/Framework/Application/SlateApplication.h
--- a/Engine/Source/Runtime/Slate/Public/Framework/Application/SlateApplication.h
+++ b/Engine/Source/Runtime/Slate/Public/Framework/Application/SlateApplication.h
@@ -680,6 +680,8 @@ protected:
 
 	virtual TOptional<EFocusCause> HasAnyUserFocus(const TSharedPtr<const SWidget> Widget) const override;
 
+	virtual bool ShowUserFocus(const TSharedPtr<const SWidget> Widget) const override;
+
 	/** 
 	 * Ticks a slate window and all of its children
 	 *
@@ -1264,11 +1266,19 @@ private:
 
 	TSharedPtr<FAnalogCursor> AnalogCursor;
 
-	/** A weak path to the widget currently focused by a user, if any. */
-	FWeakWidgetPath UserFocusedWidgetPaths[SlateApplicationDefs::MaxUsers];
-	
-	/** Reason a widget was focused by a user, if any. */
-	EFocusCause UserFocusCause[SlateApplicationDefs::MaxUsers];
+
+	struct FUserFocusEntry
+	{
+		/** A weak path to the widget currently focused by a user, if any. */
+		FWeakWidgetPath WidgetPath;
+		/** Reason a widget was focused by a user, if any. */
+		EFocusCause FocusCause;
+		/** If we should show this focus */
+		bool ShowFocus;
+	};
+
+	/** State of focus for all users */
+	FUserFocusEntry UserFocusEntries[SlateApplicationDefs::MaxUsers];
 
 	/**
 	 * Application throttling
diff --git a/Engine/Source/Runtime/Slate/Public/Widgets/SViewport.h b/Engine/Source/Runtime/Slate/Public/Widgets/SViewport.h
--- a/Engine/Source/Runtime/Slate/Public/Widgets/SViewport.h
+++ b/Engine/Source/Runtime/Slate/Public/Widgets/SViewport.h
@@ -151,6 +151,7 @@ public:
 	virtual FReply OnFocusReceived( const FGeometry& MyGeometry, const FFocusEvent& InFocusEvent ) override;
 	virtual void OnFocusLost( const FFocusEvent& InFocusEvent ) override;
 	virtual FReply OnMotionDetected( const FGeometry& MyGeometry, const FMotionEvent& InMotionEvent ) override;
+	virtual TOptional<bool> OnQueryShowFocus( const EFocusCause InFocusCause ) const override;
 	virtual TOptional<EPopupMethod> OnQueryPopupMethod() const override;
 	virtual void OnFinishedPointerInput() override;
 	virtual void OnArrangeChildren( const FGeometry& AllottedGeometry, FArrangedChildren& ArrangedChildren ) const override;
diff --git a/Engine/Source/Runtime/SlateCore/Private/Widgets/SWidget.cpp b/Engine/Source/Runtime/SlateCore/Private/Widgets/SWidget.cpp
--- a/Engine/Source/Runtime/SlateCore/Private/Widgets/SWidget.cpp
+++ b/Engine/Source/Runtime/SlateCore/Private/Widgets/SWidget.cpp
@@ -285,6 +285,11 @@ FReply SWidget::OnMotionDetected( const FGeometry& MyGeometry, const FMotionEven
 	return FReply::Unhandled();
 }
 
+TOptional<bool> SWidget::OnQueryShowFocus(const EFocusCause InFocusCause) const
+{
+	return TOptional<bool>();
+}
+
 TOptional<EPopupMethod> SWidget::OnQueryPopupMethod() const
 {
 	return TOptional<EPopupMethod>();
@@ -630,8 +635,8 @@ int32 SWidget::Paint(const FPaintArgs& Args, const FGeometry& AllottedGeometry,
 
 	if (SupportsKeyboardFocus())
 	{
-		TOptional<EFocusCause> FocusCause = HasAnyUserFocus();
-		if (FocusCause.IsSet() && FocusCause.GetValue() == EFocusCause::Navigation)
+		bool bShowUserFocus = FSlateApplicationBase::Get().ShowUserFocus(SharedThis(this));
+		if (bShowUserFocus)
 		{
 			const FSlateBrush* BrushResource = GetFocusBrush();
 
diff --git a/Engine/Source/Runtime/SlateCore/Public/Application/SlateApplicationBase.h b/Engine/Source/Runtime/SlateCore/Public/Application/SlateApplicationBase.h
--- a/Engine/Source/Runtime/SlateCore/Public/Application/SlateApplicationBase.h
+++ b/Engine/Source/Runtime/SlateCore/Public/Application/SlateApplicationBase.h
@@ -340,6 +340,13 @@ protected:
 	 */
 	virtual TOptional<EFocusCause> HasAnyUserFocus(const TSharedPtr<const SWidget> Widget) const = 0;
 
+	/**
+	 * Gets whether or not a particular widget should show user focus.
+	 *
+	 * @return True if we should show user focus
+	 */
+	virtual bool ShowUserFocus(const TSharedPtr<const SWidget> Widget) const = 0;
+
 protected:
 
 	// Holds the Slate renderer used to render this application.
diff --git a/Engine/Source/Runtime/SlateCore/Public/Rendering/RenderingCommon.h b/Engine/Source/Runtime/SlateCore/Public/Rendering/RenderingCommon.h
--- a/Engine/Source/Runtime/SlateCore/Public/Rendering/RenderingCommon.h
+++ b/Engine/Source/Runtime/SlateCore/Public/Rendering/RenderingCommon.h
@@ -465,6 +465,16 @@ public:
 		return FReply::Unhandled();
 	}
 
+	/**
+	 * Called to determine if we should render the focus brush.
+	 *
+	 * @param InFocusCause	The cause of focus
+	 */
+	virtual TOptional<bool> OnQueryShowFocus(const EFocusCause InFocusCause) const
+	{
+		return TOptional<bool>();
+	}
+
 	/**
 	 * Called after all input for this frame is processed.
 	 */
diff --git a/Engine/Source/Runtime/SlateCore/Public/Widgets/SWidget.h b/Engine/Source/Runtime/SlateCore/Public/Widgets/SWidget.h
--- a/Engine/Source/Runtime/SlateCore/Public/Widgets/SWidget.h
+++ b/Engine/Source/Runtime/SlateCore/Public/Widgets/SWidget.h
@@ -415,6 +415,13 @@ public:
 	 */
 	virtual FReply OnMotionDetected(const FGeometry& MyGeometry, const FMotionEvent& InMotionEvent);
 
+	/**
+	 * Called to determine if we should render the focus brush.
+	 *
+	 * @param InFocusCause	The cause of focus
+	 */
+	virtual TOptional<bool> OnQueryShowFocus(const EFocusCause InFocusCause) const;
+
 	/**
 	 * Popups can manifest in a NEW OS WINDOW or via an OVERLAY in an existing window.
 	 * This can be set explicitly on SMenuAnchor, or can be determined by a scoping widget.
