Fixed a regression that made non-wrapping text blocks calculate wrap data.

diff --git a/Engine/Source/Runtime/Slate/Private/Framework/Text/TextLayout.cpp b/Engine/Source/Runtime/Slate/Private/Framework/Text/TextLayout.cpp
--- a/Engine/Source/Runtime/Slate/Private/Framework/Text/TextLayout.cpp
+++ b/Engine/Source/Runtime/Slate/Private/Framework/Text/TextLayout.cpp
@@ -330,7 +330,7 @@ void FTextLayout::FlowLineLayout(const int32 LineModelIndex, const float Wrappin
 	const bool IsWrapping = WrappingWidth > 0.0f;
 
 	// if the Line doesn't have any BreakCandidates, or we're not wrapping text
-	if ( LineModel.BreakCandidates.Num() == 0 || !IsWrapping )
+	if (!IsWrapping || LineModel.BreakCandidates.Num() == 0 )
 	{
 		//Then iterate over all of it's runs
 		CreateLineViewBlocks( LineModelIndex, INDEX_NONE, /*OUT*/CurrentRunIndex, /*OUT*/CurrentRendererIndex, /*OUT*/PreviousBlockEnd, SoftLine );
@@ -501,6 +501,12 @@ void FTextLayout::ClearView()
 
 void FTextLayout::CreateWrappingCache()
 {
+	const bool IsWrapping = WrappingWidth > 0.0f;
+	if (!IsWrapping)
+	{
+		return;
+	}
+
 	for (FLineModel& LineModel : LineModels)
 	{
 		CreateLineWrappingCache(LineModel);
