SComboButton now respects the generic Cursor slate argument.

diff --git a/Engine/Source/Runtime/Slate/Private/Widgets/Input/SComboButton.cpp b/Engine/Source/Runtime/Slate/Private/Widgets/Input/SComboButton.cpp
--- a/Engine/Source/Runtime/Slate/Private/Widgets/Input/SComboButton.cpp
+++ b/Engine/Source/Runtime/Slate/Private/Widgets/Input/SComboButton.cpp
@@ -21,8 +21,8 @@ void SComboButton::Construct( const FArguments& InArgs )
 	TSharedPtr<SHorizontalBox> HBox;
 
 	SMenuAnchor::Construct( SMenuAnchor::FArguments()
-		.Placement( InArgs._MenuPlacement )
-		.Method( InArgs._Method )
+		.Placement(InArgs._MenuPlacement)
+		.Method(InArgs._Method)
 		[
 			SNew( SButton )
 			.ContentPadding( FMargin( 1, 0 ) )
@@ -33,6 +33,7 @@ void SComboButton::Construct( const FArguments& InArgs )
 			.ForegroundColor( InArgs._ForegroundColor )
 			.ButtonColorAndOpacity( InArgs._ButtonColorAndOpacity )
 			.IsFocusable( InArgs._IsFocusable )
+			.Cursor(InArgs._Cursor)
 			[
 				// Button and down arrow on the right
 				// +-------------------+---+
