The UserCard now properly handles setting the cursor hover.

diff --git a/Engine/Source/Runtime/Slate/Private/Widgets/Input/SComboButton.cpp b/Engine/Source/Runtime/Slate/Private/Widgets/Input/SComboButton.cpp
--- a/Engine/Source/Runtime/Slate/Private/Widgets/Input/SComboButton.cpp
+++ b/Engine/Source/Runtime/Slate/Private/Widgets/Input/SComboButton.cpp
@@ -33,7 +33,6 @@ void SComboButton::Construct( const FArguments& InArgs )
 			.ForegroundColor( InArgs._ForegroundColor )
 			.ButtonColorAndOpacity( InArgs._ButtonColorAndOpacity )
 			.IsFocusable( InArgs._IsFocusable )
-			.Cursor(InArgs._Cursor)
 			[
 				// Button and down arrow on the right
 				// +-------------------+---+
