Slate : MenuAnchors that reuse the current window now attempt to fit their popup within the window bounds.

This resolves UE-5772 because the menu anchors now use the full-screen window and do not attempt to use the work area.

diff --git a/Engine/Source/Runtime/Slate/Private/Framework/Application/SlateApplication.cpp b/Engine/Source/Runtime/Slate/Private/Framework/Application/SlateApplication.cpp
--- a/Engine/Source/Runtime/Slate/Private/Framework/Application/SlateApplication.cpp
+++ b/Engine/Source/Runtime/Slate/Private/Framework/Application/SlateApplication.cpp
@@ -3011,6 +3011,7 @@ void FSlateApplication::SetAnalogCursorEnable(bool bEnable, TSharedPtr<class FAn
 	}
 }
 
+
 FVector2D FSlateApplication::CalculatePopupWindowPosition( const FSlateRect& InAnchor, const FVector2D& InSize, const EOrientation Orientation ) const
 {
 	// Do nothing if this window has no size
@@ -3047,86 +3048,18 @@ FVector2D FSlateApplication::CalculatePopupWindowPosition( const FSlateRect& InA
 		AnchorRect.Top = InAnchor.Top + 1;
 		const FPlatformRect PlatformWorkArea = PlatformApplication->GetWorkArea( AnchorRect );
 
-		FSlateRect WorkAreaRect( 
+		const FSlateRect WorkAreaRect( 
 			PlatformWorkArea.Left, 
 			PlatformWorkArea.Top, 
 			PlatformWorkArea.Left+(PlatformWorkArea.Right - PlatformWorkArea.Left), 
 			PlatformWorkArea.Top+(PlatformWorkArea.Bottom - PlatformWorkArea.Top) );
 
-		// In the direction we are opening, see if there is enough room. If there is not, flip the opening direction along the same axis.
-		FVector2D NewPosition = FVector2D::ZeroVector;
-		if ( Orientation == Orient_Horizontal )
-		{
-			const bool bFitsRight = InAnchor.Right + InSize.X < WorkAreaRect.Right;
-			const bool bFitsLeft = InAnchor.Left - InSize.X >= WorkAreaRect.Left;
-
-			if ( bFitsRight || !bFitsLeft )
-			{
-				// The menu fits to the right of the anchor or it does not fit to the left, display to the right
-				NewPosition = FVector2D(InAnchor.Right, InAnchor.Top);
-			}
-			else
-			{
-				// The menu does not fit to the right of the anchor but it does fit to the left, display to the left
-				NewPosition = FVector2D(InAnchor.Left - InSize.X, InAnchor.Top);
-			}
-		}
-		else
-		{
-			const bool bFitsDown = InAnchor.Bottom + InSize.Y < WorkAreaRect.Bottom;
-			const bool bFitsUp = InAnchor.Top - InSize.Y >= WorkAreaRect.Top;
-
-			if ( bFitsDown || !bFitsUp )
-			{
-				// The menu fits below the anchor or it does not fit above, display below
-				NewPosition = FVector2D(InAnchor.Left, InAnchor.Bottom);
-			}
-			else
-			{
-				// The menu does not fit below the anchor but it does fit above, display above
-				NewPosition = FVector2D(InAnchor.Left, InAnchor.Top - InSize.Y);
-			}
-
-			if ( !bFitsDown && !bFitsUp )
-			{
-				NewPosition.X = InAnchor.Right;
-			}
-		}
-
-		// Adjust the position of popup windows so they do not go out of the visible area of the monitor(s)
-		// This can happen along the opposite axis that we are opening with
-		// Assumes this window has a valid size
-		// Adjust any menus that my not fit on the screen where they are opened
-		FVector2D StartPos = NewPosition;
-		FVector2D EndPos = NewPosition+InSize;
-		FVector2D Adjust = FVector2D::ZeroVector;
-		if (StartPos.X < WorkAreaRect.Left)
-		{
-			// Window is clipped by the left side of the work area
-			Adjust.X = WorkAreaRect.Left - StartPos.X;
-		}
-
-		if (StartPos.Y < WorkAreaRect.Top)
-		{
-			// Window is clipped by the top of the work area
-			Adjust.Y = WorkAreaRect.Top - StartPos.Y;
-		}
-
-		if (EndPos.X > WorkAreaRect.Right)
-		{
-			// Window is clipped by the right side of the work area
-			Adjust.X = WorkAreaRect.Right - EndPos.X;
-		}
-
-		if (EndPos.Y > WorkAreaRect.Bottom)
-		{
-			// Window is clipped by the bottom of the work area
-			Adjust.Y = WorkAreaRect.Bottom - EndPos.Y;
-		}
-
-		NewPosition += Adjust;
+		// Assume natural left-to-right, top-to-bottom flow; position popup below and to the right.
+		const FVector2D ProposedPlacement(
+			Orientation == Orient_Horizontal ? AnchorRect.Right : AnchorRect.Left,
+			Orientation == Orient_Horizontal ? AnchorRect.Top : AnchorRect.Bottom);
 
-		return NewPosition;
+		return ComputePopupFitInRect(InAnchor, FSlateRect(ProposedPlacement, ProposedPlacement+InSize), Orientation, WorkAreaRect);
 	}
 }
 
diff --git a/Engine/Source/Runtime/Slate/Private/Widgets/Input/SMenuAnchor.cpp b/Engine/Source/Runtime/Slate/Private/Widgets/Input/SMenuAnchor.cpp
--- a/Engine/Source/Runtime/Slate/Private/Widgets/Input/SMenuAnchor.cpp
+++ b/Engine/Source/Runtime/Slate/Private/Widgets/Input/SMenuAnchor.cpp
@@ -1,6 +1,77 @@
 // Copyright 1998-2014 Epic Games, Inc. All Rights Reserved.
  
 #include "SlatePrivatePCH.h"
+#include "LayoutUtils.h"
+
+
+static FVector2D GetMenuOffsetForPlacement(const FGeometry& AllottedGeometry, EMenuPlacement PlacementMode, FVector2D PopupSizeLocalSpace)
+{
+	switch (PlacementMode)
+	{
+		case MenuPlacement_BelowAnchor:
+			return FVector2D(0.0f, AllottedGeometry.GetLocalSize().Y);
+			break;
+		case MenuPlacement_CenteredBelowAnchor:
+			return FVector2D(-((PopupSizeLocalSpace.X / 2) - (AllottedGeometry.GetLocalSize().X / 2)), AllottedGeometry.GetLocalSize().Y);
+			break;
+		case MenuPlacement_ComboBox:
+			return FVector2D(0.0f, AllottedGeometry.GetLocalSize().Y);
+			break;
+		case MenuPlacement_ComboBoxRight:
+			return FVector2D(AllottedGeometry.GetLocalSize().X - PopupSizeLocalSpace.X, AllottedGeometry.GetLocalSize().Y);
+			break;
+		case MenuPlacement_MenuRight:
+			return FVector2D(AllottedGeometry.GetLocalSize().X, 0.0f);
+			break;
+		case MenuPlacement_AboveAnchor:
+			return FVector2D(0.0f, -AllottedGeometry.GetLocalSize().Y - PopupSizeLocalSpace.Y);
+			break;
+		case MenuPlacement_MenuLeft:
+			return FVector2D(-AllottedGeometry.GetLocalSize().X - PopupSizeLocalSpace.X, 0.0f);
+			break;
+		default:
+			ensureMsgf( false, TEXT("Unhandled placement mode: %d"), PlacementMode );
+			return FVector2D::ZeroVector;
+	}
+}
+
+/** Compute the popup size, offset, and anchor rect in local space. */
+struct FPopupPlacement
+{
+	FPopupPlacement(const FGeometry& AllottedGeometry, const FVector2D& PopupDesiredSize, EMenuPlacement PlacementMode)
+	{
+		// Compute the popup size, offset, and anchor rect  in local space
+		LocalPopupSize = (PlacementMode == MenuPlacement_ComboBox || PlacementMode == MenuPlacement_ComboBoxRight)
+			? FVector2D(FMath::Max(AllottedGeometry.Size.X, PopupDesiredSize.X), PopupDesiredSize.Y)
+			: PopupDesiredSize;
+		LocalPopupOffset = GetMenuOffsetForPlacement(AllottedGeometry, PlacementMode, LocalPopupSize);
+		AnchorLocalSpace = FSlateRect::FromPointAndExtent(FVector2D::ZeroVector, AllottedGeometry.Size);
+		Orientation = (PlacementMode == MenuPlacement_MenuRight || PlacementMode == MenuPlacement_MenuLeft) ? Orient_Horizontal : Orient_Vertical;
+	}
+
+	FVector2D LocalPopupSize;
+	FVector2D LocalPopupOffset;
+	FSlateRect AnchorLocalSpace;
+	EOrientation Orientation;
+};
+
+FGeometry ComputeMenuPlacement(const FGeometry& AllottedGeometry, const FVector2D& PopupDesiredSize, EMenuPlacement PlacementMode)
+{
+	// Compute the popup size, offset, and anchor rect  in local space
+	const FPopupPlacement Placement(AllottedGeometry, PopupDesiredSize, PlacementMode);
+
+	// ask the application to compute the proper desktop offset for the anchor. This requires the offsets to be in desktop space.
+	const FVector2D NewPositionDesktopSpace = FSlateApplication::Get().CalculatePopupWindowPosition(
+		TransformRect(AllottedGeometry.GetAccumulatedLayoutTransform(), Placement.AnchorLocalSpace),
+		TransformVector(AllottedGeometry.GetAccumulatedLayoutTransform(), Placement.LocalPopupSize),
+		Placement.Orientation);
+
+	// transform the desktop offset into local space and use that as the layout transform for the child content.
+	return AllottedGeometry.MakeChild(
+		Placement.LocalPopupSize,
+		FSlateLayoutTransform(TransformPoint(Inverse(AllottedGeometry.GetAccumulatedLayoutTransform()), NewPositionDesktopSpace)));
+}
+
 
 
 /**
@@ -29,14 +100,12 @@ void SMenuAnchor::Construct( const FArguments& InArgs )
 
 void SMenuAnchor::Tick( const FGeometry& AllottedGeometry, const double InCurrentTime, const float InDeltaTime )
 {
-	SPanel::Tick(AllottedGeometry, InCurrentTime, InDeltaTime);
-
 	TSharedPtr<SWindow> PopupWindow = PopupWindowPtr.Pin();
-	if ( IsOpenViaCreatedWindow() )
+	if ( PopupWindow.IsValid() && IsOpenViaCreatedWindow() )
 	{
 		// Figure out where our attached pop-up window should be placed.
 		const FVector2D PopupContentDesiredSize = PopupWindow->GetContent()->GetDesiredSize();
-		FGeometry PopupGeometry = ComputeMenuPlacement( AllottedGeometry, PopupContentDesiredSize, Placement.Get( ) );
+		FGeometry PopupGeometry = ComputeMenuPlacement( AllottedGeometry, PopupContentDesiredSize, Placement.Get() );
 		const FVector2D NewPosition = PopupGeometry.AbsolutePosition;
 		const FVector2D NewSize = PopupGeometry.GetDrawSize( );
 
@@ -82,12 +151,14 @@ void SMenuAnchor::Tick( const FGeometry& AllottedGeometry, const double InCurren
 void SMenuAnchor::OnArrangeChildren( const FGeometry& AllottedGeometry, FArrangedChildren& ArrangedChildren ) const
 {
 	ArrangeSingleChild( AllottedGeometry, ArrangedChildren, Children[0], FVector2D::UnitVector);
-	if ( IsOpen() )
+	const TSharedPtr<SWindow> PresentingWindow = PopupWindowPtr.Pin();
+	if (IsOpenAndReusingWindow() && PresentingWindow.IsValid())
 	{
-		// @todo umg na AllottedGeometry here is in Window-Space when computed via OnPaint.
-		//           The incorrect geometry causes the popup to fail workspace-edge tests.
-		const FGeometry PopupGeometry = ComputeMenuPlacement( AllottedGeometry, Children[1].GetWidget()->GetDesiredSize(), Placement.Get() );
-		ArrangedChildren.AddWidget( FArrangedWidget( Children[1].GetWidget(), PopupGeometry ) );
+		const FPopupPlacement LocalPlacement(AllottedGeometry, Children[1].GetWidget()->GetDesiredSize(), Placement.Get());
+		const FSlateRect WindowRectLocalSpace = TransformRect(Inverse(AllottedGeometry.GetAccumulatedLayoutTransform()), FSlateRect(FVector2D::ZeroVector, PresentingWindow->GetClientSizeInScreen()));
+		
+		const FVector2D FittedPlacement = ComputePopupFitInRect(LocalPlacement.AnchorLocalSpace, FSlateRect(LocalPlacement.LocalPopupOffset, LocalPlacement.LocalPopupOffset+LocalPlacement.LocalPopupSize), LocalPlacement.Orientation, WindowRectLocalSpace);
+		ArrangedChildren.AddWidget(AllottedGeometry.MakeChild(Children[1].GetWidget(), LocalPlacement.LocalPopupSize, FSlateLayoutTransform(FittedPlacement)));
 	}
 }
 
@@ -148,44 +219,6 @@ int32 SMenuAnchor::OnPaint( const FPaintArgs& Args, const FGeometry& AllottedGeo
 	return LayerId;
 }
 
-static FVector2D GetMenuOffsetForPlacement(const FGeometry& AllottedGeometry, EMenuPlacement PlacementMode, FVector2D PopupSizeLocalSpace)
-{
-	switch ( PlacementMode )
-	{
-	case MenuPlacement_AboveAnchor:
-		return FVector2D(0.0f, -AllottedGeometry.GetLocalSize().Y - PopupSizeLocalSpace.Y);
-	case MenuPlacement_ComboBoxRight:
-		return FVector2D(AllottedGeometry.GetLocalSize().X - PopupSizeLocalSpace.X, 0.0f);
-	case MenuPlacement_MenuLeft:
-		return FVector2D(-AllottedGeometry.GetLocalSize().X - PopupSizeLocalSpace.X, 0.0f);
-	case MenuPlacement_CenteredBelowAnchor:
-		return FVector2D(-((PopupSizeLocalSpace.X / 2 ) - (AllottedGeometry.GetLocalSize().X / 2)), 0.0f);
-	default:
-		return FVector2D::ZeroVector;
-	}
-}
-
-FGeometry SMenuAnchor::ComputeMenuPlacement( const FGeometry& AllottedGeometry, const FVector2D& PopupDesiredSize, EMenuPlacement PlacementMode )
-{
-	// Compute the popup size, offset, and anchor rect  in local space
-	const FVector2D PopupSizeLocalSpace = (PlacementMode == MenuPlacement_ComboBox || PlacementMode == MenuPlacement_ComboBoxRight)
-		? FVector2D( FMath::Max( AllottedGeometry.Size.X, PopupDesiredSize.X ), PopupDesiredSize.Y )
-		: PopupDesiredSize;
-	const FVector2D OffsetLocalSpace = GetMenuOffsetForPlacement(AllottedGeometry, PlacementMode, PopupSizeLocalSpace);
-	const FSlateRect AnchorLocalSpace = FSlateRect::FromPointAndExtent(OffsetLocalSpace, AllottedGeometry.Size);
-	const EOrientation Orientation = (PlacementMode == MenuPlacement_MenuRight || PlacementMode == MenuPlacement_MenuLeft) ? Orient_Horizontal : Orient_Vertical;
-	
-	// ask the application to compute the proper desktop offset for the anchor. This requires the offsets to be in desktop space.
-	const FVector2D NewPositionDesktopSpace = FSlateApplication::Get().CalculatePopupWindowPosition( 
-		TransformRect(AllottedGeometry.GetAccumulatedLayoutTransform(), AnchorLocalSpace), 
-		TransformVector(AllottedGeometry.GetAccumulatedLayoutTransform(), PopupSizeLocalSpace), 
-		Orientation );
-
-	// transform the desktop offset into local space and use that as the layout transform for the child content.
-	return AllottedGeometry.MakeChild(
-		PopupSizeLocalSpace,
-		FSlateLayoutTransform(TransformPoint(Inverse(AllottedGeometry.GetAccumulatedLayoutTransform()), NewPositionDesktopSpace)));
-}
 
 void SMenuAnchor::RequestClosePopupWindow( const TSharedRef<SWindow>& PopupWindow )
 {
@@ -210,12 +243,12 @@ void SMenuAnchor::OnClickedOutsidePopup()
 
 bool SMenuAnchor::IsOpenAndReusingWindow() const
 {
-	return Children[1].GetWidget() != SNullWidget::NullWidget;
+	return MethodInUse.IsSet() && MethodInUse.GetValue() == EPopupMethod::UseCurrentWindow;
 }
 
 bool SMenuAnchor::IsOpenViaCreatedWindow() const
 {
-	return PopupWindowPtr.Pin().IsValid() && Children[1].GetWidget() == SNullWidget::NullWidget;
+	return MethodInUse.IsSet() && MethodInUse.GetValue() == EPopupMethod::CreateNewWindow;
 }
 
 void SMenuAnchor::SetContent(TSharedRef<SWidget> InContent)
@@ -313,11 +346,11 @@ void SMenuAnchor::SetIsOpen( bool InIsOpen, const bool bFocusMenu )
 						TransitionEffect = FPopupTransitionEffect( FPopupTransitionEffect::SubMenu );
 					}
 
-					const EPopupMethod PopupMethod = Method.IsSet()
+					MethodInUse = Method.IsSet()
 						? Method.GetValue()
 						: QueryPopupMethod(MyWidgetPath);
 
-					if (PopupMethod == EPopupMethod::CreateNewWindow)
+					if (MethodInUse == EPopupMethod::CreateNewWindow)
 					{
 						// Open the pop-up
 						TSharedRef<SWindow> PopupWindow = FSlateApplication::Get().PushMenu( AsShared(), MenuContentRef, NewPosition, TransitionEffect, bFocusMenu, false, NewWindowSize, SummonLocationSize );
@@ -328,7 +361,8 @@ void SMenuAnchor::SetIsOpen( bool InIsOpen, const bool bFocusMenu )
 					{
 						// We are re-using the current window instead of creating a new one.
 						// The popup will be presented via an overlay service.
-						ensure(PopupMethod == EPopupMethod::UseCurrentWindow);
+						ensure(MethodInUse == EPopupMethod::UseCurrentWindow);
+						PopupWindowPtr = MyWidgetPath.GetWindow();
 						Children[1]
 						[
 							MenuContentRef
@@ -370,13 +404,15 @@ void SMenuAnchor::SetIsOpen( bool InIsOpen, const bool bFocusMenu )
 			{
 				OnMenuOpenChanged.Execute(false);
 			}
+
+			MethodInUse = TOptional<EPopupMethod>();
 		}
 	}
 }
 
 bool SMenuAnchor::IsOpen() const
 {
-	return PopupWindowPtr.Pin().IsValid() || IsOpenAndReusingWindow();
+	return MethodInUse.IsSet(); PopupWindowPtr.Pin().IsValid() || IsOpenAndReusingWindow();
 }
 
 bool SMenuAnchor::ShouldOpenDueToClick() const
@@ -408,7 +444,8 @@ bool SMenuAnchor::HasOpenSubMenus() const
 SMenuAnchor::SMenuAnchor()
 	: MenuContent( SNullWidget::NullWidget )
 	, bDismissedThisTick( false )
-	, Method(EPopupMethod::CreateNewWindow)
+	, Method()
+	, MethodInUse()
 {
 }
 
diff --git a/Engine/Source/Runtime/Slate/Public/Widgets/Input/SMenuAnchor.h b/Engine/Source/Runtime/Slate/Public/Widgets/Input/SMenuAnchor.h
--- a/Engine/Source/Runtime/Slate/Public/Widgets/Input/SMenuAnchor.h
+++ b/Engine/Source/Runtime/Slate/Public/Widgets/Input/SMenuAnchor.h
@@ -99,17 +99,6 @@ protected:
 	int32 OnPaint( const FPaintArgs& Args, const FGeometry& AllottedGeometry, const FSlateRect& MyClippingRect, FSlateWindowElementList& OutDrawElements, int32 LayerId, const FWidgetStyle& InWidgetStyle, bool bParentEnabled ) const;
 	// End of SWidget interface
 
-	/**
-	 * Computes the desired position for a popup that is attached to this anchor.
-	 *
-	 * @param   AllottedGeometry    Our widget's geometry
-	 * @param   PopupDesiredSize    How big the popup needs to be for all of its contents to be visible
-	 * @param   PlacementMode       Should the popup be below the anchor? Right of it? Some other placement?
-	 *
-	 * @return Geometry computed for the popup widget.
-	 */
-	static FGeometry ComputeMenuPlacement( const FGeometry& AllottedGeometry, const FVector2D& PopupDesiredSize, EMenuPlacement PlacementMode );
-
 	/** Invoked when the popup window is being closed by the application */
 	void RequestClosePopupWindow( const TSharedRef<SWindow>& PopupWindow );
 
@@ -144,5 +133,8 @@ protected:
 	/** Should we summon a new window for this popup or  */
 	TOptional<EPopupMethod> Method;
 
+	/** Method currently being used to show the popup. No value if popup is closed. */
+	TOptional<EPopupMethod> MethodInUse;
+
 	TPanelChildren<FSimpleSlot> Children;
 };
diff --git a/Engine/Source/Runtime/SlateCore/Private/Layout/LayoutUtils.cpp b/Engine/Source/Runtime/SlateCore/Private/Layout/LayoutUtils.cpp
new file mode 100644
--- /dev/null
+++ b/Engine/Source/Runtime/SlateCore/Private/Layout/LayoutUtils.cpp
@@ -0,0 +1,92 @@
+// Copyright 1998-2014 Epic Games, Inc. All Rights Reserved.
+
+#include "SlateCorePrivatePCH.h"
+
+FVector2D ComputePopupFitInRect(const FSlateRect& InAnchor, const FSlateRect& PopupRect, const EOrientation Orientation, const FSlateRect RectToFit)
+{
+	const bool bAdjustmentNeeded = PopupRect.IntersectionWith(RectToFit) != PopupRect;
+	if (bAdjustmentNeeded)
+	{
+		const FVector2D PopupPosition = PopupRect.GetTopLeft();
+		const FVector2D& PopupSize = PopupRect.GetSize();
+
+		// In the direction we are opening, see if there is enough room. If there is not, flip the opening direction along the same axis.
+		FVector2D NewPosition = FVector2D::ZeroVector;
+		if (Orientation == Orient_Horizontal)
+		{
+			const bool bFitsRight = InAnchor.Right + PopupSize.X < RectToFit.Right;
+			const bool bFitsLeft = InAnchor.Left - PopupSize.X >= RectToFit.Left;
+
+			if (bFitsRight || !bFitsLeft)
+			{
+				// The menu fits to the right of the anchor or it does not fit to the left, display to the right
+				NewPosition = FVector2D(InAnchor.Right, InAnchor.Top);
+			}
+			else
+			{
+				// The menu does not fit to the right of the anchor but it does fit to the left, display to the left
+				NewPosition = FVector2D(InAnchor.Left - PopupSize.X, InAnchor.Top);
+			}
+		}
+		else
+		{
+			const bool bFitsDown = InAnchor.Bottom + PopupSize.Y < RectToFit.Bottom;
+			const bool bFitsUp = InAnchor.Top - PopupSize.Y >= RectToFit.Top;
+
+			if (bFitsDown || !bFitsUp)
+			{
+				// The menu fits below the anchor or it does not fit above, display below
+				NewPosition = FVector2D(InAnchor.Left, InAnchor.Bottom);
+			}
+			else
+			{
+				// The menu does not fit below the anchor but it does fit above, display above
+				NewPosition = FVector2D(InAnchor.Left, InAnchor.Top - PopupSize.Y);
+			}
+
+			if (!bFitsDown && !bFitsUp)
+			{
+				NewPosition.X = InAnchor.Right;
+			}
+		}
+
+		// Adjust the position of popup windows so they do not go out of the visible area of the monitor(s)
+		// This can happen along the opposite axis that we are opening with
+		// Assumes this window has a valid size
+		// Adjust any menus that my not fit on the screen where they are opened
+		FVector2D StartPos = NewPosition;
+		FVector2D EndPos = NewPosition + PopupSize;
+		FVector2D Adjust = FVector2D::ZeroVector;
+		if (StartPos.X < RectToFit.Left)
+		{
+			// Window is clipped by the left side of the work area
+			Adjust.X = RectToFit.Left - StartPos.X;
+		}
+
+		if (StartPos.Y < RectToFit.Top)
+		{
+			// Window is clipped by the top of the work area
+			Adjust.Y = RectToFit.Top - StartPos.Y;
+		}
+
+		if (EndPos.X > RectToFit.Right)
+		{
+			// Window is clipped by the right side of the work area
+			Adjust.X = RectToFit.Right - EndPos.X;
+		}
+
+		if (EndPos.Y > RectToFit.Bottom)
+		{
+			// Window is clipped by the bottom of the work area
+			Adjust.Y = RectToFit.Bottom - EndPos.Y;
+		}
+
+		NewPosition += Adjust;
+
+		return NewPosition;
+	}
+	else
+	{
+		return PopupRect.GetTopLeft();
+	}
+}
diff --git a/Engine/Source/Runtime/SlateCore/Public/Layout/LayoutUtils.h b/Engine/Source/Runtime/SlateCore/Public/Layout/LayoutUtils.h
--- a/Engine/Source/Runtime/SlateCore/Public/Layout/LayoutUtils.h
+++ b/Engine/Source/Runtime/SlateCore/Public/Layout/LayoutUtils.h
@@ -119,3 +119,15 @@ void ArrangeSingleChild( const FGeometry& AllottedGeometry, FArrangedChildren& A
 	}
 
 }
+
+/**
+* Given information about a popup and the space available for displaying that popup, compute best placement for it.
+*
+* @param InAnchor          Area relative to which popup is being created (e.g. the button part of a combo box)
+* @param PopupRect         Proposed placement of popup; position may require adjustment.
+* @param Orientation       Are we trying to show the popup above/below or left/right relative to the anchor?
+* @param RectToFit         The space available for showing this popup; we want to fit entirely within it without clipping.
+*
+* @return A best position within the RectToFit such that none of the popup clips outside of the RectToFit.
+*/
+SLATECORE_API FVector2D ComputePopupFitInRect(const FSlateRect& InAnchor, const FSlateRect& PopupRect, const EOrientation Orientation, const FSlateRect RectToFit);
\ No newline at end of file
