Slate - Fixing menu anchors so that they correctly open when clicked.

diff --git a/Engine/Source/Runtime/Slate/Private/Widgets/Input/SMenuAnchor.cpp b/Engine/Source/Runtime/Slate/Private/Widgets/Input/SMenuAnchor.cpp
--- a/Engine/Source/Runtime/Slate/Private/Widgets/Input/SMenuAnchor.cpp
+++ b/Engine/Source/Runtime/Slate/Private/Widgets/Input/SMenuAnchor.cpp
@@ -412,7 +412,7 @@ void SMenuAnchor::SetIsOpen( bool InIsOpen, const bool bFocusMenu )
 
 bool SMenuAnchor::IsOpen() const
 {
-	return MethodInUse.IsSet(); PopupWindowPtr.Pin().IsValid() || IsOpenAndReusingWindow();
+	return (MethodInUse.IsSet() && PopupWindowPtr.Pin().IsValid()) || IsOpenAndReusingWindow();
 }
 
 bool SMenuAnchor::ShouldOpenDueToClick() const
