Fixed scrollbox to work with horizonal scrolling to descendant too, and also make sure existing offsets are cleared first.

#nonjira

#branch UE4

diff --git a/Engine/Source/Runtime/Slate/Private/Widgets/Layout/SScrollBox.cpp b/Engine/Source/Runtime/Slate/Private/Widgets/Layout/SScrollBox.cpp
--- a/Engine/Source/Runtime/Slate/Private/Widgets/Layout/SScrollBox.cpp
+++ b/Engine/Source/Runtime/Slate/Private/Widgets/Layout/SScrollBox.cpp
@@ -375,8 +375,13 @@ bool SScrollBox::ScrollDescendantIntoView(const FGeometry& MyGeometry, const TSh
 	FArrangedWidget* WidgetGeometry = Result.Find( WidgetToFind.ToSharedRef() );
 	if ( ensureMsg( WidgetGeometry, TEXT("Unable to scroll to descendant as it's not a child of the scrollbox") ) )
 	{
-		// Calculate how much we would need to scroll to bring this to the top of the scroll box
-		const float ScrollOffset = WidgetGeometry->Geometry.AbsolutePosition.Y - MyGeometry.AbsolutePosition.Y;
+		// Clear any existing scroll offset, since we're working with absolute positions
+		SetScrollOffset(0);
+
+		// Calculate how much we would need to scroll to bring this to the top/left of the scroll box
+		const float WidgetPosition = GetScrollComponentFromVector(WidgetGeometry->Geometry.AbsolutePosition);
+		const float MyPosition = GetScrollComponentFromVector(MyGeometry.AbsolutePosition);
+		const float ScrollOffset = WidgetPosition - MyPosition;
 		ScrollBy(MyGeometry, ScrollOffset, InAnimateScroll);
 		return true;
 	}
