Fix keybindings with control saving as command

diff --git a/Engine/Source/Runtime/Slate/Public/Framework/Commands/InputGesture.h b/Engine/Source/Runtime/Slate/Public/Framework/Commands/InputGesture.h
--- a/Engine/Source/Runtime/Slate/Public/Framework/Commands/InputGesture.h
+++ b/Engine/Source/Runtime/Slate/Public/Framework/Commands/InputGesture.h
@@ -15,11 +15,11 @@ struct SLATE_API FInputGesture
 	EModifierKey::Type ModifierKeys;
 
 #if PLATFORM_MAC
-	bool NeedsControl() const { return (ModifierKeys & EModifierKey::Control) != 0; }
-	bool NeedsCommand() const { return (ModifierKeys & EModifierKey::Command) != 0; }
-#else
 	bool NeedsControl() const { return (ModifierKeys & EModifierKey::Command) != 0; }
 	bool NeedsCommand() const { return (ModifierKeys & EModifierKey::Control) != 0; }
+#else
+	bool NeedsControl() const { return (ModifierKeys & EModifierKey::Control) != 0; }
+	bool NeedsCommand() const { return (ModifierKeys & EModifierKey::Command) != 0; }
 #endif
 	bool NeedsAlt() const { return (ModifierKeys & EModifierKey::Alt) != 0; }
 	bool NeedsShift() const { return (ModifierKeys & EModifierKey::Shift) != 0; }
