Fixed Ctrl+Enter potentially clearing text blocks

UE-5949 - Pressing ctrl+enter in string default variable text field results in text being deleted

Ctrl+Enter was being processed as a confirmation of the text block content, which, depending on the settings of the text block, would sometimes cause all of the text within it to be highlighted. Once this had been done, FTextEditHelper::OnKeyChar was processing the '\\n' character that the Ctrl+Enter combination had caused - this caused any selected to text to be removed, and had the effect of clearing the text block.

FTextEditHelper::OnKeyChar no longer allows '\\n' characters to be processed, as those are handled elsewhere by the editable text boxes that support multiple lines.

diff --git a/Engine/Source/Runtime/Slate/Private/Framework/Text/TextEditHelper.cpp b/Engine/Source/Runtime/Slate/Private/Framework/Text/TextEditHelper.cpp
--- a/Engine/Source/Runtime/Slate/Private/Framework/Text/TextEditHelper.cpp
+++ b/Engine/Source/Runtime/Slate/Private/Framework/Text/TextEditHelper.cpp
@@ -55,6 +55,14 @@ FReply FTextEditHelper::OnKeyChar( const FCharacterEvent& InCharacterEvent, cons
 		break;
 
 
+		// Newline (Ctrl+Enter - handled correctly outside this function)
+	case TCHAR( '\n' ):
+		{
+			Reply = FReply::Handled();
+		}
+		break;
+
+
 		// Swallow OnKeyChar keys that we don't want to be entered into the buffer
 	case 1:		// Swallow Ctrl+A, we handle that through OnKeyDown
 	case 3:		// Swallow Ctrl+C, we handle that through OnKeyDown
