Slate - after a drop ends, we now properly reset the widgets under cursor to cause mouse enter to be refired, since we fire mouse leave when a drag begins.  Adding ExtendBy to SlateRect.

diff --git a/Engine/Source/Runtime/Slate/Private/Framework/Application/SlateApplication.cpp b/Engine/Source/Runtime/Slate/Private/Framework/Application/SlateApplication.cpp
--- a/Engine/Source/Runtime/Slate/Private/Framework/Application/SlateApplication.cpp
+++ b/Engine/Source/Runtime/Slate/Private/Framework/Application/SlateApplication.cpp
@@ -2226,7 +2226,7 @@ void FSlateApplication::ProcessReply( const FWidgetPath& CurrentEventPath, const
 
 	if (TheReply.ShouldEndDragDrop())
 	{
-		EndDragDrop();
+		CancelDragDrop();
 	}
 
 	if ( bStartingDragDrop )
@@ -2841,7 +2841,7 @@ TSharedPtr<FDragDropOperation> FSlateApplication::GetDragDroppingContent() const
 	return DragDropContent;
 }
 
-void FSlateApplication::EndDragDrop()
+void FSlateApplication::CancelDragDrop()
 {
 	FWidgetPath WidgetsToDragLeave = WidgetsUnderCursorLastEvent.ToWidgetPath(FWeakWidgetPath::EInterruptedPathHandling::Truncate);
 	if (WidgetsToDragLeave.IsValid())
@@ -2852,6 +2852,8 @@ void FSlateApplication::EndDragDrop()
 			WidgetsToDragLeave.Widgets[WidgetIndex].Widget->OnDragLeave(DragDropEvent);
 		}
 	}
+
+	WidgetsUnderCursorLastEvent = FWeakWidgetPath();
 	DragDropContent.Reset();
 }
 
@@ -3469,7 +3471,7 @@ bool FSlateApplication::ProcessKeyDownEvent( FKeyEvent& InKeyEvent )
 	if (IsDragDropping() && InKeyEvent.GetKey() == EKeys::Escape)
 	{
 		// Pressing ESC while drag and dropping terminates the drag drop.
-		EndDragDrop();
+		CancelDragDrop();
 		Reply = FReply::Handled();
 	}
 	else
@@ -3972,6 +3974,7 @@ bool FSlateApplication::ProcessMouseButtonUpEvent( FPointerEvent& MouseEvent )
 			// @todo slate : depending on SetEventPath() is not ideal.
 			MouseEvent.SetEventPath( WidgetsUnderCursor );
 			LocalDragDropContent->OnDrop( Reply.IsEventHandled(), MouseEvent );
+			WidgetsUnderCursorLastEvent = FWeakWidgetPath();
 		}
 	}
 
diff --git a/Engine/Source/Runtime/Slate/Public/Framework/Application/SlateApplication.h b/Engine/Source/Runtime/Slate/Public/Framework/Application/SlateApplication.h
--- a/Engine/Source/Runtime/Slate/Public/Framework/Application/SlateApplication.h
+++ b/Engine/Source/Runtime/Slate/Public/Framework/Application/SlateApplication.h
@@ -543,7 +543,7 @@ public:
 	TSharedPtr<class FDragDropOperation> GetDragDroppingContent() const;
 
 	/** End any in flight drag and drops */
-	void EndDragDrop();
+	void CancelDragDrop();
 
 	/**
 	 * Returns the attribute that can be used by widgets to check if the application is in normal execution mode
diff --git a/Engine/Source/Runtime/SlateCore/Private/Layout/SlateRect.cpp b/Engine/Source/Runtime/SlateCore/Private/Layout/SlateRect.cpp
--- a/Engine/Source/Runtime/SlateCore/Private/Layout/SlateRect.cpp
+++ b/Engine/Source/Runtime/SlateCore/Private/Layout/SlateRect.cpp
@@ -11,6 +11,11 @@ FSlateRect FSlateRect::InsetBy( const FMargin& InsetAmount ) const
 	return FSlateRect(Left + InsetAmount.Left, Top + InsetAmount.Top, Right - InsetAmount.Right, Bottom - InsetAmount.Bottom);
 }
 
+FSlateRect FSlateRect::ExtendBy(const FMargin& InsetAmount) const
+{
+	return FSlateRect(Left - InsetAmount.Left, Top - InsetAmount.Top, Right + InsetAmount.Right, Bottom + InsetAmount.Bottom);
+}
+
 FSlateRect FSlateRect::OffsetBy(const FVector2D& OffsetAmount) const
 {
 	return FSlateRect(GetTopLeft() + OffsetAmount, GetBottomRight() + OffsetAmount);
diff --git a/Engine/Source/Runtime/SlateCore/Public/Layout/SlateRect.h b/Engine/Source/Runtime/SlateCore/Public/Layout/SlateRect.h
--- a/Engine/Source/Runtime/SlateCore/Public/Layout/SlateRect.h
+++ b/Engine/Source/Runtime/SlateCore/Public/Layout/SlateRect.h
@@ -116,6 +116,15 @@ public:
 	 */
 	FSlateRect InsetBy( const struct FMargin& InsetAmount ) const;
 
+	/**
+	 * Return a rectangle that is extended on each side by the amount specified in each margin.
+	 *
+	 * @param ExtendAmount The amount to extend the geometry.
+	 *
+	 * @return An extended rectangle.
+	 */
+	FSlateRect ExtendBy(const FMargin& ExtendAmount) const;
+
 	/**
 	 * Return a rectangle that is offset by the amount specified .
 	 *
