UMG / Mobile - All UMG widgets when adding to the viewport now start at zorder10 instead of 0, this ensures all UMG items always appear above items added by the engine like the virtual joystick.  This prevents a common issue people encounter with UMG, having their UIs get interleaved with the virtual joystick, which absorbs all hits.  Additionally the Virtual joystick now is able to lock the finger index when clicked allowing it to work with other UI on the screen and now lose capture if your finger moves over other UI elements.

diff --git a/Engine/Source/Runtime/Slate/Private/Widgets/Input/SVirtualJoystick.cpp b/Engine/Source/Runtime/Slate/Private/Widgets/Input/SVirtualJoystick.cpp
--- a/Engine/Source/Runtime/Slate/Private/Widgets/Input/SVirtualJoystick.cpp
+++ b/Engine/Source/Runtime/Slate/Private/Widgets/Input/SVirtualJoystick.cpp
@@ -181,7 +181,7 @@ FReply SVirtualJoystick::OnTouchStarted(const FGeometry& MyGeometry, const FPoin
 
 					if (HandleTouch(ControlIndex, LocalCoord, MyGeometry.Size)) // Never fail!
 					{
-						return FReply::Handled();
+						return FReply::Handled().CaptureMouse(SharedThis(this));
 					}
 				}
 				else
@@ -248,7 +248,8 @@ FReply SVirtualJoystick::OnTouchEnded(const FGeometry& MyGeometry, const FPointe
 				Control.bNeedUpdatedCenter = false;
 				return FReply::Unhandled();
 			}
-			return FReply::Handled();
+
+			return FReply::Handled().ReleaseMouseCapture();
 		}
 	}
 
diff --git a/Engine/Source/Runtime/UMG/Private/UserWidget.cpp b/Engine/Source/Runtime/UMG/Private/UserWidget.cpp
--- a/Engine/Source/Runtime/UMG/Private/UserWidget.cpp
+++ b/Engine/Source/Runtime/UMG/Private/UserWidget.cpp
@@ -620,7 +620,7 @@ void UUserWidget::AddToViewport()
 		{
 			if ( UGameViewportClient* ViewportClient = World->GetGameViewport() )
 			{
-				ViewportClient->AddViewportWidgetContent(RootWidget);
+				ViewportClient->AddViewportWidgetContent(RootWidget, 10);
 			}
 		}
 	}
