Added customizable software cursors.
Added ability to control globally when the focus brush renders.

diff --git a/Engine/Source/Runtime/Engine/Classes/Engine/GameViewportClient.h b/Engine/Source/Runtime/Engine/Classes/Engine/GameViewportClient.h
--- a/Engine/Source/Runtime/Engine/Classes/Engine/GameViewportClient.h
+++ b/Engine/Source/Runtime/Engine/Classes/Engine/GameViewportClient.h
@@ -144,6 +144,7 @@ public:
 	virtual bool InputTouch(FViewport* Viewport, int32 ControllerId, uint32 Handle, ETouchType::Type Type, const FVector2D& TouchLocation, FDateTime DeviceTimestamp, uint32 TouchpadIndex) override;
 	virtual bool InputMotion(FViewport* Viewport, int32 ControllerId, const FVector& Tilt, const FVector& RotationRate, const FVector& Gravity, const FVector& Acceleration) override;
 	virtual EMouseCursor::Type GetCursor(FViewport* Viewport, int32 X, int32 Y ) override;
+	virtual TOptional<TSharedRef<SWidget>> MapCursor(FViewport* Viewport, const FCursorReply& CursorReply) override;
 	virtual void Precache() override;
 	virtual void Draw(FViewport* Viewport,FCanvas* SceneCanvas) override;
 	virtual void ProcessScreenShots(FViewport* Viewport) override;
@@ -650,6 +651,9 @@ private:
 	/** Weak pointer to the highres screenshot dialog if it's open */
 	TWeakPtr<SWindow> HighResScreenshotDialog;
 
+	/** Map of Cursor Widgets*/
+	TMap<EMouseCursor::Type, TSharedRef<SWidget>> CursorWidgets;
+
 	/* Function that handles bug screen-shot requests w/ or w/o extra HUD info (project-specific) */
 	bool RequestBugScreenShot(const TCHAR* Cmd, bool bDisplayHUDInfo);
 
diff --git a/Engine/Source/Runtime/Engine/Classes/Engine/UserInterfaceSettings.h b/Engine/Source/Runtime/Engine/Classes/Engine/UserInterfaceSettings.h
--- a/Engine/Source/Runtime/Engine/Classes/Engine/UserInterfaceSettings.h
+++ b/Engine/Source/Runtime/Engine/Classes/Engine/UserInterfaceSettings.h
@@ -4,6 +4,17 @@
 #include "Curves/CurveFloat.h"
 #include "UserInterfaceSettings.generated.h"
 
+/** When to render the Focus Brush for widgets that have user focus. Based on the EFocusCause. */
+UENUM()
+enum class ERenderFocusRule : uint8
+{
+	/** Focus Brush will always be rendered for widgets that have user focus */
+	Always,
+	/** Focus Brush will be rendered for widgets that have user focus not set based on pointer causes */
+	NonPointer,
+	/** Focus Brush will be rendered for widgets that have user focus only if the focus was set by navigation. */
+	NavigationOnly,
+};
 
 /** The Side to use when scaling the UI. */
 UENUM()
@@ -24,7 +35,6 @@ namespace EUIScalingRule
 	};
 }
 
-
 /**
  * Implements user interface related settings.
  */
@@ -36,6 +46,27 @@ class ENGINE_API UUserInterfaceSettings
 
 public:
 
+	UPROPERTY(config, EditAnywhere, Category = "Focus", meta = (ToolTip = "Rule to determine if we should render the Focus Brush for widgets that have user focus."))
+	ERenderFocusRule RenderFocusRule;
+
+	UPROPERTY(config, EditAnywhere, Category = "Cursors", meta = (MetaClass = "UserWidget", ToolTip = "Widget to use when the Default Cursor is requested."))
+	FStringClassReference DefaultCursor;
+
+	UPROPERTY(config, EditAnywhere, Category = "Cursors", meta = (MetaClass = "UserWidget", ToolTip = "Widget to use when the TextEditBeam Cursor is requested."))
+	FStringClassReference TextEditBeamCursor;
+
+	UPROPERTY(config, EditAnywhere, Category = "Cursors", meta = (MetaClass = "UserWidget", ToolTip = "Widget to use when the Crosshairs Cursor is requested."))
+	FStringClassReference CrosshairsCursor;
+
+	UPROPERTY(config, EditAnywhere, Category = "Cursors", meta = (MetaClass = "UserWidget", ToolTip = "Widget to use when the GrabHand Cursor is requested."))
+	FStringClassReference GrabHandCursor;
+	
+	UPROPERTY(config, EditAnywhere, Category = "Cursors", meta = (MetaClass = "UserWidget", ToolTip = "Widget to use when the GrabHandClosed Cursor is requested."))
+	FStringClassReference GrabHandClosedCursor;
+
+	UPROPERTY(config, EditAnywhere, Category = "Cursors", meta = (MetaClass = "UserWidget", ToolTip = "Widget to use when the SlashedCircle Cursor is requested."))
+	FStringClassReference SlashedCircleCursor;
+
 	UPROPERTY(config, EditAnywhere, Category="DPI Scaling", meta=(
 		DisplayName="DPI Scale Rule",
 		ToolTip="The rule used when trying to decide what scale to apply." ))
diff --git a/Engine/Source/Runtime/Engine/Engine.Build.cs b/Engine/Source/Runtime/Engine/Engine.Build.cs
--- a/Engine/Source/Runtime/Engine/Engine.Build.cs
+++ b/Engine/Source/Runtime/Engine/Engine.Build.cs
@@ -66,12 +66,14 @@ public class Engine : ModuleRules
 				"SlateCore",
 				"SlateReflector",
 				"VectorVM",
-				"Landscape"
+				"Landscape",
+                "UMG",
 			}
         );
 
         CircularlyReferencedDependentModules.Add("AIModule");
 		CircularlyReferencedDependentModules.Add("Landscape");
+        CircularlyReferencedDependentModules.Add("UMG");
 
 		DynamicallyLoadedModuleNames.AddRange(
 			new string[]
@@ -80,7 +82,6 @@ public class Engine : ModuleRules
 				"MovieSceneCoreTypes",
 				"HeadMountedDisplay",
 				"StreamingPauseRendering",
-                "UMG",
 			}
 		);
 
diff --git a/Engine/Source/Runtime/Engine/Private/GameViewportClient.cpp b/Engine/Source/Runtime/Engine/Private/GameViewportClient.cpp
--- a/Engine/Source/Runtime/Engine/Private/GameViewportClient.cpp
+++ b/Engine/Source/Runtime/Engine/Private/GameViewportClient.cpp
@@ -28,7 +28,9 @@
 #include "Debug/DebugDrawService.h"
 #include "Components/BrushComponent.h"
 #include "Engine/GameEngine.h"
+#include "UserWidget.h"
 #include "GameFramework/GameUserSettings.h"
+#include "Runtime/Engine/Classes/Engine/UserInterfaceSettings.h"
 
 /** This variable allows forcing full screen of the first player controller viewport, even if there are multiple controllers plugged in and no cinematic playing. */
 bool GForceFullscreen = false;
@@ -223,6 +225,24 @@ void UGameViewportClient::Init(struct FWorldContext& WorldContext, UGameInstance
 
 	// remember our game instance
 	GameInstance = OwningGameInstance;
+
+	// Create the cursor Widgets
+	UUserInterfaceSettings* UISettings = GetMutableDefault<UUserInterfaceSettings>(UUserInterfaceSettings::StaticClass());
+
+#define ADD_CURSOR(enumeration, variable) \
+	if (UISettings->variable.IsValid()) { \
+	UClass* Class = LoadObject<UClass>(NULL, *UISettings->variable.ToString()); \
+	UUserWidget * UserWidget = CreateWidget<UUserWidget>(GetWorld(), Class); \
+	CursorWidgets.Add(enumeration, UserWidget->TakeWidget()); }
+
+	ADD_CURSOR(EMouseCursor::Default, DefaultCursor);
+	ADD_CURSOR(EMouseCursor::TextEditBeam, TextEditBeamCursor);
+	ADD_CURSOR(EMouseCursor::Crosshairs, CrosshairsCursor);
+	ADD_CURSOR(EMouseCursor::GrabHand, GrabHandCursor);
+	ADD_CURSOR(EMouseCursor::GrabHandClosed, GrabHandClosedCursor);
+	ADD_CURSOR(EMouseCursor::SlashedCircle, SlashedCircleCursor);
+
+#undef ADD_CURSOR
 }
 
 UWorld* UGameViewportClient::GetWorld() const
@@ -492,7 +512,7 @@ bool UGameViewportClient::RequiresUncapturedAxisInput() const
 }
 
 
-EMouseCursor::Type UGameViewportClient::GetCursor(FViewport* InViewport, int32 X, int32 Y )
+EMouseCursor::Type UGameViewportClient::GetCursor(FViewport* InViewport, int32 X, int32 Y)
 {
 	bool bIsPlayingMovie = false;//GetMoviePlayer()->IsMovieCurrentlyPlaying();
 
@@ -536,6 +556,15 @@ EMouseCursor::Type UGameViewportClient::GetCursor(FViewport* InViewport, int32 X
 	return FViewportClient::GetCursor(InViewport, X, Y);
 }
 
+TOptional<TSharedRef<SWidget>> UGameViewportClient::MapCursor(FViewport* Viewport, const FCursorReply& CursorReply)
+{
+	const TSharedRef<SWidget>* CursorWidgetPtr = CursorWidgets.Find(CursorReply.GetCursorType());
+	if (CursorWidgetPtr != nullptr)
+	{
+		return *CursorWidgetPtr;
+	}
+	return TOptional<TSharedRef<SWidget>>();
+}
 
 void UGameViewportClient::SetDropDetail(float DeltaSeconds)
 {
@@ -1244,7 +1273,9 @@ void UGameViewportClient::Precache()
 
 TOptional<bool> UGameViewportClient::QueryShowFocus(const EFocusCause InFocusCause) const
 {
-	if (InFocusCause == EFocusCause::Mouse)
+	UUserInterfaceSettings* UISettings = GetMutableDefault<UUserInterfaceSettings>(UUserInterfaceSettings::StaticClass());
+	if ((UISettings->RenderFocusRule == ERenderFocusRule::NonPointer && InFocusCause == EFocusCause::Mouse) ||
+		(UISettings->RenderFocusRule == ERenderFocusRule::NavigationOnly && InFocusCause != EFocusCause::Navigation))
 	{
 		return false;
 	}
@@ -1259,15 +1290,15 @@ void UGameViewportClient::LostFocus(FViewport* InViewport)
 	if (World)
 	{
 		for (FConstPlayerControllerIterator Iterator = World->GetPlayerControllerIterator(); Iterator; ++Iterator)
-	{
-			APlayerController* const PlayerController = *Iterator;
-		if (PlayerController)
 		{
-			PlayerController->FlushPressedKeys();
+			APlayerController* const PlayerController = *Iterator;
+			if (PlayerController)
+			{
+				PlayerController->FlushPressedKeys();
+			}
 		}
 	}
 }
-}
 
 void UGameViewportClient::ReceivedFocus(FViewport* InViewport)
 {
diff --git a/Engine/Source/Runtime/Engine/Private/Slate/SceneViewport.cpp b/Engine/Source/Runtime/Engine/Private/Slate/SceneViewport.cpp
--- a/Engine/Source/Runtime/Engine/Private/Slate/SceneViewport.cpp
+++ b/Engine/Source/Runtime/Engine/Private/Slate/SceneViewport.cpp
@@ -335,6 +335,15 @@ FCursorReply FSceneViewport::OnCursorQuery( const FGeometry& MyGeometry, const F
 	return FCursorReply::Cursor(MouseCursorToUse);
 }
 
+TOptional<TSharedRef<SWidget>> FSceneViewport::OnMapCursor(const FCursorReply& CursorReply)
+{
+	if (ViewportClient && GetSizeXY() != FIntPoint::ZeroValue)
+	{
+		return ViewportClient->MapCursor(this, CursorReply);
+	}
+	return ISlateViewport::OnMapCursor(CursorReply);
+}
+
 FReply FSceneViewport::OnMouseButtonDown( const FGeometry& InGeometry, const FPointerEvent& InMouseEvent )
 {
 	// Start a new reply state
diff --git a/Engine/Source/Runtime/Engine/Private/UserInterfaceSettings.cpp b/Engine/Source/Runtime/Engine/Private/UserInterfaceSettings.cpp
--- a/Engine/Source/Runtime/Engine/Private/UserInterfaceSettings.cpp
+++ b/Engine/Source/Runtime/Engine/Private/UserInterfaceSettings.cpp
@@ -10,7 +10,8 @@
 
 
 UUserInterfaceSettings::UUserInterfaceSettings(const FObjectInitializer& ObjectInitializer)
-: Super(ObjectInitializer)
+	: Super(ObjectInitializer),
+	  RenderFocusRule(ERenderFocusRule::NavigationOnly)
 {
 }
 
diff --git a/Engine/Source/Runtime/Engine/Public/Slate/SceneViewport.h b/Engine/Source/Runtime/Engine/Public/Slate/SceneViewport.h
--- a/Engine/Source/Runtime/Engine/Public/Slate/SceneViewport.h
+++ b/Engine/Source/Runtime/Engine/Public/Slate/SceneViewport.h
@@ -185,6 +185,7 @@ public:
 	virtual FSlateShaderResource* GetViewportRenderTargetTexture() const override;
 	virtual void OnDrawViewport( const FGeometry& AllottedGeometry, const FSlateRect& MyClippingRect, FSlateWindowElementList& OutDrawElements, int32 LayerId, const FWidgetStyle& InWidgetStyle, bool bParentEnabled ) override;
 	virtual FCursorReply OnCursorQuery( const FGeometry& MyGeometry, const FPointerEvent& CursorEvent ) override;
+	virtual TOptional<TSharedRef<SWidget>> OnMapCursor(const FCursorReply& CursorReply) override;
 	virtual FReply OnMouseButtonDown( const FGeometry& InGeometry, const FPointerEvent& MouseEvent ) override;
 	virtual FReply OnMouseButtonUp( const FGeometry& InGeometry, const FPointerEvent& MouseEvent ) override;
 	virtual void OnMouseEnter( const FGeometry& MyGeometry, const FPointerEvent& MouseEvent ) override;
diff --git a/Engine/Source/Runtime/Engine/Public/UnrealClient.h b/Engine/Source/Runtime/Engine/Public/UnrealClient.h
--- a/Engine/Source/Runtime/Engine/Public/UnrealClient.h
+++ b/Engine/Source/Runtime/Engine/Public/UnrealClient.h
@@ -13,8 +13,11 @@
 
 class FCanvas;
 class FViewportClient;
-enum class EPopupMethod : uint8;
 
+class SWidget;
+class FCursorReply;
+
+enum class EPopupMethod : uint8;
 enum class EFocusCause : uint8;
 
 /**
@@ -790,10 +793,17 @@ public:
 	virtual EMouseCursor::Type GetCursor(FViewport* Viewport,int32 X,int32 Y) { return EMouseCursor::Default; }
 
 	/**
-	* Called to determine if we should render the focus brush.
-	*
-	* @param InFocusCause	The cause of focus
-	*/
+	 * Called to map a cursor reply to an actual widget to render.
+	 *
+	 * @return	the widget that should be rendered for this cursor, return TOptional<TSharedRef<SWidget>>() if no mapping.
+	 */
+	virtual TOptional<TSharedRef<SWidget>> MapCursor(FViewport* Viewport, const FCursorReply& CursorReply) { return TOptional<TSharedRef<SWidget>>(); }
+	
+	/**
+	 * Called to determine if we should render the focus brush.
+	 *
+	 * @param InFocusCause	The cause of focus
+	 */
 	virtual TOptional<bool> QueryShowFocus(const EFocusCause InFocusCause) const { return TOptional<bool>(); }
 
 	virtual void LostFocus(FViewport* Viewport) {}
diff --git a/Engine/Source/Runtime/Slate/Private/Framework/Application/SlateApplication.cpp b/Engine/Source/Runtime/Slate/Private/Framework/Application/SlateApplication.cpp
--- a/Engine/Source/Runtime/Slate/Private/Framework/Application/SlateApplication.cpp
+++ b/Engine/Source/Runtime/Slate/Private/Framework/Application/SlateApplication.cpp
@@ -914,6 +914,30 @@ void FSlateApplication::DrawWindowAndChildren( const TSharedRef<SWindow>& Window
 				0,
 				FWidgetStyle(),
 				WindowToDraw->IsEnabled() );
+
+			// Draw Software Cursor
+			TSharedPtr<SWindow> CursorWindow = CursorWindowPtr.Pin();
+			if (CursorWindow.IsValid() && WindowToDraw == CursorWindow)
+			{
+				TSharedPtr<SWidget> CursorWidget = CursorWidgetPtr.Pin();
+				
+				if (CursorWidget.IsValid())
+				{
+					CursorWidget->SlatePrepass();
+
+					FVector2D CursorPosInWindowSpace = WindowToDraw->GetWindowGeometryInScreen().AbsoluteToLocal(GetCursorPos());
+					CursorPosInWindowSpace += (CursorWidget->GetDesiredSize() * -0.5);
+					const FGeometry CursorGeometry = FGeometry::MakeRoot(CursorWidget->GetDesiredSize(), FSlateLayoutTransform(CursorPosInWindowSpace));
+
+					CursorWidget->Paint(
+						FPaintArgs(WindowToDraw, *WindowToDraw->GetHittestGrid(), WindowToDraw->GetPositionInScreen(), GetCurrentTime(), GetDeltaTime()),
+						CursorGeometry, WindowToDraw->GetClippingRectangleInWindow(),
+						WindowElementList,
+						++MaxLayerId,
+						FWidgetStyle(),
+						WindowToDraw->IsEnabled());
+				}
+			}
 		}
 
 		// The widget reflector may want to paint some additional stuff as part of the Widget introspection that it performs.
@@ -1929,10 +1953,10 @@ bool FSlateApplication::SetKeyboardFocus(const FWidgetPath& InFocusPath, const E
 
 bool FSlateApplication::SetUserFocus(const uint32 InUserIndex, const FWidgetPath& InFocusPath, const EFocusCause InCause)
 {
-	FUserFocusEntry& UserFocusEntry = UserFocusEntries[InUserIndex];
-
 	check(InUserIndex >= 0 && InUserIndex < SlateApplicationDefs::MaxUsers);
 
+	FUserFocusEntry& UserFocusEntry = UserFocusEntries[InUserIndex];
+
 	TSharedPtr<IWidgetReflector> WidgetReflector = WidgetReflectorPtr.Pin();
 	const bool bReflectorShowingFocus = WidgetReflector.IsValid() && WidgetReflector->IsShowingFocus();
 
@@ -2403,19 +2427,13 @@ void FSlateApplication::QueryCursor()
 	if ( PlatformApplication->Cursor.IsValid() )
 	{
 		// drag-drop overrides cursor
-		FCursorReply CursorResult = FCursorReply::Unhandled();
-
+		FCursorReply CursorReply = FCursorReply::Unhandled();
 		if(IsDragDropping())
 		{
-			CursorResult = DragDropContent->OnCursorQuery();
-			if (CursorResult.IsEventHandled())
-			{
-				// Query was handled, so we should set the cursor.
-				PlatformApplication->Cursor->SetType( CursorResult.GetCursor() );
-			}
+			CursorReply = DragDropContent->OnCursorQuery();
 		}
 		
-		if(!CursorResult.IsEventHandled())
+		if (!CursorReply.IsEventHandled())
 		{
 			FWidgetPath WidgetsToQueryForCursor;
 			const TSharedPtr<SWindow> ActiveModalWindow = GetActiveModalWindow();
@@ -2456,39 +2474,65 @@ void FSlateApplication::QueryCursor()
 				// Switch worlds for widgets in the current path
 				FScopedSwitchWorldHack SwitchWorld( WidgetsToQueryForCursor );
 
-				const FVector2D CurrentCursorPosition = GetCursorPos();
-				const FVector2D LastCursorPosition = GetLastCursorPos();
-				const FPointerEvent CursorEvent(
-					CursorPointerIndex,
-					CurrentCursorPosition,
-					LastCursorPosition,
-					CurrentCursorPosition - LastCursorPosition,
-					PressedMouseButtons,
-					PlatformApplication->GetModifierKeys()
-				);
-
-				CursorResult = FEventRouter::Route<FCursorReply>(this, FEventRouter::FBubblePolicy(WidgetsToQueryForCursor), CursorEvent, [](const FArrangedWidget& WidgetToQuery, const FPointerEvent& PointerEvent)
+				for (int32 WidgetIndex = WidgetsToQueryForCursor.Widgets.Num() - 1; WidgetIndex >= 0; --WidgetIndex)
 				{
-					return WidgetToQuery.Widget->OnCursorQuery( WidgetToQuery.Geometry, PointerEvent );
-				});
-
-				if (CursorResult.IsEventHandled())
-				{
-					// Query was handled, so we should set the cursor.
-					PlatformApplication->Cursor->SetType( CursorResult.GetCursor() );
+					const FArrangedWidget& ArrangedWidget = WidgetsToQueryForCursor.Widgets[WidgetIndex];
+					CursorReply = ArrangedWidget.Widget->OnCursorQuery(ArrangedWidget.Geometry, CursorEvent);
+					if (CursorReply.IsEventHandled())
+					{
+						if (!CursorReply.GetCursorWidget().IsValid())
+						{
+							for (; WidgetIndex >= 0; --WidgetIndex)
+							{
+								TOptional<TSharedRef<SWidget>> CursorWidget = WidgetsToQueryForCursor.Widgets[WidgetIndex].Widget->OnMapCursor(CursorReply);
+								if (CursorWidget.IsSet())
+								{
+									CursorReply.SetCursorWidget(WidgetsToQueryForCursor.GetWindow(), CursorWidget.GetValue());
+									break;
+								}
+							}
+						}
+						break;
+					}
 				}
-				else if (WidgetsToQueryForCursor.IsValid())
+
+				if (!CursorReply.IsEventHandled() && WidgetsToQueryForCursor.IsValid())
 				{
 					// Query was NOT handled, and we are still over a slate window.
-					PlatformApplication->Cursor->SetType(EMouseCursor::Default);
+					CursorReply = FCursorReply::Cursor(EMouseCursor::Default);
 				}
 			}
 			else
 			{
 				// Set the default cursor when there isn't an active window under the cursor and the mouse isn't captured
-				PlatformApplication->Cursor->SetType(EMouseCursor::Default);
+				CursorReply = FCursorReply::Cursor(EMouseCursor::Default);
 			}
 		}
+		ProcessCursorReply(CursorReply);
+	}
+}
+
+void FSlateApplication::ProcessCursorReply(const FCursorReply& CursorReply)
+{
+	if (CursorReply.IsEventHandled())
+	{
+		CursorWidgetPtr = CursorReply.GetCursorWidget();
+		if (CursorReply.GetCursorWidget().IsValid())
+		{
+			CursorReply.GetCursorWidget()->SetVisibility(EVisibility::HitTestInvisible);
+			CursorWindowPtr = CursorReply.GetCursorWindow();
+			PlatformApplication->Cursor->SetType(EMouseCursor::None);
+		}
+		else
+		{
+			CursorWindowPtr.Reset();
+			PlatformApplication->Cursor->SetType(CursorReply.GetCursorType());
+		}
+	}
+	else
+	{
+		CursorWindowPtr.Reset();
+		CursorWidgetPtr.Reset();
 	}
 }
 
@@ -4321,7 +4365,7 @@ bool FSlateApplication::ProcessMouseMoveEvent( FPointerEvent& MouseEvent, bool b
 		if (CursorResult.IsEventHandled())
 		{
 			// Query was handled, so we should set the cursor.
-			PlatformApplication->Cursor->SetType( CursorResult.GetCursor() );
+			PlatformApplication->Cursor->SetType(CursorResult.GetCursorType());
 		}
 		else
 		{
diff --git a/Engine/Source/Runtime/Slate/Private/Widgets/SViewport.cpp b/Engine/Source/Runtime/Slate/Private/Widgets/SViewport.cpp
--- a/Engine/Source/Runtime/Slate/Private/Widgets/SViewport.cpp
+++ b/Engine/Source/Runtime/Slate/Private/Widgets/SViewport.cpp
@@ -94,7 +94,7 @@ int32 SViewport::OnPaint( const FPaintArgs& Args, const FGeometry& AllottedGeome
 				TSet<FKey>(),
 				FModifierKeysState() )
 		 );
-		EMouseCursor::Type CursorType = Reply.GetCursor();
+		EMouseCursor::Type CursorType = Reply.GetCursorType();
 
 		const FSlateBrush* Brush = FCoreStyle::Get().GetBrush(TEXT("SoftwareCursor_Grab"));
 		if( CursorType == EMouseCursor::CardinalCross )
@@ -140,6 +140,10 @@ FCursorReply SViewport::OnCursorQuery( const FGeometry& MyGeometry, const FPoint
 	return ViewportInterface.IsValid() ? ViewportInterface.Pin()->OnCursorQuery(MyGeometry, CursorEvent) : FCursorReply::Unhandled();
 }
 
+TOptional<TSharedRef<SWidget>> SViewport::OnMapCursor(const FCursorReply& CursorReply) const
+{
+	return ViewportInterface.IsValid() ? ViewportInterface.Pin()->OnMapCursor(CursorReply) : TOptional<TSharedRef<SWidget>>();
+}
 
 FReply SViewport::OnMouseButtonDown( const FGeometry& MyGeometry, const FPointerEvent& MouseEvent )
 {
diff --git a/Engine/Source/Runtime/Slate/Public/Framework/Application/SlateApplication.h b/Engine/Source/Runtime/Slate/Public/Framework/Application/SlateApplication.h
--- a/Engine/Source/Runtime/Slate/Public/Framework/Application/SlateApplication.h
+++ b/Engine/Source/Runtime/Slate/Public/Framework/Application/SlateApplication.h
@@ -872,7 +872,7 @@ public:
 	void DestroyWindowsImmediately();
 
 	/**
-	 * Apply any requests form the Reply to the application. E.g. Capture mouse
+	 * Apply any requests from the Reply to the application. E.g. Capture mouse
 	 *
 	 * @param CurrentEventPath   The WidgetPath along which the reply-generated event was routed
 	 * @param TheReply           The reply generated by an event that was being processed.
@@ -888,6 +888,13 @@ public:
 	 * Bubble a request for which cursor to display for widgets under the mouse or the widget that captured the mouse.
 	 */
 	void QueryCursor();
+
+	/**
+	 * Apply any requests from the CursorReply
+	 *
+	 * @param CursorReply        The reply generated by an event that was being processed.
+	 */
+	void ProcessCursorReply(const FCursorReply& CursorReply);
 	
 	/**
 	 * Spawns a tool tip window.  If an existing tool tip window is open, it will be dismissed first.
@@ -1273,10 +1280,15 @@ private:
 	MouseCaptorHelper MouseCaptor;
 
 	TSharedPtr<FAnalogCursor> AnalogCursor;
+	
+	/** The cursor widget and window to render that cursor for the current software cursor.*/
+	TWeakPtr<SWindow> CursorWindowPtr;
+	TWeakPtr<SWidget> CursorWidgetPtr;
 
 	/** The hit-test radius of the cursor. Default value is 0. */
 	float CursorRadius;
 
+
 	struct FUserFocusEntry
 	{
 		/** A weak path to the widget currently focused by a user, if any. */
diff --git a/Engine/Source/Runtime/Slate/Public/Widgets/SViewport.h b/Engine/Source/Runtime/Slate/Public/Widgets/SViewport.h
--- a/Engine/Source/Runtime/Slate/Public/Widgets/SViewport.h
+++ b/Engine/Source/Runtime/Slate/Public/Widgets/SViewport.h
@@ -144,6 +144,7 @@ public:
 	virtual FReply OnTouchEnded( const FGeometry& MyGeometry, const FPointerEvent& InTouchEvent ) override;
 	virtual FReply OnTouchGesture( const FGeometry& MyGeometry, const FPointerEvent& GestureEvent ) override;
 	virtual FCursorReply OnCursorQuery( const FGeometry& MyGeometry, const FPointerEvent& CursorEvent ) const override;
+	virtual TOptional<TSharedRef<SWidget>> OnMapCursor(const FCursorReply& CursorReply) const override;
 	virtual FReply OnMouseButtonDown( const FGeometry& MyGeometry, const FPointerEvent& MouseEvent ) override;
 	virtual FReply OnMouseButtonUp( const FGeometry& MyGeometry, const FPointerEvent& MouseEvent ) override;
 	virtual void OnMouseEnter( const FGeometry& MyGeometry, const FPointerEvent& MouseEvent ) override;
diff --git a/Engine/Source/Runtime/SlateCore/Private/Widgets/SWidget.cpp b/Engine/Source/Runtime/SlateCore/Private/Widgets/SWidget.cpp
--- a/Engine/Source/Runtime/SlateCore/Private/Widgets/SWidget.cpp
+++ b/Engine/Source/Runtime/SlateCore/Private/Widgets/SWidget.cpp
@@ -197,6 +197,11 @@ FCursorReply SWidget::OnCursorQuery( const FGeometry& MyGeometry, const FPointer
 		: FCursorReply::Unhandled();
 }
 
+TOptional<TSharedRef<SWidget>> SWidget::OnMapCursor(const FCursorReply& CursorReply) const
+{
+	return TOptional<TSharedRef<SWidget>>();
+}
+
 
 FReply SWidget::OnMouseButtonDoubleClick( const FGeometry& InMyGeometry, const FPointerEvent& InMouseEvent )
 {
diff --git a/Engine/Source/Runtime/SlateCore/Public/Input/CursorReply.h b/Engine/Source/Runtime/SlateCore/Public/Input/CursorReply.h
--- a/Engine/Source/Runtime/SlateCore/Public/Input/CursorReply.h
+++ b/Engine/Source/Runtime/SlateCore/Public/Input/CursorReply.h
@@ -28,9 +28,18 @@ public:
 	{
 		return FCursorReply( InCursor );
 	}
-		
-	/** @return The requested MouseCursor when the event was handled. Undefined otherwise. */
-	EMouseCursor::Type GetCursor() { return MouseCursor; }
+
+	/** @return The window to render the Cursor Widget in. */
+	TSharedPtr<SWindow> GetCursorWindow() const { return CursorWindow; }
+	
+	/** @return The custom Cursor Widget to render if set and the event was handled. */
+	TSharedPtr<SWidget> GetCursorWidget() const { return CursorWidget; }
+
+	/** @return The requested MouseCursor if no custom widget is set and the event was handled. */
+	EMouseCursor::Type GetCursorType() const { return MouseCursor; }
+
+	/** Set the Cursor Widget, used by slate application to set the cursor widget if the MapCursor returns a widget. */
+	void SetCursorWidget(TSharedPtr<SWindow> InCursorWindow, TSharedPtr<SWidget> InCursorWidget) { CursorWindow = InCursorWindow; CursorWidget = InCursorWidget; }
 
 private:
 
@@ -45,8 +54,13 @@ private:
 		: TReplyBase<FCursorReply>(true)
 		, MouseCursor( InCursorType )
 	{ }
-	
+
+	/** Window to render for cursor */
+	TSharedPtr<SWindow> CursorWindow;
+
+	/** Custom widget to render for cursor */
+	TSharedPtr<SWidget> CursorWidget;
 		
-	/** The value of the reply, if bHasValue is true. */
+	/** The cursor type must be set is CursorWidget is invalid */
 	EMouseCursor::Type MouseCursor;
 };
diff --git a/Engine/Source/Runtime/SlateCore/Public/Rendering/RenderingCommon.h b/Engine/Source/Runtime/SlateCore/Public/Rendering/RenderingCommon.h
--- a/Engine/Source/Runtime/SlateCore/Public/Rendering/RenderingCommon.h
+++ b/Engine/Source/Runtime/SlateCore/Public/Rendering/RenderingCommon.h
@@ -261,6 +261,16 @@ public:
 		return FCursorReply::Unhandled();
 	}
 
+	/**
+	 * After OnCursorQuery has specified a cursor type the system asks each widget under the mouse to map that cursor to a widget. This event is bubbled.
+	 * 
+	 * @return TOptional<TSharedRef<SWidget>>() if you don't have a mapping otherwise return the Widget to show.
+	 */
+	virtual TOptional<TSharedRef<SWidget>> OnMapCursor(const FCursorReply& CursorReply)
+	{
+		return TOptional<TSharedRef<SWidget>>();
+	}
+
 	/**
 	 *	Returns whether the software cursor is currently visible
 	 */
diff --git a/Engine/Source/Runtime/SlateCore/Public/Widgets/SWidget.h b/Engine/Source/Runtime/SlateCore/Public/Widgets/SWidget.h
--- a/Engine/Source/Runtime/SlateCore/Public/Widgets/SWidget.h
+++ b/Engine/Source/Runtime/SlateCore/Public/Widgets/SWidget.h
@@ -292,6 +292,13 @@ public:
 	 */
 	virtual FCursorReply OnCursorQuery(const FGeometry& MyGeometry, const FPointerEvent& CursorEvent) const;
 
+	/**
+	 * After OnCursorQuery has specified a cursor type the system asks each widget under the mouse to map that cursor to a widget. This event is bubbled.
+	 * 
+	 * @return TOptional<TSharedRef<SWidget>>() if you don't have a mapping otherwise return the Widget to show.
+	 */
+	virtual TOptional<TSharedRef<SWidget>> OnMapCursor(const FCursorReply& CursorReply) const;
+
 	/**
 	 * Called when a mouse button is double clicked.  Override this in derived classes.
 	 *
