Slate - Fixing some build errors from the change to ESlateCheckBoxState.  Adding back EndDragDrop to slate application and deprecating it.

diff --git a/Engine/Source/Runtime/Slate/Public/Framework/Application/SlateApplication.h b/Engine/Source/Runtime/Slate/Public/Framework/Application/SlateApplication.h
--- a/Engine/Source/Runtime/Slate/Public/Framework/Application/SlateApplication.h
+++ b/Engine/Source/Runtime/Slate/Public/Framework/Application/SlateApplication.h
@@ -542,7 +542,13 @@ public:
 	/** Get the current drag-dropping content */
 	TSharedPtr<class FDragDropOperation> GetDragDroppingContent() const;
 
-	/** End any in flight drag and drops */
+	DEPRECATED(4.6, "Use CancelDragDrop")
+	void EndDragDrop()
+	{
+		CancelDragDrop();
+	}
+
+	/** Cancels any in flight drag and drops */
 	void CancelDragDrop();
 
 	/**
