Slate : fixed issue where EPopupMethor::UseCurrentWindow menu anchors would sometimes have the wrong geometry during event processing.

diff --git a/Engine/Source/Runtime/Slate/Private/Widgets/Input/SMenuAnchor.cpp b/Engine/Source/Runtime/Slate/Private/Widgets/Input/SMenuAnchor.cpp
--- a/Engine/Source/Runtime/Slate/Private/Widgets/Input/SMenuAnchor.cpp
+++ b/Engine/Source/Runtime/Slate/Private/Widgets/Input/SMenuAnchor.cpp
@@ -143,6 +143,21 @@ void SMenuAnchor::Tick( const FGeometry& AllottedGeometry, const double InCurren
 			}
 		}
 	}
+	else if (PopupWindow.IsValid() && IsOpenAndReusingWindow())
+	{
+		// Ideally, do this in OnArrangeChildren(); currently not possible because OnArrangeChildren()
+		// can be called in DesktopSpace or WindowSpace, and we will not know which version of the Window
+		// geometry to use. Tick() is always in DesktopSpace, so cache the solution here and just use
+		// it in OnArrangeChildren().
+		const FPopupPlacement LocalPlacement(AllottedGeometry, Children[1].GetWidget()->GetDesiredSize(), Placement.Get());
+		const FSlateRect WindowRectLocalSpace = TransformRect(Inverse(AllottedGeometry.GetAccumulatedLayoutTransform()), PopupWindow->GetClientRectInScreen());
+		const FVector2D FittedPlacement = ComputePopupFitInRect(
+			LocalPlacement.AnchorLocalSpace,
+			FSlateRect(LocalPlacement.LocalPopupOffset, LocalPlacement.LocalPopupOffset + LocalPlacement.LocalPopupSize),
+			LocalPlacement.Orientation, WindowRectLocalSpace);
+
+		LocalPopupPosition = FittedPlacement;
+	}
 
 	/** The tick is ending, so the window was not dismissed this tick. */
 	bDismissedThisTick = false;
@@ -150,15 +165,12 @@ void SMenuAnchor::Tick( const FGeometry& AllottedGeometry, const double InCurren
 
 void SMenuAnchor::OnArrangeChildren( const FGeometry& AllottedGeometry, FArrangedChildren& ArrangedChildren ) const
 {
-	ArrangeSingleChild( AllottedGeometry, ArrangedChildren, Children[0], FVector2D::UnitVector);
+	ArrangeSingleChild( AllottedGeometry, ArrangedChildren, Children[0], FVector2D::UnitVector );
 	const TSharedPtr<SWindow> PresentingWindow = PopupWindowPtr.Pin();
 	if (IsOpenAndReusingWindow() && PresentingWindow.IsValid())
 	{
 		const FPopupPlacement LocalPlacement(AllottedGeometry, Children[1].GetWidget()->GetDesiredSize(), Placement.Get());
-		const FSlateRect WindowRectLocalSpace = TransformRect(Inverse(AllottedGeometry.GetAccumulatedLayoutTransform()), FSlateRect(FVector2D::ZeroVector, PresentingWindow->GetClientSizeInScreen()));
-		
-		const FVector2D FittedPlacement = ComputePopupFitInRect(LocalPlacement.AnchorLocalSpace, FSlateRect(LocalPlacement.LocalPopupOffset, LocalPlacement.LocalPopupOffset+LocalPlacement.LocalPopupSize), LocalPlacement.Orientation, WindowRectLocalSpace);
-		ArrangedChildren.AddWidget(AllottedGeometry.MakeChild(Children[1].GetWidget(), LocalPlacement.LocalPopupSize, FSlateLayoutTransform(FittedPlacement)));
+		ArrangedChildren.AddWidget(AllottedGeometry.MakeChild(Children[1].GetWidget(), LocalPlacement.LocalPopupSize, FSlateLayoutTransform(LocalPopupPosition)));
 	}
 }
 
@@ -446,6 +458,7 @@ SMenuAnchor::SMenuAnchor()
 	, bDismissedThisTick( false )
 	, Method()
 	, MethodInUse()
+	, LocalPopupPosition( FVector2D::ZeroVector )
 {
 }
 
diff --git a/Engine/Source/Runtime/Slate/Public/Widgets/Input/SMenuAnchor.h b/Engine/Source/Runtime/Slate/Public/Widgets/Input/SMenuAnchor.h
--- a/Engine/Source/Runtime/Slate/Public/Widgets/Input/SMenuAnchor.h
+++ b/Engine/Source/Runtime/Slate/Public/Widgets/Input/SMenuAnchor.h
@@ -136,5 +136,13 @@ protected:
 	/** Method currently being used to show the popup. No value if popup is closed. */
 	TOptional<EPopupMethod> MethodInUse;
 
+	/**
+	 * @todo Slate : Unify geometry so that this is not necessary.
+	 * Stores the local offset of the popup as we have to compute it in OnTick;
+	 * Cannot compute in OnArrangeChildren() because it can be
+	 * called in different spaces (window or desktop.)
+	 */
+	FVector2D LocalPopupPosition;
+
 	TPanelChildren<FSimpleSlot> Children;
 };
