Refactored the "analogcursor" on the slate side to be an input preprocessor.
[REDACTED_HANDLE] nick.atamas

diff --git a/Engine/Source/Runtime/Slate/Private/Framework/Application/SlateApplication.cpp b/Engine/Source/Runtime/Slate/Private/Framework/Application/SlateApplication.cpp
--- a/Engine/Source/Runtime/Slate/Private/Framework/Application/SlateApplication.cpp
+++ b/Engine/Source/Runtime/Slate/Private/Framework/Application/SlateApplication.cpp
@@ -8,7 +8,7 @@
 #include "IWidgetReflector.h"
 #include "GenericCommands.h"
 #include "NotificationManager.h"
-#include "AnalogCursor.h"
+#include "IInputProcessor.h"
 
 
 class FEventRouter
@@ -1125,9 +1125,9 @@ void FSlateApplication::FinishedInputThisFrame()
 {
 	const float DeltaTime = GetDeltaTime();
 
-	if (AnalogCursor.IsValid() && PlatformApplication->Cursor.IsValid())
+	if (InputPreProcessor.IsValid() && PlatformApplication->Cursor.IsValid())
 	{
-		AnalogCursor->Tick(DeltaTime, *this, PlatformApplication->Cursor.ToSharedRef());
+		InputPreProcessor->Tick(DeltaTime, *this, PlatformApplication->Cursor.ToSharedRef());
 	}
 
 	// All the input events have been processed.
@@ -3068,15 +3068,15 @@ void FSlateApplication::SetDragTriggerDistnace( float ScreenPixels )
 	DragTriggerDistnace = ScreenPixels;
 }
 
-void FSlateApplication::SetAnalogCursorEnable(bool bEnable, TSharedPtr<class FAnalogCursor> OptionalNewAnalogCursor)
+void FSlateApplication::SetInputPreProcessor(bool bEnable, TSharedPtr<class IInputProcessor> NewInputProcessor)
 {
-	if (bEnable && !AnalogCursor.IsValid())
+	if (bEnable && NewInputProcessor.IsValid())
 	{
-		AnalogCursor = OptionalNewAnalogCursor.IsValid() ? OptionalNewAnalogCursor : MakeShareable(new FAnalogCursor);
+		InputPreProcessor = NewInputProcessor;
 	}
 	else
 	{
-		AnalogCursor.Reset();
+		InputPreProcessor.Reset();
 	}
 }
 
@@ -3525,7 +3525,7 @@ bool FSlateApplication::OnKeyDown( const int32 KeyCode, const uint32 CharacterCo
 bool FSlateApplication::ProcessKeyDownEvent( FKeyEvent& InKeyEvent )
 {
 	// Analog cursor gets first chance at the input
-	if (AnalogCursor.IsValid() && AnalogCursor->HandleKeyDownEvent(*this, InKeyEvent))
+	if (InputPreProcessor.IsValid() && InputPreProcessor->HandleKeyDownEvent(*this, InKeyEvent))
 	{
 		return true;
 	}
@@ -3610,7 +3610,7 @@ bool FSlateApplication::OnKeyUp( const int32 KeyCode, const uint32 CharacterCode
 bool FSlateApplication::ProcessKeyUpEvent( FKeyEvent& InKeyEvent )
 {
 	// Analog cursor gets first chance at the input
-	if (AnalogCursor.IsValid() && AnalogCursor->HandleKeyUpEvent(*this, InKeyEvent))
+	if (InputPreProcessor.IsValid() && InputPreProcessor->HandleKeyUpEvent(*this, InKeyEvent))
 	{
 		return true;
 	}
@@ -3642,7 +3642,7 @@ bool FSlateApplication::ProcessKeyUpEvent( FKeyEvent& InKeyEvent )
 bool FSlateApplication::ProcessAnalogInputEvent(FAnalogInputEvent& InAnalogInputEvent)
 {
 	// Analog cursor gets first chance at the input
-	if (AnalogCursor.IsValid() && AnalogCursor->HandleAnalogInputEvent(*this, InAnalogInputEvent))
+	if (InputPreProcessor.IsValid() && InputPreProcessor->HandleAnalogInputEvent(*this, InAnalogInputEvent))
 	{
 		return true;
 	}
diff --git a/Engine/Source/Runtime/Slate/Public/Framework/Application/AnalogCursor.h b/Engine/Source/Runtime/Slate/Public/Framework/Application/AnalogCursor.h
--- a/Engine/Source/Runtime/Slate/Public/Framework/Application/AnalogCursor.h
+++ b/Engine/Source/Runtime/Slate/Public/Framework/Application/AnalogCursor.h
@@ -2,12 +2,12 @@
 
 #pragma once
 
-class FSlateApplication;
+#include "IInputProcessor.h"
 
 /**
  * A class that simulates a cursor driven by an analog stick.
  */
-class SLATE_API FAnalogCursor
+class SLATE_API FAnalogCursor : public IInputProcessor
 {
 public:
 	FAnalogCursor();
@@ -16,11 +16,11 @@ public:
 	virtual ~FAnalogCursor()
 	{}
 
-	virtual void Tick(const float DeltaTime, FSlateApplication& SlateApp, TSharedRef<ICursor> Cursor);
+	virtual void Tick(const float DeltaTime, FSlateApplication& SlateApp, TSharedRef<ICursor> Cursor) override;
 
-	bool HandleKeyDownEvent(FSlateApplication& SlateApp, const FKeyEvent& InKeyEvent);
-	bool HandleKeyUpEvent(FSlateApplication& SlateApp, const FKeyEvent& InKeyEvent);
-	bool HandleAnalogInputEvent(FSlateApplication& SlateApp, const FAnalogInputEvent& InAnalogInputEvent);
+	virtual bool HandleKeyDownEvent(FSlateApplication& SlateApp, const FKeyEvent& InKeyEvent) override;
+	virtual bool HandleKeyUpEvent(FSlateApplication& SlateApp, const FKeyEvent& InKeyEvent) override;
+	virtual bool HandleAnalogInputEvent(FSlateApplication& SlateApp, const FAnalogInputEvent& InAnalogInputEvent) override;
 
 protected:
 
diff --git a/Engine/Source/Runtime/Slate/Public/Framework/Application/IInputProcessor.h b/Engine/Source/Runtime/Slate/Public/Framework/Application/IInputProcessor.h
new file mode 100644
--- /dev/null
+++ b/Engine/Source/Runtime/Slate/Public/Framework/Application/IInputProcessor.h
@@ -0,0 +1,23 @@
+// Copyright 1998-2015 Epic Games, Inc. All Rights Reserved.
+
+#pragma once
+
+class FSlateApplication;
+
+/**
+ * Interface for a Slate Input Handler
+ */
+class SLATE_API IInputProcessor
+{
+public:
+	IInputProcessor(){};
+	virtual ~IInputProcessor(){}
+
+	virtual void Tick(const float DeltaTime, FSlateApplication& SlateApp, TSharedRef<ICursor> Cursor) = 0;
+
+	virtual bool HandleKeyDownEvent(FSlateApplication& SlateApp, const FKeyEvent& InKeyEvent) = 0;
+	virtual bool HandleKeyUpEvent(FSlateApplication& SlateApp, const FKeyEvent& InKeyEvent) = 0;
+	virtual bool HandleAnalogInputEvent(FSlateApplication& SlateApp, const FAnalogInputEvent& InAnalogInputEvent) = 0;
+
+};
+
diff --git a/Engine/Source/Runtime/Slate/Public/Framework/Application/SlateApplication.h b/Engine/Source/Runtime/Slate/Public/Framework/Application/SlateApplication.h
--- a/Engine/Source/Runtime/Slate/Public/Framework/Application/SlateApplication.h
+++ b/Engine/Source/Runtime/Slate/Public/Framework/Application/SlateApplication.h
@@ -15,7 +15,7 @@ class FHittestGrid;
 struct FPopupTransitionEffect;
 class FMenuStack;
 class IWidgetReflector;
-class FAnalogCursor;
+class IInputProcessor;
 
 /** A Delegate for querying whether source code access is possible */
 DECLARE_DELEGATE_RetVal(bool, FQueryAccessSourceCode);
@@ -958,7 +958,7 @@ public:
 	void SetDragTriggerDistnace( float ScreenPixels );
 	
 	/** Set the analog cursor to be enabled or disabled. */
-	void SetAnalogCursorEnable(bool bEnable, TSharedPtr<class FAnalogCursor> OptionalNewAnalogCursor = nullptr);
+	void SetInputPreProcessor(bool bEnable, TSharedPtr<class IInputProcessor> NewInputProcessor = nullptr);
 
 	/** Sets the hit detection radius of the cursor */
 	void SetCursorRadius(float NewRadius);
@@ -1313,7 +1313,8 @@ private:
 	/** The current mouse captor for the application, if any. */
 	MouseCaptorHelper MouseCaptor;
 
-	TSharedPtr<FAnalogCursor> AnalogCursor;
+	/** An input preprocessor, gets an opportunity to parse input before anything else. */
+	TSharedPtr<IInputProcessor> InputPreProcessor;
 	
 	/** The cursor widget and window to render that cursor for the current software cursor.*/
 	TWeakPtr<SWindow> CursorWindowPtr;
