Slate : fixed tooltip drift.
Fixes: UE-6437

diff --git a/Engine/Source/Runtime/Slate/Private/Framework/Application/SlateApplication.cpp b/Engine/Source/Runtime/Slate/Private/Framework/Application/SlateApplication.cpp
--- a/Engine/Source/Runtime/Slate/Private/Framework/Application/SlateApplication.cpp
+++ b/Engine/Source/Runtime/Slate/Private/Framework/Application/SlateApplication.cpp
@@ -3122,9 +3122,10 @@ FVector2D FSlateApplication::CalculatePopupWindowPosition( const FSlateRect& InA
 		// Calculate the rectangle around our work area
 		// Use our own rect.  This window as probably doesn't have a size or position yet.
 		// Use a size of 1 to get the closest monitor to the start point
-		AnchorRect.Left = InAnchor.Left + 1;
-		AnchorRect.Top = InAnchor.Top + 1;
-		const FPlatformRect PlatformWorkArea = PlatformApplication->GetWorkArea( AnchorRect );
+		FPlatformRect WorkAreaFinderRect(AnchorRect);
+		WorkAreaFinderRect.Left = AnchorRect.Left + 1;
+		WorkAreaFinderRect.Top = AnchorRect.Top + 1;
+		const FPlatformRect PlatformWorkArea = PlatformApplication->GetWorkArea(WorkAreaFinderRect);
 
 		const FSlateRect WorkAreaRect( 
 			PlatformWorkArea.Left, 
