Removed reference to deleted SlateWordWrapper.h from SMultiLineEditableText.cpp.

diff --git a/Engine/Source/Runtime/Slate/Private/Widgets/Text/SMultiLineEditableText.cpp b/Engine/Source/Runtime/Slate/Private/Widgets/Text/SMultiLineEditableText.cpp
--- a/Engine/Source/Runtime/Slate/Private/Widgets/Text/SMultiLineEditableText.cpp
+++ b/Engine/Source/Runtime/Slate/Private/Widgets/Text/SMultiLineEditableText.cpp
@@ -5,7 +5,6 @@
 #if WITH_FANCY_TEXT
 
 #include "TextEditHelper.h"
-#include "SlateWordWrapper.h"
 #include "PlainTextLayoutMarshaller.h"
 #include "GenericCommands.h"
 
