Slate: Fixed some whitespaces and line breaks; NULL to nullptr; removed redundant function documentation

diff --git a/Engine/Source/Runtime/Slate/Private/Framework/Application/SlateApplication.cpp b/Engine/Source/Runtime/Slate/Private/Framework/Application/SlateApplication.cpp
--- a/Engine/Source/Runtime/Slate/Private/Framework/Application/SlateApplication.cpp
+++ b/Engine/Source/Runtime/Slate/Private/Framework/Application/SlateApplication.cpp
@@ -16,15 +16,9 @@ class FEventRouter
 {
 
 // @todo slate : making too many event copies when translating events( i.e. Translate<EventType>::PointerEvent ).
-
-
-
 // @todo slate : Widget Reflector should log: (1) Every process reply (2) Every time the event is handled and by who.
-
 // @todo slate : Remove remaining [&]-style mass captures.
-
 // @todo slate : Eliminate all ad-hoc uses of SetEventPath()
-
 // @todo slate : Remove CALL_WIDGET_FUNCTION
 
 public:
@@ -225,7 +219,7 @@ public:
 
 	static void ProcessReply( FSlateApplication* Application, const FWidgetPath& RoutingPath, const FReply& Reply, const FWidgetPath* WidgetsUnderCursor, const FInputEvent* PointerEvent )
 	{
-		Application->ProcessReply(RoutingPath, Reply, WidgetsUnderCursor, NULL);
+		Application->ProcessReply(RoutingPath, Reply, WidgetsUnderCursor, nullptr);
 	}
 
 	static void ProcessReply( FSlateApplication* Application, const FWidgetPath& RoutingPath, const FReply& Reply, const FWidgetPath* WidgetsUnderCursor, const FPointerEvent* PointerEvent )
@@ -245,6 +239,7 @@ public:
 
 };
 
+
 template<>
 struct FEventRouter::Translate<FPointerEvent>
 {
@@ -308,6 +303,7 @@ DECLARE_CYCLE_STAT( TEXT("DrawWindows"), STAT_SlateDrawWindowTime, STATGROUP_Sla
 	}
 #endif
 
+
 namespace SlateDefs
 {
 	// How far tool tips should be offset from the mouse cursor position, in pixels
@@ -317,6 +313,7 @@ namespace SlateDefs
 	static const FVector2D ToolTipOffsetFromForceField( 4.0f, 3.0f );
 }
 
+
 /** True if we should allow throttling based on mouse movement activity.  int32 instead of bool only for console variable system. */
 TAutoConsoleVariable<int32> ThrottleWhenMouseIsMoving( 
 	TEXT( "Slate.ThrottleWhenMouseIsMoving" ),
@@ -334,6 +331,7 @@ TAutoConsoleVariable<int32> SkipSecondPrepass(
 	0,
 	TEXT("Whether to skip the second Slate PrePass call (the one right before rendering)."));
 
+
 //////////////////////////////////////////////////////////////////////////
 bool FSlateApplication::MouseCaptorHelper::HasCapture() const
 {
@@ -647,7 +645,7 @@ void FSlateApplication::Shutdown()
 }
 
 
-TSharedPtr<FSlateApplication> FSlateApplication::CurrentApplication = NULL;
+TSharedPtr<FSlateApplication> FSlateApplication::CurrentApplication = nullptr;
 
 FSlateApplication::FSlateApplication()
 	: bAppIsActive(true)
@@ -686,7 +684,7 @@ FSlateApplication::FSlateApplication()
 		ToolTipFadeInDuration,
 		TEXT( "How long it takes for a tool-tip to fade in, in seconds." ) )
 	, bIsExternalUIOpened( false )
-	, SlateTextField( NULL )
+	, SlateTextField( nullptr )
 	, bIsFakingTouch(FParse::Param(FCommandLine::Get(), TEXT("simmobile")) || FParse::Param(FCommandLine::Get(), TEXT("faketouches")))
 	, bIsGameFakingTouch( false )
 	, bIsFakingTouched( false )
@@ -717,10 +715,10 @@ FSlateApplication::FSlateApplication()
 
 FSlateApplication::~FSlateApplication()
 {
-	if (SlateTextField != NULL)
+	if (SlateTextField != nullptr)
 	{
 		delete SlateTextField;
-		SlateTextField = NULL;
+		SlateTextField = nullptr;
 	}
 }
 
@@ -834,11 +832,7 @@ bool FSlateApplication::IsWindowHousingInteractiveTooltip(const TSharedRef<const
 	return bIsHousingInteractiveTooltip;
 }
 
-/** 
- * Ticks a single slate window
- *
- * @param WindowToTick	The window to tick
- */
+
 void FSlateApplication::TickWindowAndChildren( TSharedRef<SWindow> WindowToTick )
 {
 	if ( WindowToTick->IsVisible() && !WindowToTick->IsWindowMinimized() )
@@ -1179,11 +1173,10 @@ void FSlateApplication::FinishedInputThisFrame()
 	}
 }
 
+
 extern SLATECORE_API int32 bFoldTick;
 
-/**
- * Ticks this application
- */
+
 void FSlateApplication::Tick()
 {
 	SCOPE_CYCLE_COUNTER( STAT_SlateTickTime );
@@ -1428,7 +1421,7 @@ TSharedRef<SWindow> FSlateApplication::AddWindow( TSharedRef<SWindow> InSlateWin
 
 TSharedRef< FGenericWindow > FSlateApplication::MakeWindow( TSharedRef<SWindow> InSlateWindow, const bool bShowImmediately )
 {
-	TSharedPtr<FGenericWindow> NativeParent = NULL;
+	TSharedPtr<FGenericWindow> NativeParent = nullptr;
 	TSharedPtr<SWindow> ParentWindow = InSlateWindow->GetParentWindow();
 	if ( ParentWindow.IsValid() )
 	{
@@ -1498,12 +1491,7 @@ bool FSlateApplication::CanDisplayWindows() const
 	return Renderer.IsValid() && Renderer->AreShadersInitialized();
 }
 
-/**
- * Adds a modal window to the application.  
- * In most cases, this function does not return until the modal window is closed (the only exception is a modal window for slow tasks)  
- *
- * @param InSlateWindow  A SlateWindow to which to add a native window.
- */
+
 void FSlateApplication::AddModalWindow( TSharedRef<SWindow> InSlateWindow, const TSharedPtr<const SWidget> InParentWidget, bool bSlowTaskWindow )
 {
 	if( !CanAddModalWindow() )
@@ -1559,10 +1547,10 @@ void FSlateApplication::AddModalWindow( TSharedRef<SWindow> InSlateWindow, const
 	PressedMouseButtons.Empty();
 
 	// Also force the platform capture off as the call to ReleaseMouseCapture() above still relies on mouse up messages to clear the capture
-	PlatformApplication->SetCapture( NULL );
+	PlatformApplication->SetCapture( nullptr );
 
 	// Disable high precision mouse mode when a modal window is added.  On some OS'es even when a window is diabled, raw input is sent to it.
-	PlatformApplication->SetHighPrecisionMouseMode( false, NULL );
+	PlatformApplication->SetHighPrecisionMouseMode( false, nullptr );
 
 	// Block on all modal windows unless its a slow task.  In that case the game thread is allowed to run.
 	if( !bSlowTaskWindow )
@@ -1666,12 +1654,7 @@ int32 FSlateApplication::GetLocationInMenuStack( TSharedRef<SWindow> WindowToFin
 	return MenuStack.FindLocationInStack( WindowToFind );
 }
 
-/**
- * Destroying windows has implications on some OSs ( e.g. destroying Win32 HWNDs can cause events to be lost ).
- * Slate strictly controls when windows are destroyed. 
- *
- * @param WindowToDestroy  Window to queue for destruction.  This will also queue children of this window for destruction
- */
+
 void FSlateApplication::RequestDestroyWindow( TSharedRef<SWindow> InWindowToDestroy )
 {
 	struct local
@@ -1708,10 +1691,7 @@ void FSlateApplication::DestroyWindowImmediately( TSharedRef<SWindow> WindowToDe
 	DestroyWindowsImmediately();
 }
 
-/**
- * Disable Slate components when an external, non-slate, modal window is brought up.  In the case of multiple
- * external modal windows, we will only increment our tracking counter.
- */
+
 void FSlateApplication::ExternalModalStart()
 {
 	if( NumExternalModalWindowsActive++ == 0 )
@@ -1743,10 +1723,7 @@ void FSlateApplication::ExternalModalStart()
 	}
 }
 
-/**
- * Re-enable disabled Slate components when a non-slate modal window is dismissed.  Slate components
- * will only be re-enabled when all tracked external modal windows have been dismissed.
- */
+
 void FSlateApplication::ExternalModalStop()
 {
 	check(NumExternalModalWindowsActive > 0);
@@ -1791,7 +1768,7 @@ void FSlateApplication::RegisterGameViewport( TSharedRef<SViewport> InViewport )
 		// Without this we don't get WM_KEYDOWN or WM_CHAR messages in play in viewport sessions.
 		PathToViewport.GetWindow()->GetNativeWindow()->SetWindowFocus();
 
-		ProcessReply( PathToViewport, Reply, NULL, NULL );
+		ProcessReply( PathToViewport, Reply, nullptr, nullptr );
 	}
 }
 
@@ -1916,10 +1893,10 @@ void FSlateApplication::ResetToDefaultPointerInputSettings()
 {
 	for (auto MouseCaptorPath : MouseCaptor.ToWidgetPaths())
 	{
-		ProcessReply(MouseCaptorPath, FReply::Handled().ReleaseMouseCapture(), NULL, NULL);
+		ProcessReply(MouseCaptorPath, FReply::Handled().ReleaseMouseCapture(), nullptr, nullptr);
 	}
 
-	ProcessReply(FWidgetPath(), FReply::Handled().ReleaseMouseLock(), NULL, NULL);
+	ProcessReply(FWidgetPath(), FReply::Handled().ReleaseMouseLock(), nullptr, nullptr);
 
 	if ( PlatformApplication->Cursor.IsValid() )
 	{
@@ -1946,7 +1923,7 @@ TSharedPtr<SWindow> FSlateApplication::GetActiveTopLevelWindow() const
 
 TSharedPtr<SWindow> FSlateApplication::GetActiveModalWindow() const
 {
-	return (ActiveModalWindows.Num() > 0) ? ActiveModalWindows.Last() : NULL;
+	return (ActiveModalWindows.Num() > 0) ? ActiveModalWindows.Last() : nullptr;
 }
 
 bool FSlateApplication::SetKeyboardFocus(const FWidgetPath& InFocusPath, const EFocusCause InCause /*= EFocusCause::SetDirectly*/)
@@ -2086,7 +2063,7 @@ bool FSlateApplication::SetUserFocus(const uint32 InUserIndex, const FWidgetPath
 		FReply Reply = NewFocusedWidget->OnFocusReceived(WidgetToFocus.Geometry, FFocusEvent(InCause));
 		if (Reply.IsEventHandled())
 		{
-			ProcessReply(InFocusPath, Reply, NULL, NULL);
+			ProcessReply(InFocusPath, Reply, nullptr, nullptr);
 		}
 		//@Todo slate: Remove this backwards compatibility code as soon as the deprecated functionality in Swidget is removed
 		// Run the deprecated function to maintain backwards compatibility for now.
@@ -2095,7 +2072,7 @@ bool FSlateApplication::SetUserFocus(const uint32 InUserIndex, const FWidgetPath
 			FReply Reply = NewFocusedWidget->OnFocusReceived(WidgetToFocus.Geometry, FFocusEvent(InCause));
 			if (Reply.IsEventHandled())
 			{
-				ProcessReply(InFocusPath, Reply, NULL, NULL);
+				ProcessReply(InFocusPath, Reply, nullptr, nullptr);
 			}
 		}
 		
@@ -2217,7 +2194,7 @@ TSharedPtr<SWindow> FSlateApplication::FindWidgetWindow( TSharedRef< const SWidg
 	{
 		return OutWidgetPath.TopLevelWindow;
 	}
-	return NULL;
+	return nullptr;
 }
 
 
@@ -2260,8 +2237,8 @@ void FSlateApplication::ProcessReply( const FWidgetPath& CurrentEventPath, const
 	{
 		checkf( !this->DragDropContent.IsValid(), TEXT("Drag and Drop already in progress!") );
 		check( true == TheReply.IsEventHandled() );
-		check( WidgetsUnderMouse != NULL );
-		check( InMouseEvent != NULL );
+		check( WidgetsUnderMouse != nullptr );
+		check( InMouseEvent != nullptr );
 		DragDropContent = ReplyDragDropContent;
 
 		// We have entered drag and drop mode.
@@ -2299,8 +2276,8 @@ void FSlateApplication::ProcessReply( const FWidgetPath& CurrentEventPath, const
 		}
 		else if ( PlatformApplication->IsUsingHighPrecisionMouseMode() )
 		{
-			PlatformApplication->SetHighPrecisionMouseMode( false, NULL );
-			PlatformApplication->SetCapture( NULL );
+			PlatformApplication->SetHighPrecisionMouseMode( false, nullptr );
+			PlatformApplication->SetCapture( nullptr );
 		}
 	}
 
@@ -2320,7 +2297,7 @@ void FSlateApplication::ProcessReply( const FWidgetPath& CurrentEventPath, const
 	else if( TheReply.ShouldReleaseMouseLock() )
 	{
 		// Unlock the mouse
-		LockCursor( NULL );
+		LockCursor( nullptr );
 	}
 	
 	// If we have a valid Navigation request and enough time has passed since the last navigation then try to navigate.
@@ -2428,7 +2405,7 @@ void FSlateApplication::LockCursorToPath(const FWidgetPath& WidgetPath)
 void FSlateApplication::UnlockCursor()
 {
 	// Unlock the mouse
-	PlatformApplication->Cursor->Lock(NULL);
+	PlatformApplication->Cursor->Lock(nullptr);
 	CursorLock.PathToLockingWidget = FWeakWidgetPath();
 }
 
@@ -2701,7 +2678,7 @@ void FSlateApplication::UpdateToolTip( bool AllowSpawningOfNewToolTips )
 		// Remove existing tooltip if there is one.
 		if (TooltipVisualizerPtr.IsValid())
 		{
-			TooltipVisualizerPtr.Pin()->OnVisualizeTooltip( NULL );
+			TooltipVisualizerPtr.Pin()->OnVisualizeTooltip( nullptr );
 		}
 
 		bool bOnVisualizeTooltipHandled = false;
@@ -2810,7 +2787,7 @@ void FSlateApplication::UpdateToolTip( bool AllowSpawningOfNewToolTips )
 			}
 			else
 			{
-				NewToolTip = NULL;
+				NewToolTip = nullptr;
 			}
 
 			ActiveToolTip = NewToolTip;
@@ -3032,7 +3009,7 @@ void FSlateApplication::OnLogSlateEvent(EEventLog::Type Elovent, const FString&
 #if LOG_SLATE_EVENTS
 	if (EventLogger.IsValid())
 	{
-		LOG_EVENT_CONTENT(Event, AdditionalContent, NULL);
+		LOG_EVENT_CONTENT(Event, AdditionalContent, nullptr);
 	}
 #endif
 }
@@ -3042,7 +3019,7 @@ void FSlateApplication::OnLogSlateEvent(EEventLog::Type Event, const FText& Addi
 #if LOG_SLATE_EVENTS
 	if (EventLogger.IsValid())
 	{
-		LOG_EVENT_CONTENT(Event, AdditionalContent.ToString(), NULL);
+		LOG_EVENT_CONTENT(Event, AdditionalContent.ToString(), nullptr);
 	}
 #endif
 };
@@ -3425,6 +3402,12 @@ bool FSlateApplication::HasFocusedDescendants( const TSharedRef< const SWidget >
 }
 
 
+bool FSlateApplication::IsExternalUIOpened()
+{
+	return bIsExternalUIOpened;
+}
+
+
 TSharedRef<SWidget> FSlateApplication::MakeImage( const TAttribute<const FSlateBrush*>& Image, const TAttribute<FSlateColor>& Color, const TAttribute<EVisibility>& Visibility ) const
 {
 	return SNew(SImage)
@@ -4065,7 +4048,7 @@ bool FSlateApplication::ProcessMouseButtonUpEvent( FPointerEvent& MouseEvent )
 	if ( PressedMouseButtons.Num() == 0 )
 	{
 		// Release Capture
-		PlatformApplication->SetCapture( NULL );
+		PlatformApplication->SetCapture( nullptr );
 	}
 
 	return true;
@@ -4085,13 +4068,13 @@ bool FSlateApplication::OnMouseWheel( const float Delta )
 		PlatformApplication->GetModifierKeys()
 		);
 
-	return ProcessMouseWheelOrGestureEvent( MouseWheelEvent, NULL );
+	return ProcessMouseWheelOrGestureEvent( MouseWheelEvent, nullptr );
 }
 
 bool FSlateApplication::ProcessMouseWheelOrGestureEvent( FPointerEvent& InWheelEvent, const FPointerEvent* InGestureEvent )
 {
 	const bool bShouldProcessEvent = InWheelEvent.GetWheelDelta() != 0 ||
-		(InGestureEvent!=NULL && InGestureEvent->GetGestureDelta()!=FVector2D::ZeroVector);
+		(InGestureEvent != nullptr && InGestureEvent->GetGestureDelta()!=FVector2D::ZeroVector);
 	
 	if ( !bShouldProcessEvent )
 	{
@@ -4113,7 +4096,7 @@ bool FSlateApplication::ProcessMouseWheelOrGestureEvent( FPointerEvent& InWheelE
 	{
 		FReply TempReply = FReply::Unhandled();
 		// Gesture event gets first shot, if slate doesn't respond to it, we'll try the wheel event.
-		if( InGestureEvent!=NULL )
+		if( InGestureEvent != nullptr )
 		{
 			TempReply = CurWidget.Widget->OnTouchGesture( CurWidget.Geometry, *InGestureEvent );
 		}
@@ -4684,7 +4667,7 @@ void FSlateApplication::ProcessMotionDetectedEvent( FMotionEvent& MotionEvent )
 
 		/* Send the message to the widget */
 		FReply Reply = ArrangedWidget.Widget->OnMotionDetected(ArrangedWidget.Geometry, MotionEvent).SetHandler(ArrangedWidget.Widget);
-		ProcessReply(PathToWidget, Reply, NULL, NULL, MotionEvent.GetUserIndex());
+		ProcessReply(PathToWidget, Reply, nullptr, nullptr, MotionEvent.GetUserIndex());
 	}
 }
 
@@ -5229,5 +5212,3 @@ void FSlateApplication::SetWidgetReflector(const TSharedRef<IWidgetReflector>& W
 
 	WidgetReflectorPtr = WidgetReflector;
 }
-
-
diff --git a/Engine/Source/Runtime/Slate/Public/Framework/Application/SlateApplication.h b/Engine/Source/Runtime/Slate/Public/Framework/Application/SlateApplication.h
--- a/Engine/Source/Runtime/Slate/Public/Framework/Application/SlateApplication.h
+++ b/Engine/Source/Runtime/Slate/Public/Framework/Application/SlateApplication.h
@@ -7,6 +7,7 @@
 #include "SlateDelegates.h"
 #include "SlateApplicationBase.h"
 
+
 class SToolTip;
 class SViewport;
 class SWindow;
@@ -17,6 +18,7 @@ class FMenuStack;
 class IWidgetReflector;
 class IInputProcessor;
 
+
 /** A Delegate for querying whether source code access is possible */
 DECLARE_DELEGATE_RetVal(bool, FQueryAccessSourceCode);
 
@@ -26,6 +28,7 @@ DECLARE_DELEGATE(FModalWindowStackEnded)
 
 extern SLATE_API const FName NAME_UnrealOS;
 
+
 /**
  * Design constraints for Slate applications
  */
@@ -35,6 +38,7 @@ namespace SlateApplicationDefs
 	static const int32 MaxUsers = 8;
 }
 
+
 /** Allow widgets to find out when someone clicked outside them. Currently needed by MenuAnchros. */
 class SLATE_API FPopupSupport
 {
@@ -95,10 +99,8 @@ class SLATE_API FSlateApplication
 {
 public:
 
-	/**
-	 * Destructor.
-	 */
-	virtual ~FSlateApplication( );
+	/** Virtual destructor. */
+	virtual ~FSlateApplication();
 
 public:
 
@@ -212,34 +214,22 @@ public:
 	 */
 	void SetCursorPos( const FVector2D& MouseCoordinate );
 
-	/**
-	 * Polls game devices for input
-	 */
+	/** Polls game devices for input */
 	void PollGameDeviceState();
 
-	/**
-	 * Occurs before Tick(), after all pointer and keyboard input has been processed.
-	 */
+	/** Occurs before Tick(), after all pointer and keyboard input has been processed. */
 	void FinishedInputThisFrame();
 
-	/**
-	 * Ticks this application
-	 */
+	/** Ticks this application */
 	void Tick();
 
-	/**
-	 * Pumps OS messages when a modal window or intra-frame debugging session exists
-	 */
+	/** Pumps OS messages when a modal window or intra-frame debugging session exists */
 	void PumpMessages();
 
-	/**
-	 * Returns true if this slate application is ready to open modal windows
-	 */
+	/** Returns true if this slate application is ready to open modal windows */
 	bool CanAddModalWindow() const;
 
-	/**
-	 * Returns true if this slate application is ready to display windows.
-	 */
+	/** Returns true if this slate application is ready to display windows. */
 	bool CanDisplayWindows() const;
 	
 	/**
@@ -290,9 +280,7 @@ public:
 	/** @return	Returns true if there are any pop-up menus summoned */
 	bool AnyMenusVisible() const;
 
-	/**
-	 * Dismisses any open menus
-	 */
+	/** Dismisses any open menus */
 	void DismissAllMenus();
 
 	/**
@@ -322,13 +310,13 @@ public:
 	 * Disable Slate components when an external, non-slate, modal window is brought up.  In the case of multiple
 	 * external modal windows, we will only increment our tracking counter.
 	 */
-	void ExternalModalStart( );
+	void ExternalModalStart();
 
 	/**
 	 * Re-enable disabled Slate components when a non-slate modal window is dismissed.  Slate components
 	 * will only be re-enabled when all tracked external modal windows have been dismissed.
 	 */
-	void ExternalModalStop( );
+	void ExternalModalStop();
 
 	/** 
 	 * Removes references to FViewportRHI's.  
@@ -432,14 +420,14 @@ public:
 	 * Mouse capture
 	 */
 
-	/** returning platform-specific value designating window that captures mouse, or NULL if mouse isn't captured */
+	/** returning platform-specific value designating window that captures mouse, or nullptr if mouse isn't captured */
 	virtual void* GetMouseCaptureWindow( void ) const;
 
 
 	/** Releases the mouse capture from whatever it currently is on. */
 	void ReleaseMouseCapture();
 
-	/** @return The active modal window or NULL if there is no modal window. */
+	/** @return The active modal window or nullptr if there is no modal window. */
 	TSharedPtr<SWindow> GetActiveModalWindow() const;
 
 	/**
@@ -475,7 +463,6 @@ public:
 	 * Finds the window that the provided widget resides in
 	 * 
 	 * @param InWidget		The widget to find the window for
-	 *
 	 * @return The window where the widget resides, or null if the widget wasn't found.  Remember, a widget might not be found simply because its parent decided not to report the widget in ArrangeChildren.
 	 */
 	TSharedPtr<SWindow> FindWidgetWindow( TSharedRef< const SWidget > InWidget ) const;
@@ -485,7 +472,6 @@ public:
 	 * 
 	 * @param InWidget		The widget to find the window for
 	 * @param OutWidgetPath Full widget path generated 
-	 *
 	 * @return The window where the widget resides, or null if the widget wasn't found.  Remember, a widget might not be found simply because its parent decided not to report the widget in ArrangeChildren.
 	 */
 	TSharedPtr<SWindow> FindWidgetWindow( TSharedRef< const SWidget > InWidget, FWidgetPath& OutWidgetPath) const;
@@ -681,11 +667,8 @@ protected:
 	friend class FEventRouter;
 
 	virtual bool HasMouseCapture(const TSharedPtr<const SWidget> Widget) const override;
-
 	virtual TOptional<EFocusCause> HasUserFocus(const TSharedPtr<const SWidget> Widget, int32 UserIndex) const override;
-
 	virtual TOptional<EFocusCause> HasAnyUserFocus(const TSharedPtr<const SWidget> Widget) const override;
-
 	virtual bool ShowUserFocus(const TSharedPtr<const SWidget> Widget) const override;
 
 	/** 
@@ -695,15 +678,13 @@ protected:
 	 */
 	void TickWindowAndChildren( TSharedRef<SWindow> WindowToTick );
 
-	/**
-	 * Draws Slate windows.  Should only be called by the application's main loop or renderer.
-	 */
+	/** Draws Slate windows. Should only be called by the application's main loop or renderer. */
 	void DrawWindows();
 
 	/**
 	 * Draws slate windows, optionally only drawing the passed in window
 	 */
-	void PrivateDrawWindows( TSharedPtr<SWindow> DrawOnlyThisWindow = NULL );
+	void PrivateDrawWindows( TSharedPtr<SWindow> DrawOnlyThisWindow = nullptr );
 
 	/**
 	 * Pre-pass step before drawing windows to compute geometry size and reshape autosized windows
@@ -731,7 +712,6 @@ public:
 	 *
 	 * @param  InMouseEvent  Mouse event
 	 * @param  bIsSynthetic  True when the even is synthesized by slate.
-	 *
 	 * @return  Was this event handled by the Slate application?
 	 */
 	bool ProcessMouseMoveEvent( FPointerEvent& MouseEvent, bool bIsSynthetic = false );
@@ -742,7 +722,6 @@ public:
 	 * @param  PlatformWindow  The platform window the event originated from, used to set focus at the platform level. 
 	 *                         If Invalid the Mouse event will work but there will be no effect on the platform.
 	 * @param  InMouseEvent    Mouse event
-	 *
 	 * @return  Was this event handled by the Slate application?
 	 */
 	bool ProcessMouseButtonDownEvent(const TSharedPtr< FGenericWindow >& PlatformWindow, FPointerEvent& InMouseEvent);
@@ -751,7 +730,6 @@ public:
 	 * Called by the native application in response to a mouse button release. Routs the event to Slate Widgets.
 	 *
 	 * @param  InMouseEvent  Mouse event
-	 *
 	 * @return  Was this event handled by the Slate application?
 	 */
 	bool ProcessMouseButtonUpEvent( FPointerEvent& MouseEvent );
@@ -760,7 +738,6 @@ public:
 	 * Called by the native application in response to a mouse release. Routs the event to Slate Widgets.
 	 *
 	 * @param  InMouseEvent  Mouse event
-	 *
 	 * @return  Was this event handled by the Slate application?
 	 */
 	bool ProcessMouseButtonDoubleClickEvent( const TSharedPtr< FGenericWindow >& PlatformWindow, FPointerEvent& InMouseEvent );
@@ -770,7 +747,6 @@ public:
 	 *
 	 * @param  InWheelEvent    Mouse wheel event details
 	 * @param  InGestureEvent  Optional gesture event details
-	 *
 	 * @return  Was this event handled by the Slate application?
 	 */
 	bool ProcessMouseWheelOrGestureEvent( FPointerEvent& InWheelEvent, const FPointerEvent* InGestureEvent );
@@ -779,7 +755,6 @@ public:
 	 * Called when a character is entered
 	 *
 	 * @param  InCharacterEvent  Character event
-	 *
 	 * @return  Was this event handled by the Slate application?
 	 */
 	bool ProcessKeyCharEvent( FCharacterEvent& InCharacterEvent );
@@ -788,7 +763,6 @@ public:
 	 * Called when a key is pressed
 	 *
 	 * @param  InKeyEvent  Keyb event
-	 *
 	 * @return  Was this event handled by the Slate application?
 	 */
 	bool ProcessKeyDownEvent( FKeyEvent& InKeyEvent );
@@ -797,7 +771,6 @@ public:
 	 * Called when a key is released
 	 *
 	 * @param  InKeyEvent  Key event
-	 *
 	 * @return  Was this event handled by the Slate application?
 	 */
 	bool ProcessKeyUpEvent( FKeyEvent& InKeyEvent );
@@ -806,7 +779,6 @@ public:
 	 * Called when a analog input values change
 	 *
 	 * @param  InAnalogInputEvent Analog input event
-	 *
 	 * @return  Was this event handled by the Slate application?
 	 */
 	bool ProcessAnalogInputEvent(FAnalogInputEvent& InAnalogInputEvent);
@@ -816,7 +788,6 @@ public:
 	 *
 	 * @param WindowEntered  The window that was entered by the drag and drop
 	 * @param DragDropEvent  Describes the mouse state (position, pressed buttons, etc) and associated payload
-	 *
 	 * @return true if the drag enter was handled and can be processed by some widget in this window; false otherwise
 	 */
 	bool ProcessDragEnterEvent( TSharedRef<SWindow> WindowEntered, FDragDropEvent& DragDropEvent );
@@ -853,7 +824,6 @@ public:
 	 * Called by the native application in response to an activation or deactivation. 
 	 *
 	 * @param ActivateEvent Information about the window activation/deactivation
-	 *
 	 * @return  Was this event handled by the Slate application?
 	 */
 	bool ProcessWindowActivatedEvent( const FWindowActivateEvent& ActivateEvent );
@@ -888,9 +858,7 @@ public:
 	 */
 	void ProcessReply(const FWidgetPath& CurrentEventPath, const FReply TheReply, const FWidgetPath* WidgetsUnderMouse, const FPointerEvent* InMouseEvent, uint32 UserIndex = 0);
 	
-	/**
-	 * Bubble a request for which cursor to display for widgets under the mouse or the widget that captured the mouse.
-	 */
+	/** Bubble a request for which cursor to display for widgets under the mouse or the widget that captured the mouse. */
 	void QueryCursor();
 
 	/**
@@ -908,9 +876,7 @@ public:
 	 */
 	void SpawnToolTip( const TSharedRef<IToolTip>& InToolTip, const FVector2D& InSpawnLocation );
 
-	/**
-	 * Closes the open tool-tip, if a tool-tip is open
-	 */
+	/** Closes the open tool-tip, if a tool-tip is open */
 	void CloseToolTip();
 
 	/**
@@ -982,154 +948,99 @@ public:
 		return FSlateWindowHelper::FindPathToWidget(GetInteractiveTopLevelWindows(), InWidget, OutWidgetPath, VisibilityFilter);
 	}
 
-	virtual const double GetCurrentTime( ) const override
+	virtual const double GetCurrentTime() const override
 	{
 		return CurrentTime;
 	}
 
-	virtual TSharedPtr<SWindow> GetActiveTopLevelWindow( ) const override;
-
-	virtual const FSlateBrush* GetAppIcon( ) const override;
+	virtual TSharedPtr<SWindow> GetActiveTopLevelWindow() const override;
+	virtual const FSlateBrush* GetAppIcon() const override;
 
-	virtual float GetApplicationScale( ) const override
+	virtual float GetApplicationScale() const override
 	{
 		return Scale;
 	}
 
-	virtual FVector2D GetCursorPos( ) const override;
+	virtual FVector2D GetCursorPos() const override;
+	virtual FVector2D GetLastCursorPos() const override;
+	virtual FVector2D GetCursorSize() const override;	
 
-	virtual FVector2D GetLastCursorPos( ) const override;
-
-	virtual FVector2D GetCursorSize( ) const override;	
-
-	virtual bool GetSoftwareCursorAvailable( ) const override
+	virtual bool GetSoftwareCursorAvailable() const override
 	{
 		return bSoftwareCursorAvailable;
 	}
 
-	virtual EVisibility GetSoftwareCursorVis( ) const override;	
-	
-	virtual TSharedPtr<SWidget> GetKeyboardFocusedWidget( ) const override;
+	virtual EVisibility GetSoftwareCursorVis() const override;	
+	virtual TSharedPtr<SWidget> GetKeyboardFocusedWidget() const override;
 
 protected:
+
 	virtual TSharedPtr< SWidget > GetMouseCaptorImpl() const override;
 
 public:
-	virtual bool HasAnyMouseCaptor( ) const override;
 
-	virtual FSlateRect GetPreferredWorkArea( ) const override;
+	// FSlateApplicationBase interface
 
+	virtual bool HasAnyMouseCaptor() const override;
+	virtual FSlateRect GetPreferredWorkArea() const override;
 	virtual bool HasFocusedDescendants( const TSharedRef<const SWidget>& Widget ) const override;
-
-	virtual bool IsExternalUIOpened() override
-	{
-		return bIsExternalUIOpened;
-	}
-
+	virtual bool IsExternalUIOpened() override;
 	virtual FWidgetPath LocateWindowUnderMouse( FVector2D ScreenspaceMouseCoordinate, const TArray<TSharedRef<SWindow>>& Windows, bool bIgnoreEnabledStatus = false ) override;
-
 	virtual bool IsWindowHousingInteractiveTooltip(const TSharedRef<const SWindow>& WindowToTest) const override;
-
 	virtual TSharedRef<SWidget> MakeImage( const TAttribute<const FSlateBrush*>& Image, const TAttribute<FSlateColor>& Color, const TAttribute<EVisibility>& Visibility ) const override;
-
 	virtual TSharedRef<SWidget> MakeWindowTitleBar( const TSharedRef<SWindow>& Window, const TSharedPtr<SWidget>& CenterContent, EHorizontalAlignment CenterContentAlignment, TSharedPtr<IWindowTitleBar>& OutTitleBar ) const override;
-
 	virtual TSharedRef<IToolTip> MakeToolTip( const TAttribute<FString>& ToolTipString ) override;
-
 	virtual TSharedRef<IToolTip> MakeToolTip( const TAttribute<FText>& ToolTipText ) override;
-
 	virtual TSharedRef<IToolTip> MakeToolTip( const FText& ToolTipText ) override;
-
 	virtual void RequestDestroyWindow( TSharedRef<SWindow> WindowToDestroy ) override;
-
 	virtual bool SetKeyboardFocus( const FWidgetPath& InFocusPath, const EFocusCause InCause ) override;
-
 	virtual bool SetUserFocus(const uint32 InUserIndex, const FWidgetPath& InFocusPath, const EFocusCause InCause) override;
 
-	// End FSlateApplicationBase interface
-
 public:
 
-	// Begin FGenericApplicationMessageHandler interface
+	// FGenericApplicationMessageHandler interface
 
 	virtual bool ShouldProcessUserInputMessages( const TSharedPtr< FGenericWindow >& PlatformWindow ) const override;
-
 	virtual bool OnKeyChar( const TCHAR Character, const bool IsRepeat ) override;
-
 	virtual bool OnKeyDown( const int32 KeyCode, const uint32 CharacterCode, const bool IsRepeat ) override;
-
 	virtual bool OnKeyUp( const int32 KeyCode, const uint32 CharacterCode, const bool IsRepeat ) override;
-
 	virtual bool OnMouseDown( const TSharedPtr< FGenericWindow >& PlatformWindow, const EMouseButtons::Type Button ) override;
-
 	virtual bool OnMouseUp( const EMouseButtons::Type Button ) override;
-
 	virtual bool OnMouseDoubleClick( const TSharedPtr< FGenericWindow >& PlatformWindow, const EMouseButtons::Type Button ) override;
-
 	virtual bool OnMouseWheel( const float Delta ) override;
-
 	virtual bool OnMouseMove() override;
-
 	virtual bool OnRawMouseMove( const int32 X, const int32 Y ) override;
-
 	virtual bool OnCursorSet() override;
-
 	virtual bool OnControllerAnalog( EControllerButtons::Type Button, int32 ControllerId, float AnalogValue ) override;
-
 	virtual bool OnControllerButtonPressed( EControllerButtons::Type Button, int32 ControllerId, bool IsRepeat ) override;
-
 	virtual bool OnControllerButtonReleased( EControllerButtons::Type Button, int32 ControllerId, bool IsRepeat ) override;
-
 	virtual bool OnTouchGesture( EGestureEvent::Type GestureType, const FVector2D& Delta, float WheelDelta ) override;
-	
 	virtual bool OnTouchStarted( const TSharedPtr< FGenericWindow >& PlatformWindow, const FVector2D& Location, int32 TouchIndex, int32 ControllerId ) override;
-
 	virtual bool OnTouchMoved( const FVector2D& Location, int32 TouchIndex, int32 ControllerId ) override;
-
 	virtual bool OnTouchEnded( const FVector2D& Location, int32 TouchIndex, int32 ControllerId ) override;
-
 	virtual bool OnMotionDetected(const FVector& Tilt, const FVector& RotationRate, const FVector& Gravity, const FVector& Acceleration, int32 ControllerId) override;
-
 	virtual bool OnSizeChanged( const TSharedRef< FGenericWindow >& PlatformWindow, const int32 Width, const int32 Height, bool bWasMinimized = false ) override;
-
 	virtual void OnOSPaint( const TSharedRef< FGenericWindow >& PlatformWindow ) override;
-
 	virtual FWindowSizeLimits GetSizeLimitsForWindow(const TSharedRef<FGenericWindow>& Window) const override;
-
 	virtual void OnResizingWindow( const TSharedRef< FGenericWindow >& PlatformWindow ) override;
-
 	virtual bool BeginReshapingWindow( const TSharedRef< FGenericWindow >& PlatformWindow ) override;
-
 	virtual void FinishedReshapingWindow( const TSharedRef< FGenericWindow >& PlatformWindow ) override;
-
 	virtual void OnMovedWindow( const TSharedRef< FGenericWindow >& PlatformWindow, const int32 X, const int32 Y ) override;
-
 	virtual bool OnWindowActivationChanged( const TSharedRef< FGenericWindow >& PlatformWindow, const EWindowActivation::Type ActivationType ) override;
-
 	virtual bool OnApplicationActivationChanged( const bool IsActive ) override;
-
 	virtual bool OnConvertibleDeviceModeChanged(const EConvertibleLaptopModes NewMode) override;
-
 	virtual EWindowZone::Type GetWindowZoneForPoint( const TSharedRef< FGenericWindow >& PlatformWindow, const int32 X, const int32 Y ) override;
-
 	virtual void OnWindowClose( const TSharedRef< FGenericWindow >& PlatformWindow ) override;
-
 	virtual EDropEffect::Type OnDragEnterText( const TSharedRef< FGenericWindow >& Window, const FString& Text ) override;
-
 	virtual EDropEffect::Type OnDragEnterFiles( const TSharedRef< FGenericWindow >& Window, const TArray< FString >& Files ) override;
 
 	EDropEffect::Type OnDragEnter( const TSharedRef< SWindow >& Window, const TSharedRef<FExternalDragOperation>& DragDropOperation );
 
 	virtual EDropEffect::Type OnDragOver( const TSharedPtr< FGenericWindow >& Window ) override;
-
 	virtual void OnDragLeave( const TSharedPtr< FGenericWindow >& Window ) override;
-
 	virtual EDropEffect::Type OnDragDrop( const TSharedPtr< FGenericWindow >& Window ) override;
-
 	virtual bool OnWindowAction( const TSharedRef< FGenericWindow >& PlatformWindow, const EWindowAction::Type InActionType ) override;
 
-	// End FGenericApplicationMessageHandler interface
-
 private:
 
 	TSharedRef< FGenericWindow > MakeWindow( TSharedRef<SWindow> InSlateWindow, const bool bShowImmediately );
@@ -1159,9 +1070,6 @@ private:
 	int32 GetUserIndexForController(int32 ControllerId) const;
 
 private:
-	//
-	// Cursor Lock
-	// 
 
 	/** Lock the cursor such that it cannot leave the bounds of the specified widget. Null widget implies no cursor lock. */
 	void LockCursor(const TSharedPtr<SWidget>& Widget);
@@ -1215,7 +1123,7 @@ private:
 	/** Application-wide scale for supporting monitors of varying pixel density */
 	float Scale;
 
-	/** The deadzone distance in virtual desktop pixels (a.k.a screen pixels) that the user has to move their finder before it is considered a drag.*/
+	/** The dead zone distance in virtual desktop pixels (a.k.a screen pixels) that the user has to move their finder before it is considered a drag.*/
 	float DragTriggerDistnace;
 
 	/** All the top-level windows owned by this application; they are tracked here in a platform-agnostic way. */
@@ -1307,7 +1215,7 @@ private:
 		/** Call the OnMouseCaptureLost() handler for the widget captured by the specific pointer index */
 		void InformCurrentCaptorOfCaptureLoss(uint32 PointerIndex) const;
 
-		/** A map of pointer idices to weak widget paths for the active mouse captures */
+		/** A map of pointer indices to weak widget paths for the active mouse captures */
 		TMap<uint32, FWeakWidgetPath> PointerIndexToMouseCaptorWeakPathMap;
 	};
 	/** The current mouse captor for the application, if any. */
@@ -1369,7 +1277,7 @@ private:
 		FKey DetectDragButton;
 	} DragDetector;
 
-	/** Support for auto-dismissing popups */
+	/** Support for auto-dismissing pop-ups */
 	FPopupSupport PopupSupport;
 	
 	/** Pointer to the currently registered game viewport widget if any */
