Slate - The IToolTip interface is now told when tooltips are closed so that it can lazily create and destroy the contents of a tooltip widget in custom implementations.

diff --git a/Engine/Source/Runtime/Slate/Private/Framework/Application/SlateApplication.cpp b/Engine/Source/Runtime/Slate/Private/Framework/Application/SlateApplication.cpp
--- a/Engine/Source/Runtime/Slate/Private/Framework/Application/SlateApplication.cpp
+++ b/Engine/Source/Runtime/Slate/Private/Framework/Application/SlateApplication.cpp
@@ -2597,6 +2597,12 @@ void FSlateApplication::CloseToolTip()
 			SourceWidget->OnToolTipClosing();
 		}
 
+		TSharedPtr<IToolTip> StableActiveToolTip = ActiveToolTip.Pin();
+		if ( StableActiveToolTip.IsValid() )
+		{
+			StableActiveToolTip->OnClosed();
+		}
+
 		// Hide the tool tip window.  We don't destroy the window, because we want to reuse it for future tool tips.
 		PinnedToolTipWindow->HideWindow();
 
diff --git a/Engine/Source/Runtime/Slate/Public/Widgets/SToolTip.h b/Engine/Source/Runtime/Slate/Public/Widgets/SToolTip.h
--- a/Engine/Source/Runtime/Slate/Public/Widgets/SToolTip.h
+++ b/Engine/Source/Runtime/Slate/Public/Widgets/SToolTip.h
@@ -70,6 +70,8 @@ public:
 
 	virtual bool IsEmpty( ) const override;
 	virtual bool IsInteractive( ) const override;
+	virtual void OnClosed() override { }
+
 
 	virtual const FText& GetTextTooltip() const
 	{
diff --git a/Engine/Source/Runtime/SlateCore/Public/Widgets/IToolTip.h b/Engine/Source/Runtime/SlateCore/Public/Widgets/IToolTip.h
--- a/Engine/Source/Runtime/SlateCore/Public/Widgets/IToolTip.h
+++ b/Engine/Source/Runtime/SlateCore/Public/Widgets/IToolTip.h
@@ -45,6 +45,11 @@ public:
 	 */
 	virtual bool IsInteractive( ) const = 0;
 
+	/**
+	 * Called when the tooltip widget is closed and the tooltip is no longer needed.
+	 */
+	virtual void OnClosed( ) = 0;
+
 public:
 
 	/** Virtual destructor. */
