Improved the comment about newline characters to state where the handing code can be found

diff --git a/Engine/Source/Runtime/Slate/Private/Framework/Text/TextEditHelper.cpp b/Engine/Source/Runtime/Slate/Private/Framework/Text/TextEditHelper.cpp
--- a/Engine/Source/Runtime/Slate/Private/Framework/Text/TextEditHelper.cpp
+++ b/Engine/Source/Runtime/Slate/Private/Framework/Text/TextEditHelper.cpp
@@ -55,7 +55,7 @@ FReply FTextEditHelper::OnKeyChar( const FCharacterEvent& InCharacterEvent, cons
 		break;
 
 
-		// Newline (Ctrl+Enter - handled correctly outside this function)
+		// Newline (Ctrl+Enter), we handle adding new lines via SMultiLineEditableText::OnEnter rather than processing \n characters
 	case TCHAR( '\n' ):
 		{
 			Reply = FReply::Handled();
