Slate : SFxWidget's HAlign and VAlign settings are now arguments; they were previously misleadingly created as attributes but behaved as attributes.

diff --git a/Engine/Source/Runtime/Slate/Private/Widgets/Layout/SFxWidget.cpp b/Engine/Source/Runtime/Slate/Private/Widgets/Layout/SFxWidget.cpp
--- a/Engine/Source/Runtime/Slate/Private/Widgets/Layout/SFxWidget.cpp
+++ b/Engine/Source/Runtime/Slate/Private/Widgets/Layout/SFxWidget.cpp
@@ -11,12 +11,10 @@ void SFxWidget::Construct( const FArguments& InArgs )
 	VisualOffset = InArgs._VisualOffset;
 	bIgnoreClipping = InArgs._IgnoreClipping;
 	ColorAndOpacity = InArgs._ColorAndOpacity;
-	HAlign = InArgs._HAlign;
-	VAlign = InArgs._VAlign;
 	
 	this->ChildSlot
-	.HAlign(HAlign.Get())
-	.VAlign(VAlign.Get())
+	.HAlign(InArgs._HAlign)
+	.VAlign(InArgs._VAlign)
 	[
 		InArgs._Content.Widget
 	];
diff --git a/Engine/Source/Runtime/Slate/Public/Widgets/Layout/SFxWidget.h b/Engine/Source/Runtime/Slate/Public/Widgets/Layout/SFxWidget.h
--- a/Engine/Source/Runtime/Slate/Public/Widgets/Layout/SFxWidget.h
+++ b/Engine/Source/Runtime/Slate/Public/Widgets/Layout/SFxWidget.h
@@ -40,17 +40,17 @@ public:
 	SLATE_ATTRIBUTE( FLinearColor, ColorAndOpacity )
 
 	/** The horizontal alignment of the child widget */
-	SLATE_ATTRIBUTE( EHorizontalAlignment, HAlign )
+	SLATE_ARGUMENT( EHorizontalAlignment, HAlign )
 
 	/** The vertical alignment of the child widget */
-	SLATE_ATTRIBUTE( EVerticalAlignment, VAlign )
+	SLATE_ARGUMENT( EVerticalAlignment, VAlign )
 
 	/** The content that should be modified. */
 	SLATE_DEFAULT_SLOT( FArguments, Content )
 
 	SLATE_END_ARGS()
 
-	virtual void Construct( const FArguments& InArgs );
+	void Construct( const FArguments& InArgs );
 
 	/** @see VisualOffset */
 	void SetVisualOffset( const TAttribute<FVector2D>& InOffset );
@@ -82,6 +82,4 @@ protected:
 	TAttribute<float> LayoutScale;
 	TAttribute<FVector2D> VisualOffset;
 	TAttribute<bool> bIgnoreClipping;
-	TAttribute<EHorizontalAlignment> HAlign;
-	TAttribute<EVerticalAlignment> VAlign;
 };
\ No newline at end of file
