Slate : Slate's hit testing now correctly handles non-slate windows on top of Slate Windows.
Added platform method IsCursorOverSlateWindow(); requires Mac and Linux implementation.

diff --git a/Engine/Source/Runtime/Core/Private/Windows/WindowsApplication.cpp b/Engine/Source/Runtime/Core/Private/Windows/WindowsApplication.cpp
--- a/Engine/Source/Runtime/Core/Private/Windows/WindowsApplication.cpp
+++ b/Engine/Source/Runtime/Core/Private/Windows/WindowsApplication.cpp
@@ -272,6 +272,33 @@ FModifierKeysState FWindowsApplication::GetModifierKeys() const
 	return FModifierKeysState(bIsLeftShiftDown, bIsRightShiftDown, bIsLeftControlDown, bIsRightControlDown, bIsLeftAltDown, bIsRightAltDown, false, false, bAreCapsLocked); // Win key is ignored
 }
 
+
+static TSharedPtr< FWindowsWindow > FindWindowByHWND(const TArray< TSharedRef< FWindowsWindow > >& WindowsToSearch, HWND HandleToFind)
+{
+	for (int32 WindowIndex = 0; WindowIndex < WindowsToSearch.Num(); ++WindowIndex)
+	{
+		TSharedRef< FWindowsWindow > Window = WindowsToSearch[WindowIndex];
+		if (Window->GetHWnd() == HandleToFind)
+		{
+			return Window;
+		}
+	}
+
+	return TSharedPtr< FWindowsWindow >(NULL);
+}
+
+
+bool FWindowsApplication::IsCursorDirectlyOverSlateWindow() const
+{
+	POINT CursorPos;
+	::GetCursorPos(&CursorPos);
+	HWND hWnd = ::WindowFromPoint(CursorPos);
+	
+	TSharedPtr< FWindowsWindow > SlatWindowUnderCursor = FindWindowByHWND(Windows, hWnd);
+	return SlatWindowUnderCursor.IsValid();
+}
+
+
 void FWindowsApplication::SetCapture( const TSharedPtr< FGenericWindow >& InWindow )
 {
 	if ( InWindow.IsValid() )
@@ -580,21 +607,6 @@ EWindowTitleAlignment::Type FWindowsApplication::GetWindowTitleAlignment() const
 	return EWindowTitleAlignment::Left;
 }
 
-static TSharedPtr< FWindowsWindow > FindWindowByHWND( const TArray< TSharedRef< FWindowsWindow > >& WindowsToSearch, HWND HandleToFind )
-{
-	for (int32 WindowIndex=0; WindowIndex < WindowsToSearch.Num(); ++WindowIndex)
-	{
-		TSharedRef< FWindowsWindow > Window = WindowsToSearch[ WindowIndex ];
-		if ( Window->GetHWnd() == HandleToFind )
-		{
-			return Window;
-		}
-	}
-
-	return TSharedPtr< FWindowsWindow >( NULL );
-}
-
-
 /** All WIN32 messages sent to our app go here; this method simply passes them on */
 LRESULT CALLBACK FWindowsApplication::AppWndProc(HWND hwnd, uint32 msg, WPARAM wParam, LPARAM lParam)
 {
diff --git a/Engine/Source/Runtime/Core/Public/GenericPlatform/GenericApplication.h b/Engine/Source/Runtime/Core/Public/GenericPlatform/GenericApplication.h
--- a/Engine/Source/Runtime/Core/Public/GenericPlatform/GenericApplication.h
+++ b/Engine/Source/Runtime/Core/Public/GenericPlatform/GenericApplication.h
@@ -401,6 +401,9 @@ public:
 
 	virtual FModifierKeysState GetModifierKeys() const  { return FModifierKeysState(); }
 
+	/** @return true if the system cursor is currently directly over a slate window. */
+	virtual bool IsCursorDirectlyOverSlateWindow() const { return true; }
+
 	virtual void SetHighPrecisionMouseMode( const bool Enable, const TSharedPtr< FGenericWindow >& InWindow ) { };
 
 	virtual bool IsUsingHighPrecisionMouseMode() const { return false; }
diff --git a/Engine/Source/Runtime/Core/Public/Windows/WindowsApplication.h b/Engine/Source/Runtime/Core/Public/Windows/WindowsApplication.h
--- a/Engine/Source/Runtime/Core/Public/Windows/WindowsApplication.h
+++ b/Engine/Source/Runtime/Core/Public/Windows/WindowsApplication.h
@@ -249,6 +249,8 @@ public:
 
 	virtual FModifierKeysState GetModifierKeys() const override;
 
+	virtual bool IsCursorDirectlyOverSlateWindow() const override;
+
 	virtual FPlatformRect GetWorkArea( const FPlatformRect& CurrentWindow ) const override;
 
 	virtual bool TryCalculatePopupWindowPosition( const FPlatformRect& InAnchor, const FVector2D& InSize, const EPopUpOrientation::Type Orientation, /*OUT*/ FVector2D* const CalculatedPopUpPosition ) const override;
diff --git a/Engine/Source/Runtime/Slate/Private/Framework/Application/SlateApplication.cpp b/Engine/Source/Runtime/Slate/Private/Framework/Application/SlateApplication.cpp
--- a/Engine/Source/Runtime/Slate/Private/Framework/Application/SlateApplication.cpp
+++ b/Engine/Source/Runtime/Slate/Private/Framework/Application/SlateApplication.cpp
@@ -799,31 +799,35 @@ FWidgetPath FSlateApplication::LocateWindowUnderMouse( FVector2D ScreenspaceMous
 	bool bPrevWindowWasModal = false;
 	FArrangedChildren OutWidgetPath(EVisibility::Visible);
 
-	for( int32 WindowIndex=Windows.Num()-1; WindowIndex >= 0 && OutWidgetPath.Num() == 0; --WindowIndex )
+	if (PlatformApplication->IsCursorDirectlyOverSlateWindow())
 	{
-		const TSharedRef<SWindow>& Window = Windows[ WindowIndex ];
 
-		// Hittest the window's children first.
-		FWidgetPath ResultingPath = LocateWindowUnderMouse( ScreenspaceMouseCoordinate, Window->GetChildWindows(), bIgnoreEnabledStatus );
-		if ( ResultingPath.IsValid() )
+		for (int32 WindowIndex = Windows.Num() - 1; WindowIndex >= 0 && OutWidgetPath.Num() == 0; --WindowIndex)
 		{
-			return ResultingPath;
-		}
+			const TSharedRef<SWindow>& Window = Windows[WindowIndex];
+
+			// Hittest the window's children first.
+			FWidgetPath ResultingPath = LocateWindowUnderMouse(ScreenspaceMouseCoordinate, Window->GetChildWindows(), bIgnoreEnabledStatus);
+			if (ResultingPath.IsValid())
+			{
+				return ResultingPath;
+			}
 
-		// If none of the children were hit, hittest the parent.
+			// If none of the children were hit, hittest the parent.
 
-		// Only accept input if the current window accepts input and the current window is not under a modal window or an interactive tooltip
-		
-		if (!bPrevWindowWasModal)
-		{
-			FWidgetPath PathToLocatedWidget = LocateWidgetInWindow(ScreenspaceMouseCoordinate, Window, bIgnoreEnabledStatus);
-			if (PathToLocatedWidget.IsValid())
+			// Only accept input if the current window accepts input and the current window is not under a modal window or an interactive tooltip
+
+			if (!bPrevWindowWasModal)
 			{
-				return PathToLocatedWidget;
+				FWidgetPath PathToLocatedWidget = LocateWidgetInWindow(ScreenspaceMouseCoordinate, Window, bIgnoreEnabledStatus);
+				if (PathToLocatedWidget.IsValid())
+				{
+					return PathToLocatedWidget;
+				}
 			}
 		}
 	}
-
+	
 	return FWidgetPath();
 }
 
