Fixed issue with output log sometimes not auto-scrolling to bottom.

diff --git a/Engine/Source/Developer/OutputLog/Private/SOutputLog.cpp b/Engine/Source/Developer/OutputLog/Private/SOutputLog.cpp
--- a/Engine/Source/Developer/OutputLog/Private/SOutputLog.cpp
+++ b/Engine/Source/Developer/OutputLog/Private/SOutputLog.cpp
@@ -665,6 +665,7 @@ void SOutputLog::Construct( const FArguments& InArgs )
 		.Marshaller(MessagesTextMarshaller)
 		.IsReadOnly(true)
 		.AlwaysShowScrollbars(true)
+		.OnVScrollBarUserScrolled(this, &SOutputLog::OnUserScrolled)
 		.ContextMenuExtender(this, &SOutputLog::ExtendTextBoxMenu);
 
 	ChildSlot
@@ -692,6 +693,7 @@ void SOutputLog::Construct( const FArguments& InArgs )
 	
 	GLog->AddOutputDevice(this);
 
+	bIsUserScrolled = false;
 	RequestForceScroll();
 }
 END_SLATE_FUNCTION_BUILD_OPTIMIZATION
@@ -769,8 +771,7 @@ void SOutputLog::Serialize( const TCHAR* V, ELogVerbosity::Type Verbosity, const
 	if ( MessagesTextMarshaller->AppendMessage(V, Verbosity, Category) )
 	{
 		// Don't scroll to the bottom automatically when the user is scrolling the view or has scrolled it away from the bottom.
-		const float DistanceFromBottom = MessagesTextBox->GetVScrollBar()->DistanceFromBottom();
-		if( DistanceFromBottom <= KINDA_SMALL_NUMBER )
+		if( !bIsUserScrolled )
 		{
 			MessagesTextBox->ScrollTo(FTextLocation(MessagesTextMarshaller->GetNumMessages() - 1));
 		}
@@ -799,6 +800,12 @@ void SOutputLog::OnClearLog()
 
 	MessagesTextMarshaller->ClearMessages();
 	MessagesTextBox->Refresh();
+	bIsUserScrolled = false;
+}
+
+void SOutputLog::OnUserScrolled(float ScrollOffset)
+{
+	bIsUserScrolled = !FMath::IsNearlyEqual(ScrollOffset, 1.0f);
 }
 
 bool SOutputLog::CanClearLog() const
@@ -816,5 +823,6 @@ void SOutputLog::RequestForceScroll()
 	if(MessagesTextMarshaller->GetNumMessages() > 0)
 	{
 		MessagesTextBox->ScrollTo(FTextLocation(MessagesTextMarshaller->GetNumMessages() - 1));
+		bIsUserScrolled = false;
 	}
 }
diff --git a/Engine/Source/Developer/OutputLog/Private/SOutputLog.h b/Engine/Source/Developer/OutputLog/Private/SOutputLog.h
--- a/Engine/Source/Developer/OutputLog/Private/SOutputLog.h
+++ b/Engine/Source/Developer/OutputLog/Private/SOutputLog.h
@@ -171,6 +171,11 @@ private:
 	 */
 	void OnClearLog();
 
+	/**
+	 * Called when the user scrolls the log window vertically
+	 */
+	void OnUserScrolled(float ScrollOffset);
+
 	/**
 	 * Called to determine whether delete all is currently a valid command
 	 */
@@ -187,4 +192,7 @@ private:
 
 	/** The editable text showing all log messages */
 	TSharedPtr< SMultiLineEditableTextBox > MessagesTextBox;
+
+	/** True if the user has scrolled the window upwards */
+	bool bIsUserScrolled;
 };
diff --git a/Engine/Source/Runtime/Slate/Private/Widgets/Input/SMultiLineEditableTextBox.cpp b/Engine/Source/Runtime/Slate/Private/Widgets/Input/SMultiLineEditableTextBox.cpp
--- a/Engine/Source/Runtime/Slate/Private/Widgets/Input/SMultiLineEditableTextBox.cpp
+++ b/Engine/Source/Runtime/Slate/Private/Widgets/Input/SMultiLineEditableTextBox.cpp
@@ -98,6 +98,8 @@ void SMultiLineEditableTextBox::Construct( const FArguments& InArgs )
 					.AutoWrapText(InArgs._AutoWrapText)
 					.HScrollBar(HScrollBar)
 					.VScrollBar(VScrollBar)
+					.OnHScrollBarUserScrolled(InArgs._OnHScrollBarUserScrolled)
+					.OnVScrollBarUserScrolled(InArgs._OnVScrollBarUserScrolled)
 					.ModiferKeyForNewLine(InArgs._ModiferKeyForNewLine)
 				]
 
diff --git a/Engine/Source/Runtime/Slate/Private/Widgets/Text/SMultiLineEditableText.cpp b/Engine/Source/Runtime/Slate/Private/Widgets/Text/SMultiLineEditableText.cpp
--- a/Engine/Source/Runtime/Slate/Private/Widgets/Text/SMultiLineEditableText.cpp
+++ b/Engine/Source/Runtime/Slate/Private/Widgets/Text/SMultiLineEditableText.cpp
@@ -266,6 +266,8 @@ void SMultiLineEditableText::Construct( const FArguments& InArgs )
 	OnTextChanged = InArgs._OnTextChanged;
 	OnTextCommitted = InArgs._OnTextCommitted;
 	OnCursorMoved = InArgs._OnCursorMoved;
+	OnHScrollBarUserScrolled = InArgs._OnHScrollBarUserScrolled;
+	OnVScrollBarUserScrolled = InArgs._OnVScrollBarUserScrolled;
 
 	Marshaller = InArgs._Marshaller;
 	if (!Marshaller.IsValid())
@@ -456,11 +458,13 @@ void SMultiLineEditableText::ForceRefreshTextLayout(const FText& CurrentText)
 void SMultiLineEditableText::OnHScrollBarMoved(const float InScrollOffsetFraction)
 {
 	ScrollOffset.X = FMath::Clamp<float>(InScrollOffsetFraction, 0.0, 1.0) * GetDesiredSize().X;
+	OnHScrollBarUserScrolled.ExecuteIfBound(InScrollOffsetFraction);
 }
 
 void SMultiLineEditableText::OnVScrollBarMoved(const float InScrollOffsetFraction)
 {
 	ScrollOffset.Y = FMath::Clamp<float>(InScrollOffsetFraction, 0.0, 1.0) * GetDesiredSize().Y;
+	OnVScrollBarUserScrolled.ExecuteIfBound(InScrollOffsetFraction);
 }
 
 FReply SMultiLineEditableText::OnFocusReceived( const FGeometry& MyGeometry, const FFocusEvent& InFocusEvent )
@@ -2347,7 +2351,12 @@ FReply SMultiLineEditableText::OnMouseWheel( const FGeometry& MyGeometry, const
 		const float ScrollMax = ContentSize - MyGeometry.Size.Y;
 		ScrollOffset.Y = FMath::Clamp(ScrollOffset.Y, ScrollMin, ScrollMax);
 
-		return (PreviousScrollOffset != ScrollOffset.Y) ? FReply::Handled() : FReply::Unhandled();
+		if (PreviousScrollOffset != ScrollOffset.Y)
+		{
+			const float ScrollbarOffset = (ScrollMax != 0.0f) ? ScrollOffset.Y / ScrollMax : 0.0f;
+			OnVScrollBarUserScrolled.ExecuteIfBound(ScrollbarOffset);
+			return FReply::Handled();
+		}
 	}
 
 	return FReply::Unhandled();
diff --git a/Engine/Source/Runtime/Slate/Public/Widgets/Input/SMultiLineEditableTextBox.h b/Engine/Source/Runtime/Slate/Public/Widgets/Input/SMultiLineEditableTextBox.h
--- a/Engine/Source/Runtime/Slate/Public/Widgets/Input/SMultiLineEditableTextBox.h
+++ b/Engine/Source/Runtime/Slate/Public/Widgets/Input/SMultiLineEditableTextBox.h
@@ -109,6 +109,12 @@ public:
 		/** Called whenever the text is committed.  This happens when the user presses enter or the text box loses focus. */
 		SLATE_EVENT( FOnTextCommitted, OnTextCommitted )
 
+		/** Called whenever the horizontal scrollbar is moved by the user */
+		SLATE_EVENT( FOnUserScrolled, OnHScrollBarUserScrolled )
+
+		/** Called whenever the vertical scrollbar is moved by the user */
+		SLATE_EVENT( FOnUserScrolled, OnVScrollBarUserScrolled )
+
 		/** Called when the cursor is moved within the text area */
 		SLATE_EVENT( SMultiLineEditableText::FOnCursorMoved, OnCursorMoved )
 
diff --git a/Engine/Source/Runtime/Slate/Public/Widgets/Text/SMultiLineEditableText.h b/Engine/Source/Runtime/Slate/Public/Widgets/Text/SMultiLineEditableText.h
--- a/Engine/Source/Runtime/Slate/Public/Widgets/Text/SMultiLineEditableText.h
+++ b/Engine/Source/Runtime/Slate/Public/Widgets/Text/SMultiLineEditableText.h
@@ -78,6 +78,12 @@ public:
 		/** Called whenever the text is committed.  This happens when the user presses enter or the text box loses focus. */
 		SLATE_EVENT(FOnTextCommitted, OnTextCommitted)
 
+		/** Called whenever the horizontal scrollbar is moved by the user */
+		SLATE_EVENT(FOnUserScrolled, OnHScrollBarUserScrolled)
+
+		/** Called whenever the vertical scrollbar is moved by the user */
+		SLATE_EVENT(FOnUserScrolled, OnVScrollBarUserScrolled)
+
 		/** Called when the cursor is moved within the text area */
 		SLATE_EVENT(FOnCursorMoved, OnCursorMoved)
 
@@ -633,6 +639,12 @@ private:
 	/** The vertical scroll bar widget */
 	TSharedPtr< SScrollBar > VScrollBar;
 
+	/** Called whenever the horizontal scrollbar is moved by the user */
+	FOnUserScrolled OnHScrollBarUserScrolled;
+
+	/** Called whenever the vertical scrollbar is moved by the user */
+	FOnUserScrolled OnVScrollBarUserScrolled;
+
 	/** Menu extender for right-click context menu */
 	TSharedPtr<FExtender> MenuExtender;
 
