Slate - Adding OnOpening to IToolTip.

diff --git a/Engine/Source/Runtime/Slate/Private/Framework/Application/SlateApplication.cpp b/Engine/Source/Runtime/Slate/Private/Framework/Application/SlateApplication.cpp
--- a/Engine/Source/Runtime/Slate/Private/Framework/Application/SlateApplication.cpp
+++ b/Engine/Source/Runtime/Slate/Private/Framework/Application/SlateApplication.cpp
@@ -2712,12 +2712,20 @@ void FSlateApplication::UpdateToolTip( bool AllowSpawningOfNewToolTips )
 			TooltipVisualizerPtr.Pin()->OnVisualizeTooltip( nullptr );
 		}
 
+		// Notify the new tooltip that it's about to be opened.
+		if ( NewToolTip.IsValid() )
+		{
+			NewToolTip->OnOpening();
+		}
+
+		TSharedPtr<SWidget> NewToolTipWidget = NewToolTip.IsValid() ? NewToolTip->AsWidget() : TSharedPtr<SWidget>();
+
 		bool bOnVisualizeTooltipHandled = false;
 		// Some widgets might want to provide an alternative Tooltip Handler.
 		for ( int32 WidgetIndex=WidgetsToQueryForToolTip.Widgets.Num()-1; !bOnVisualizeTooltipHandled && WidgetIndex >= 0; --WidgetIndex )
 		{
 			const FArrangedWidget& CurWidgetGeometry = WidgetsToQueryForToolTip.Widgets[WidgetIndex];
-			bOnVisualizeTooltipHandled = CurWidgetGeometry.Widget->OnVisualizeTooltip( NewToolTip.IsValid() ? NewToolTip->AsWidget() : TSharedPtr<SWidget>() );
+			bOnVisualizeTooltipHandled = CurWidgetGeometry.Widget->OnVisualizeTooltip( NewToolTipWidget );
 			if (bOnVisualizeTooltipHandled)
 			{
 				// Someone is taking care of visualizing this tooltip
diff --git a/Engine/Source/Runtime/Slate/Public/Widgets/SToolTip.h b/Engine/Source/Runtime/Slate/Public/Widgets/SToolTip.h
--- a/Engine/Source/Runtime/Slate/Public/Widgets/SToolTip.h
+++ b/Engine/Source/Runtime/Slate/Public/Widgets/SToolTip.h
@@ -70,6 +70,7 @@ public:
 
 	virtual bool IsEmpty( ) const override;
 	virtual bool IsInteractive( ) const override;
+	virtual void OnOpening() override { }
 	virtual void OnClosed() override { }
 
 
diff --git a/Engine/Source/Runtime/SlateCore/Public/Widgets/IToolTip.h b/Engine/Source/Runtime/SlateCore/Public/Widgets/IToolTip.h
--- a/Engine/Source/Runtime/SlateCore/Public/Widgets/IToolTip.h
+++ b/Engine/Source/Runtime/SlateCore/Public/Widgets/IToolTip.h
@@ -45,6 +45,11 @@ public:
 	 */
 	virtual bool IsInteractive( ) const = 0;
 
+	/**
+	 * Called when the tooltip widget is about to be requested for opening.
+	 */
+	virtual void OnOpening( ) = 0;
+
 	/**
 	 * Called when the tooltip widget is closed and the tooltip is no longer needed.
 	 */
diff --git a/Engine/Source/Runtime/UMG/Private/Components/Widget.cpp b/Engine/Source/Runtime/UMG/Private/Components/Widget.cpp
--- a/Engine/Source/Runtime/UMG/Private/Components/Widget.cpp
+++ b/Engine/Source/Runtime/UMG/Private/Components/Widget.cpp
@@ -78,6 +78,10 @@ public:
 		CachedToolTip.Reset();
 	}
 
+	virtual void OnOpening() override
+	{
+	}
+
 public:
 	UWidget::FGetWidget ToolTipWidgetDelegate;
 
