Fix to CL 2390838 based on feedback from NickA.

diff --git a/Engine/Source/Runtime/Slate/Private/Widgets/Views/SHeaderRow.cpp b/Engine/Source/Runtime/Slate/Private/Widgets/Views/SHeaderRow.cpp
--- a/Engine/Source/Runtime/Slate/Private/Widgets/Views/SHeaderRow.cpp
+++ b/Engine/Source/Runtime/Slate/Private/Widgets/Views/SHeaderRow.cpp
@@ -53,12 +53,12 @@ public:
 
 		if (Column.HeaderContent.Widget == SNullWidget::NullWidget)
 		{
-			if (!Column.DefaultText.IsBound() && Column.DefaultText.Get().IsEmpty())
+			if (!Column.DefaultText.IsSet())
 			{
 				LabelText = FText::FromString( Column.ColumnId.ToString() + TEXT("[LabelMissing]") );
 			}
 
-			if (!Column.DefaultTooltip.IsBound() && Column.DefaultTooltip.Get().IsEmpty())
+			if (!Column.DefaultTooltip.IsSet())
 			{
 				TooltipText = LabelText;
 			}
