Enabled shared precompiled header for Slate module (SlateBasics.h)
- Some compile time performance gains (20 modules use Slate as their primary included module.)
- This change is somewhat experimental, we'll try it out for awhile and see how it goes

diff --git a/Engine/Source/Runtime/Slate/Slate.Build.cs b/Engine/Source/Runtime/Slate/Slate.Build.cs
--- a/Engine/Source/Runtime/Slate/Slate.Build.cs
+++ b/Engine/Source/Runtime/Slate/Slate.Build.cs
@@ -6,7 +6,7 @@ public class Slate : ModuleRules
 {
 	public Slate(TargetInfo Target)
 	{
-//		SharedPCHHeaderFile = "Runtime/Slate/Public/Slate.h";
+		SharedPCHHeaderFile = "Runtime/Slate/Public/SlateBasics.h";
 
 		PublicDependencyModuleNames.AddRange(
 			new string[] { 
