---- Merging with SlateDev branch ----

diff --git a/Engine/Source/Editor/UnrealEd/Private/FeedbackContextEditor.cpp b/Engine/Source/Editor/UnrealEd/Private/FeedbackContextEditor.cpp
--- a/Engine/Source/Editor/UnrealEd/Private/FeedbackContextEditor.cpp
+++ b/Engine/Source/Editor/UnrealEd/Private/FeedbackContextEditor.cpp
@@ -42,6 +42,9 @@ public:
 		OnCancelClickedDelegate = InArgs._OnCancelClickedDelegate;
 		ScopeStack = InArgs._ScopeStack;
 
+		// This is a temporary widget that needs to be updated over its entire lifespan => has an active tick registered for its entire lifespan
+		RegisterActiveTick( 0.f, FWidgetActiveTickDelegate::CreateSP( this, &SSlowTaskWidget::UpdateProgress ) );
+
 		TSharedRef<SVerticalBox> VerticalBox = SNew(SVerticalBox)
 
 			// Construct the main progress bar and text
@@ -131,12 +134,15 @@ public:
 		UpdateDynamicProgressBars();
 	}
 
-	virtual void Tick(const FGeometry& AllottedGeometry, const double InCurrentTime, const float InDeltaTime) override
+private:
+
+	/** Active tick to update the progress bars */
+	EActiveTickReturnType UpdateProgress(double InCurrentTime, float InDeltaTime)
 	{
 		UpdateDynamicProgressBars();
-	}
 
-private:
+		return EActiveTickReturnType::KeepTicking;
+	}
 
 	/** Updates the dynamic progress bars for this widget */
 	void UpdateDynamicProgressBars()
diff --git a/Engine/Source/Editor/UnrealEd/Private/SEditorViewport.cpp b/Engine/Source/Editor/UnrealEd/Private/SEditorViewport.cpp
--- a/Engine/Source/Editor/UnrealEd/Private/SEditorViewport.cpp
+++ b/Engine/Source/Editor/UnrealEd/Private/SEditorViewport.cpp
@@ -61,6 +61,11 @@ void SEditorViewport::Construct( const FArguments& InArgs )
 	ViewportWidget->SetViewportInterface(SceneViewport.ToSharedRef());
 	Client = ViewportClient;
 
+	if ( Client->IsRealtime() )
+	{
+		ActiveTickHandle = RegisterActiveTick( 0.f, FWidgetActiveTickDelegate::CreateSP( this, &SEditorViewport::EnsureTickWhileRealtime ) );
+	}
+
 	CommandList = MakeShareable( new FUICommandList );
 	// Ensure the commands are registered
 	FEditorViewportCommands::Register();
@@ -323,7 +328,20 @@ EVisibility SEditorViewport::OnGetViewportContentVisibility() const
 
 void SEditorViewport::OnToggleRealtime()
 {
-	 Client->SetRealtime( ! Client->IsRealtime() );
+	if (Client->IsRealtime())
+	{
+		Client->SetRealtime( false );
+		if ( ActiveTickHandle.IsValid() )
+		{
+			UnRegisterActiveTick( ActiveTickHandle.Pin().ToSharedRef() );
+		}
+		
+	}
+	else
+	{
+		Client->SetRealtime( true );
+		ActiveTickHandle = RegisterActiveTick( 0.f, FWidgetActiveTickDelegate::CreateSP( this, &SEditorViewport::EnsureTickWhileRealtime ) );
+	}
 }
 
 void SEditorViewport::OnToggleStats()
@@ -334,7 +352,11 @@ void SEditorViewport::OnToggleStats()
 	if( !bIsEnabled )
 	{
 		// We cannot show stats unless realtime rendering is enabled
-		 Client->SetRealtime( true );
+		if ( !Client->IsRealtime() )
+		{
+			Client->SetRealtime( true );
+			ActiveTickHandle = RegisterActiveTick( 0.f, FWidgetActiveTickDelegate::CreateSP( this, &SEditorViewport::EnsureTickWhileRealtime ) );
+		}
 
 		 // let the user know how they can enable stats via the console
 		 FNotificationInfo Info(LOCTEXT("StatsEnableHint", "Stats display can be toggled via the STAT [type] console command"));
@@ -496,4 +518,9 @@ bool SEditorViewport::OnIsSurfaceSnapEnabled()
 	return GetDefault<ULevelEditorViewportSettings>()->SnapToSurface.bEnabled;
 }
 
+EActiveTickReturnType SEditorViewport::EnsureTickWhileRealtime( double InCurrentTime, float InDeltaTime )
+{
+	return EActiveTickReturnType::KeepTicking;
+}
+
 #undef LOCTEXT_NAMESPACE
diff --git a/Engine/Source/Editor/UnrealEd/Private/SSkeletonWidget.cpp b/Engine/Source/Editor/UnrealEd/Private/SSkeletonWidget.cpp
--- a/Engine/Source/Editor/UnrealEd/Private/SSkeletonWidget.cpp
+++ b/Engine/Source/Editor/UnrealEd/Private/SSkeletonWidget.cpp
@@ -1093,13 +1093,24 @@ public:
 	}
 
 
-	// FlEditorViewportClient interface
-	virtual FSceneInterface* GetScene() const
+	// FEditorViewportClient interface
+	virtual void Tick(float DeltaTime) override
+	{
+		if (PreviewScene)
+		{
+			PreviewScene->GetWorld()->Tick(LEVELTICK_All, DeltaTime);
+		}
+	}
+
+	virtual FSceneInterface* GetScene() const override
 	{
 		return PreviewScene->GetScene();
 	}
 
-	virtual FLinearColor GetBackgroundColor() const override { return FLinearColor::White; }
+	virtual FLinearColor GetBackgroundColor() const override 
+	{ 
+		return FLinearColor::White; 
+	}
 
 	// End of FEditorViewportClient
 
@@ -1118,6 +1129,8 @@ public:
 // SBasePoseViewport
 void SBasePoseViewport::Construct(const FArguments& InArgs)
 {
+	//RegisterActiveTick( 0.f, FTickWidgetDelegate::CreateSP( this, &SBasePoseViewport::TickViewport ) );
+
 	this->ChildSlot
 	[
 		SNew(SVerticalBox)
@@ -1222,12 +1235,5 @@ bool SBasePoseViewport::IsVisible() const
 {
 	return true;
 }
-
-void SBasePoseViewport::Tick(const FGeometry& AllottedGeometry, const double InCurrentTime, const float InDeltaTime)
-{
-	PreviewScene.GetWorld()->Tick(LEVELTICK_All, InDeltaTime);
-	LevelViewportClient->Tick(InDeltaTime);
-	SCompoundWidget::Tick(AllottedGeometry, InCurrentTime, InDeltaTime);
-}
 #undef LOCTEXT_NAMESPACE 
 
diff --git a/Engine/Source/Editor/UnrealEd/Private/SViewportToolBar.cpp b/Engine/Source/Editor/UnrealEd/Private/SViewportToolBar.cpp
--- a/Engine/Source/Editor/UnrealEd/Private/SViewportToolBar.cpp
+++ b/Engine/Source/Editor/UnrealEd/Private/SViewportToolBar.cpp
@@ -24,7 +24,7 @@ void SViewportToolBar::Construct( const FArguments& InArgs )
 
 	FadeInSequence = FCurveSequence( 0.0f, ToolBarConstants::FadeTime );
 	FadeOutSequence = FCurveSequence( ToolBarConstants::TimeToFadeOut, ToolBarConstants::FadeTime );
-	FadeOutSequence.Play();
+	FadeOutSequence.JumpToEnd();
 }
 
 TWeakPtr<SMenuAnchor> SViewportToolBar::GetOpenMenu() const
@@ -73,13 +73,13 @@ void SViewportToolBar::OnMouseEnter( const FGeometry& MyGeometry, const FPointer
 		bIsHovered = true;
 		if( FadeOutSequence.IsPlaying() )
 		{
-			// Fade out is already playing so just force the fade in curve to the end so we dont have a "pop" 
+			// Fade out is already playing so just force the fade in curve to the end so we don't have a "pop" 
 			// effect from quickly resetting the alpha
 			FadeInSequence.JumpToEnd();
 		}
 		else
 		{
-			FadeInSequence.Play();
+			FadeInSequence.Play( this->AsShared() );
 		}
 	}
 
@@ -92,7 +92,7 @@ void SViewportToolBar::OnMouseLeave( const FPointerEvent& MouseEvent )
 	if( !FSlateApplication::Get().IsUsingHighPrecisionMouseMovment() )
 	{
 		bIsHovered = false;
-		FadeOutSequence.Play();
+		FadeOutSequence.Play( this->AsShared() );
 	}
 }
 
diff --git a/Engine/Source/Editor/UnrealEd/Public/AssetThumbnail.h b/Engine/Source/Editor/UnrealEd/Public/AssetThumbnail.h
--- a/Engine/Source/Editor/UnrealEd/Public/AssetThumbnail.h
+++ b/Engine/Source/Editor/UnrealEd/Public/AssetThumbnail.h
@@ -107,7 +107,7 @@ private:
 /**
  * Utility class for keeping track of, rendering, and recycling thumbnails rendered in Slate              
  */
-class FAssetThumbnailPool
+class FAssetThumbnailPool : public FTickableEditorObject
 {
 public:
 
@@ -124,10 +124,16 @@ public:
 	/** Destructor to free all remaining resources */
 	UNREALED_API ~FAssetThumbnailPool();
 
-	/**
-	 * Ticks the pool, rendering new thumbnails as needed
-	 */
-	UNREALED_API void Tick( float DeltaTime );
+	// Begin FTickableObject interface
+	UNREALED_API virtual TStatId GetStatId() const override;
+
+	/** Checks if any new thumbnails are queued */
+	UNREALED_API virtual bool IsTickable() const override;
+
+	/** Ticks the pool, rendering new thumbnails as needed */
+	UNREALED_API virtual void Tick( float DeltaTime ) override;
+
+	// End FTickableObject interface
 
 	/**
 	 * Releases all rendering resources held by the pool
diff --git a/Engine/Source/Editor/UnrealEd/Public/CreditsScreen.h b/Engine/Source/Editor/UnrealEd/Public/CreditsScreen.h
--- a/Engine/Source/Editor/UnrealEd/Public/CreditsScreen.h
+++ b/Engine/Source/Editor/UnrealEd/Public/CreditsScreen.h
@@ -17,9 +17,9 @@ public:
 	 */
 	UNREALED_API void Construct(const FArguments& InArgs);
 
-	UNREALED_API virtual void Tick(const FGeometry& AllottedGeometry, const double InCurrentTime, const float InDeltaTime) override;
-
 private:
+	/** Animates the credits during play */
+	EActiveTickReturnType RollCredits( double InCurrentTime, float InDeltaTime );
 
 	/** Handles the user clicking the play/pause toggle button. */
 	FReply HandleTogglePlayPause();
@@ -36,6 +36,9 @@ private:
 	/** The widget that scrolls the credits text */
 	TSharedPtr<SScrollBox> ScrollBox;
 
+	/** The handle to the active roll credits tick */
+	TWeakPtr<FActiveTickHandle> ActiveTickHandle;
+
 	/** The auto scroll rate in pixels per second */
 	float ScrollPixelsPerSecond;
 
diff --git a/Engine/Source/Editor/UnrealEd/Public/Dialogs/SDeleteAssetsDialog.h b/Engine/Source/Editor/UnrealEd/Public/Dialogs/SDeleteAssetsDialog.h
--- a/Engine/Source/Editor/UnrealEd/Public/Dialogs/SDeleteAssetsDialog.h
+++ b/Engine/Source/Editor/UnrealEd/Public/Dialogs/SDeleteAssetsDialog.h
@@ -39,7 +39,6 @@ public:
 	void Construct( const FArguments& InArgs, TSharedRef<FAssetDeleteModel> InDeleteModel );
 
 	virtual FReply OnKeyDown( const FGeometry& MyGeometry, const FKeyEvent& InKeyEvent ) override;
-	virtual void Tick( const FGeometry& AllottedGeometry, const double InCurrentTime, const float InDeltaTime ) override;
 
 private:
 	TSharedRef<SWidget> CreateThumbnailWidget();
@@ -52,6 +51,8 @@ private:
 	TSharedRef<SWidget> BuildDeleteDialog();
 
 private:
+	/** Active tick to tick the delete model until it reaches a "Finished" state */
+	EActiveTickReturnType TickDeleteModel( double InCurrentTime, float InDeltaTime );
 
 	void HandleDeleteModelStateChanged( FAssetDeleteModel::EState NewState );
 
@@ -119,6 +120,9 @@ private:
 
 private:
 
+	/** Whether the active tick is currently registered */
+	bool bIsActiveTickRegistered;
+
 	/** The model used for deleting assets */
 	TSharedPtr<FAssetDeleteModel> DeleteModel;
 
diff --git a/Engine/Source/Editor/UnrealEd/Public/EditorViewportClient.h b/Engine/Source/Editor/UnrealEd/Public/EditorViewportClient.h
--- a/Engine/Source/Editor/UnrealEd/Public/EditorViewportClient.h
+++ b/Engine/Source/Editor/UnrealEd/Public/EditorViewportClient.h
@@ -219,8 +219,8 @@ public:
 	 */
 	FMatrix ComputeOrbitMatrix() const;
 private:
-	/** Curve for animating between locations */
-	TSharedPtr<struct FCurveSequence> TransitionCurve;
+	/** The time when a transition to the desired location began */
+	double TransitionStartTime;
 	/** Current viewport Position. */
 	FVector	ViewLocation;
 	/** Current Viewport orientation; valid only for perspective projections. */
diff --git a/Engine/Source/Editor/UnrealEd/Public/SEditorViewport.h b/Engine/Source/Editor/UnrealEd/Public/SEditorViewport.h
--- a/Engine/Source/Editor/UnrealEd/Public/SEditorViewport.h
+++ b/Engine/Source/Editor/UnrealEd/Public/SEditorViewport.h
@@ -17,7 +17,7 @@ public:
 	virtual FReply OnKeyDown( const FGeometry& MyGeometry, const FKeyEvent& InKeyEvent ) override;
 	virtual bool SupportsKeyboardFocus() const override;
 	virtual FReply OnFocusReceived( const FGeometry& MyGeometry, const FFocusEvent& InFocusEvent ) override;
-	virtual void Tick(const FGeometry& AllottedGeometry, const double InCurrentTime, const float InDeltaTime) override;
+	virtual void Tick( const FGeometry& AllottedGeometry, const double InCurrentTime, const float InDeltaTime ) override;
 	
 	/**
 	 * @return True if the viewport is being updated in realtime
@@ -155,14 +155,26 @@ protected:
 
 protected:
 	TSharedPtr<SOverlay> ViewportOverlay;
+
 	/** Viewport that renders the scene provided by the viewport client */
 	TSharedPtr<class FSceneViewport> SceneViewport;
+	
 	/** Widget where the scene viewport is drawn in */
 	TSharedPtr<SViewport> ViewportWidget;
+	
 	/** The client responsible for setting up the scene */
 	TSharedPtr<FEditorViewportClient> Client;
+	
 	/** Commandlist used in the viewport (Maps commands to viewport specific actions) */
 	TSharedPtr<FUICommandList> CommandList;
+	
 	/** The last time the viewport was ticked (for visibility determination) */
 	double LastTickTime;
+
+private:
+	/** Ensures the viewport is ticked actively when realtime is enabled */
+	EActiveTickReturnType EnsureTickWhileRealtime( double InCurrentTime, float InDeltaTime );
+
+	/** The handle to the active tick */
+	TWeakPtr<FActiveTickHandle> ActiveTickHandle;
 };
diff --git a/Engine/Source/Editor/UnrealEd/Public/SSkeletonWidget.h b/Engine/Source/Editor/UnrealEd/Public/SSkeletonWidget.h
--- a/Engine/Source/Editor/UnrealEd/Public/SSkeletonWidget.h
+++ b/Engine/Source/Editor/UnrealEd/Public/SSkeletonWidget.h
@@ -283,8 +283,6 @@ private:
 	class UDebugSkelMeshComponent* PreviewComponent;
 
 	bool IsVisible() const;
-
-	virtual void Tick(const FGeometry& AllottedGeometry, const double InCurrentTime, const float InDeltaTime) override;
 };
 
 /////////////////////////////////////////////
diff --git a/Engine/Source/Editor/UserFeedback/Private/UserFeedback.cpp b/Engine/Source/Editor/UserFeedback/Private/UserFeedback.cpp
--- a/Engine/Source/Editor/UserFeedback/Private/UserFeedback.cpp
+++ b/Engine/Source/Editor/UserFeedback/Private/UserFeedback.cpp
@@ -59,7 +59,7 @@ public:
 		// Animate ourselves in if we're running at the target frame rate
 		if (FSlateApplication::Get().IsRunningAtTargetFrameRate())
 		{
-			Sequence.Play();
+			Sequence.Play( this->AsShared() );
 		}
 		else
 		{
diff --git a/Engine/Source/Editor/WorldBrowser/Private/Tiles/SWorldComposition.cpp b/Engine/Source/Editor/WorldBrowser/Private/Tiles/SWorldComposition.cpp
--- a/Engine/Source/Editor/WorldBrowser/Private/Tiles/SWorldComposition.cpp
+++ b/Engine/Source/Editor/WorldBrowser/Private/Tiles/SWorldComposition.cpp
@@ -131,8 +131,6 @@ public:
 		//	FCanExecuteAction::CreateSP(this, &SWorldCompositionGrid::AreAnyItemsSelected));
 
 		//
-		BounceCurve.AddCurve(0.0f, 1.0f);
-		BounceCurve.Play();
 		WorldModel = InArgs._InWorldModel;
 		bUpdatingSelection = false;
 	
@@ -183,7 +181,7 @@ public:
 	}
 		
 	/**  SWidget interface */
-	virtual void Tick(const FGeometry& AllottedGeometry, const double InCurrentTime, const float InDeltaTime) override
+	void Tick( const FGeometry& AllottedGeometry, const double InCurrentTime, const float InDeltaTime )
 	{
 		SNodePanel::Tick(AllottedGeometry, InCurrentTime, InDeltaTime);
 
@@ -206,7 +204,7 @@ public:
 	
 		WorldModel->UpdateStreamingPreview(WorldMouseLocation, bShowPotentiallyVisibleLevels);
 			
-		// deffered scroll and zooming requests
+		// deferred scroll and zooming requests
 		if (bHasScrollToRequest || bHasScrollByRequest)
 		{
 			// zoom to
@@ -423,7 +421,7 @@ public:
 					// Update the amount to pan panel
 					UpdateViewOffset(MyGeometry, MouseEvent.GetScreenSpacePosition());
 
-					const bool bCursorInDeadZone = TotalMouseDelta <= FSlateApplication::Get().GetDragTriggerDistnace();
+					const bool bCursorInDeadZone = TotalMouseDelta <= FSlateApplication::Get().GetDragTriggerDistance();
 
 					if ( NodeBeingDragged.IsValid() )
 					{
@@ -896,8 +894,7 @@ private:
 
 	/** Geometry cache */
 	mutable FVector2D						CachedAllottedGeometryScaledSize;
-	/** Bouncing curve */
-	FCurveSequence							BounceCurve;
+
 	bool									bUpdatingSelection;
 	TArray<FIntRect>						OccupiedCells;
 	const TSharedRef<FUICommandList>		CommandList;
diff --git a/Engine/Source/Editor/WorldBrowser/Private/Tiles/SWorldTileItem.cpp b/Engine/Source/Editor/WorldBrowser/Private/Tiles/SWorldTileItem.cpp
--- a/Engine/Source/Editor/WorldBrowser/Private/Tiles/SWorldTileItem.cpp
+++ b/Engine/Source/Editor/WorldBrowser/Private/Tiles/SWorldTileItem.cpp
@@ -161,7 +161,10 @@ void SWorldTileItem::Construct(const FArguments& InArgs)
 	SetToolTip(CreateToolTipWidget());
 
 	CurveSequence = FCurveSequence(0.0f, 0.5f);
-	CurveSequence.Play();
+	if (TileModel->IsLoading())
+	{
+		CurveSequence.Play( this->AsShared(), true );
+	}
 }
 
 void SWorldTileItem::RequestRefresh()
@@ -358,6 +361,18 @@ FVector2D SWorldTileItem::ComputeDesiredSize() const
 	return TileModel->GetLevelSize2D();
 }
 
+void SWorldTileItem::Tick( const FGeometry& AllottedGeometry, const double InCurrentTime, const float InDeltaTime )
+{
+	if ( TileModel->IsLoading() && !CurveSequence.IsPlaying() )
+	{
+		CurveSequence.Play( this->AsShared() );
+	}
+	else if ( !TileModel->IsLoading() && CurveSequence.IsPlaying() )
+	{
+		CurveSequence.JumpToEnd();
+	}
+}
+
 int32 SWorldTileItem::OnPaint(const FPaintArgs& Args, const FGeometry& AllottedGeometry, const FSlateRect& ClippingRect, FSlateWindowElementList& OutDrawElements, int32 LayerId, const FWidgetStyle& InWidgetStyle, bool bParentEnabled) const
 {
 	const bool bIsVisible = FSlateRect::DoRectanglesIntersect(AllottedGeometry.GetClippingRect(), ClippingRect);
@@ -400,7 +415,7 @@ int32 SWorldTileItem::OnPaint(const FPaintArgs& Args, const FGeometry& AllottedG
 		// Draw progress bar if level is currently loading
 		if (TileModel->IsLoading())
 		{
-			const float ProgressBarAnimOffset = ProgressBarImage->ImageSize.X * CurveSequence.GetLerpLooping();
+			const float ProgressBarAnimOffset = ProgressBarImage->ImageSize.X * CurveSequence.GetLerp();
 			const float ProgressBarImageSize = ProgressBarImage->ImageSize.X;
 			const float ProgressBarImageHeight = ProgressBarImage->ImageSize.Y;
 			const FVector2D ProggresBarOffset = FVector2D(ProgressBarAnimOffset - ProgressBarImageSize, 0);
diff --git a/Engine/Source/Editor/WorldBrowser/Private/Tiles/SWorldTileItem.h b/Engine/Source/Editor/WorldBrowser/Private/Tiles/SWorldTileItem.h
--- a/Engine/Source/Editor/WorldBrowser/Private/Tiles/SWorldTileItem.h
+++ b/Engine/Source/Editor/WorldBrowser/Private/Tiles/SWorldTileItem.h
@@ -52,6 +52,7 @@ public:
 	~SWorldTileItem();
 
 	void Construct(const FArguments& InArgs);
+	virtual void Tick( const FGeometry& AllottedGeometry, const double InCurrentTime, const float InDeltaTime ) override;
 	
 	// SNodePanel::SNode interface start
 	virtual FVector2D GetDesiredSizeForMarquee() const override;
diff --git a/Engine/Source/Runtime/AppFramework/Private/Framework/Testing/SWidgetGallery.cpp b/Engine/Source/Runtime/AppFramework/Private/Framework/Testing/SWidgetGallery.cpp
--- a/Engine/Source/Runtime/AppFramework/Private/Framework/Testing/SWidgetGallery.cpp
+++ b/Engine/Source/Runtime/AppFramework/Private/Framework/Testing/SWidgetGallery.cpp
@@ -66,7 +66,7 @@ public:
 			TextComboBoxOptions.Add(MakeShareable(new FString(TEXT("Option iii"))));
 
 			ProgressCurve = FCurveSequence(0.0f,15.0f);
-			ProgressCurve.Play();
+			ProgressCurve.Play(this->AsShared());
 		}
 		
 		ChildSlot
@@ -993,7 +993,7 @@ private:
 	TOptional<float> HandleProgressBarPercent( ) const
 	{
 		// Show some marquee, some progress and some 100% filled state.
-		const float Progress = ProgressCurve.GetLerpLooping();
+		const float Progress = ProgressCurve.GetLerp();
 		const float MarqueeTimeFraction = 0.5f;
 
 		return (Progress < MarqueeTimeFraction)
diff --git a/Engine/Source/Runtime/AppFramework/Private/Widgets/Colors/SColorPicker.cpp b/Engine/Source/Runtime/AppFramework/Private/Widgets/Colors/SColorPicker.cpp
--- a/Engine/Source/Runtime/AppFramework/Private/Widgets/Colors/SColorPicker.cpp
+++ b/Engine/Source/Runtime/AppFramework/Private/Widgets/Colors/SColorPicker.cpp
@@ -63,6 +63,8 @@ void SColorPicker::Construct( const FArguments& InArgs )
 	DisplayGamma = InArgs._DisplayGamma;
 	bClosedViaOkOrCancel = false;
 
+	RegisterActiveTick( 0.f, FWidgetActiveTickDelegate::CreateSP( this, &SColorPicker::AnimatePostConstruct ) );
+
 	// We need a parent window to set the close callback
 	if (ParentWindowPtr.IsValid())
 	{
@@ -570,35 +572,36 @@ void SColorPicker::GenerateDefaultColorPickerContent( bool bAdvancedSectionExpan
 	
 	HideSmallTrash();
 }
-END_SLATE_FUNCTION_BUILD_OPTIMIZATION
-
 
-void SColorPicker::Tick( const FGeometry& AllottedGeometry, const double InCurrentTime, const float InDeltaTime )
+EActiveTickReturnType SColorPicker::AnimatePostConstruct( double InCurrentTime, float InDeltaTime )
 {
-	// animate the wheel and all the sliders
-	static const float AnimationTime = 0.2f;
-	if (CurrentTime < AnimationTime)
+	static const float AnimationTime = 0.25f;
+
+	EActiveTickReturnType TickReturnVal = EActiveTickReturnType::KeepTicking;
+	if ( CurrentTime < AnimationTime )
 	{
-		CurrentColorHSV = FMath::Lerp(ColorBegin, ColorEnd, CurrentTime / AnimationTime);
-		if (CurrentColorHSV.R < 0.f)
+		CurrentColorHSV = FMath::Lerp( ColorBegin, ColorEnd, CurrentTime / AnimationTime );
+		if ( CurrentColorHSV.R < 0.f )
 		{
 			CurrentColorHSV.R += 360.f;
 		}
-		else if (CurrentColorHSV.R > 360.f)
+		else if ( CurrentColorHSV.R > 360.f )
 		{
 			CurrentColorHSV.R -= 360.f;
 		}
 
 		CurrentTime += InDeltaTime;
-		if (CurrentTime >= AnimationTime)
+		if ( CurrentTime >= AnimationTime )
 		{
 			CurrentColorHSV = ColorEnd;
+			TickReturnVal = EActiveTickReturnType::StopTicking;
 		}
 
 		CurrentColorRGB = CurrentColorHSV.HSVToLinearRGB();
 	}
-}
 
+	return TickReturnVal;
+}
 
 void SColorPicker::GenerateInlineColorPickerContent()
 {
diff --git a/Engine/Source/Runtime/AppFramework/Private/Widgets/Testing/STestSuite.cpp b/Engine/Source/Runtime/AppFramework/Private/Widgets/Testing/STestSuite.cpp
--- a/Engine/Source/Runtime/AppFramework/Private/Widgets/Testing/STestSuite.cpp
+++ b/Engine/Source/Runtime/AppFramework/Private/Widgets/Testing/STestSuite.cpp
@@ -465,7 +465,7 @@ public:
 		];
 	}
 
-	virtual void Tick( const FGeometry& AllottedGeometry, const double InCurrentTime, const float InDeltaTime )
+	virtual void Tick( const FGeometry& AllottedGeometry, const double InCurrentTime, const float InDeltaTime ) override
 	{
 		CenterRotation += InDeltaTime*.3;
 		if( CenterRotation > 2*PI)
@@ -1946,7 +1946,8 @@ public:
 		InlineEditableText = LOCTEXT( "TestingInlineEditableTextBlock", "Testing inline editable text block!" );
 
 		Animation = FCurveSequence(0, 5);
-		Animation.Play();
+		
+		Animation.Play(this->AsShared(), true);
 
 		this->ChildSlot
 		[
@@ -2061,7 +2062,7 @@ public:
 
 	FSlateColor GetLoopingColor() const
 	{
-		return FLinearColor( 360*Animation.GetLerpLooping(), 0.8f, 1.0f).HSVToLinearRGB();
+		return FLinearColor( 360*Animation.GetLerp(), 0.8f, 1.0f).HSVToLinearRGB();
 	}
 
 	FReply LaunchPopUp_OnClicked ()
@@ -3433,7 +3434,7 @@ public:
 
 	FReply PlayAnimation_OnClicked()
 	{
-		SpawnAnimation.Play();
+		SpawnAnimation.Play(this->AsShared());
 
 		return FReply::Handled();
 	}
@@ -3447,7 +3448,7 @@ public:
 
 	FReply PlayReverse_OnClicked()
 	{
-		SpawnAnimation.PlayReverse();
+		SpawnAnimation.PlayReverse(this->AsShared());
 
 		return FReply::Handled();
 	}
diff --git a/Engine/Source/Runtime/AppFramework/Public/Widgets/Colors/SColorPicker.h b/Engine/Source/Runtime/AppFramework/Public/Widgets/Colors/SColorPicker.h
--- a/Engine/Source/Runtime/AppFramework/Public/Widgets/Colors/SColorPicker.h
+++ b/Engine/Source/Runtime/AppFramework/Public/Widgets/Colors/SColorPicker.h
@@ -220,14 +220,11 @@ protected:
 	 */
 	TSharedRef<SWidget> MakeColorPreviewBox() const;
 
-protected:	
-
-	// SWidget interface
-
-	virtual void Tick( const FGeometry& AllottedGeometry, const double InCurrentTime, const float InDeltaTime ) override;
-
 private:
 
+	// Callback for the active tick to animate the color post-construct
+	EActiveTickReturnType AnimatePostConstruct( double InCurrentTime, float InDeltaTime );
+
 	// Callback for getting the end color of a color spin box gradient.
 	FLinearColor GetGradientEndColor( EColorPickerChannels Channel ) const;
 
diff --git a/Engine/Source/Runtime/Engine/Public/Slate/SceneViewport.h b/Engine/Source/Runtime/Engine/Public/Slate/SceneViewport.h
--- a/Engine/Source/Runtime/Engine/Public/Slate/SceneViewport.h
+++ b/Engine/Source/Runtime/Engine/Public/Slate/SceneViewport.h
@@ -117,7 +117,7 @@ public:
 	/**
 	 * Ticks the viewport
 	 */
-	virtual void Tick( const FGeometry& AllottedGeometry, double InCurrentTime, float DeltaTime ) override;
+	virtual void Tick( const FGeometry& AllottedGeometry, double InCurrentTime, float InDeltaTime ) override;
 
 	/**
 	 * Performs a resize when in swapping viewports while viewing the play world.
diff --git a/Engine/Source/Runtime/GameMenuBuilder/Private/SGameMenuPageWidget.cpp b/Engine/Source/Runtime/GameMenuBuilder/Private/SGameMenuPageWidget.cpp
--- a/Engine/Source/Runtime/GameMenuBuilder/Private/SGameMenuPageWidget.cpp
+++ b/Engine/Source/Runtime/GameMenuBuilder/Private/SGameMenuPageWidget.cpp
@@ -41,7 +41,7 @@ void FMenuPanel::Tick(float Delta)
 	}
 }
 
-void FMenuPanel::ClosePanel()
+void FMenuPanel::ClosePanel(TSharedRef<SWidget> OwnerWidget)
 {
 	if ((CurrentState != EPanelState::Closed) && (CurrentState != EPanelState::Closing))
 	{
@@ -51,7 +51,7 @@ void FMenuPanel::ClosePanel()
 		}
 		else
 		{
-			AnimationSequence.PlayReverse();
+			AnimationSequence.PlayReverse(OwnerWidget);
 		}
 		CurrentState = EPanelState::Closing;
 	}
@@ -68,7 +68,7 @@ void FMenuPanel::ClosePanel()
 	}
 }
 
-void FMenuPanel::OpenPanel()
+void FMenuPanel::OpenPanel(TSharedRef<SWidget> OwnerWidget)
 {
 	if ((CurrentState != EPanelState::Open) && (CurrentState != EPanelState::Opening) )
 	{
@@ -78,7 +78,7 @@ void FMenuPanel::OpenPanel()
 		}
 		else
 		{
-			AnimationSequence.Play();
+			AnimationSequence.Play(OwnerWidget);
 		}
 		CurrentState = EPanelState::Opening;
 	}
@@ -329,8 +329,8 @@ void SGameMenuPageWidget::HideMenu()
 {
 	if (!bMenuHiding)
 	{
-		MainMenuPanel.ClosePanel();
-		SubMenuPanel.ClosePanel();
+		MainMenuPanel.ClosePanel(this->AsShared());
+		SubMenuPanel.ClosePanel(this->AsShared());
 		PendingMainMenu.Reset();
 		PendingSubMenu.Reset();
 		bMenuHiding = true;
@@ -461,7 +461,7 @@ void SGameMenuPageWidget::EnterSubMenu(TSharedPtr<class FGameMenuPage> InSubMenu
 	{
 		MenuHistory.Push(CurrentMenu);
 		CurrentMenu->SelectedIndex = SelectedIndex;
-		MainMenuPanel.ClosePanel();
+		MainMenuPanel.ClosePanel(this->AsShared());
 		PendingMainMenu = InSubMenu;
 		FSlateApplication::Get().PlaySound(MenuStyle->MenuEnterSound);
 	}
@@ -474,7 +474,7 @@ void SGameMenuPageWidget::MenuGoBack(bool bIsCancel)
 		TSharedPtr< FGameMenuPage > MenuInfo = MenuHistory.Pop();
 		if (MainMenuPanel.CurrentState == EPanelState::Closing)
 		{
-			MainMenuPanel.OpenPanel();
+			MainMenuPanel.OpenPanel(this->AsShared());
 			PendingMainMenu.Reset();
 			PendingSubMenu.Reset();
 		}
@@ -485,7 +485,7 @@ void SGameMenuPageWidget::MenuGoBack(bool bIsCancel)
 			{
 				// single menu layout - close this panel and replace with prev menu
 				PendingMainMenu = MenuInfo;
-				MainMenuPanel.ClosePanel();
+				MainMenuPanel.ClosePanel(this->AsShared());
 			}
 			else
 			{
@@ -572,7 +572,7 @@ void SGameMenuPageWidget::OpenMainPanel(TSharedPtr< class FGameMenuPage > InMenu
 	BuildPanelButtons(CurrentMenu, MainPanel, PreviousIndex);
 	CurrentMenuTitle = InMenu->MenuTitle;
 
-	MainMenuPanel.OpenPanel();
+	MainMenuPanel.OpenPanel(this->AsShared());
 	CurrentMenu->MenuOpening();
 
 	OpenPendingSubMenu();	
@@ -622,14 +622,11 @@ void SGameMenuPageWidget::OnSubPanelStateChange(bool bWasOpened)
 	}
 }
 
-void SGameMenuPageWidget::Tick(const FGeometry& AllottedGeometry, const double InCurrentTime, const float InDeltaTime)
+void SGameMenuPageWidget::Tick( const FGeometry& AllottedGeometry, const double InCurrentTime, const float InDeltaTime )
 {
-	//Always tick the super
-	SCompoundWidget::Tick(AllottedGeometry, InCurrentTime, InDeltaTime);
-
 	//ugly code seeing if the console is open
-	UConsole* ViewportConsole = (GEngine->GameViewport != NULL) ? GEngine->GameViewport->ViewportConsole : NULL;
-	if (ViewportConsole != NULL && (ViewportConsole->ConsoleState == "Typing" || ViewportConsole->ConsoleState == "Open"))
+	UConsole* ViewportConsole = ( GEngine->GameViewport != nullptr ) ? GEngine->GameViewport->ViewportConsole : nullptr;
+	if ( ViewportConsole != nullptr && ( ViewportConsole->ConsoleState == "Typing" || ViewportConsole->ConsoleState == "Open" ) )
 	{
 		if (!bConsoleVisible)
 		{
@@ -648,7 +645,7 @@ void SGameMenuPageWidget::Tick(const FGeometry& AllottedGeometry, const double I
 
 	if (GEngine && GEngine->GameViewport )
 	{
-		if (GEngine->GameViewport->ViewportFrame != NULL )
+		if ( GEngine->GameViewport->ViewportFrame != nullptr )
 		{
 			FViewport* Viewport = GEngine->GameViewport->ViewportFrame->GetViewport();
 			if (Viewport)
@@ -866,7 +863,7 @@ FReply SGameMenuPageWidget::SelectionChanged(int32 SelectionIndex)
 		}
 		if ((SubMenuPanel.CurrentState == EPanelState::Open) || (SubMenuPanel.CurrentState == EPanelState::Opening))
 		{
-			SubMenuPanel.ClosePanel();
+			SubMenuPanel.ClosePanel(this->AsShared());
 		}
 		else
 		{
@@ -884,7 +881,7 @@ FReply SGameMenuPageWidget::SelectionChanged(int32 SelectionIndex)
 void SGameMenuPageWidget::FadeIn()
 {
 	//Start the menu widget playing
-	MenuWidgetAnimation.Play();
+	MenuWidgetAnimation.Play(this->AsShared());
 	SetTitleAnimation(true);
 
 	//Go into UI mode
@@ -1002,11 +999,11 @@ void SGameMenuPageWidget::SetTitleAnimation(bool bShowTitle)
 	{
 		if (bShowTitle == true)
 		{
-			TitleWidgetAnimation.Play();
+			TitleWidgetAnimation.Play(this->AsShared());
 		}
 		else
 		{
-			TitleWidgetAnimation.PlayReverse();
+			TitleWidgetAnimation.PlayReverse(this->AsShared());
 		}		
 	}
 	else
@@ -1017,7 +1014,7 @@ void SGameMenuPageWidget::SetTitleAnimation(bool bShowTitle)
 		}
 		else
 		{
-			TitleWidgetAnimation.PlayReverse();
+			TitleWidgetAnimation.PlayReverse(this->AsShared());
 		}
 	}
 }
@@ -1032,7 +1029,7 @@ void SGameMenuPageWidget::OpenPendingSubMenu()
 	if( (PendingSubMenu.IsValid() ) && ( MenuStyle->LayoutType == GameMenuLayoutType::SideBySide ) )
 	{
 		BuildPanelButtons(PendingSubMenu, SubPanel, INDEX_NONE);		
-		SubMenuPanel.OpenPanel();
+		SubMenuPanel.OpenPanel(this->AsShared());
 		NextMenu = PendingSubMenu;
 		NextMenu->MenuOpening();
 		PendingSubMenu.Reset();
diff --git a/Engine/Source/Runtime/GameMenuBuilder/Public/SGameMenuPageWidget.h b/Engine/Source/Runtime/GameMenuBuilder/Public/SGameMenuPageWidget.h
--- a/Engine/Source/Runtime/GameMenuBuilder/Public/SGameMenuPageWidget.h
+++ b/Engine/Source/Runtime/GameMenuBuilder/Public/SGameMenuPageWidget.h
@@ -31,10 +31,10 @@ public:
 	void Tick(float Delta);
 
 	/* Close the panel */
-	void ClosePanel();
+	void ClosePanel(TSharedRef<SWidget> OwnerWidget);
 	
 	/* Open the panel */
-	void OpenPanel();
+	void OpenPanel(TSharedRef<SWidget> OwnerWidget);
 
 	/* Force the panel to be fully open. */
 	void ForcePanelOpen();
@@ -90,7 +90,7 @@ public:
 
 	// Begin SWidget interface
 	/** Update function. */
-	virtual void Tick(const FGeometry& AllottedGeometry, const double InCurrentTime, const float InDeltaTime) override;
+	virtual void Tick( const FGeometry& AllottedGeometry, const double InCurrentTime, const float InDeltaTime ) override;
 
 	/** Mouse button down handler. */
 	virtual FReply OnMouseButtonDown(const FGeometry& MyGeometry, const FPointerEvent& MouseEvent) override;
diff --git a/Engine/Source/Runtime/Slate/Private/Framework/Application/SlateApplication.cpp b/Engine/Source/Runtime/Slate/Private/Framework/Application/SlateApplication.cpp
--- a/Engine/Source/Runtime/Slate/Private/Framework/Application/SlateApplication.cpp
+++ b/Engine/Source/Runtime/Slate/Private/Framework/Application/SlateApplication.cpp
@@ -269,8 +269,9 @@ DECLARE_CYCLE_STAT( TEXT("Update Tooltip Time"), STAT_SlateUpdateTooltip, STATGR
 DECLARE_CYCLE_STAT( TEXT("Tick Window And Children Time"), STAT_SlateTickWindowAndChildren, STATGROUP_Slate );
 DECLARE_CYCLE_STAT( TEXT("Total Slate Tick Time"), STAT_SlateTickTime, STATGROUP_Slate );
 DECLARE_CYCLE_STAT( TEXT("SlatePrepass"), STAT_SlateCacheDesiredSize, STATGROUP_Slate );
-DECLARE_CYCLE_STAT( TEXT("TickWidgets"), STAT_SlateTickWidgets, STATGROUP_Slate );
 DECLARE_CYCLE_STAT( TEXT("DrawWindows"), STAT_SlateDrawWindowTime, STATGROUP_Slate );
+DECLARE_CYCLE_STAT( TEXT("TickWidgets"), STAT_SlateTickWidgets, STATGROUP_Slate );
+DECLARE_CYCLE_STAT( TEXT("TickRegisteredWidgets"), STAT_SlateTickRegisteredWidgets, STATGROUP_Slate );
 
 
 // Slate Event Logging is enabled to allow crash log dumping
@@ -332,11 +333,23 @@ TAutoConsoleVariable<int32> TargetFrameRateForResponsiveness(
 	35,	// Frames per second
 	TEXT( "Minimum sustained average frame rate required before we consider the editor to be \"responsive\" for a smooth UI experience" ) );
 
+/** Whether to skip the second Slate PrePass call (the one right before rendering). */
 TAutoConsoleVariable<int32> SkipSecondPrepass(
 	TEXT("Slate.SkipSecondPrepass"),
 	0,
 	TEXT("Whether to skip the second Slate PrePass call (the one right before rendering)."));
 
+/** Whether Slate should go to sleep when there are no active ticks and the user is idle */
+TAutoConsoleVariable<int32> AllowSlateToSleep(
+	TEXT("Slate.AllowSlateToSleep"),
+	false,
+	TEXT("Whether Slate should go to sleep when there are no active ticks and the user is idle"));
+
+/** The amount of time that must pass without any user action before Slate is put to sleep (provided that there are no active ticks). */
+TAutoConsoleVariable<float> SleepBufferPostInput(
+	TEXT("Slate.SleepBufferPostInput"),
+	1.0f,
+	TEXT("The amount of time that must pass without any user action before Slate is put to sleep (provided that there are no active ticks)."));
 
 //////////////////////////////////////////////////////////////////////////
 bool FSlateApplication::MouseCaptorHelper::HasCapture() const
@@ -654,13 +667,14 @@ void FSlateApplication::Shutdown()
 TSharedPtr<FSlateApplication> FSlateApplication::CurrentApplication = nullptr;
 
 FSlateApplication::FSlateApplication()
-	: bAppIsActive(true)
-	, bSlateWindowActive(true)
+	: bAppIsActive( true )
+	, bSlateWindowActive( true )
 	, Scale( 1.0f )
-	, DragTriggerDistnace( 5.0f )
 	, CursorRadius(0.0f)
 	, LastUserInteractionTime( 0.0 )
+	, DragTriggerDistance( 5.0f )
 	, LastUserInteractionTimeForThrottling( 0.0 )
+	, LastMouseMoveTime( 0.0 )
 	, SlateSoundDevice( MakeShareable(new FNullSlateSoundDevice()) )
 	, CurrentTime( FPlatformTime::Seconds() )
 	, LastTickTime( 0.0 )
@@ -843,7 +857,6 @@ bool FSlateApplication::IsWindowHousingInteractiveTooltip(const TSharedRef<const
 	return bIsHousingInteractiveTooltip;
 }
 
-
 void FSlateApplication::TickWindowAndChildren( TSharedRef<SWindow> WindowToTick )
 {
 	if ( WindowToTick->IsVisible() && !WindowToTick->IsWindowMinimized() )
@@ -856,7 +869,7 @@ void FSlateApplication::TickWindowAndChildren( TSharedRef<SWindow> WindowToTick
 		// the Tick() function.
 		
 		{
-			SCOPE_CYCLE_COUNTER( STAT_SlateCacheDesiredSize )
+			SCOPE_CYCLE_COUNTER( STAT_SlateCacheDesiredSize );
 			WindowToTick->SlatePrepass();
 		}
 
@@ -866,7 +879,7 @@ void FSlateApplication::TickWindowAndChildren( TSharedRef<SWindow> WindowToTick
 		}
 
 		{
-			SCOPE_CYCLE_COUNTER( STAT_SlateTickWidgets )
+			SCOPE_CYCLE_COUNTER( STAT_SlateTickWidgets );
 			// Tick this window and all of the widgets in this window
 			WindowToTick->TickWidgetsRecursively( WindowToTick->GetWindowGeometryInScreen(), GetCurrentTime(), GetDeltaTime() );
 		}
@@ -1271,51 +1284,68 @@ void FSlateApplication::Tick()
 	// Update auto-throttling based on elapsed time since user interaction
 	ThrottleApplicationBasedOnMouseMovement();
 
+	TSharedPtr<SWindow> ActiveModalWindow = GetActiveModalWindow();
+
+	const float SleepThreshold = SleepBufferPostInput.GetValueOnGameThread();
+	const double TimeSinceInput = LastTickTime - LastUserInteractionTime;
+	const double TimeSinceMouseMove = LastTickTime - LastMouseMoveTime;
 	
-	if (!bFoldTick)
-	{
-		TSharedPtr<SWindow> ActiveModalWindow = GetActiveModalWindow();
-	
-		SCOPE_CYCLE_COUNTER( STAT_SlateTickWindowAndChildren );
+	const bool bIsUserIdle = (TimeSinceInput > SleepThreshold) && (TimeSinceMouseMove > SleepThreshold);
+	const bool bAnyActiveTicksPending = AnyActiveTicksArePending();
 
-		if ( ActiveModalWindow.IsValid() )
+	// skip tick/draw if we are idle and there are no active ticks registered that we need to drive slate for.
+	// This effectively means the slate application is totally idle and we don't need to update the UI.
+	// This relies on Widgets properly registering for Active tick when they need something to happen even
+	// when the user is not providing any input (ie, animations, viewport rendering, async polling, etc).
+	if ( !AllowSlateToSleep.GetValueOnGameThread() || bAnyActiveTicksPending || !bIsUserIdle )
+	{
+		if (!bFoldTick)
 		{
-			// There is a modal window, and we just need to tick it.
-			TickWindowAndChildren( ActiveModalWindow.ToSharedRef() );
-			// And also tick any top-level windows.
-			for( TArray< TSharedRef<SWindow> >::TIterator CurrentWindowIt( SlateWindows ); CurrentWindowIt; ++CurrentWindowIt )
+			TSharedPtr<SWindow> ActiveModalWindow = GetActiveModalWindow();
+
 			{
-				TSharedRef<SWindow>& CurrentWindow = *CurrentWindowIt;
-				if (CurrentWindow->IsTopmostWindow())
+				SCOPE_CYCLE_COUNTER(STAT_SlateTickWindowAndChildren);
+
+				if ( ActiveModalWindow.IsValid() )
+				{
+					// There is a modal window, and we just need to tick it.
+					TickWindowAndChildren(ActiveModalWindow.ToSharedRef());
+					// And also tick any topmost windows (like tooltips, etc)
+					for ( TArray< TSharedRef<SWindow> >::TIterator CurrentWindowIt(SlateWindows); CurrentWindowIt; ++CurrentWindowIt )
+					{
+						TSharedRef<SWindow>& CurrentWindow = *CurrentWindowIt;
+						if ( CurrentWindow->IsTopmostWindow() )
+						{
+							TickWindowAndChildren(CurrentWindow);
+						}
+					}
+					// also tick the notification manager's windows
+					TArray< TSharedRef<SWindow> > NotificationWindows;
+					FSlateNotificationManager::Get().GetWindows(NotificationWindows);
+					for ( auto CurrentWindowIt(NotificationWindows.CreateIterator()); CurrentWindowIt; ++CurrentWindowIt )
+					{
+						TickWindowAndChildren(*CurrentWindowIt);
+					}
+				}
+				else
 				{
-					TickWindowAndChildren(CurrentWindow);
+					// No modal window; tick all slate windows.
+					for ( TArray< TSharedRef<SWindow> >::TIterator CurrentWindowIt(SlateWindows); CurrentWindowIt; ++CurrentWindowIt )
+					{
+						TSharedRef<SWindow>& CurrentWindow = *CurrentWindowIt;
+						TickWindowAndChildren(CurrentWindow);
+					}
 				}
 			}
-			// also tick the notification manager's windows
-			TArray< TSharedRef<SWindow> > NotificationWindows;
-			FSlateNotificationManager::Get().GetWindows(NotificationWindows);
-			for( auto CurrentWindowIt( NotificationWindows.CreateIterator() ); CurrentWindowIt; ++CurrentWindowIt )
-			{
-				TickWindowAndChildren(*CurrentWindowIt);
-			}		
 		}
-		else
-		{
-			// No modal window; tick all slate windows.
-			for( TArray< TSharedRef<SWindow> >::TIterator CurrentWindowIt( SlateWindows ); CurrentWindowIt; ++CurrentWindowIt )
-			{
-				TSharedRef<SWindow>& CurrentWindow = *CurrentWindowIt;
-				TickWindowAndChildren( CurrentWindow );
-			}
-		}
-	}
 
-	// Update any notifications - this needs to be done after windows have updated themselves 
-	// (so they know their size)
-	FSlateNotificationManager::Get().Tick();
+		// Update any notifications - this needs to be done after windows have updated themselves 
+		// (so they know their size)
+		FSlateNotificationManager::Get().Tick();
 
-	// Draw all windows
-	DrawWindows();
+		// Draw all windows
+		DrawWindows();
+	}
 }
 
 
@@ -1342,7 +1372,7 @@ void FSlateApplication::ThrottleApplicationBasedOnMouseMovement()
 
 		const FVector2D& CursorPos = GetCursorPos();
 		static FVector2D LastCursorPos = GetCursorPos();
-		static double LastMouseMoveTime = FPlatformTime::Seconds();
+		//static double LastMouseMoveTime = FPlatformTime::Seconds();
 		static bool bIsMouseMoving = false;
 		if( CursorPos != LastCursorPos )
 		{
@@ -2616,7 +2646,7 @@ void FSlateApplication::CloseToolTip()
 	TSharedPtr< SWindow > PinnedToolTipWindow( ToolTipWindow.Pin() );
 	if( PinnedToolTipWindow.IsValid() && PinnedToolTipWindow->IsVisible() )
 	{
-		// Notify the source widget that it's tooltip is closing
+		// Notify the source widget that its tooltip is closing
 		if (SWidget* SourceWidget = ActiveToolTipWidgetSource.Pin().Get())
 		{
 			SourceWidget->OnToolTipClosing();
@@ -2702,7 +2732,7 @@ void FSlateApplication::UpdateToolTip( bool AllowSpawningOfNewToolTips )
 	// Did the tool tip change from last time?
 	const bool bToolTipChanged = (NewToolTip != ActiveToolTip.Pin());
 	
-	// Any widgets that wish to handle visualizing the tooltip get a change here.
+	// Any widgets that wish to handle visualizing the tooltip get a chance here.
 	TSharedPtr<SWidget> NewTooltipVisualizer;
 	if (bToolTipChanged)
 	{
@@ -3075,14 +3105,14 @@ void FSlateApplication::SetUnhandledKeyDownEventHandler( const FOnKeyEvent& NewH
 	UnhandledKeyDownEventHandler = NewHandler;
 }
 
-float FSlateApplication::GetDragTriggerDistnace() const
+float FSlateApplication::GetDragTriggerDistance() const
 {
-	return DragTriggerDistnace;
+	return DragTriggerDistance;
 }
 
-void FSlateApplication::SetDragTriggerDistnace( float ScreenPixels )
+void FSlateApplication::SetDragTriggerDistance( float ScreenPixels )
 {
-	DragTriggerDistnace = ScreenPixels;
+	DragTriggerDistance = ScreenPixels;
 }
 
 void FSlateApplication::SetInputPreProcessor(bool bEnable, TSharedPtr<class IInputProcessor> NewInputProcessor)
@@ -4171,6 +4201,8 @@ bool FSlateApplication::OnMouseMove()
 	const FVector2D LastCursorPosition = GetLastCursorPos();
 	if ( LastCursorPosition != CurrentCursorPosition )
 	{
+		LastMouseMoveTime = GetCurrentTime();
+
 		FPointerEvent MouseEvent(
 			CursorPointerIndex,
 			CurrentCursorPosition,
@@ -4228,7 +4260,7 @@ bool FSlateApplication::ProcessMouseMoveEvent( FPointerEvent& MouseEvent, bool b
 	if( DragDetector.DetectDragForWidget.IsValid() )
 	{	
 		const FVector2D DragDelta = (DragDetector.DetectDragStartLocation - MouseEvent.GetScreenSpacePosition());
-		bDragDetected = ( DragDelta.Size() > FSlateApplication::Get().GetDragTriggerDistnace() );
+		bDragDetected = ( DragDelta.Size() > FSlateApplication::Get().GetDragTriggerDistance() );
 		if (bDragDetected)
 		{
 			FWidgetPath DragDetectPath = DragDetector.DetectDragForWidget.ToWidgetPath(FWeakWidgetPath::EInterruptedPathHandling::ReturnInvalid);
diff --git a/Engine/Source/Runtime/Slate/Private/Framework/Docking/SDockingTabStack.cpp b/Engine/Source/Runtime/Slate/Private/Framework/Docking/SDockingTabStack.cpp
--- a/Engine/Source/Runtime/Slate/Private/Framework/Docking/SDockingTabStack.cpp
+++ b/Engine/Source/Runtime/Slate/Private/Framework/Docking/SDockingTabStack.cpp
@@ -866,11 +866,11 @@ void SDockingTabStack::SetTabWellHidden( bool bShouldHideTabWell )
 {
 	if (bShouldHideTabWell)
 	{
-		ShowHideTabWell.PlayReverse();
+		ShowHideTabWell.PlayReverse( this->AsShared() );
 	}
 	else
 	{
-		ShowHideTabWell.Play();
+		ShowHideTabWell.Play( this->AsShared() );
 	}
 }
 
diff --git a/Engine/Source/Runtime/Slate/Private/Framework/Layout/ScrollyZoomy.cpp b/Engine/Source/Runtime/Slate/Private/Framework/Layout/ScrollyZoomy.cpp
--- a/Engine/Source/Runtime/Slate/Private/Framework/Layout/ScrollyZoomy.cpp
+++ b/Engine/Source/Runtime/Slate/Private/Framework/Layout/ScrollyZoomy.cpp
@@ -185,7 +185,7 @@ FCursorReply FScrollyZoomy::OnCursorQuery( ) const
 
 bool FScrollyZoomy::IsRightClickScrolling( ) const
 {
-	return (AmountScrolledWhileRightMouseDown >= FSlateApplication::Get().GetDragTriggerDistnace());
+	return (AmountScrolledWhileRightMouseDown >= FSlateApplication::Get().GetDragTriggerDistance());
 }
 
 
diff --git a/Engine/Source/Runtime/Slate/Private/Widgets/Docking/SDockTab.cpp b/Engine/Source/Runtime/Slate/Private/Widgets/Docking/SDockTab.cpp
--- a/Engine/Source/Runtime/Slate/Private/Widgets/Docking/SDockTab.cpp
+++ b/Engine/Source/Runtime/Slate/Private/Widgets/Docking/SDockTab.cpp
@@ -95,24 +95,33 @@ FReply SDockTab::OnMouseButtonUp( const FGeometry& MyGeometry, const FPointerEve
 
 void SDockTab::OnDragEnter( const FGeometry& MyGeometry, const FDragDropEvent& DragDropEvent )
 {
-	// Start the drag timer
-	DragTimer = SDockTabDefs::DragTimerActivate;
+	// Register to activate the tab after a delay
+	if ( !ActiveTickHandle.IsValid() )
+	{
+		ActiveTickHandle = RegisterActiveTick( SDockTabDefs::DragTimerActivate, FWidgetActiveTickDelegate::CreateSP( this, &SDockTab::TriggerActivateTab ) );
+	}
 
 	SBorder::OnDragEnter( MyGeometry, DragDropEvent );
 }
 
 void SDockTab::OnDragLeave( const FDragDropEvent& DragDropEvent )
 {
-	// Stop the drag timer
-	DragTimer = 0.0f;
+	// Unregister the activation timer if it hasn't fired yet
+	if ( ActiveTickHandle.IsValid() )
+	{
+		UnRegisterActiveTick( ActiveTickHandle.Pin().ToSharedRef() );
+	}
 
 	SBorder::OnDragLeave( DragDropEvent );
 }
 
 FReply SDockTab::OnDrop( const FGeometry& MyGeometry, const FDragDropEvent& DragDropEvent )
 {
-	// Stop the drag timer
-	DragTimer = 0.0f;
+	// Unregister the activation timer if it hasn't fired yet
+	if ( ActiveTickHandle.IsValid() )
+	{
+		UnRegisterActiveTick( ActiveTickHandle.Pin().ToSharedRef() );
+	}
 
 	return SBorder::OnDrop( MyGeometry, DragDropEvent );
 }
@@ -304,7 +313,6 @@ SDockTab::SDockTab()
 	, ContentAreaPadding( 2 )
 	, bShouldAutosize(false)
 	, TabColorScale(FLinearColor(0,0,0,0))
-	, DragTimer( 0.0f )
 {
 
 }
@@ -370,13 +378,13 @@ void SDockTab::ProvideDefaultIcon( const FSlateBrush* InDefaultIcon )
 
 void SDockTab::PlaySpawnAnim()
 {
-	SpawnAnimCurve.Play();
+	SpawnAnimCurve.Play( this->AsShared() );
 }
 
 void SDockTab::FlashTab()
 {
 	FlashTabCurve = FCurveSequence(0, SDockTabDefs::TabFlashDuration, ECurveEaseFunction::Linear);
-	FlashTabCurve.Play();
+	FlashTabCurve.Play( this->AsShared() );
 }
 
 float SDockTab::GetFlashValue() const
@@ -529,19 +537,10 @@ void SDockTab::Construct( const FArguments& InArgs )
 	);
 }
 
-void SDockTab::Tick( const FGeometry& AllottedGeometry, const double InCurrentTime, const float InDeltaTime )
+EActiveTickReturnType SDockTab::TriggerActivateTab( double InCurrentTime, float InDeltaTime )
 {
-	SBorder::Tick( AllottedGeometry, InCurrentTime, InDeltaTime );
-
-	// Check to see if the window needs to activate because the user has dragged something to it
-	if ( DragTimer > 0.0f )
-	{
-		DragTimer -= InDeltaTime;
-		if ( DragTimer <= 0.0f )
-		{
-			ActivateInParent(ETabActivationCause::UserClickedOnTab);
-		}
-	}
+	ActivateInParent( ETabActivationCause::UserClickedOnTab );
+	return EActiveTickReturnType::StopTicking;
 }
 
 const FDockTabStyle& SDockTab::GetCurrentStyle() const
diff --git a/Engine/Source/Runtime/Slate/Private/Widgets/Images/SSpinningImage.cpp b/Engine/Source/Runtime/Slate/Private/Widgets/Images/SSpinningImage.cpp
--- a/Engine/Source/Runtime/Slate/Private/Widgets/Images/SSpinningImage.cpp
+++ b/Engine/Source/Runtime/Slate/Private/Widgets/Images/SSpinningImage.cpp
@@ -9,9 +9,8 @@ void SSpinningImage::Construct(const FArguments& InArgs)
 	ColorAndOpacity = InArgs._ColorAndOpacity;
 	OnMouseButtonDownHandler = InArgs._OnMouseButtonDown;
 	
-	Sequence = FCurveSequence();
-	Curve = Sequence.AddCurve(0.f, InArgs._Period);
-	Sequence.Play();
+	SpinAnimationSequence = FCurveSequence( 0.f, InArgs._Period );
+	SpinAnimationSequence.Play( this->AsShared(), true );
 }
 
 // Override SImage's OnPaint to draw rotated
@@ -26,7 +25,7 @@ int32 SSpinningImage::OnPaint( const FPaintArgs& Args, const FGeometry& Allotted
 
 		const FColor FinalColorAndOpacity( InWidgetStyle.GetColorAndOpacityTint() * ColorAndOpacity.Get().GetColor(InWidgetStyle) * ImageBrush->GetTint( InWidgetStyle ) );
 		
-		const float Angle = Curve.GetLerpLooping() * 2.0f * PI;
+		const float Angle = SpinAnimationSequence.GetLerp() * 2.0f * PI;
 
 		FSlateDrawElement::MakeRotatedBox( 
 			OutDrawElements,
diff --git a/Engine/Source/Runtime/Slate/Private/Widgets/Images/SThrobber.cpp b/Engine/Source/Runtime/Slate/Private/Widgets/Images/SThrobber.cpp
--- a/Engine/Source/Runtime/Slate/Private/Widgets/Images/SThrobber.cpp
+++ b/Engine/Source/Runtime/Slate/Private/Widgets/Images/SThrobber.cpp
@@ -28,7 +28,7 @@ void SThrobber::ConstructPieces()
 	{
 		ThrobberCurve.Add(AnimCurves.AddCurve(PieceIndex*0.05f, 1.5f));
 	}
-	AnimCurves.Play();
+	AnimCurves.Play(this->AsShared(), true);
 
 	HBox->ClearChildren();
 	for (int32 PieceIndex = 0; PieceIndex < NumPieces; ++PieceIndex)
@@ -73,7 +73,7 @@ void SThrobber::SetAnimate(EAnimation InAnimate)
 
 FVector2D SThrobber::GetPieceScale(int32 PieceIndex) const
 {
-	const float Value = FMath::Sin( 2 * PI * ThrobberCurve[PieceIndex].GetLerpLooping());
+	const float Value = FMath::Sin( 2 * PI * ThrobberCurve[PieceIndex].GetLerp());
 	
 	const bool bAnimateHorizontally = ((0 != (Animate & Horizontal)));
 	const bool bAnimateVertically = (0 != (Animate & Vertical));
@@ -89,7 +89,7 @@ FLinearColor SThrobber::GetPieceColor(int32 PieceIndex) const
 	const bool bAnimateOpacity = (0 != (Animate & Opacity));
 	if (bAnimateOpacity)
 	{
-		const float Value = FMath::Sin( 2 * PI * ThrobberCurve[PieceIndex].GetLerpLooping());
+		const float Value = FMath::Sin( 2 * PI * ThrobberCurve[PieceIndex].GetLerp());
 		return FLinearColor(1.0f,1.0f,1.0f, Value);
 	}
 	else
@@ -134,7 +134,7 @@ void SCircularThrobber::ConstructSequence()
 {
 	Sequence = FCurveSequence();
 	Curve = Sequence.AddCurve(0.f, Period);
-	Sequence.Play();
+	Sequence.Play(this->AsShared(), true);
 }
 
 int32 SCircularThrobber::OnPaint( const FPaintArgs& Args, const FGeometry& AllottedGeometry, const FSlateRect& MyClippingRect, FSlateWindowElementList& OutDrawElements, int32 LayerId, const FWidgetStyle& InWidgetStyle, bool bParentEnabled ) const
@@ -142,7 +142,7 @@ int32 SCircularThrobber::OnPaint( const FPaintArgs& Args, const FGeometry& Allot
 	const FColor FinalColorAndOpacity( InWidgetStyle.GetColorAndOpacityTint() * PieceImage->GetTint( InWidgetStyle ) );
 	const FVector2D LocalOffset = (AllottedGeometry.Size - PieceImage->ImageSize) * 0.5f;
 	const float DeltaAngle = NumPieces > 0 ? 2 * PI / NumPieces : 0;
-	const float Phase = Curve.GetLerpLooping() * 2 * PI;
+	const float Phase = Curve.GetLerp() * 2 * PI;
 
 	for (int32 PieceIdx = 0; PieceIdx < NumPieces; ++PieceIdx)
 	{
diff --git a/Engine/Source/Runtime/Slate/Private/Widgets/Input/SButton.cpp b/Engine/Source/Runtime/Slate/Private/Widgets/Input/SButton.cpp
--- a/Engine/Source/Runtime/Slate/Private/Widgets/Input/SButton.cpp
+++ b/Engine/Source/Runtime/Slate/Private/Widgets/Input/SButton.cpp
@@ -250,7 +250,7 @@ FReply SButton::OnMouseButtonUp( const FGeometry& MyGeometry, const FPointerEven
 
 FReply SButton::OnMouseMove( const FGeometry& MyGeometry, const FPointerEvent& MouseEvent )
 {
-	const float SlateDragStartDistance = FSlateApplication::Get().GetDragTriggerDistnace();
+	const float SlateDragStartDistance = FSlateApplication::Get().GetDragTriggerDistance();
 	if ( IsPreciseTapOrClick(MouseEvent) && MouseEvent.GetCursorDelta().SizeSquared() > ( SlateDragStartDistance*SlateDragStartDistance ) )
 	{
 		bIsPressed = false;
diff --git a/Engine/Source/Runtime/Slate/Private/Widgets/Input/SEditableText.cpp b/Engine/Source/Runtime/Slate/Private/Widgets/Input/SEditableText.cpp
--- a/Engine/Source/Runtime/Slate/Private/Widgets/Input/SEditableText.cpp
+++ b/Engine/Source/Runtime/Slate/Private/Widgets/Input/SEditableText.cpp
@@ -12,19 +12,20 @@ SEditableText* SEditableText::EditableTextUnderCursor = nullptr;
 
 /** Constructor */
 SEditableText::SEditableText()
-	: ScrollHelper(),
-	  CaretPosition( 0 ),
-	  CaretVisualPositionSpring(),
-	  LastCaretInteractionTime( -1000.0 ),
-	  LastSelectionInteractionTime( -1000.0 ),
-	  bIsDragSelecting( false ),
-	  bWasFocusedByLastMouseDown( false ),
-	  bHasDragSelectedSinceFocused( false ),
-	  bHasDragSelectedSinceMouseDown( false ),
-	  CurrentUndoLevel( INDEX_NONE ),
-	  bIsChangingText( false ),
-	  UICommandList( new FUICommandList() ),
-	  bTextChangedByVirtualKeyboard(false)
+	: ScrollHelper()
+	, CaretPosition( 0 )
+	, CaretVisualPositionSpring()
+	, LastCaretInteractionTime( -1000.0 )
+	, LastSelectionInteractionTime( -1000.0 )
+	, bIsDragSelecting( false )
+	, bWasFocusedByLastMouseDown( false )
+	, bHasDragSelectedSinceFocused( false )
+	, bHasDragSelectedSinceMouseDown( false )
+	, CurrentUndoLevel( INDEX_NONE )
+	, bIsChangingText( false )
+	, UICommandList( new FUICommandList() )
+	, bTextChangedByVirtualKeyboard( false )
+	, bIsSpringing( false )
 {
 	// Setup springs
 	FFloatSpring1D::FSpringConfig SpringConfig;
@@ -209,23 +210,11 @@ void SEditableText::SetFont( const TAttribute< FSlateFontInfo >& InNewFont )
 	Font = InNewFont;
 }
 
-
-/**
- * Ticks this widget.  Override in derived classes, but always call the parent implementation.
- *
- * @param  AllottedGeometry Space allotted to this widget
- * @param  InCurrentTime  Current absolute real time
- * @param  InDeltaTime  Real time passed since last tick
- */
 void SEditableText::Tick( const FGeometry& AllottedGeometry, const double InCurrentTime, const float InDeltaTime )
 {
-	// Call parent implementation.
-	SLeafWidget::Tick( AllottedGeometry, InCurrentTime, InDeltaTime );
-
 	// Poll to see if the text has been updated by a virtual keyboard
 	if (bTextChangedByVirtualKeyboard)
 	{
-		//OnTextChanged.ExecuteIfBound(EditedText);
 		OnEnter();
 		bTextChangedByVirtualKeyboard = false;
 	}
@@ -235,39 +224,29 @@ void SEditableText::Tick( const FGeometry& AllottedGeometry, const double InCurr
 		TextInputMethodContext->CachedGeometry = AllottedGeometry;
 		TextInputMethodChangeNotifier->NotifyLayoutChanged(ITextInputMethodChangeNotifier::ELayoutChangeType::Changed);
 	}
-
-	// We'll draw with the 'focused' look if we're either focused or we have a context menu summoned
-	const bool bShouldAppearFocused = HasKeyboardFocus() || ContextMenuWindow.IsValid();
-
-
+	
 	const FString VisibleText = GetStringToRender();
 	const FSlateFontInfo& FontInfo = Font.Get();
 	const TSharedRef< FSlateFontMeasure > FontMeasureService = FSlateApplication::Get().GetRenderer()->GetFontMeasureService();
 	const float FontMaxCharHeight = FontMeasureService->Measure( FString( ), FontInfo, AllottedGeometry.Scale ).Y;
 
-
-	// Update caret position
+	const bool bShouldAppearFocused = ShouldAppearFocused();
+	
+	// Update caret position if focused
+	if( bShouldAppearFocused )
 	{
-		// If we don't have keyboard focus, then push the caret back to the beginning of the line so that
-		// the start of the text is always scrolled into view
-		float FinalCaretTargetPosition = 0.0f;
-		if( bShouldAppearFocused )
-		{
-			FString TextBeforeCaret( VisibleText.Left( CaretPosition ) );
-			const FVector2D TextBeforeCaretSize( FontMeasureService->Measure( TextBeforeCaret, FontInfo, AllottedGeometry.Scale ) );
-
-			FinalCaretTargetPosition = TextBeforeCaretSize.X;
-		}
+		FString TextBeforeCaret( VisibleText.Left( CaretPosition ) );
+		const FVector2D TextBeforeCaretSize( FontMeasureService->Measure( TextBeforeCaret, FontInfo, AllottedGeometry.Scale ) );
 
-		// Spring toward the caret's position!
-		if( FSlateApplication::Get().IsRunningAtTargetFrameRate() )
+		if (FSlateApplication::Get().IsRunningAtTargetFrameRate())
 		{
-			CaretVisualPositionSpring.SetTarget( FinalCaretTargetPosition );
-			CaretVisualPositionSpring.Tick( InDeltaTime );
+			// Spring toward the caret's position!
+			CaretVisualPositionSpring.SetTarget(TextBeforeCaretSize.X);
 		}
 		else
 		{
-			CaretVisualPositionSpring.SetPosition( FinalCaretTargetPosition );
+			// Just set the caret position directly
+			CaretVisualPositionSpring.SetPosition(TextBeforeCaretSize.X);
 		}
 	}
 
@@ -317,6 +296,7 @@ void SEditableText::Tick( const FGeometry& AllottedGeometry, const double InCurr
 
 	// Update selection 'target' effect
 	{
+		bool bSelectionTargetChanged = false;
 		if( AnyTextSelected() )
 		{
 			const float SelectionLeftX = EditableTextDefs::SelectionRectLeftOffset + FontMeasureService->Measure( VisibleText.Left( Selection.GetMinIndex() ), FontInfo, AllottedGeometry.Scale ).X;
@@ -325,14 +305,29 @@ void SEditableText::Tick( const FGeometry& AllottedGeometry, const double InCurr
 			SelectionTargetLeftSpring.SetTarget( SelectionLeftX );
 			SelectionTargetRightSpring.SetTarget( SelectionRightX );
 		}
+	}
+}
 
+EActiveTickReturnType SEditableText::AnimateSpringsWhileFocused(double InCurrentTime, float InDeltaTime)
+{
+	// Keep ticking as long as we should appear focused or the caret is in transit
+	const bool bShouldAppearFocused = ShouldAppearFocused();
+	if (bShouldAppearFocused || !CaretVisualPositionSpring.IsAtRest())
+	{
+		// Update all the springs
+		CaretVisualPositionSpring.Tick(InDeltaTime);
+		
 		const float TimeSinceSelectionInteraction = (float)( InCurrentTime - LastSelectionInteractionTime );
-		if( TimeSinceSelectionInteraction <= EditableTextDefs::SelectionTargetEffectDuration || ( bShouldAppearFocused && AnyTextSelected() ) )
+		if (TimeSinceSelectionInteraction <= EditableTextDefs::SelectionTargetEffectDuration || ( bShouldAppearFocused && AnyTextSelected() ))
 		{
-			SelectionTargetLeftSpring.Tick( InDeltaTime );
-			SelectionTargetRightSpring.Tick( InDeltaTime );
+			SelectionTargetLeftSpring.Tick(InDeltaTime);
+			SelectionTargetRightSpring.Tick(InDeltaTime);
 		}
+
+		return EActiveTickReturnType::KeepTicking;
 	}
+
+	return EActiveTickReturnType::StopTicking;
 }
 
 bool SEditableText::FTextInputMethodContext::IsReadOnly()
@@ -1272,8 +1267,7 @@ int32 SEditableText::OnPaint( const FPaintArgs& Args, const FGeometry& AllottedG
 	const float FontMaxCharHeight = FTextEditHelper::GetFontHeight(FontInfo) * AllottedGeometry.Scale;
 	const double CurrentTime = FSlateApplication::Get().GetCurrentTime();
 
-	// We'll draw with the 'focused' look if we're either focused or we have a context menu summoned
-	const bool bShouldAppearFocused = HasKeyboardFocus() || ContextMenuWindow.IsValid();
+	const bool bShouldAppearFocused = ShouldAppearFocused();
 
 	// Draw selection background
 	if( AnyTextSelected() && ( bShouldAppearFocused || bIsReadonly ) )
@@ -1504,9 +1498,11 @@ bool SEditableText::SupportsKeyboardFocus() const
 
 FReply SEditableText::OnFocusReceived( const FGeometry& MyGeometry, const FFocusEvent& InFocusEvent )
 {
-	// Skip the focus received code if it's due to the context menu closing
+	// Skip the rest of the focus received code if it's due to the context menu closing
 	if ( !ContextMenuWindow.IsValid() )
 	{
+		RegisterActiveTick(0.f, FWidgetActiveTickDelegate::CreateSP(this, &SEditableText::AnimateSpringsWhileFocused));
+
 		// Don't unselect text when focus is set because another widget lost focus, as that may be in response to a
 		// dismissed context menu where the user clicked an option to select text.
 		if( InFocusEvent.GetCause() != EFocusCause::OtherWidgetLostFocus )
@@ -1571,6 +1567,17 @@ void SEditableText::OnFocusLost( const FFocusEvent& InFocusEvent )
 	// Skip the focus lost code if it's due to the context menu opening
 	if ( !ContextMenuWindow.IsValid() )
 	{
+		// Place the caret at the front of the text
+		if (FSlateApplication::Get().IsRunningAtTargetFrameRate())
+		{
+			// Spring back when the frame rate is high enough
+			CaretVisualPositionSpring.SetTarget(0.f);
+		}
+		else
+		{
+			CaretVisualPositionSpring.SetPosition(0.f);
+		}
+
 		FSlateApplication& SlateApplication = FSlateApplication::Get();
 		if (FPlatformMisc::GetRequiresVirtualKeyboard())
 		{
@@ -2046,6 +2053,11 @@ void SEditableText::SetCaretPosition( int32 Position )
 	LastCaretInteractionTime = FSlateApplication::Get().GetCurrentTime();
 }
 
+bool SEditableText::ShouldAppearFocused() const
+{
+	// We'll draw with the 'focused' look if we're either focused or we have a context menu summoned
+	return HasKeyboardFocus() || ContextMenuWindow.IsValid();
+}
 
 bool SEditableText::DoesClipboardHaveAnyText() const 
 {
@@ -2059,7 +2071,7 @@ FString SEditableText::GetStringToRender() const
 {
 	FString VisibleText;
 
-	const bool bShouldAppearFocused = HasKeyboardFocus() || ContextMenuWindow.IsValid();
+	const bool bShouldAppearFocused = ShouldAppearFocused();
 	if(bShouldAppearFocused)
 	{
 		VisibleText = EditedText.ToString();
diff --git a/Engine/Source/Runtime/Slate/Private/Widgets/Input/SSearchBox.cpp b/Engine/Source/Runtime/Slate/Private/Widgets/Input/SSearchBox.cpp
--- a/Engine/Source/Runtime/Slate/Private/Widgets/Input/SSearchBox.cpp
+++ b/Engine/Source/Runtime/Slate/Private/Widgets/Input/SSearchBox.cpp
@@ -3,11 +3,13 @@
 #include "SlatePrivatePCH.h"
 #include "SSearchBox.h"
 
+const double SSearchBox::FilterDelayAfterTyping = 0.25f;
 
 void SSearchBox::Construct( const FArguments& InArgs )
 {
 	check(InArgs._Style);
 
+	bIsActiveTickRegistered = false;
 	OnSearchDelegate = InArgs._OnSearch;
 	OnTextChangedDelegate = InArgs._OnTextChanged;
 	OnTextCommittedDelegate = InArgs._OnTextCommitted;
@@ -16,10 +18,6 @@ void SSearchBox::Construct( const FArguments& InArgs )
 	InactiveFont = InArgs._Style->TextBoxStyle.Font;
 	ActiveFont = InArgs._Style->ActiveFontInfo;
 
-	bTypingFilterText = false;
-	LastTypeTime = 0;
-	FilterDelayAfterTyping = 0.25f;
-
 	SEditableTextBox::Construct( SEditableTextBox::FArguments()
 		.Style( &InArgs._Style->TextBoxStyle )
 		.Font( this, &SSearchBox::GetWidgetFont )
@@ -111,15 +109,13 @@ void SSearchBox::Construct( const FArguments& InArgs )
 	];
 }
 
-void SSearchBox::Tick( const FGeometry& AllottedGeometry, const double InCurrentTime, const float InDeltaTime )
+EActiveTickReturnType SSearchBox::TriggerOnTextChanged( double InCurrentTime, float InDeltaTime )
 {
-	CurrentTime = InCurrentTime;
+	// Reset the flag first in case the delegate winds up triggering HandleTextChanged
+	bIsActiveTickRegistered = false;
 
-	if ( bTypingFilterText && InCurrentTime > LastTypeTime + FilterDelayAfterTyping)
-{
 	OnTextChangedDelegate.ExecuteIfBound( LastPendingTextChangedValue );
-		bTypingFilterText = false;
-	}
+	return EActiveTickReturnType::StopTicking;
 }
 
 void SSearchBox::HandleTextChanged(const FText& NewText)
@@ -127,8 +123,14 @@ void SSearchBox::HandleTextChanged(const FText& NewText)
 	if ( DelayChangeNotificationsWhileTyping.Get() )
 	{
 		LastPendingTextChangedValue = NewText;
-		LastTypeTime = CurrentTime;
-		bTypingFilterText = true;
+
+		// Remove the existing registered tick if necessary
+		if ( ActiveTickHandle.IsValid() )
+		{
+			UnRegisterActiveTick( ActiveTickHandle.Pin().ToSharedRef() );
+		}
+		bIsActiveTickRegistered = true;
+		ActiveTickHandle = RegisterActiveTick( FilterDelayAfterTyping, FWidgetActiveTickDelegate::CreateSP( this, &SSearchBox::TriggerOnTextChanged ) );
 	}
 	else
 	{
@@ -138,7 +140,12 @@ void SSearchBox::HandleTextChanged(const FText& NewText)
 
 void SSearchBox::HandleTextCommitted(const FText& NewText, ETextCommit::Type CommitType)
 {
-	bTypingFilterText = false;
+	if ( bIsActiveTickRegistered && ActiveTickHandle.IsValid() )
+	{
+		bIsActiveTickRegistered = false;
+		UnRegisterActiveTick( ActiveTickHandle.Pin().ToSharedRef() );
+	}
+
 	OnTextCommittedDelegate.ExecuteIfBound( NewText, CommitType );
 }
 
diff --git a/Engine/Source/Runtime/Slate/Private/Widgets/Input/SVirtualJoystick.cpp b/Engine/Source/Runtime/Slate/Private/Widgets/Input/SVirtualJoystick.cpp
--- a/Engine/Source/Runtime/Slate/Private/Widgets/Input/SVirtualJoystick.cpp
+++ b/Engine/Source/Runtime/Slate/Private/Widgets/Input/SVirtualJoystick.cpp
@@ -300,8 +300,6 @@ bool SVirtualJoystick::HandleTouch(int32 ControlIndex, const FVector2D& LocalCoo
 
 void SVirtualJoystick::Tick( const FGeometry& AllottedGeometry, const double InCurrentTime, const float InDeltaTime )
 {
-	SLeafWidget::Tick(AllottedGeometry, InCurrentTime, InDeltaTime);
-
 	if (State == State_WaitForStart || State == State_CountingDownToStart)
 	{
 		CurrentOpacity = 0.f;
diff --git a/Engine/Source/Runtime/Slate/Private/Widgets/Input/SVirtualKeyboardEntry.cpp b/Engine/Source/Runtime/Slate/Private/Widgets/Input/SVirtualKeyboardEntry.cpp
--- a/Engine/Source/Runtime/Slate/Private/Widgets/Input/SVirtualKeyboardEntry.cpp
+++ b/Engine/Source/Runtime/Slate/Private/Widgets/Input/SVirtualKeyboardEntry.cpp
@@ -117,9 +117,6 @@ bool SVirtualKeyboardEntry::GetIsReadOnly() const
  */
 void SVirtualKeyboardEntry::Tick( const FGeometry& AllottedGeometry, const double InCurrentTime, const float InDeltaTime )
 {
-	// Call parent implementation.
-	SLeafWidget::Tick( AllottedGeometry, InCurrentTime, InDeltaTime );
-
 	if(bNeedsUpdate)
 	{
 		// Let outsiders know that the text content has been changed
diff --git a/Engine/Source/Runtime/Slate/Private/Widgets/Layout/SExpandableArea.cpp b/Engine/Source/Runtime/Slate/Private/Widgets/Layout/SExpandableArea.cpp
--- a/Engine/Source/Runtime/Slate/Private/Widgets/Layout/SExpandableArea.cpp
+++ b/Engine/Source/Runtime/Slate/Private/Widgets/Layout/SExpandableArea.cpp
@@ -147,12 +147,12 @@ void SExpandableArea::OnToggleContentVisibility()
 	if( !bAreaCollapsed )
 	{
 		RolloutCurve = FCurveSequence( 0.0f, 0.1f, ECurveEaseFunction::CubicOut );
-		RolloutCurve.Play();
+		RolloutCurve.Play( this->AsShared() );
 	}
 	else
 	{
 		RolloutCurve = FCurveSequence( 0.0f, 0.1f, ECurveEaseFunction::CubicIn );
-		RolloutCurve.PlayReverse();
+		RolloutCurve.PlayReverse( this->AsShared() );
 	}
 
 	// Allow some section-specific code to be executed when the section becomes visible or collapsed
diff --git a/Engine/Source/Runtime/Slate/Private/Widgets/Notifications/SErrorHint.cpp b/Engine/Source/Runtime/Slate/Private/Widgets/Notifications/SErrorHint.cpp
--- a/Engine/Source/Runtime/Slate/Private/Widgets/Notifications/SErrorHint.cpp
+++ b/Engine/Source/Runtime/Slate/Private/Widgets/Notifications/SErrorHint.cpp
@@ -29,7 +29,7 @@ void SErrorHint::SetError( const FText& InErrorText )
 {
 	if ( ErrorText.IsEmpty() && !InErrorText.IsEmpty() )
 	{
-		ExpandAnimation.Play();
+		ExpandAnimation.Play( this->AsShared() );
 	}
 
 	ErrorText = InErrorText;
diff --git a/Engine/Source/Runtime/Slate/Private/Widgets/Notifications/SErrorText.cpp b/Engine/Source/Runtime/Slate/Private/Widgets/Notifications/SErrorText.cpp
--- a/Engine/Source/Runtime/Slate/Private/Widgets/Notifications/SErrorText.cpp
+++ b/Engine/Source/Runtime/Slate/Private/Widgets/Notifications/SErrorText.cpp
@@ -32,7 +32,7 @@ void SErrorText::SetError( const FText& InErrorText )
 {
 	if ( TextBlock->GetText().IsEmpty() && !InErrorText.IsEmpty() )
 	{
-		ExpandAnimation.Play();
+		ExpandAnimation.Play( this->AsShared() );
 	}
 
 	TextBlock->SetText( InErrorText );
diff --git a/Engine/Source/Runtime/Slate/Private/Widgets/Notifications/SNotificationList.cpp b/Engine/Source/Runtime/Slate/Private/Widgets/Notifications/SNotificationList.cpp
--- a/Engine/Source/Runtime/Slate/Private/Widgets/Notifications/SNotificationList.cpp
+++ b/Engine/Source/Runtime/Slate/Private/Widgets/Notifications/SNotificationList.cpp
@@ -43,7 +43,7 @@ public:
 		{
 			CompletionStateAnimation = FCurveSequence();
 			GlowCurve = CompletionStateAnimation.AddCurve(0.f, 0.75f);
-			CompletionStateAnimation.Play();
+			CompletionStateAnimation.Play( this->AsShared() );
 		}
 	}
 
@@ -54,7 +54,7 @@ public:
 		FadeAnimation.AddCurve(FadeOutDuration.Get(), ExpireDuration.Get());
 		// Add the actual fade curve
 		FadeCurve = FadeAnimation.AddCurve(0.f, FadeOutDuration.Get());
-		FadeAnimation.PlayReverse();
+		FadeAnimation.PlayReverse(this->AsShared());
 	}
 
 	/** Begins the fadein of this message */
@@ -66,14 +66,14 @@ public:
 		// Play Fadein animation
 		FadeAnimation = FCurveSequence();
 		FadeCurve = FadeAnimation.AddCurve(0.f, FadeInDuration.Get());
-		FadeAnimation.Play();
+		FadeAnimation.Play( this->AsShared() );
 
 		// Scale up/flash animation
 		IntroAnimation = FCurveSequence();
 		ScaleCurveX = IntroAnimation.AddCurve(0.2f, 0.3f, ECurveEaseFunction::QuadOut);
 		ScaleCurveY = IntroAnimation.AddCurve(0.f, 0.2f);
 		GlowCurve = IntroAnimation.AddCurve(0.5f, 0.55f, ECurveEaseFunction::QuadOut);
-		IntroAnimation.Play();
+		IntroAnimation.Play( this->AsShared() );
 
 		// When a fade in occurs, we need a high framerate for the animation to look good
 		if( FadeInDuration.Get() > KINDA_SMALL_NUMBER && bAllowThrottleWhenFrameRateIsLow && !ThrottleHandle.IsValid() )
@@ -91,7 +91,7 @@ public:
 		// Start fade animation
 		FadeAnimation = FCurveSequence();
 		FadeCurve = FadeAnimation.AddCurve(0.f, FadeOutDuration.Get());
-		FadeAnimation.PlayReverse();
+		FadeAnimation.PlayReverse(this->AsShared());
 	}
 
 	/** Sets the ExpireDuration */
@@ -114,8 +114,6 @@ public:
 
 	void Tick( const FGeometry& AllottedGeometry, const double InCurrentTime, const float InDeltaTime ) override
 	{
-		SCompoundWidget::Tick(AllottedGeometry, InCurrentTime, InDeltaTime);
-
 		const bool bIsFadingOut = FadeAnimation.IsInReverse();
 		const bool bIsCurrentlyPlaying = FadeAnimation.IsPlaying();
 		const bool bIsIntroPlaying = IntroAnimation.IsPlaying();
diff --git a/Engine/Source/Runtime/Slate/Private/Widgets/Notifications/SProgressBar.cpp b/Engine/Source/Runtime/Slate/Private/Widgets/Notifications/SProgressBar.cpp
--- a/Engine/Source/Runtime/Slate/Private/Widgets/Notifications/SProgressBar.cpp
+++ b/Engine/Source/Runtime/Slate/Private/Widgets/Notifications/SProgressBar.cpp
@@ -6,11 +6,12 @@
 void SProgressBar::Construct( const FArguments& InArgs )
 {
 	check(InArgs._Style);
+
+	MarqueeAnimation = FCurveSequence(0.0f, 0.5f);
 	
 	Style = InArgs._Style;
 
-	Percent = InArgs._Percent;
-
+	SetPercent(InArgs._Percent);
 	BarFillType = InArgs._BarFillType;
 	
 	BackgroundImage = InArgs._BackgroundImage;
@@ -19,14 +20,25 @@ void SProgressBar::Construct( const FArguments& InArgs )
 	
 	FillColorAndOpacity = InArgs._FillColorAndOpacity;
 	BorderPadding = InArgs._BorderPadding;
-
-	CurveSequence = FCurveSequence(0.0f, 0.5f);
-	CurveSequence.Play();
 }
 
 void SProgressBar::SetPercent(TAttribute< TOptional<float> > InPercent)
 {
 	Percent = InPercent;
+
+	if ( InPercent.Get().IsSet() )
+	{
+		// No need for the marquee animation
+		MarqueeAnimation.JumpToEnd();
+	}
+	else
+	{
+		// Kick off the marquee animation
+		if ( !MarqueeAnimation.IsPlaying() )
+		{
+			MarqueeAnimation.Play(this->AsShared(), true);
+		}
+	}
 }
 
 void SProgressBar::SetStyle(const FProgressBarStyle* InStyle)
@@ -217,7 +229,7 @@ int32 SProgressBar::OnPaint( const FPaintArgs& Args, const FGeometry& AllottedGe
 		const FSlateBrush* CurrentMarqueeImage = GetMarqueeImage();
 		
 		// Draw Marquee
-		const float MarqueeAnimOffset = CurrentMarqueeImage->ImageSize.X * CurveSequence.GetLerpLooping();
+		const float MarqueeAnimOffset = CurrentMarqueeImage->ImageSize.X * MarqueeAnimation.GetLerp();
 		const float MarqueeImageSize = CurrentMarqueeImage->ImageSize.X;
 
 		FSlateDrawElement::MakeBox(
diff --git a/Engine/Source/Runtime/Slate/Private/Widgets/SViewport.cpp b/Engine/Source/Runtime/Slate/Private/Widgets/SViewport.cpp
--- a/Engine/Source/Runtime/Slate/Private/Widgets/SViewport.cpp
+++ b/Engine/Source/Runtime/Slate/Private/Widgets/SViewport.cpp
@@ -5,7 +5,7 @@
 
 DECLARE_CYCLE_STAT( TEXT("OnPaint SViewport"), STAT_SlateOnPaint_SViewport, STATGROUP_Slate );
 
-/* SViewport structors
+/* SViewport constructors
  *****************************************************************************/
 
 SViewport::SViewport()
@@ -122,7 +122,6 @@ int32 SViewport::OnPaint( const FPaintArgs& Args, const FGeometry& AllottedGeome
 	return Layer;
 }
 
-
 void SViewport::Tick( const FGeometry& AllottedGeometry, const double InCurrentTime, const float InDeltaTime )
 {
 	if(ViewportInterface.IsValid())
diff --git a/Engine/Source/Runtime/Slate/Private/Widgets/Text/SInlineEditableTextBlock.cpp b/Engine/Source/Runtime/Slate/Private/Widgets/Text/SInlineEditableTextBlock.cpp
--- a/Engine/Source/Runtime/Slate/Private/Widgets/Text/SInlineEditableTextBlock.cpp
+++ b/Engine/Source/Runtime/Slate/Private/Widgets/Text/SInlineEditableTextBlock.cpp
@@ -134,9 +134,9 @@ FReply SInlineEditableTextBlock::OnMouseButtonDown( const FGeometry& MyGeometry,
 
 	if(IsSelected.IsBound())
 	{
-		if(IsSelected.Execute() && !bIsReadOnly.Get())
+		if(IsSelected.Execute() && !bIsReadOnly.Get() && !ActiveTickHandle.IsValid())
 		{
-			DoubleSelectDelay = 0.5f;
+			RegisterActiveTick(0.5f, FWidgetActiveTickDelegate::CreateSP(this, &SInlineEditableTextBlock::TriggerEditMode));
 		}
 	}
 	else
@@ -156,27 +156,28 @@ FReply SInlineEditableTextBlock::OnMouseButtonDown( const FGeometry& MyGeometry,
 FReply SInlineEditableTextBlock::OnDragOver( const FGeometry& MyGeometry, const FDragDropEvent& DragDropEvent )
 {
 	// Cancel during a drag over event, otherwise the widget may enter edit mode.
-	DoubleSelectDelay = 0.0f;
+	auto PinnedActiveTickHandle = ActiveTickHandle.Pin();
+	if (PinnedActiveTickHandle.IsValid())
+	{
+		UnRegisterActiveTick(PinnedActiveTickHandle.ToSharedRef());
+	}
 	return FReply::Unhandled();
 }
 
 FReply SInlineEditableTextBlock::OnMouseButtonDoubleClick( const FGeometry& InMyGeometry, const FPointerEvent& InMouseEvent )
 {
-	DoubleSelectDelay = 0.0f;
+	auto PinnedActiveTickHandle = ActiveTickHandle.Pin();
+	if (PinnedActiveTickHandle.IsValid())
+	{
+		UnRegisterActiveTick(PinnedActiveTickHandle.ToSharedRef());
+	}
 	return FReply::Unhandled();
 }
 
-void SInlineEditableTextBlock::Tick( const FGeometry& AllottedGeometry, const double InCurrentTime, const float InDeltaTime )
+EActiveTickReturnType SInlineEditableTextBlock::TriggerEditMode(double InCurrentTime, float InDeltaTime)
 {
-	if(DoubleSelectDelay > 0.0f && IsSelected.IsBound())
-	{
-		DoubleSelectDelay -= InDeltaTime;
-
-		if(DoubleSelectDelay <= 0.0f && IsSelected.Execute())
-		{
-			EnterEditingMode();
-		}
-	}
+	EnterEditingMode();
+	return EActiveTickReturnType::StopTicking;
 }
 
 FReply SInlineEditableTextBlock::OnKeyDown( const FGeometry& MyGeometry, const FKeyEvent& InKeyEvent )
diff --git a/Engine/Source/Runtime/Slate/Private/Widgets/Views/STableViewBase.cpp b/Engine/Source/Runtime/Slate/Private/Widgets/Views/STableViewBase.cpp
--- a/Engine/Source/Runtime/Slate/Private/Widgets/Views/STableViewBase.cpp
+++ b/Engine/Source/Runtime/Slate/Private/Widgets/Views/STableViewBase.cpp
@@ -11,6 +11,9 @@ namespace ListConstants
 
 void STableViewBase::ConstructChildren( const TAttribute<float>& InItemWidth, const TAttribute<float>& InItemHeight, const TAttribute<EListItemAlignment>& InItemAlignment, const TSharedPtr<SHeaderRow>& InHeaderRow, const TSharedPtr<SScrollBar>& InScrollBar )
 {
+	bItemsNeedRefresh = true;
+	//RegisterActiveTick(0.f, FWidgetActiveTickDelegate::CreateSP(this, &STableViewBase::EnsureTickToRefresh));
+
 	HeaderRow = InHeaderRow;
 
 	// If the user provided a scrollbar, we do not need to make one of our own.
@@ -148,57 +151,75 @@ static FEndOfListResult ComputeOffsetForEndOfList( const FGeometry& ListPanelGeo
 	return FEndOfListResult( OffsetFromEndOfList, ItemsAboveView );
 }
 
-void STableViewBase::TickInertialScroll( const FGeometry& AllottedGeometry, const double InCurrentTime, const float InDeltaTime )
+EActiveTickReturnType STableViewBase::UpdateInertialScroll(double InCurrentTime, float InDeltaTime)
 {
-	if ( IsRightClickScrolling() )
+	bool bKeepTicking = false;
+	if (ItemsPanel.IsValid())
 	{
-		// We sample for the inertial scroll on tick rather than on mouse/touch move so
-		// that we still get samples even if the mouse has not moved.
-		if ( CanUseInertialScroll(TickScrollDelta) )
+		if (IsRightClickScrolling())
 		{
-			this->InertialScrollManager.AddScrollSample(TickScrollDelta, InCurrentTime);
-		}
-	}
-	else
-	{
-		// If we are not right click scrolling then we can perform inertial scroll
-		this->InertialScrollManager.UpdateScrollVelocity(InDeltaTime);
-		const float ScrollVelocity = this->InertialScrollManager.GetScrollVelocity();
+			bKeepTicking = true;
 
-		if ( ScrollVelocity != 0.f )
+			// We sample for the inertial scroll on tick rather than on mouse/touch move so
+			// that we still get samples even if the mouse has not moved.
+			if (CanUseInertialScroll(TickScrollDelta))
+			{
+				InertialScrollManager.AddScrollSample(TickScrollDelta, InCurrentTime);
+			}
+		}
+		else
 		{
-			if ( CanUseInertialScroll(ScrollVelocity) )
+			InertialScrollManager.UpdateScrollVelocity(InDeltaTime);
+			const float ScrollVelocity = InertialScrollManager.GetScrollVelocity();
+
+			if (ScrollVelocity != 0.f)
 			{
-				this->ScrollBy(AllottedGeometry, ScrollVelocity * InDeltaTime, AllowOverscroll);
+				if (CanUseInertialScroll(ScrollVelocity))
+				{
+					bKeepTicking = true;
+					ScrollBy(CachedGeometry, ScrollVelocity * InDeltaTime, AllowOverscroll);
+				}
+				else
+				{
+					InertialScrollManager.ClearScrollVelocity();
+				}
 			}
-			else
+
+			if (AllowOverscroll == EAllowOverscroll::Yes)
 			{
-				this->InertialScrollManager.ClearScrollVelocity();
+				// If we are currently in overscroll, the list will need refreshing.
+				// Do this before UpdateOverscroll, as that could cause GetOverscroll() to be 0
+				if (Overscroll.GetOverscroll() != 0.0f)
+				{
+					bKeepTicking = true;
+					RequestListRefresh();
+				}
+
+				Overscroll.UpdateOverscroll(InDeltaTime);
 			}
 		}
+
+		TickScrollDelta = 0.f;
 	}
 
-	TickScrollDelta = 0.f;
+	bIsScrollingActiveTickRegistered = bKeepTicking;
+	return bKeepTicking ? EActiveTickReturnType::KeepTicking : EActiveTickReturnType::StopTicking;
 }
 
-void STableViewBase::Tick( const FGeometry& AllottedGeometry, const double InCurrentTime, const float InDeltaTime )
+EActiveTickReturnType STableViewBase::EnsureTickToRefresh(double InCurrentTime, float InDeltaTime)
 {
-	if (ItemsPanel.IsValid())
-	{
-		TickInertialScroll(AllottedGeometry, InCurrentTime, InDeltaTime);
+	// Actual refresh isn't implemented here as it can be needed in response to changes in the panel geometry.
+	// Since that isn't known until Tick (called after this when registered), refreshing here could result in two refreshes in one frame.
 
-		if ( !IsRightClickScrolling() && AllowOverscroll == EAllowOverscroll::Yes)
-		{
-			// If we are currently in overscroll, the list will need refreshing.
-			// Do this before UpdateOverscroll, as that could cause GetOverscroll() to be 0
-			if ( Overscroll.GetOverscroll() != 0.0f )
-			{
-				this->RequestListRefresh();
-			}
+	return EActiveTickReturnType::StopTicking;
+}
 
-			Overscroll.UpdateOverscroll( InDeltaTime );
-		}
+void STableViewBase::Tick( const FGeometry& AllottedGeometry, const double InCurrentTime, const float InDeltaTime )
+{
+	CachedGeometry = AllottedGeometry;
 
+	if (ItemsPanel.IsValid())
+	{
 		FGeometry PanelGeometry = FindChildGeometry( AllottedGeometry, ItemsPanel.ToSharedRef() );
 		if ( bItemsNeedRefresh || PanelGeometryLastTick.Size != PanelGeometry.Size)
 		{
@@ -361,6 +382,13 @@ FReply STableViewBase::OnMouseMove( const FGeometry& MyGeometry, const FPointerE
 		// the mouse and dragging the view?
 		if( IsRightClickScrolling() )
 		{
+			// Make sure the active tick is registered to update the inertial scroll
+			if (!bIsScrollingActiveTickRegistered)
+			{
+				bIsScrollingActiveTickRegistered = true;
+				RegisterActiveTick(0.f, FWidgetActiveTickDelegate::CreateSP(this, &STableViewBase::UpdateInertialScroll));
+			}
+
 			TickScrollDelta -= ScrollByAmount;
 
 			const float AmountScrolled = this->ScrollBy( MyGeometry, -ScrollByAmount, AllowOverscroll );
@@ -460,7 +488,7 @@ FReply STableViewBase::OnTouchMoved( const FGeometry& MyGeometry, const FPointer
 		AmountScrolledWhileRightMouseDown += FMath::Abs( ScrollByAmount );
 		TickScrollDelta -= ScrollByAmount;
 
-		if (AmountScrolledWhileRightMouseDown > FSlateApplication::Get().GetDragTriggerDistnace())
+		if (AmountScrolledWhileRightMouseDown > FSlateApplication::Get().GetDragTriggerDistance())
 		{
 			const float AmountScrolled = this->ScrollBy( MyGeometry, -ScrollByAmount, EAllowOverscroll::Yes );
 
@@ -508,7 +536,7 @@ TSharedPtr<SHeaderRow> STableViewBase::GetHeaderRow() const
 
 bool STableViewBase::IsRightClickScrolling() const
 {
-	return AmountScrolledWhileRightMouseDown >= FSlateApplication::Get().GetDragTriggerDistnace() &&
+	return AmountScrolledWhileRightMouseDown >= FSlateApplication::Get().GetDragTriggerDistance() &&
 		(this->ScrollBar->IsNeeded() || AllowOverscroll == EAllowOverscroll::Yes);
 }
 
@@ -520,7 +548,12 @@ bool STableViewBase::IsUserScrolling() const
 
 void STableViewBase::RequestListRefresh()
 {
-	bItemsNeedRefresh = true;
+	//bItemsNeedRefresh = true;
+	if (!bItemsNeedRefresh)
+	{
+		bItemsNeedRefresh = true;
+		RegisterActiveTick(0.f, FWidgetActiveTickDelegate::CreateSP(this, &STableViewBase::EnsureTickToRefresh));
+	}
 	ItemsPanel->SetRefreshPending(true);
 }
 
@@ -564,7 +597,8 @@ STableViewBase::STableViewBase( ETableViewMode::Type InTableViewMode )
 	, bShowSoftwareCursor( false )
 	, Overscroll()
 	, AllowOverscroll(EAllowOverscroll::Yes)
-	, bItemsNeedRefresh( true )	
+	, bItemsNeedRefresh( false )	
+	, bIsScrollingActiveTickRegistered(false)
 {
 }
 
@@ -676,6 +710,8 @@ void STableViewBase::OnRightMouseButtonUp(const FVector2D& SummonLocation)
 		{
 			bShowSoftwareCursor = false;
 			FSlateApplication::Get().PushMenu( AsShared(), MenuContent.ToSharedRef(), SummonLocation, FPopupTransitionEffect( FPopupTransitionEffect::ContextMenu ) );
+
+
 		}
 	}
 
diff --git a/Engine/Source/Runtime/Slate/Public/Framework/Application/SWindowTitleBar.h b/Engine/Source/Runtime/Slate/Public/Framework/Application/SWindowTitleBar.h
--- a/Engine/Source/Runtime/Slate/Public/Framework/Application/SWindowTitleBar.h
+++ b/Engine/Source/Runtime/Slate/Public/Framework/Application/SWindowTitleBar.h
@@ -127,7 +127,7 @@ public:
 	void Flash( )
 	{
 		TitleFlashSequence = FCurveSequence(0, SWindowTitleBarDefs::WindowFlashDuration, ECurveEaseFunction::Linear);
-		TitleFlashSequence.Play();
+		TitleFlashSequence.Play(this->AsShared());
 	}
 
 	// Begin IWindowTitleBar interface
diff --git a/Engine/Source/Runtime/Slate/Public/Framework/Application/SlateApplication.h b/Engine/Source/Runtime/Slate/Public/Framework/Application/SlateApplication.h
--- a/Engine/Source/Runtime/Slate/Public/Framework/Application/SlateApplication.h
+++ b/Engine/Source/Runtime/Slate/Public/Framework/Application/SlateApplication.h
@@ -920,10 +920,10 @@ public:
 	double GetLastUserInteractionTime() const { return LastUserInteractionTime; }
 
 	/** @return the deadzone size for dragging in screen pixels (aka virtual desktop pixels) */
-	float GetDragTriggerDistnace() const;
+	float GetDragTriggerDistance() const;
 
 	/** Set the size of the deadzone for dragging in screen pixels */
-	void SetDragTriggerDistnace( float ScreenPixels );
+	void SetDragTriggerDistance( float ScreenPixels );
 	
 	/** Set the analog cursor to be enabled or disabled. */
 	void SetInputPreProcessor(bool bEnable, TSharedPtr<class IInputProcessor> NewInputProcessor = nullptr);
@@ -1126,7 +1126,7 @@ private:
 	float Scale;
 
 	/** The dead zone distance in virtual desktop pixels (a.k.a screen pixels) that the user has to move their finder before it is considered a drag.*/
-	float DragTriggerDistnace;
+	float DragTriggerDistance;
 
 	/** All the top-level windows owned by this application; they are tracked here in a platform-agnostic way. */
 	TArray< TSharedRef<SWindow> > SlateWindows;
@@ -1263,6 +1263,9 @@ private:
 	/** Subset of LastUserInteractionTime that is used only when considering when to throttle */
 	double LastUserInteractionTimeForThrottling;
 
+	/** Used when considering whether to put Slate to sleep */
+	double LastMouseMoveTime;
+
 	/** Helper for detecting when a drag should begin */
 	struct FDragDetector
 	{
diff --git a/Engine/Source/Runtime/Slate/Public/Widgets/Docking/SDockTab.h b/Engine/Source/Runtime/Slate/Public/Widgets/Docking/SDockTab.h
--- a/Engine/Source/Runtime/Slate/Public/Widgets/Docking/SDockTab.h
+++ b/Engine/Source/Runtime/Slate/Public/Widgets/Docking/SDockTab.h
@@ -78,8 +78,6 @@ public:
 	/** Construct the widget from the declaration. */
 	void Construct( const FArguments& InArgs );
 
-	virtual void Tick( const FGeometry& AllottedGeometry, const double InCurrentTime, const float InDeltaTime ) override;
-
 	virtual FReply OnMouseButtonDown( const FGeometry& MyGeometry, const FPointerEvent& MouseEvent ) override;
 
 	virtual FReply OnMouseButtonDoubleClick( const FGeometry& MyGeometry, const FPointerEvent& MouseEvent ) override;
@@ -238,7 +236,7 @@ public:
 	bool HasSiblingTab(const FTabId& SiblingTabId, const bool TreatIndexNoneAsWildcard = true) const;
 
 protected:
-
+	
 	/** Gets the dock tab stack this dockable tab resides within, if any */
 	TSharedPtr<SDockingTabStack> GetParentDockTabStack() const;
 
@@ -275,6 +273,13 @@ protected:
 	/** @return if the close button should be visible. */
 	EVisibility HandleIsCloseButtonVisible() const;
 
+private:
+	/** Activates the tab in its tab well */
+	EActiveTickReturnType TriggerActivateTab( double InCurrentTime, float InDeltaTime );
+
+	/** The handle to the active tab activation tick */
+	TWeakPtr<FActiveTickHandle> ActiveTickHandle;
+
 protected:
 
 	/** The tab manager that created this tab. */
@@ -329,11 +334,13 @@ protected:
 
 	/** @return the scaling of the tab based on the opening/closing animation */
 	FVector2D GetAnimatedScale() const;
+
 	/** Animation that shows the tab opening up */
 	FCurveSequence SpawnAnimCurve;
 
 	/** Animation that causes the tab to flash */
 	FCurveSequence FlashTabCurve;
+
 	/** Get the desired color of tab. These change during flashing. */
 	float GetFlashValue() const;
 
@@ -342,10 +349,7 @@ protected:
 
 	/** Widget used to show the label on the tab */
 	TSharedPtr<STextBlock> LabelWidget;
+	
 	/** Widget used to show the icon on the tab */
 	TSharedPtr<SImage> IconWidget;
-
-
-	/** The amount of time spent until we bring this tab to the front */
-	float DragTimer;
 };
diff --git a/Engine/Source/Runtime/Slate/Public/Widgets/Images/SSpinningImage.h b/Engine/Source/Runtime/Slate/Public/Widgets/Images/SSpinningImage.h
--- a/Engine/Source/Runtime/Slate/Public/Widgets/Images/SSpinningImage.h
+++ b/Engine/Source/Runtime/Slate/Public/Widgets/Images/SSpinningImage.h
@@ -35,6 +35,5 @@ public:
 
 private:
 	/** The sequence to drive the spinning animation */
-	FCurveSequence Sequence;
-	FCurveHandle Curve;
+	FCurveSequence SpinAnimationSequence;
 };
\ No newline at end of file
diff --git a/Engine/Source/Runtime/Slate/Public/Widgets/Input/SEditableText.h b/Engine/Source/Runtime/Slate/Public/Widgets/Input/SEditableText.h
--- a/Engine/Source/Runtime/Slate/Public/Widgets/Input/SEditableText.h
+++ b/Engine/Source/Runtime/Slate/Public/Widgets/Input/SEditableText.h
@@ -463,6 +463,12 @@ protected:
 	void OnWindowClosed(const TSharedRef<SWindow>&);
 
 private:
+	/** Animates the caret and highlight selection springs */
+	EActiveTickReturnType AnimateSpringsWhileFocused(double InCurrentTime, float InDeltaTime);
+	
+	/** @return Whether the editable text should appear focused */
+	bool ShouldAppearFocused() const;
+
 	/** 
 	 * @return whether there is anything in the clipboard
 	 */
@@ -607,6 +613,10 @@ private:
 	/** The type of virtual keyboard to use for editing this text on mobile */
 	TAttribute<EKeyboardType> VirtualKeyboardType;
 
-	/** Whether the text has been changed by a virtual keyboard */
+	/** True if the text has been changed by a virtual keyboard */
 	bool bTextChangedByVirtualKeyboard;
+
+	/** True if a spring animation is currently in progress */
+	bool bIsSpringing;
+
 };
diff --git a/Engine/Source/Runtime/Slate/Public/Widgets/Input/SSearchBox.h b/Engine/Source/Runtime/Slate/Public/Widgets/Input/SSearchBox.h
--- a/Engine/Source/Runtime/Slate/Public/Widgets/Input/SSearchBox.h
+++ b/Engine/Source/Runtime/Slate/Public/Widgets/Input/SSearchBox.h
@@ -62,12 +62,6 @@ public:
 
 	void Construct( const FArguments& InArgs );
 
-public:
-
-	// SWidget overrides
-
-	virtual void Tick( const FGeometry& AllottedGeometry, const double InCurrentTime, const float InDeltaTime ) override;
-
 private:
 
 	/** Callback for changes in the editable text box. */
@@ -78,6 +72,9 @@ private:
 
 private:
 
+	/** Fires the text changed delegate on a delay */
+	EActiveTickReturnType TriggerOnTextChanged( double InCurrentTime, float InDeltaTime );
+
 	/** @return should we show the X to clear search? */
 	EVisibility GetXVisibility() const;
 
@@ -92,6 +89,13 @@ private:
 	/** Invoked to get the font to use for the editable text box */
 	FSlateFontInfo GetWidgetFont() const;
 
+private:
+	/** The amount of time to delay filtering after the user types */
+	static const double FilterDelayAfterTyping;
+
+	/** Handle to the active trigger text changed timer */
+	TWeakPtr<FActiveTickHandle> ActiveTickHandle;
+
 	/** Delegate that is invoked when the user does next or previous */
 	FOnSearch OnSearchDelegate;
 
@@ -107,12 +111,9 @@ private:
 	/** Fonts that specify how to render search text when inactive, and active */
 	FSlateFontInfo ActiveFont, InactiveFont;
 
-	/** When true, the user is typing in the search box. This is used to delay the actual filter until the user is done typing. */
-	double CurrentTime;
-	double LastTypeTime;
-	double FilterDelayAfterTyping;
-	bool bTypingFilterText;
 	FText LastPendingTextChangedValue;
+
+	bool bIsActiveTickRegistered : 1;
 };
 
 
diff --git a/Engine/Source/Runtime/Slate/Public/Widgets/Input/SSpinBox.h b/Engine/Source/Runtime/Slate/Public/Widgets/Input/SSpinBox.h
--- a/Engine/Source/Runtime/Slate/Public/Widgets/Input/SSpinBox.h
+++ b/Engine/Source/Runtime/Slate/Public/Widgets/Input/SSpinBox.h
@@ -275,7 +275,7 @@ public:
 			}
 
 			bDragging = false;
-			if ( DistanceDragged < FSlateApplication::Get().GetDragTriggerDistnace() )
+			if ( DistanceDragged < FSlateApplication::Get().GetDragTriggerDistance() )
 			{
 				EnterTextMode();
 				return FReply::Handled().ReleaseMouseCapture().SetUserFocus(EditableText.ToSharedRef(), EFocusCause::SetDirectly).SetMousePos(CachedMousePosition);
@@ -304,7 +304,7 @@ public:
 			if (!bDragging)
 			{
 				DistanceDragged += FMath::Abs(MouseEvent.GetCursorDelta().X);
-				if ( DistanceDragged > FSlateApplication::Get().GetDragTriggerDistnace() )
+				if ( DistanceDragged > FSlateApplication::Get().GetDragTriggerDistance() )
 				{
 					ExitTextMode();
 					bDragging = true;
diff --git a/Engine/Source/Runtime/Slate/Public/Widgets/Layout/SScrollBox.h b/Engine/Source/Runtime/Slate/Public/Widgets/Layout/SScrollBox.h
--- a/Engine/Source/Runtime/Slate/Public/Widgets/Layout/SScrollBox.h
+++ b/Engine/Source/Runtime/Slate/Public/Widgets/Layout/SScrollBox.h
@@ -189,21 +189,17 @@ private:
 	/** Does the user need a hint that they can scroll to the end of the list? */
 	FSlateColor GetEndShadowOpacity() const;
 
+	/** Active tick to update inertial scrolling as needed */
+	EActiveTickReturnType UpdateInertialScroll(double InCurrentTime, float InDeltaTime);
+
+private:
+
 	/** The panel which stacks the child slots */
 	TSharedPtr<class SScrollPanel> ScrollPanel;
 
 	/** The scrollbar which controls scrolling for the scrollbox. */
 	TSharedPtr<SScrollBar> ScrollBar;
 
-	/** Whether or not the user supplied an external scrollbar to control scrolling. */
-	bool bScrollBarIsExternal;
-	
-	/** Are we actively scrolling right now */
-	bool bIsScrolling;
-
-	/** Should the current scrolling be animated or immediately jump to the desired scroll offer */
-	bool bAnimateScroll;
-
 	/** How much we scrolled while the rmb has been held */
 	float AmountScrolledWhileRightMouseDown;
 
@@ -213,9 +209,6 @@ private:
 	/**	The current position of the software cursor */
 	FVector2D SoftwareCursorPosition;
 
-	/**	Whether the software cursor should be drawn in the viewport */
-	bool bShowSoftwareCursor;
-
 	/** Fired when the user scrolls the scrollbox */
 	FOnUserScrolled OnUserScrolled;
 
@@ -228,8 +221,26 @@ private:
 	/** Style resource for the scrollbar */
 	const FScrollBarStyle* ScrollBarStyle;
 
+	/** Cached geometry for use with the active tick */
+	FGeometry CachedGeometry;
+
+	/**	Whether the software cursor should be drawn in the viewport */
+	bool bShowSoftwareCursor : 1;
+
+	/** Whether or not the user supplied an external scrollbar to control scrolling. */
+	bool bScrollBarIsExternal : 1;
+
+	/** Are we actively scrolling right now */
+	bool bIsScrolling : 1;
+
+	/** Should the current scrolling be animated or immediately jump to the desired scroll offer */
+	bool bAnimateScroll : 1;
+
 	/** If true, will scroll to the end next Tick */
-	bool bScrollToEnd;
+	bool bScrollToEnd : 1;
+
+	/** Whether the active tick to update the inertial scroll is registered */
+	bool bIsActiveTickRegistered : 1;
 };
 
 
diff --git a/Engine/Source/Runtime/Slate/Public/Widgets/Notifications/SProgressBar.h b/Engine/Source/Runtime/Slate/Public/Widgets/Notifications/SProgressBar.h
--- a/Engine/Source/Runtime/Slate/Public/Widgets/Notifications/SProgressBar.h
+++ b/Engine/Source/Runtime/Slate/Public/Widgets/Notifications/SProgressBar.h
@@ -136,6 +136,6 @@ private:
 	TAttribute<FVector2D> BorderPadding;
 
 	/** Curve sequence to drive progress bar animation */
-	FCurveSequence CurveSequence;
+	FCurveSequence MarqueeAnimation;
 };
 
diff --git a/Engine/Source/Runtime/Slate/Public/Widgets/Text/SInlineEditableTextBlock.h b/Engine/Source/Runtime/Slate/Public/Widgets/Text/SInlineEditableTextBlock.h
--- a/Engine/Source/Runtime/Slate/Public/Widgets/Text/SInlineEditableTextBlock.h
+++ b/Engine/Source/Runtime/Slate/Public/Widgets/Text/SInlineEditableTextBlock.h
@@ -96,7 +96,7 @@ class SLATE_API SInlineEditableTextBlock: public SCompoundWidget
 
 	virtual FReply OnKeyDown( const FGeometry& MyGeometry, const FKeyEvent& InKeyEvent ) override;
 
-	virtual void Tick( const FGeometry& AllottedGeometry, const double InCurrentTime, const float InDeltaTime ) override;
+	//virtual void Tick( const FGeometry& AllottedGeometry, const double InCurrentTime, const float InDeltaTime ) override;
 	 
 	/** Switches the widget to editing mode */
 	void EnterEditingMode();
@@ -156,4 +156,11 @@ protected:
 
 	/** Attribute to look up if the widget is read-only */
 	TAttribute< bool > bIsReadOnly;
+
+private:
+	/** Active tick to trigger entry into edit mode after a delay */
+	EActiveTickReturnType TriggerEditMode(double InCurrentTime, float InDeltaTime);
+
+	/** The handle to the active tick */
+	TWeakPtr<FActiveTickHandle> ActiveTickHandle;
 };
diff --git a/Engine/Source/Runtime/Slate/Public/Widgets/Views/STableViewBase.h b/Engine/Source/Runtime/Slate/Public/Widgets/Views/STableViewBase.h
--- a/Engine/Source/Runtime/Slate/Public/Widgets/Views/STableViewBase.h
+++ b/Engine/Source/Runtime/Slate/Public/Widgets/Views/STableViewBase.h
@@ -274,8 +274,17 @@ private:
 	/** Check whether the current state of the table warrants inertial scroll by the specified amount */
 	bool CanUseInertialScroll( float ScrollAmount ) const;
 
-	/** Called every tick to update and perform the inertial scroll */
-	void TickInertialScroll( const FGeometry& AllottedGeometry, const double InCurrentTime, const float InDeltaTime );
+	/** Active tick to update the inertial scroll */
+	EActiveTickReturnType UpdateInertialScroll(double InCurrentTime, float InDeltaTime);
+
+	/** One-off active tick to refresh the contents of the table as needed */
+	EActiveTickReturnType EnsureTickToRefresh(double InCurrentTime, float InDeltaTime);
+
+	/** Whether the active tick to update the inertial scrolling is currently registered */
+	bool bIsScrollingActiveTickRegistered;
+
+	/** Cached geometry for use by the active tick */
+	FGeometry CachedGeometry;
 
 protected:
 
diff --git a/Engine/Source/Runtime/Slate/Public/Widgets/Views/STreeView.h b/Engine/Source/Runtime/Slate/Public/Widgets/Views/STreeView.h
--- a/Engine/Source/Runtime/Slate/Public/Widgets/Views/STreeView.h
+++ b/Engine/Source/Runtime/Slate/Public/Widgets/Views/STreeView.h
@@ -339,7 +339,7 @@ public:
 	 * @param  InCurrentTime  Current absolute real time.
 	 * @param  InDeltaTime  Real time passed since last tick.
 	 */
-	virtual void Tick( const FGeometry& AllottedGeometry, const double InCurrentTime, const float InDeltaTime )
+	virtual void Tick( const FGeometry& AllottedGeometry, const double InCurrentTime, const float InDeltaTime ) override
 	{
 		if ( bTreeItemsAreDirty )
 		{
diff --git a/Engine/Source/Runtime/SlateCore/Private/Animation/CurveHandle.cpp b/Engine/Source/Runtime/SlateCore/Private/Animation/CurveHandle.cpp
--- a/Engine/Source/Runtime/SlateCore/Private/Animation/CurveHandle.cpp
+++ b/Engine/Source/Runtime/SlateCore/Private/Animation/CurveHandle.cpp
@@ -6,7 +6,7 @@
 #define LOCTEXT_NAMESPACE "FCurveHandle"
 
 
-/* FCurveHandle structors
+/* FCurveHandle constructors
  *****************************************************************************/
 
 FCurveHandle::FCurveHandle( const struct FCurveSequence* InOwnerSequence, int32 InCurveIndex )
@@ -37,8 +37,11 @@ float FCurveHandle::GetLerp( ) const
 	return ApplyEasing(Time, TheCurve.EaseFunction);
 }
 
-
-float FCurveHandle::GetLerpLooping( ) const
+float FCurveHandle::GetLerpLooping() const
+{
+	return DEPRECATED_GetLerpLooping();
+}
+float FCurveHandle::DEPRECATED_GetLerpLooping() const
 {
 	if (OwnerSequence == nullptr)
 	{
@@ -46,7 +49,7 @@ float FCurveHandle::GetLerpLooping( ) const
 	}
 
 	// How far we've played through the curve sequence so far.
-	const float CurSequenceTime = OwnerSequence->GetSequenceTimeLooping();
+	const float CurSequenceTime = OwnerSequence->DEPRECATED_GetSequenceTimeLooping();
 
 	const FCurveSequence::FSlateCurve& TheCurve = OwnerSequence->GetCurve(CurveIndex);
 	const float TimeSinceStarted = CurSequenceTime - TheCurve.StartTime;
