Fixes to work with new active ticking

diff --git a/Engine/Source/Runtime/Slate/Private/Framework/MultiBox/SMenuEntryBlock.cpp b/Engine/Source/Runtime/Slate/Private/Framework/MultiBox/SMenuEntryBlock.cpp
--- a/Engine/Source/Runtime/Slate/Private/Framework/MultiBox/SMenuEntryBlock.cpp
+++ b/Engine/Source/Runtime/Slate/Private/Framework/MultiBox/SMenuEntryBlock.cpp
@@ -159,9 +159,6 @@ TSharedRef< class IMultiBlockBaseWidget > FMenuEntryBlock::ConstructWidget() con
 	return SNew( SMenuEntryBlock );
 }
 
-
-
-
 /**
  * Construct this widget
  *
@@ -170,8 +167,8 @@ TSharedRef< class IMultiBlockBaseWidget > FMenuEntryBlock::ConstructWidget() con
 void SMenuEntryBlock::Construct( const FArguments& InArgs )
 {
 	// No initial sub-menus should be opened
-	TimeToSubMenuOpen = 0.0f;
-	SubMenuRequestState = Idle;
+	/*TimeToSubMenuOpen = 0.0f;*/
+	//SubMenuRequestState = Idle;
 
 	// No images by default
 	CheckedImage = nullptr;
@@ -814,17 +811,25 @@ void SMenuEntryBlock::BuildMultiBlockWidget(const ISlateStyle* StyleSet, const F
 void SMenuEntryBlock::RequestSubMenuToggle( bool bOpenMenu, const bool bClobber )
 {
 	// Reset the time before the menu opens
-	TimeToSubMenuOpen = bClobber ? MultiBoxConstants::SubMenuClobberTime : MultiBoxConstants::SubMenuOpenTime;
-	SubMenuRequestState = bOpenMenu ? WantOpen : WantClose;
-	UpdateSubMenuState();
+	float TimeToSubMenuOpen = bClobber ? MultiBoxConstants::SubMenuClobberTime : MultiBoxConstants::SubMenuOpenTime;
+	//SubMenuRequestState = bOpenMenu ? WantOpen : WantClose;
+	if (!ActiveTickHandle.IsValid())
+	{
+		ActiveTickHandle = RegisterActiveTick(TimeToSubMenuOpen, FWidgetActiveTickDelegate::CreateSP(this, &SMenuEntryBlock::UpdateSubMenuState, bOpenMenu));
+	}
+	//UpdateSubMenuState();
 }
 
 void SMenuEntryBlock::CancelPendingSubMenu()
 {
 	// Reset any pending sub-menu openings
-	TimeToSubMenuOpen = 0.0f;
-	SubMenuRequestState = Idle;
-	UpdateSubMenuState();
+	//SubMenuRequestState = Idle;
+	
+	auto PinnedActiveTickHandle = ActiveTickHandle.Pin();
+	if (PinnedActiveTickHandle.IsValid())
+	{
+		UnRegisterActiveTick(PinnedActiveTickHandle.ToSharedRef());
+	}
 }
 
 bool SMenuEntryBlock::ShouldSubMenuAppearHovered() const
@@ -1085,62 +1090,103 @@ FReply SMenuEntryBlock::OnKeyDown( const FGeometry& MyGeometry, const FKeyEvent&
 	return FReply::Unhandled();
 }
 
-void SMenuEntryBlock::UpdateSubMenuState()
+//void SMenuEntryBlock::UpdateSubMenuState()
+//{
+//	// Check to see if there is a pending sub-menu request
+//	if( SubMenuRequestState != Idle )
+//	{
+//		// Reduce the time until the new menu opens
+//		TimeToSubMenuOpen -= FSlateApplication::Get().GetDeltaTime();
+//		if( TimeToSubMenuOpen <= 0.0f )
+//		{
+//			const bool bSubMenuNeedsToOpen = ( SubMenuRequestState == WantOpen );
+//			SubMenuRequestState = Idle;
+//
+//			// The menu should be opened now as our timer is up
+//			TSharedRef< const FMenuEntryBlock > MenuEntryBlock = StaticCastSharedRef< const FMenuEntryBlock >( MultiBlock.ToSharedRef() );
+//
+//			TSharedPtr< SMultiBoxWidget > PinnedOwnerMultiBoxWidget( OwnerMultiBoxWidget.Pin() );
+//			check( PinnedOwnerMultiBoxWidget.IsValid() );
+//
+//			if( bSubMenuNeedsToOpen )
+//			{
+//				// For menu bar entries, we also need to handle mouse enter/leave events, so we can show and hide
+//				// the pull-down menu appropriately
+//				check( MenuEntryBlock->EntryBuilder.IsBound() || MenuEntryBlock->MenuBuilder.IsBound() || MenuEntryBlock->EntryWidget.IsValid() );
+//
+//				// Close other open pull-down menus from this menu bar
+//				// Do we have a different pull-down menu open?
+//				TSharedPtr< SMenuAnchor > PinnedMenuAnchor( MenuAnchor.Pin() );
+//				if( PinnedOwnerMultiBoxWidget->GetOpenMenu() != PinnedMenuAnchor )
+//				{
+//					PinnedOwnerMultiBoxWidget->CloseSummonedMenus();
+//
+//					// Summon the new pull-down menu!
+//					if( PinnedMenuAnchor.IsValid() )
+//					{
+//						PinnedMenuAnchor->SetIsOpen( true );
+//					}
+//
+//					// Also tell the multibox about this open pull-down menu, so it can be closed later if we need to
+//					PinnedOwnerMultiBoxWidget->SetSummonedMenu( PinnedMenuAnchor.ToSharedRef() );
+//				}
+//			}
+//			else
+//			{
+//				PinnedOwnerMultiBoxWidget->CloseSummonedMenus();
+//			}
+//		}
+//
+//	}
+//}
+
+EActiveTickReturnType SMenuEntryBlock::UpdateSubMenuState(double InCurrentTime, float InDeltaTime, bool bWantsOpen)
 {
-	// Check to see if there is a pending sub-menu request
-	if( SubMenuRequestState != Idle )
-	{
-		// Reduce the time until the new menu opens
-		TimeToSubMenuOpen -= FSlateApplication::Get().GetDeltaTime();
-		if( TimeToSubMenuOpen <= 0.0f )
-		{
-			const bool bSubMenuNeedsToOpen = ( SubMenuRequestState == WantOpen );
-			SubMenuRequestState = Idle;
+	//const bool bSubMenuNeedsToOpen = ( SubMenuRequestState == WantOpen );
+	//SubMenuRequestState = Idle;
 
-			// The menu should be opened now as our timer is up
-			TSharedRef< const FMenuEntryBlock > MenuEntryBlock = StaticCastSharedRef< const FMenuEntryBlock >( MultiBlock.ToSharedRef() );
+	// The menu should be opened now as our timer is up
+	TSharedRef< const FMenuEntryBlock > MenuEntryBlock = StaticCastSharedRef< const FMenuEntryBlock >(MultiBlock.ToSharedRef());
 
-			TSharedPtr< SMultiBoxWidget > PinnedOwnerMultiBoxWidget( OwnerMultiBoxWidget.Pin() );
-			check( PinnedOwnerMultiBoxWidget.IsValid() );
-
-			if( bSubMenuNeedsToOpen )
-			{
-				// For menu bar entries, we also need to handle mouse enter/leave events, so we can show and hide
-				// the pull-down menu appropriately
-				check( MenuEntryBlock->EntryBuilder.IsBound() || MenuEntryBlock->MenuBuilder.IsBound() || MenuEntryBlock->EntryWidget.IsValid() );
+	TSharedPtr< SMultiBoxWidget > PinnedOwnerMultiBoxWidget(OwnerMultiBoxWidget.Pin());
+	check(PinnedOwnerMultiBoxWidget.IsValid());
 
-				// Close other open pull-down menus from this menu bar
-				// Do we have a different pull-down menu open?
-				TSharedPtr< SMenuAnchor > PinnedMenuAnchor( MenuAnchor.Pin() );
-				if( PinnedOwnerMultiBoxWidget->GetOpenMenu() != PinnedMenuAnchor )
-				{
-					PinnedOwnerMultiBoxWidget->CloseSummonedMenus();
+	if (bWantsOpen)
+	{
+		// For menu bar entries, we also need to handle mouse enter/leave events, so we can show and hide
+		// the pull-down menu appropriately
+		check(MenuEntryBlock->EntryBuilder.IsBound() || MenuEntryBlock->MenuBuilder.IsBound() || MenuEntryBlock->EntryWidget.IsValid());
 
-					// Summon the new pull-down menu!
-					if( PinnedMenuAnchor.IsValid() )
-					{
-						PinnedMenuAnchor->SetIsOpen( true );
-					}
+		// Close other open pull-down menus from this menu bar
+		// Do we have a different pull-down menu open?
+		TSharedPtr< SMenuAnchor > PinnedMenuAnchor(MenuAnchor.Pin());
+		if (PinnedOwnerMultiBoxWidget->GetOpenMenu() != PinnedMenuAnchor)
+		{
+			PinnedOwnerMultiBoxWidget->CloseSummonedMenus();
 
-					// Also tell the multibox about this open pull-down menu, so it can be closed later if we need to
-					PinnedOwnerMultiBoxWidget->SetSummonedMenu( PinnedMenuAnchor.ToSharedRef() );
-				}
-			}
-			else
+			// Summon the new pull-down menu!
+			if (PinnedMenuAnchor.IsValid())
 			{
-				PinnedOwnerMultiBoxWidget->CloseSummonedMenus();
+				PinnedMenuAnchor->SetIsOpen(true);
 			}
-		}
 
+			// Also tell the multibox about this open pull-down menu, so it can be closed later if we need to
+			PinnedOwnerMultiBoxWidget->SetSummonedMenu(PinnedMenuAnchor.ToSharedRef());
+		}
+	}
+	else
+	{
+		PinnedOwnerMultiBoxWidget->CloseSummonedMenus();
 	}
-}
-
 
-void SMenuEntryBlock::Tick( const FGeometry& AllottedGeometry, const double InCurrentTime, const float InDeltaTime )
-{
-	UpdateSubMenuState();
+	return EActiveTickReturnType::StopTicking;
 }
 
+//void SMenuEntryBlock::Tick( const FGeometry& AllottedGeometry, const double InCurrentTime, const float InDeltaTime )
+//{
+//	UpdateSubMenuState();
+//}
+
 /**
  * Called to create content for a pull-down or sub-menu window when it's summoned by the user
  *
diff --git a/Engine/Source/Runtime/Slate/Private/Framework/MultiBox/SMenuEntryBlock.h b/Engine/Source/Runtime/Slate/Private/Framework/MultiBox/SMenuEntryBlock.h
--- a/Engine/Source/Runtime/Slate/Private/Framework/MultiBox/SMenuEntryBlock.h
+++ b/Engine/Source/Runtime/Slate/Private/Framework/MultiBox/SMenuEntryBlock.h
@@ -229,7 +229,7 @@ protected:
 	virtual void OnMouseEnter( const FGeometry& MyGeometry, const FPointerEvent& MouseEvent ) override;
 	virtual void OnMouseLeave( const FPointerEvent& MouseEvent ) override;
 	virtual FReply OnKeyDown( const FGeometry& MyGeometry, const FKeyEvent& KeyEvent ) override;
-	virtual void Tick( const FGeometry& AllottedGeometry, const double InCurrentTime, const float InDeltaTime ) override;
+	//virtual void Tick( const FGeometry& AllottedGeometry, const double InCurrentTime, const float InDeltaTime ) override;
 	// End of SWidget interface
 
 	/**
@@ -311,31 +311,38 @@ protected:
 	FSlateColor InvertOnHover() const;
 
 	/** Updates state machine for sub-menu opening logic.  Called in the widget's Tick as well as on demand in some cases. */
-	void UpdateSubMenuState();
+	//void UpdateSubMenuState();
 
 private:
 	const FSlateBrush* GetCheckBoxImageBrushFromStyle(const FCheckBoxStyle* Style) const;
 
+	/** One-off delayed active tick to update the open/closed state of the sub menu. */
+	EActiveTickReturnType UpdateSubMenuState(double InCurrentTime, float InDeltaTime, bool bWantsOpen);
+
+	/** The handle to the active tick to update the sub-menu state */
+	TWeakPtr<FActiveTickHandle> ActiveTickHandle;
+
 	/** The brush to use when an item should appear checked */
 	const FSlateBrush* CheckedImage;
 	/** The brush to use when an item should appear unchecked */
 	const FSlateBrush* UncheckedImage;
 	/** For pull-down or sub-menu entries, this stores a weak reference to the menu anchor widget that we'll use to summon the menu */
 	TWeakPtr< SMenuAnchor > MenuAnchor;
+	
 	/** The time until a SubMenu action is taken (open or close).  <0 means no action taken */
-	float TimeToSubMenuOpen;
+	//float TimeToSubMenuOpen;
 
 	/** Style for menu bar button with sub menu opened */
 	const FSlateBrush* MenuBarButtonBorderSubmenuOpen;
 	/** Style for menu bar button with no sub menu opened */
 	const FSlateBrush* MenuBarButtonBorderSubmenuClosed;
 
-	/** The pending SubMenu request state */
-	enum
-	{
-		Idle,
-		WantOpen,
-		WantClose
-	} 
-	SubMenuRequestState;
+	///** The pending SubMenu request state */
+	//enum
+	//{
+	//	Idle,
+	//	WantOpen,
+	//	WantClose
+	//} 
+	//SubMenuRequestState;
 };
diff --git a/Engine/Source/Runtime/Slate/Private/Widgets/Layout/SScrollBox.cpp b/Engine/Source/Runtime/Slate/Private/Widgets/Layout/SScrollBox.cpp
--- a/Engine/Source/Runtime/Slate/Private/Widgets/Layout/SScrollBox.cpp
+++ b/Engine/Source/Runtime/Slate/Private/Widgets/Layout/SScrollBox.cpp
@@ -160,6 +160,7 @@ void SScrollBox::Construct( const FArguments& InArgs )
 	OnUserScrolled = InArgs._OnUserScrolled;
 	Orientation = InArgs._Orientation;
 	bScrollToEnd = false;
+	bIsActiveTickRegistered = false;
 
 	if (InArgs._ExternalScrollbar.IsValid())
 	{
@@ -340,7 +341,7 @@ void SScrollBox::ClearChildren()
 
 bool SScrollBox::IsRightClickScrolling() const
 {
-	return AmountScrolledWhileRightMouseDown >= FSlateApplication::Get().GetDragTriggerDistnace() && this->ScrollBar->IsNeeded();
+	return AmountScrolledWhileRightMouseDown >= FSlateApplication::Get().GetDragTriggerDistance() && this->ScrollBar->IsNeeded();
 }
 
 float SScrollBox::GetScrollOffset()
@@ -432,21 +433,33 @@ void SScrollBox::SetScrollBarThickness(FVector2D InThickness)
 	ScrollBar->SetThickness(InThickness);
 }
 
-void SScrollBox::Tick( const FGeometry& AllottedGeometry, const double InCurrentTime, const float InDeltaTime )
+EActiveTickReturnType SScrollBox::UpdateInertialScroll(double InCurrentTime, float InDeltaTime)
 {
-	if ( !IsRightClickScrolling() )
-	{
-		this->InertialScrollManager.UpdateScrollVelocity(InDeltaTime);
+	InertialScrollManager.UpdateScrollVelocity(InDeltaTime);
 
-		const float ScrollVelocity = this->InertialScrollManager.GetScrollVelocity();
-		if ( ScrollVelocity != 0.f )
+	if (bIsScrolling)
+	{
+		const float ScrollVelocity = InertialScrollManager.GetScrollVelocity();
+		if (ScrollVelocity != 0.f)
 		{
-			this->ScrollBy(AllottedGeometry, ScrollVelocity * InDeltaTime, true);
+			ScrollBy(CachedGeometry, ScrollVelocity * InDeltaTime, true);
 		}
+
+		return EActiveTickReturnType::KeepTicking;
 	}
 
+	// Clear the scroll velocity so there isn't any delayed minuscule movement
+	InertialScrollManager.ClearScrollVelocity();
+
+	bIsActiveTickRegistered = false;
+	return EActiveTickReturnType::StopTicking;
+}
+
+void SScrollBox::Tick( const FGeometry& AllottedGeometry, const double InCurrentTime, const float InDeltaTime )
+{
+	CachedGeometry = AllottedGeometry;
+
 	const FGeometry ScrollPanelGeometry = FindChildGeometry( AllottedGeometry, ScrollPanel.ToSharedRef() );
-	
 	const float ContentSize = GetScrollComponentFromVector(ScrollPanel->GetDesiredSize());
 
 	if ( bScrollToEnd )
@@ -461,7 +474,7 @@ void SScrollBox::Tick( const FGeometry& AllottedGeometry, const double InCurrent
 	
 	// Update the scrollbar with the clamped version of the offset
 	const float TargetPhysicalOffset = GetScrollComponentFromVector(ViewOffset*ScrollPanel->GetDesiredSize());
-	bIsScrolling = !FMath::IsNearlyEqual(TargetPhysicalOffset, ScrollPanel->PhysicalOffset, 0.001f);	
+	bIsScrolling = !FMath::IsNearlyEqual(TargetPhysicalOffset, ScrollPanel->PhysicalOffset, 0.001f);
 	ScrollPanel->PhysicalOffset = (bAnimateScroll)
 		? FMath::FInterpTo(ScrollPanel->PhysicalOffset, TargetPhysicalOffset, InDeltaTime, 15.f)
 		: TargetPhysicalOffset;
@@ -472,6 +485,12 @@ void SScrollBox::Tick( const FGeometry& AllottedGeometry, const double InCurrent
 		// We cannot scroll, so ensure that there is no offset.
 		ScrollPanel->PhysicalOffset = 0.0f;
 	}
+	else if (bIsScrolling && !bIsActiveTickRegistered)
+	{
+		// If scrolling and the scrollbar is needed, make sure the active tick is registered (possible b/c we may need to scroll when our geometry changes)
+		bIsActiveTickRegistered = true;
+		RegisterActiveTick(0.f, FWidgetActiveTickDelegate::CreateSP(this, &SScrollBox::UpdateInertialScroll));
+	}
 }
 
 FReply SScrollBox::OnMouseButtonDown( const FGeometry& MyGeometry, const FPointerEvent& MouseEvent )
@@ -498,6 +517,15 @@ FReply SScrollBox::OnMouseButtonUp( const FGeometry& MyGeometry, const FPointerE
 {
 	if ( MouseEvent.GetEffectingButton() == EKeys::RightMouseButton )
 	{
+		if (!bIsActiveTickRegistered && IsRightClickScrolling())
+		{
+			// Register the active tick to handle the inertial scrolling
+			CachedGeometry = MyGeometry;
+			bIsScrolling = true;
+			bIsActiveTickRegistered = true;
+			RegisterActiveTick(0.f, FWidgetActiveTickDelegate::CreateSP(this, &SScrollBox::UpdateInertialScroll));
+		}
+
 		AmountScrolledWhileRightMouseDown = 0;
 
 		FReply Reply = FReply::Handled().ReleaseMouseCapture();
@@ -575,10 +603,19 @@ FReply SScrollBox::OnMouseWheel( const FGeometry& MyGeometry, const FPointerEven
 	if (ScrollBar->IsNeeded())
 	{
 		// Make sure scroll velocity is cleared so it doesn't fight with the mouse wheel input
-		this->InertialScrollManager.ClearScrollVelocity();
+		InertialScrollManager.ClearScrollVelocity();
 
 		const bool bScrollWasHandled = this->ScrollBy( MyGeometry, -MouseEvent.GetWheelDelta()*WheelScrollAmount );
 
+		if (bScrollWasHandled && !bIsActiveTickRegistered)
+		{
+			// Register the active tick to handle the inertial scrolling
+			CachedGeometry = MyGeometry;
+			bIsScrolling = true;
+			bIsActiveTickRegistered = true;
+			RegisterActiveTick(0.f, FWidgetActiveTickDelegate::CreateSP(this, &SScrollBox::UpdateInertialScroll));
+		}
+
 		return bScrollWasHandled ? FReply::Handled() : FReply::Unhandled();
 	}
 	else
