Merging using UE4-To-UE4-LauncherReleases

Added a new ListItemAlignment option to the SListPanel that allows for STileView items to stretch horizontally instead of increasing the padding between the items.

diff --git a/Engine/Source/Runtime/Slate/Private/Widgets/Views/SListPanel.cpp b/Engine/Source/Runtime/Slate/Private/Widgets/Views/SListPanel.cpp
--- a/Engine/Source/Runtime/Slate/Private/Widgets/Views/SListPanel.cpp
+++ b/Engine/Source/Runtime/Slate/Private/Widgets/Views/SListPanel.cpp
@@ -61,11 +61,13 @@ void SListPanel::OnArrangeChildren( const FGeometry& AllottedGeometry, FArranged
 
 		// This is a tile view list, arrange items horizontally until there is no more room then create a new row.
 		const float AllottedWidth = AllottedGeometry.Size.X;
-		const float ItemPadding = (ListItemAlignment == EListItemAlignment::EvenlyDistributed) ? GetItemPadding(AllottedGeometry) : 0.0f;
+		const float ItemPadding = GetItemPadding(AllottedGeometry, ListItemAlignment);
 		const float HalfItemPadding = ItemPadding * 0.5;
+
+		const float LocalItemWidth = GetItemWidth(AllottedGeometry, ListItemAlignment);
+		const float LocalItemHeight = GetItemHeight();
+
 		float WidthSoFar = 0;
-		float LocalItemWidth = ItemWidth.Get();
-		float LocalItemHeight = ItemHeight.Get();
 		float HeightSoFar = -FMath::FloorToInt(SmoothScrollOffsetInItems * LocalItemHeight) - OverscrollAmount;
 
 		bool bIsNewLine = true;
@@ -89,9 +91,9 @@ void SListPanel::OnArrangeChildren( const FGeometry& AllottedGeometry, FArranged
 			}
 
 			ArrangedChildren.AddWidget(
-				AllottedGeometry.MakeChild( Children[ItemIndex].GetWidget(), FVector2D(WidthSoFar + HalfItemPadding, HeightSoFar), FVector2D(LocalItemWidth, LocalItemHeight) )
+				AllottedGeometry.MakeChild(Children[ItemIndex].GetWidget(), FVector2D(WidthSoFar + HalfItemPadding, HeightSoFar), FVector2D(LocalItemWidth, LocalItemHeight))
 				);
-
+		
 			WidthSoFar += LocalItemWidth + ItemPadding;
 
 			if ( WidthSoFar + LocalItemWidth + ItemPadding > AllottedWidth )
@@ -129,9 +131,10 @@ void SListPanel::Tick( const FGeometry& AllottedGeometry, const double InCurrent
 {
 	if (ShouldArrangeHorizontally())
 	{
+		const EListItemAlignment ListItemAlignment = ItemAlignment.Get();
 		const float AllottedWidth = AllottedGeometry.Size.X;
-		const float ItemPadding = GetItemPadding(AllottedGeometry);
-		const float LocalItemWidth = ItemWidth.Get();
+		const float ItemPadding = GetItemPadding(AllottedGeometry, ListItemAlignment);
+		const float LocalItemWidth = GetItemWidth(AllottedGeometry, ListItemAlignment);
 		const float TotalItemSize = LocalItemWidth + ItemPadding;
 		const int32 NumChildren = Children.Num();
 
@@ -217,29 +220,59 @@ void SListPanel::ClearItems()
 	Children.Empty();
 }
 
-float SListPanel::GetItemWidth() const
+float SListPanel::GetDesiredItemWidth() const
 {
 	return ItemWidth.Get();
 }
 
 float SListPanel::GetItemPadding(const FGeometry& AllottedGeometry) const
 {
-	const float LocalItemWidth = ItemWidth.Get();
+	return GetItemPadding(AllottedGeometry, ItemAlignment.Get());
+}
+
+float SListPanel::GetItemPadding(const FGeometry& AllottedGeometry, const EListItemAlignment ListItemAlignment) const
+{
+	const float LocalItemWidth = GetDesiredItemWidth();
 	const int32 NumItemsWide = LocalItemWidth > 0 ? FMath::FloorToInt(AllottedGeometry.Size.X / LocalItemWidth) : 0;
-	
+	const bool IsFillAlignment = ListItemAlignment == EListItemAlignment::Fill;
+
 	// Only add padding between items if we have more total items that we can fit on a single row.  Otherwise,
 	// the padding around items would continue to change proportionately with the (ample) free horizontal space
 	float Padding = 0.0f;
-	if( NumItemsWide > 0 && Children.Num() > NumItemsWide )
+	if (!IsFillAlignment && NumItemsWide > 0 && Children.Num() > NumItemsWide)
 	{
 		// Subtract a tiny amount from the available width to avoid floating point precision problems when arranging children
 		static const float FloatingPointPrecisionOffset = 0.001f;
 
 		Padding = (AllottedGeometry.Size.X - FloatingPointPrecisionOffset - NumItemsWide * LocalItemWidth) / NumItemsWide;
 	}
+
 	return Padding;
 }
 
+float SListPanel::GetItemWidth(const FGeometry& AllottedGeometry) const
+{
+	return GetItemWidth(AllottedGeometry, ItemAlignment.Get());
+}
+
+float SListPanel::GetItemWidth(const FGeometry& AllottedGeometry, const EListItemAlignment ListItemAlignment) const
+{
+	const float DesiredWidth = GetDesiredItemWidth();
+	const int32 NumItemsWide = DesiredWidth > 0 ? FMath::Min(Children.Num(), FMath::FloorToInt(AllottedGeometry.Size.X / DesiredWidth)) : 0;
+	const bool IsFillAlignment = ListItemAlignment == EListItemAlignment::Fill;
+
+	float ExtraWidth = 0.0f;
+	if (IsFillAlignment && NumItemsWide > 0)
+	{
+		// Subtract a tiny amount from the available width to avoid floating point precision problems when arranging children
+		static const float FloatingPointPrecisionOffset = 0.001f;
+
+		ExtraWidth = (AllottedGeometry.Size.X - FloatingPointPrecisionOffset - NumItemsWide * DesiredWidth) / NumItemsWide;
+	}
+
+	return DesiredWidth + ExtraWidth;
+}
+
 float SListPanel::GetLinePadding(const FGeometry& AllottedGeometry, const int32 LineStartIndex) const
 {
 	const int32 NumItemsLeft = Children.Num() - LineStartIndex;
@@ -248,7 +281,7 @@ float SListPanel::GetLinePadding(const FGeometry& AllottedGeometry, const int32
 		return 0.0f;
 	}
 
-	const float LocalItemWidth = ItemWidth.Get();
+	const float LocalItemWidth = GetItemWidth(AllottedGeometry);
 	const int32 NumItemsWide = LocalItemWidth > 0 ? FMath::FloorToInt(AllottedGeometry.Size.X / LocalItemWidth) : 0;
 	const int32 NumItemsOnLine = FMath::Min(NumItemsLeft, NumItemsWide);
 
@@ -276,7 +309,7 @@ bool SListPanel::IsRefreshPending() const
 
 bool SListPanel::ShouldArrangeHorizontally() const
 {
-	return ItemWidth.Get() > 0;
+	return GetDesiredItemWidth() > 0;
 }
 
 void SListPanel::SetItemHeight(TAttribute<float> Height)
diff --git a/Engine/Source/Runtime/Slate/Private/Widgets/Views/SListPanel.h b/Engine/Source/Runtime/Slate/Private/Widgets/Views/SListPanel.h
--- a/Engine/Source/Runtime/Slate/Private/Widgets/Views/SListPanel.h
+++ b/Engine/Source/Runtime/Slate/Private/Widgets/Views/SListPanel.h
@@ -69,12 +69,21 @@ public:
 	/** Remove all the children from this panel */
 	void ClearItems();
 
+	/** @return the uniform desired item width used when arranging children. */
+	float GetDesiredItemWidth() const;
+
+	/** @return the uniform item width used when arranging children. */
+	float GetItemWidth(const FGeometry& AllottedGeometry) const;
+
 	/** @return the uniform item width used when arranging children. */
-	float GetItemWidth() const;
+	float GetItemWidth(const FGeometry& AllottedGeometry, const EListItemAlignment ListItemAlignment) const;
 
 	/** @return the horizontal padding applied to each tile item */
 	float GetItemPadding(const FGeometry& AllottedGeometry) const;
 
+	/** @return the horizontal padding applied to each tile item */
+	float GetItemPadding(const FGeometry& AllottedGeometry, const EListItemAlignment ListItemAlignment) const;
+
 	/** @return the horizontal padding applied to all the items on a line */
 	float GetLinePadding(const FGeometry& AllottedGeometry, const int32 LineStartIndex) const;
 	
diff --git a/Engine/Source/Runtime/Slate/Private/Widgets/Views/STableViewBase.cpp b/Engine/Source/Runtime/Slate/Private/Widgets/Views/STableViewBase.cpp
--- a/Engine/Source/Runtime/Slate/Private/Widgets/Views/STableViewBase.cpp
+++ b/Engine/Source/Runtime/Slate/Private/Widgets/Views/STableViewBase.cpp
@@ -664,8 +664,7 @@ void STableViewBase::ClearWidgets()
  */
 float STableViewBase::GetItemWidth() const
 {
-	// Casting to a STilePanel here because we constructed it in Construct()
-	return ItemsPanel->GetItemWidth() + ItemsPanel->GetItemPadding(PanelGeometryLastTick);
+	return ItemsPanel->GetItemWidth(PanelGeometryLastTick) + ItemsPanel->GetItemPadding(PanelGeometryLastTick);
 }
 
 /**
diff --git a/Engine/Source/Runtime/Slate/Public/Widgets/Views/STableViewBase.h b/Engine/Source/Runtime/Slate/Public/Widgets/Views/STableViewBase.h
--- a/Engine/Source/Runtime/Slate/Public/Widgets/Views/STableViewBase.h
+++ b/Engine/Source/Runtime/Slate/Public/Widgets/Views/STableViewBase.h
@@ -30,6 +30,9 @@ enum class EListItemAlignment : uint8
 
 	/** Items are center aligned on the row (any extra space is halved and added to the left of the items) */
 	CenterAligned,
+
+	/** Items are evenly horizontally stretched to distribute any extra space */
+	Fill,
 };
 
 
