Fixes for clang build

diff --git a/Engine/Source/Runtime/Slate/Private/Widgets/Views/STableViewBase.cpp b/Engine/Source/Runtime/Slate/Private/Widgets/Views/STableViewBase.cpp
--- a/Engine/Source/Runtime/Slate/Private/Widgets/Views/STableViewBase.cpp
+++ b/Engine/Source/Runtime/Slate/Private/Widgets/Views/STableViewBase.cpp
@@ -595,10 +595,10 @@ STableViewBase::STableViewBase( ETableViewMode::Type InTableViewMode )
 	, SelectionMode( ESelectionMode::Multi )
 	, SoftwareCursorPosition( ForceInitToZero )
 	, bShowSoftwareCursor( false )
+	, bIsScrollingActiveTickRegistered( false )
 	, Overscroll()
 	, AllowOverscroll(EAllowOverscroll::Yes)
-	, bItemsNeedRefresh( false )	
-	, bIsScrollingActiveTickRegistered(false)
+	, bItemsNeedRefresh( false )
 {
 }
 
