Removed some redundant event dispatches for focus.

diff --git a/Engine/Source/Runtime/Slate/Private/Framework/Application/SlateApplication.cpp b/Engine/Source/Runtime/Slate/Private/Framework/Application/SlateApplication.cpp
--- a/Engine/Source/Runtime/Slate/Private/Framework/Application/SlateApplication.cpp
+++ b/Engine/Source/Runtime/Slate/Private/Framework/Application/SlateApplication.cpp
@@ -2035,10 +2035,6 @@ bool FSlateApplication::SetUserFocus(const uint32 InUserIndex, const FWidgetPath
 			if (SomeWidget.IsValid())
 			{
 				SomeWidget->OnFocusChanging(OldFocusedWidgetPath, NewFocusedWidgetPath);
-
-				//@Todo slate: Remove this backwards compatibility code as doon as the deprecated functionality in Swidget is removed
-				// Run the deprecated function to maintain backwards compatibility for now.
-				SomeWidget->OnFocusChanging(OldFocusedWidgetPath, NewFocusedWidgetPath);
 			}
 		}
 	}
@@ -2054,10 +2050,6 @@ bool FSlateApplication::SetUserFocus(const uint32 InUserIndex, const FWidgetPath
 			if (SomeWidget.IsValid())
 			{
 				SomeWidget->OnFocusChanging(OldFocusedWidgetPath, NewFocusedWidgetPath);
-
-				//@Todo slate: Remove this backwards compatibility code as doon as the deprecated functionality in Swidget is removed
-				// Run the deprecated function to maintain backwards compatibility for now.
-				SomeWidget->OnFocusChanging(OldFocusedWidgetPath, NewFocusedWidgetPath);
 			}
 		}
 	}
@@ -2094,10 +2086,6 @@ bool FSlateApplication::SetUserFocus(const uint32 InUserIndex, const FWidgetPath
 
 		// Let previously-focused widget know that it's losing focus
 		OldFocusedWidget->OnFocusLost(FFocusEvent(InCause));
-
-		//@Todo slate: Remove this backwards compatibility code as doon as the deprecated functionality in Swidget is removed
-		// Run the deprecated function to maintain backwards compatibility for now.
-		OldFocusedWidget->OnFocusLost(FFocusEvent(InCause));
 	}
 
 	if (bReflectorShowingFocus)
@@ -2120,16 +2108,6 @@ bool FSlateApplication::SetUserFocus(const uint32 InUserIndex, const FWidgetPath
 		{
 			ProcessReply(InFocusPath, Reply, nullptr, nullptr);
 		}
-		//@Todo slate: Remove this backwards compatibility code as soon as the deprecated functionality in Swidget is removed
-		// Run the deprecated function to maintain backwards compatibility for now.
-		else
-		{
-			FReply Reply = NewFocusedWidget->OnFocusReceived(WidgetToFocus.Geometry, FFocusEvent(InCause));
-			if (Reply.IsEventHandled())
-			{
-				ProcessReply(InFocusPath, Reply, nullptr, nullptr);
-			}
-		}
 		
 		// Set the windows restore state.
 		//@Todo Slate: We need to store the full focus state 
