Slate : Synthetic Mouse Move now only happens after input events or ActiveTicks. i.e. Any input or animation will trigger a SynthesizeMouseMove().
Only Synthetic Mouse Moves check that we are directly over a Slate window. Natural mouse moves come from the OS, so we can safely assume they were intended for a Slate Window.

diff --git a/Engine/Source/Runtime/Slate/Private/Framework/Application/SlateApplication.cpp b/Engine/Source/Runtime/Slate/Private/Framework/Application/SlateApplication.cpp
--- a/Engine/Source/Runtime/Slate/Private/Framework/Application/SlateApplication.cpp
+++ b/Engine/Source/Runtime/Slate/Private/Framework/Application/SlateApplication.cpp
@@ -663,12 +663,12 @@ void FSlateApplication::Shutdown()
 	}
 }
 
-
 TSharedPtr<FSlateApplication> FSlateApplication::CurrentApplication = nullptr;
 
 FSlateApplication::FSlateApplication()
-	: bAppIsActive( true )
-	, bSlateWindowActive( true )
+	: SynthesizeMouseMovePending(0)
+	, bAppIsActive(true)
+	, bSlateWindowActive(true)
 	, Scale( 1.0f )
 	, DragTriggerDistance( 5.0f )
 	, CursorRadius( 0.0f )
@@ -724,6 +724,7 @@ FSlateApplication::FSlateApplication()
 		GConfig->GetBool(TEXT("CursorControl"), TEXT("bAllowSoftwareCursor"), bSoftwareCursorAvailable, GEngineIni);
 	}
 
+
 	// causes InputCore to initialize, even if statically linked
 	FInputCoreModule& InputCore = FModuleManager::LoadModuleChecked<FInputCoreModule>(TEXT("InputCore"));
 
@@ -813,31 +814,27 @@ FWidgetPath FSlateApplication::LocateWindowUnderMouse( FVector2D ScreenspaceMous
 	bool bPrevWindowWasModal = false;
 	FArrangedChildren OutWidgetPath(EVisibility::Visible);
 
-	if (PlatformApplication->IsCursorDirectlyOverSlateWindow())
-	{
+	for (int32 WindowIndex = Windows.Num() - 1; WindowIndex >= 0 && OutWidgetPath.Num() == 0; --WindowIndex)
+	{ 
+		const TSharedRef<SWindow>& Window = Windows[WindowIndex];
 
-		for (int32 WindowIndex = Windows.Num() - 1; WindowIndex >= 0 && OutWidgetPath.Num() == 0; --WindowIndex)
+		// Hittest the window's children first.
+		FWidgetPath ResultingPath = LocateWindowUnderMouse(ScreenspaceMouseCoordinate, Window->GetChildWindows(), bIgnoreEnabledStatus);
+		if (ResultingPath.IsValid())
 		{
-			const TSharedRef<SWindow>& Window = Windows[WindowIndex];
-
-			// Hittest the window's children first.
-			FWidgetPath ResultingPath = LocateWindowUnderMouse(ScreenspaceMouseCoordinate, Window->GetChildWindows(), bIgnoreEnabledStatus);
-			if (ResultingPath.IsValid())
-			{
-				return ResultingPath;
-			}
+			return ResultingPath;
+		}
 
-			// If none of the children were hit, hittest the parent.
+		// If none of the children were hit, hittest the parent.
 
-			// Only accept input if the current window accepts input and the current window is not under a modal window or an interactive tooltip
+		// Only accept input if the current window accepts input and the current window is not under a modal window or an interactive tooltip
 
-			if (!bPrevWindowWasModal)
+		if (!bPrevWindowWasModal)
+		{
+			FWidgetPath PathToLocatedWidget = LocateWidgetInWindow(ScreenspaceMouseCoordinate, Window, bIgnoreEnabledStatus);
+			if (PathToLocatedWidget.IsValid())
 			{
-				FWidgetPath PathToLocatedWidget = LocateWidgetInWindow(ScreenspaceMouseCoordinate, Window, bIgnoreEnabledStatus);
-				if (PathToLocatedWidget.IsValid())
-				{
-					return PathToLocatedWidget;
-				}
+				return PathToLocatedWidget;
 			}
 		}
 	}
@@ -1227,6 +1224,9 @@ void FSlateApplication::Tick()
 		PlatformApplication->ProcessDeferredEvents( DeltaTime );
 	}
 
+	// The widget locking the cursor to its bounds may have been reshaped.
+	// Check if the widget was reshaped and update the cursor lock
+	// bounds if needed.
 	UpdateCursorLockRegion();
 
 	// When Slate captures the mouse, it is up to us to set the cursor 
@@ -1278,8 +1278,13 @@ void FSlateApplication::Tick()
 		LastTickTime = CurrentTime - MaxQuantumBeforeClamp;
 	}
 
-	// Force a mouse move event to make sure all widgets know whether there is a mouse cursor hovering over them
-	SynthesizeMouseMove();
+	const bool bNeedsSyntheticMouseMouse = SynthesizeMouseMovePending > 0;
+	if (bNeedsSyntheticMouseMouse)
+	{
+		// Force a mouse move event to make sure all widgets know whether there is a mouse cursor hovering over them
+		SynthesizeMouseMove();
+		--SynthesizeMouseMovePending;
+	}
 
 	// Update auto-throttling based on elapsed time since user interaction
 	ThrottleApplicationBasedOnMouseMovement();
@@ -1292,12 +1297,18 @@ void FSlateApplication::Tick()
 	
 	const bool bIsUserIdle = (TimeSinceInput > SleepThreshold) && (TimeSinceMouseMove > SleepThreshold);
 	const bool bAnyActiveTicksPending = AnyActiveTicksArePending();
+	if (bAnyActiveTicksPending)
+	{
+		// Some UI might slide under the cursor. To a widget, this is
+		// as if the cursor moved over it.
+		QueueSynthesizedMouseMove();
+	}
 
 	// skip tick/draw if we are idle and there are no active ticks registered that we need to drive slate for.
 	// This effectively means the slate application is totally idle and we don't need to update the UI.
 	// This relies on Widgets properly registering for Active tick when they need something to happen even
 	// when the user is not providing any input (ie, animations, viewport rendering, async polling, etc).
-	if ( !AllowSlateToSleep.GetValueOnGameThread() || bAnyActiveTicksPending || !bIsUserIdle )
+	if (!AllowSlateToSleep.GetValueOnGameThread() || bAnyActiveTicksPending || !bIsUserIdle || bNeedsSyntheticMouseMouse)
 	{
 		if (!bFoldTick)
 		{
@@ -3051,6 +3062,11 @@ void FSlateApplication::SynthesizeMouseMove()
 	}
 }
 
+void FSlateApplication::QueueSynthesizedMouseMove()
+{
+	SynthesizeMouseMovePending = 2;
+}
+
 void FSlateApplication::OnLogSlateEvent(EEventLog::Type Elovent, const FString& AdditionalContent)
 {
 #if LOG_SLATE_EVENTS
@@ -3556,6 +3572,8 @@ bool FSlateApplication::OnKeyDown( const int32 KeyCode, const uint32 CharacterCo
 
 bool FSlateApplication::ProcessKeyDownEvent( FKeyEvent& InKeyEvent )
 {
+	QueueSynthesizedMouseMove();
+
 	// Analog cursor gets first chance at the input
 	if (InputPreProcessor.IsValid() && InputPreProcessor->HandleKeyDownEvent(*this, InKeyEvent))
 	{
@@ -3651,6 +3669,8 @@ bool FSlateApplication::OnKeyUp( const int32 KeyCode, const uint32 CharacterCode
 
 bool FSlateApplication::ProcessKeyUpEvent( FKeyEvent& InKeyEvent )
 {
+	QueueSynthesizedMouseMove();
+
 	// Analog cursor gets first chance at the input
 	if (InputPreProcessor.IsValid() && InputPreProcessor->HandleKeyUpEvent(*this, InKeyEvent))
 	{
@@ -3683,6 +3703,8 @@ bool FSlateApplication::ProcessKeyUpEvent( FKeyEvent& InKeyEvent )
 
 bool FSlateApplication::ProcessAnalogInputEvent(FAnalogInputEvent& InAnalogInputEvent)
 {
+	QueueSynthesizedMouseMove();
+
 	// Analog cursor gets first chance at the input
 	if (InputPreProcessor.IsValid() && InputPreProcessor->HandleAnalogInputEvent(*this, InAnalogInputEvent))
 	{
@@ -3711,6 +3733,8 @@ bool FSlateApplication::ProcessAnalogInputEvent(FAnalogInputEvent& InAnalogInput
 
 	LOG_EVENT_CONTENT(EEventLog::AnalogInput, GetKeyName(InAnalogInputEvent.GetKey()).ToString(), Reply);
 
+	QueueSynthesizedMouseMove();
+
 	return Reply.IsEventHandled();
 }
 
@@ -3781,6 +3805,7 @@ bool FSlateApplication::OnMouseDown( const TSharedPtr< FGenericWindow >& Platfor
 
 bool FSlateApplication::ProcessMouseButtonDownEvent( const TSharedPtr< FGenericWindow >& PlatformWindow, FPointerEvent& MouseEvent )
 {
+	QueueSynthesizedMouseMove();
 	LastUserInteractionTime = this->GetCurrentTime();
 	LastUserInteractionTimeForThrottling = LastUserInteractionTime;
 	
@@ -3949,6 +3974,7 @@ bool FSlateApplication::OnMouseDoubleClick( const TSharedPtr< FGenericWindow >&
 
 bool FSlateApplication::ProcessMouseButtonDoubleClickEvent( const TSharedPtr< FGenericWindow >& PlatformWindow, FPointerEvent& InMouseEvent )
 {
+	QueueSynthesizedMouseMove();
 	LastUserInteractionTime = this->GetCurrentTime();
 	LastUserInteractionTimeForThrottling = LastUserInteractionTime;
 	
@@ -3969,6 +3995,7 @@ bool FSlateApplication::ProcessMouseButtonDoubleClickEvent( const TSharedPtr< FG
 	LOG_EVENT( EEventLog::MouseButtonDoubleClick, Reply );
 
 	PointerIndexLastPositionMap.Add(InMouseEvent.GetPointerIndex(), InMouseEvent.GetScreenSpacePosition());
+
 	return Reply.IsEventHandled();
 }
 
@@ -3998,6 +4025,7 @@ bool FSlateApplication::OnMouseUp( const EMouseButtons::Type Button )
 
 bool FSlateApplication::ProcessMouseButtonUpEvent( FPointerEvent& MouseEvent )
 {
+	QueueSynthesizedMouseMove();
 	LastUserInteractionTime = this->GetCurrentTime();
 	LastUserInteractionTimeForThrottling = LastUserInteractionTime;
 	PressedMouseButtons.Remove( MouseEvent.GetEffectingButton() );
@@ -4120,6 +4148,8 @@ bool FSlateApplication::OnMouseWheel( const float Delta )
 
 bool FSlateApplication::ProcessMouseWheelOrGestureEvent( FPointerEvent& InWheelEvent, const FPointerEvent* InGestureEvent )
 {
+	QueueSynthesizedMouseMove();
+
 	const bool bShouldProcessEvent = InWheelEvent.GetWheelDelta() != 0 ||
 		(InGestureEvent != nullptr && InGestureEvent->GetGestureDelta()!=FVector2D::ZeroVector);
 	
@@ -4220,6 +4250,8 @@ bool FSlateApplication::ProcessMouseMoveEvent( FPointerEvent& MouseEvent, bool b
 {
 	if ( !bIsSynthetic )
 	{
+		QueueSynthesizedMouseMove();
+
 		// Detecting a mouse move of zero delta is our way of filtering out synthesized move events
 		const bool AllowSpawningOfToolTips = true;
 		UpdateToolTip( AllowSpawningOfToolTips );
@@ -4228,7 +4260,15 @@ bool FSlateApplication::ProcessMouseMoveEvent( FPointerEvent& MouseEvent, bool b
 		LastUserInteractionTime = this->GetCurrentTime();
 	}
 
-	FWidgetPath WidgetsUnderCursor = LocateWindowUnderMouse( MouseEvent.GetScreenSpacePosition(), GetInteractiveTopLevelWindows() );
+	// When the event came from the OS, we are guaranteed to be over a slate window.
+	// Otherwise, we are synthesizing a MouseMove ourselves, and must verify that the
+	// cursor is indeed over a Slate window.
+	const bool bOverSlateWindow = !bIsSynthetic || PlatformApplication->IsCursorDirectlyOverSlateWindow();
+	
+	FWidgetPath WidgetsUnderCursor = bOverSlateWindow
+		? LocateWindowUnderMouse(MouseEvent.GetScreenSpacePosition(), GetInteractiveTopLevelWindows())
+		: FWidgetPath();
+
 	bool bHandled = false;
 
 	FWeakWidgetPath LastWidgetsUnderCursor;
@@ -4704,6 +4744,7 @@ bool FSlateApplication::OnMotionDetected(const FVector& Tilt, const FVector& Rot
 
 void FSlateApplication::ProcessMotionDetectedEvent( FMotionEvent& MotionEvent )
 {
+	QueueSynthesizedMouseMove();
 	LastUserInteractionTime = this->GetCurrentTime();
 	
 	if (MotionEvent.GetUserIndex() < ARRAY_COUNT(UserFocusEntries) && UserFocusEntries[MotionEvent.GetUserIndex()].WidgetPath.IsValid())
diff --git a/Engine/Source/Runtime/Slate/Public/Framework/Application/SlateApplication.h b/Engine/Source/Runtime/Slate/Public/Framework/Application/SlateApplication.h
--- a/Engine/Source/Runtime/Slate/Public/Framework/Application/SlateApplication.h
+++ b/Engine/Source/Runtime/Slate/Public/Framework/Application/SlateApplication.h
@@ -886,12 +886,6 @@ public:
 	 */
 	void UpdateToolTip( bool AllowSpawningOfNewToolTips );
 
-	/**
-	 * Creates a mouse move event for the last known cursor position.  This should be called every tick to make
-	 * sure that widgets that appear (or vanish from) underneath the cursor have hover state set appropriately.
-	 */
-	virtual void SynthesizeMouseMove();
-
 	/** @return an array of top-level windows that can be interacted with. e.g. when a modal window is up, only return the modal window */
 	TArray< TSharedRef<SWindow> > GetInteractiveTopLevelWindows();
 
@@ -1105,6 +1099,15 @@ private:
 
 private:
 
+	/**
+	 * Creates a mouse move event for the last known cursor position.  This should be called every tick to make
+	 * sure that widgets that appear (or vanish from) underneath the cursor have hover state set appropriately.
+	 */
+	void SynthesizeMouseMove();
+
+	/** Signal that a synthesized mouse move will be required after this operation. */
+	void QueueSynthesizedMouseMove();
+
 	/**
 	 * Will be invoked when the size of the geometry of the virtual
 	 * desktop changes (e.g. resolution change or monitors re-arranged)
@@ -1116,6 +1119,9 @@ private:
 
 	TSet<FKey> PressedMouseButtons;
 
+	/** After processing an event or performing an active tick, we need to synthesize a mouse move. @see SynthesizeMouseMove */
+	int32 SynthesizeMouseMovePending;
+
 	/** true when the slate app is active; i.e. the current foreground window is from our Slate app*/
 	bool bAppIsActive;
 
