Changing Slate.SleepBufferPostInput to default to 0

diff --git a/Engine/Source/Runtime/Slate/Private/Framework/Application/SlateApplication.cpp b/Engine/Source/Runtime/Slate/Private/Framework/Application/SlateApplication.cpp
--- a/Engine/Source/Runtime/Slate/Private/Framework/Application/SlateApplication.cpp
+++ b/Engine/Source/Runtime/Slate/Private/Framework/Application/SlateApplication.cpp
@@ -348,7 +348,7 @@ TAutoConsoleVariable<int32> AllowSlateToSleep(
 /** The amount of time that must pass without any user action before Slate is put to sleep (provided that there are no active ticks). */
 TAutoConsoleVariable<float> SleepBufferPostInput(
 	TEXT("Slate.SleepBufferPostInput"),
-	1.0f,
+	0.0f,
 	TEXT("The amount of time that must pass without any user action before Slate is put to sleep (provided that there are no active ticks)."));
 
 //////////////////////////////////////////////////////////////////////////
