Modifications to search box to improve customization

#nonjira

#branch UE4

Added ImagePadding to FSearchBoxStyle which allows custom padding around the various icons in the search box

Added Style to SuperSearchBox as TOptional (because it's in it's own module, it's easier to add "optional" args this way), and modified SearchBox to also use the same method, which means we reduce duplicate instances of the default search box style around the code.

diff --git a/Engine/Source/Developer/SuperSearch/Private/SSuperSearch.cpp b/Engine/Source/Developer/SuperSearch/Private/SSuperSearch.cpp
--- a/Engine/Source/Developer/SuperSearch/Private/SSuperSearch.cpp
+++ b/Engine/Source/Developer/SuperSearch/Private/SSuperSearch.cpp
@@ -50,12 +50,13 @@ void SSuperSearchBox::Construct( const FArguments& InArgs )
 		.Method( EPopupMethod::UseCurrentWindow )
 			[
 				SAssignNew(InputText, SSearchBox)
+					.Style(InArgs._Style)
 					.OnTextCommitted(this, &SSuperSearchBox::OnTextCommitted)
 					.HintText( NSLOCTEXT( "SuperSearchBox", "HelpHint", "Search For Help" ) )
 					.OnTextChanged(this, &SSuperSearchBox::OnTextChanged)
 					.SelectAllTextWhenFocused(false)
 					.DelayChangeNotificationsWhileTyping(true)
-			.MinDesiredWidth(200)
+					.MinDesiredWidth(200)
 			]
 			.MenuContent
 			(
diff --git a/Engine/Source/Developer/SuperSearch/Private/SSuperSearch.h b/Engine/Source/Developer/SuperSearch/Private/SSuperSearch.h
--- a/Engine/Source/Developer/SuperSearch/Private/SSuperSearch.h
+++ b/Engine/Source/Developer/SuperSearch/Private/SSuperSearch.h
@@ -29,9 +29,13 @@ class SSuperSearchBox
 public:
 
 	SLATE_BEGIN_ARGS( SSuperSearchBox )
-		: _SuggestionListPlacement( MenuPlacement_ComboBoxRight )
+		: _Style()
+		, _SuggestionListPlacement( MenuPlacement_ComboBoxRight )
 		{}
 
+		/** Style used to draw this search box */
+		SLATE_ARGUMENT( TOptional<const FSearchBoxStyle*>, Style )
+
 		/** Where to place the suggestion list */
 		SLATE_ARGUMENT( EMenuPlacement, SuggestionListPlacement )
 
diff --git a/Engine/Source/Developer/SuperSearch/Private/SuperSearchModule.cpp b/Engine/Source/Developer/SuperSearch/Private/SuperSearchModule.cpp
--- a/Engine/Source/Developer/SuperSearch/Private/SuperSearchModule.cpp
+++ b/Engine/Source/Developer/SuperSearch/Private/SuperSearchModule.cpp
@@ -18,9 +18,9 @@ void FSuperSearchModule::ShutdownModule()
 {
 }
 
-TSharedRef< SWidget > FSuperSearchModule::MakeSearchBox( TSharedPtr< SEditableTextBox >& OutExposedEditableTextBox ) const
+TSharedRef< SWidget > FSuperSearchModule::MakeSearchBox(TSharedPtr< SEditableTextBox >& OutExposedEditableTextBox, const TOptional<const FSearchBoxStyle*> InStyle) const
 {
-	TSharedRef< SSuperSearchBox > NewSearchBox = SNew( SSuperSearchBox );
+	TSharedRef< SSuperSearchBox > NewSearchBox = SNew(SSuperSearchBox).Style(InStyle);
 	OutExposedEditableTextBox = NewSearchBox->GetEditableTextBox();
 	return NewSearchBox;
 }
diff --git a/Engine/Source/Developer/SuperSearch/Public/SuperSearchModule.h b/Engine/Source/Developer/SuperSearch/Public/SuperSearchModule.h
--- a/Engine/Source/Developer/SuperSearch/Public/SuperSearchModule.h
+++ b/Engine/Source/Developer/SuperSearch/Public/SuperSearchModule.h
@@ -12,5 +12,5 @@ public:
 
 	/** Generates aSuperSearch box widget.  Remember, this widget will become invalid if the
 	   SuperSearch DLL is unloaded on the fly. */
-	virtual TSharedRef< SWidget > MakeSearchBox( TSharedPtr< SEditableTextBox >& OutExposedEditableTextBox ) const;
+	virtual TSharedRef< SWidget > MakeSearchBox(TSharedPtr< SEditableTextBox >& OutExposedEditableTextBox, const TOptional<const FSearchBoxStyle*> InStyle = TOptional<const FSearchBoxStyle*>()) const;
 };
diff --git a/Engine/Source/Runtime/Slate/Private/Widgets/Input/SSearchBox.cpp b/Engine/Source/Runtime/Slate/Private/Widgets/Input/SSearchBox.cpp
--- a/Engine/Source/Runtime/Slate/Private/Widgets/Input/SSearchBox.cpp
+++ b/Engine/Source/Runtime/Slate/Private/Widgets/Input/SSearchBox.cpp
@@ -7,7 +7,8 @@ const double SSearchBox::FilterDelayAfterTyping = 0.25f;
 
 void SSearchBox::Construct( const FArguments& InArgs )
 {
-	check(InArgs._Style);
+	// Allow style to be optionally overridden, but fallback to default if not specified
+	const FSearchBoxStyle* InStyle = InArgs._Style.IsSet() ? InArgs._Style.GetValue() : &FCoreStyle::Get().GetWidgetStyle<FSearchBoxStyle>("SearchBox");
 
 	bIsActiveTimerRegistered = false;
 	OnSearchDelegate = InArgs._OnSearch;
@@ -15,11 +16,11 @@ void SSearchBox::Construct( const FArguments& InArgs )
 	OnTextCommittedDelegate = InArgs._OnTextCommitted;
 	DelayChangeNotificationsWhileTyping = InArgs._DelayChangeNotificationsWhileTyping;
 
-	InactiveFont = InArgs._Style->TextBoxStyle.Font;
-	ActiveFont = InArgs._Style->ActiveFontInfo;
+	InactiveFont = InStyle->TextBoxStyle.Font;
+	ActiveFont = InStyle->ActiveFontInfo;
 
 	SEditableTextBox::Construct( SEditableTextBox::FArguments()
-		.Style( &InArgs._Style->TextBoxStyle )
+		.Style( &InStyle->TextBoxStyle )
 		.Font( this, &SSearchBox::GetWidgetFont )
 		.Text( InArgs._InitialText )
 		.HintText( InArgs._HintText )
@@ -36,6 +37,7 @@ void SSearchBox::Construct( const FArguments& InArgs )
 	{
 		Box->AddSlot()
 		.AutoWidth()
+		.Padding(InStyle->ImagePadding)
 		.HAlign(HAlign_Center)
 		.VAlign(VAlign_Center)
 		[
@@ -49,12 +51,13 @@ void SSearchBox::Construct( const FArguments& InArgs )
 			.IsFocusable(false)
 			[
 				SNew(SImage)
-				.Image( &InArgs._Style->UpArrowImage )
+				.Image( &InStyle->UpArrowImage )
 				.ColorAndOpacity( FSlateColor::UseForeground() )
 			]
 		];
 		Box->AddSlot()
 		.AutoWidth()
+		.Padding(InStyle->ImagePadding)
 		.HAlign(HAlign_Center)
 		.VAlign(VAlign_Center)
 		[
@@ -68,7 +71,7 @@ void SSearchBox::Construct( const FArguments& InArgs )
 			.IsFocusable(false)
 			[
 				SNew(SImage)
-				.Image( &InArgs._Style->DownArrowImage )
+				.Image( &InStyle->DownArrowImage )
 				.ColorAndOpacity( FSlateColor::UseForeground() )
 			]
 		];
@@ -77,18 +80,22 @@ void SSearchBox::Construct( const FArguments& InArgs )
 	// Add a search glass image so that the user knows this text box is for searching
 	Box->AddSlot()
 	.AutoWidth()
+	.Padding(InStyle->ImagePadding)
 	.HAlign(HAlign_Center)
 	.VAlign(VAlign_Center)
 	[
 		SNew(SImage)
 		.Visibility(this, &SSearchBox::GetSearchGlassVisibility)
-		.Image( &InArgs._Style->GlassImage )
+		.Image( &InStyle->GlassImage )
 		.ColorAndOpacity( FSlateColor::UseForeground() )
 	];
 
 	// Add an X to clear the search whenever there is some text typed into it
 	Box->AddSlot()
 	.AutoWidth()
+	.Padding(InStyle->ImagePadding)
+	.HAlign(HAlign_Center)
+	.VAlign(VAlign_Center)
 	[
 		SNew(SButton)
 		.Visibility(this, &SSearchBox::GetXVisibility)
@@ -103,7 +110,7 @@ void SSearchBox::Construct( const FArguments& InArgs )
 		.IsFocusable(false)
 		[
 			SNew(SImage)
-			.Image( &InArgs._Style->ClearImage )
+			.Image( &InStyle->ClearImage )
 			.ColorAndOpacity( FSlateColor::UseForeground() )
 		]
 	];
diff --git a/Engine/Source/Runtime/Slate/Public/Widgets/Input/SSearchBox.h b/Engine/Source/Runtime/Slate/Public/Widgets/Input/SSearchBox.h
--- a/Engine/Source/Runtime/Slate/Public/Widgets/Input/SSearchBox.h
+++ b/Engine/Source/Runtime/Slate/Public/Widgets/Input/SSearchBox.h
@@ -21,7 +21,7 @@ public:
 	DECLARE_DELEGATE_OneParam(FOnSearch, SSearchBox::SearchDirection);
 
 	SLATE_BEGIN_ARGS(SSearchBox)
-		: _Style( &FCoreStyle::Get().GetWidgetStyle<FSearchBoxStyle>("SearchBox") )
+		: _Style()
 		, _HintText( LOCTEXT("SearchHint", "Search") )
 		, _InitialText()
 		, _OnTextChanged()
@@ -32,7 +32,7 @@ public:
 	{ }
 
 		/** Style used to draw this search box */
-		SLATE_STYLE_ARGUMENT( FSearchBoxStyle, Style )
+		SLATE_ARGUMENT( TOptional<const FSearchBoxStyle*>, Style )
 
 		/** The text displayed in the SearchBox when no text has been entered */
 		SLATE_ATTRIBUTE( FText, HintText )
diff --git a/Engine/Source/Runtime/SlateCore/Public/Styling/SlateTypes.h b/Engine/Source/Runtime/SlateCore/Public/Styling/SlateTypes.h
--- a/Engine/Source/Runtime/SlateCore/Public/Styling/SlateTypes.h
+++ b/Engine/Source/Runtime/SlateCore/Public/Styling/SlateTypes.h
@@ -795,6 +795,11 @@ struct SLATECORE_API FSearchBoxStyle : public FSlateWidgetStyle
 	UPROPERTY(EditAnywhere, Category=Appearance)
 	FSlateBrush ClearImage;
 	FSearchBoxStyle& SetClearImage( const FSlateBrush& InClearImage ){ ClearImage = InClearImage; return *this; }
+
+	/** Padding to use around the images */
+	UPROPERTY(EditAnywhere, Category = Appearance)
+	FMargin ImagePadding;
+	FSearchBoxStyle& SetImagePadding(const FMargin& InImagePadding){ ImagePadding = InImagePadding; return *this; }
 };
 
 
