Added funcs to insert slots (and fetch the number of slots)

#nonjira

#branch UE4

#change I encountered an issue with SSearchBox customizations where it wasn't possible to reposition the search button to before the text area as slots are added from left to right (top to bottom) and the first slot that's created IS the text area, so I needed to add the ability to be able to add slots into the list at any specified index.

#note I also tried adding Slots as args to the editable text box to insert them prior to the text box being created, however you can't copy slots and managing a ptr to an array of them instead may have been (argueably) more complex.

diff --git a/Engine/Source/Runtime/Slate/Private/Widgets/Input/SSearchBox.cpp b/Engine/Source/Runtime/Slate/Private/Widgets/Input/SSearchBox.cpp
--- a/Engine/Source/Runtime/Slate/Private/Widgets/Input/SSearchBox.cpp
+++ b/Engine/Source/Runtime/Slate/Private/Widgets/Input/SSearchBox.cpp
@@ -32,10 +32,13 @@ void SSearchBox::Construct( const FArguments& InArgs )
 		.MinDesiredWidth( InArgs._MinDesiredWidth )
 	);
 
+	// If we want to have the buttons appear to the left of the text box we have to insert the slots instead of add them
+	int32 SlotIndex = InStyle->bLeftAlignButtons ? 0 : Box->NumSlots();
+
 	// If a search delegate was bound, add a previous and next button
 	if (OnSearchDelegate.IsBound())
 	{
-		Box->AddSlot()
+		Box->InsertSlot(SlotIndex++)
 		.AutoWidth()
 		.Padding(InStyle->ImagePadding)
 		.HAlign(HAlign_Center)
@@ -55,7 +58,7 @@ void SSearchBox::Construct( const FArguments& InArgs )
 				.ColorAndOpacity( FSlateColor::UseForeground() )
 			]
 		];
-		Box->AddSlot()
+		Box->InsertSlot(SlotIndex++)
 		.AutoWidth()
 		.Padding(InStyle->ImagePadding)
 		.HAlign(HAlign_Center)
@@ -78,7 +81,7 @@ void SSearchBox::Construct( const FArguments& InArgs )
 	}
 
 	// Add a search glass image so that the user knows this text box is for searching
-	Box->AddSlot()
+	Box->InsertSlot(SlotIndex++)
 	.AutoWidth()
 	.Padding(InStyle->ImagePadding)
 	.HAlign(HAlign_Center)
@@ -91,7 +94,7 @@ void SSearchBox::Construct( const FArguments& InArgs )
 	];
 
 	// Add an X to clear the search whenever there is some text typed into it
-	Box->AddSlot()
+	Box->InsertSlot(SlotIndex++)
 	.AutoWidth()
 	.Padding(InStyle->ImagePadding)
 	.HAlign(HAlign_Center)
diff --git a/Engine/Source/Runtime/SlateCore/Private/Styling/SlateTypes.cpp b/Engine/Source/Runtime/SlateCore/Private/Styling/SlateTypes.cpp
--- a/Engine/Source/Runtime/SlateCore/Private/Styling/SlateTypes.cpp
+++ b/Engine/Source/Runtime/SlateCore/Private/Styling/SlateTypes.cpp
@@ -341,6 +341,7 @@ const FExpandableAreaStyle& FExpandableAreaStyle::GetDefault()
 
 
 FSearchBoxStyle::FSearchBoxStyle()
+	: bLeftAlignButtons(false)
 {
 }
 
diff --git a/Engine/Source/Runtime/SlateCore/Public/Styling/SlateTypes.h b/Engine/Source/Runtime/SlateCore/Public/Styling/SlateTypes.h
--- a/Engine/Source/Runtime/SlateCore/Public/Styling/SlateTypes.h
+++ b/Engine/Source/Runtime/SlateCore/Public/Styling/SlateTypes.h
@@ -800,6 +800,11 @@ struct SLATECORE_API FSearchBoxStyle : public FSlateWidgetStyle
 	UPROPERTY(EditAnywhere, Category = Appearance)
 	FMargin ImagePadding;
 	FSearchBoxStyle& SetImagePadding(const FMargin& InImagePadding){ ImagePadding = InImagePadding; return *this; }
+
+	/** If true, buttons appear to the left of the search text */
+	UPROPERTY(EditAnywhere, Category = Appearance)
+	bool bLeftAlignButtons;
+	FSearchBoxStyle& SetLeftAlignButtons(bool bInLeftAlignButtons){ bLeftAlignButtons = bInLeftAlignButtons; return *this; }
 };
 
 
diff --git a/Engine/Source/Runtime/SlateCore/Public/Widgets/SBoxPanel.h b/Engine/Source/Runtime/SlateCore/Public/Widgets/SBoxPanel.h
--- a/Engine/Source/Runtime/SlateCore/Public/Widgets/SBoxPanel.h
+++ b/Engine/Source/Runtime/SlateCore/Public/Widgets/SBoxPanel.h
@@ -200,6 +200,18 @@ public:
 		return NewSlot;
 	}
 
+	FSlot& InsertSlot( int32 Index )
+	{
+		SHorizontalBox::FSlot& NewSlot = *new SHorizontalBox::FSlot();
+		this->Children.Insert(&NewSlot, Index);
+		return NewSlot;
+	}
+
+	int32 NumSlots() const
+	{
+		return this->Children.Num();
+	}
+
 	FORCENOINLINE SHorizontalBox()
 	: SBoxPanel( Orient_Horizontal )
 	{
@@ -327,6 +339,18 @@ public:
 		return NewSlot;
 	}
 
+	FSlot& InsertSlot(int32 Index)
+	{
+		SVerticalBox::FSlot& NewSlot = *new SVerticalBox::FSlot();
+		this->Children.Insert(&NewSlot, Index);
+		return NewSlot;
+	}
+
+	int32 NumSlots() const
+	{
+		return this->Children.Num();
+	}
+
 	FORCENOINLINE SVerticalBox()
 	: SBoxPanel( Orient_Vertical )
 	{
