Fix for nonunity compiler error

#nonjira

#branch UE4

diff --git a/Engine/Source/Runtime/Slate/Public/Widgets/Notifications/INotificationWidget.h b/Engine/Source/Runtime/Slate/Public/Widgets/Notifications/INotificationWidget.h
--- a/Engine/Source/Runtime/Slate/Public/Widgets/Notifications/INotificationWidget.h
+++ b/Engine/Source/Runtime/Slate/Public/Widgets/Notifications/INotificationWidget.h
@@ -2,6 +2,8 @@
 
 #pragma once
 
+#include "SNotificationList.h"
+
 class SLATE_API INotificationWidget
 {
 public:
