Rollback TOptional change for SSearchBox

#nonjira

#branch UE4

diff --git a/Engine/Source/Developer/SuperSearch/Private/SSuperSearch.cpp b/Engine/Source/Developer/SuperSearch/Private/SSuperSearch.cpp
--- a/Engine/Source/Developer/SuperSearch/Private/SSuperSearch.cpp
+++ b/Engine/Source/Developer/SuperSearch/Private/SSuperSearch.cpp
@@ -43,6 +43,9 @@ SSuperSearchBox::SSuperSearchBox()
 BEGIN_SLATE_FUNCTION_BUILD_OPTIMIZATION
 void SSuperSearchBox::Construct( const FArguments& InArgs )
 {
+	// Allow style to be optionally overridden, but fallback to SSearchBox default if not specified
+	const FSearchBoxStyle* InStyle = InArgs._Style.IsSet() ? InArgs._Style.GetValue() : &FCoreStyle::Get().GetWidgetStyle<FSearchBoxStyle>("SearchBox");
+
 	ChildSlot
 	[
 		SAssignNew( SuggestionBox, SMenuAnchor )
@@ -50,7 +53,7 @@ void SSuperSearchBox::Construct( const FArguments& InArgs )
 		.Method( EPopupMethod::UseCurrentWindow )
 			[
 				SAssignNew(InputText, SSearchBox)
-					.Style(InArgs._Style)
+					.Style(InStyle)
 					.OnTextCommitted(this, &SSuperSearchBox::OnTextCommitted)
 					.HintText( NSLOCTEXT( "SuperSearchBox", "HelpHint", "Search For Help" ) )
 					.OnTextChanged(this, &SSuperSearchBox::OnTextChanged)
diff --git a/Engine/Source/Runtime/Slate/Private/Widgets/Input/SSearchBox.cpp b/Engine/Source/Runtime/Slate/Private/Widgets/Input/SSearchBox.cpp
--- a/Engine/Source/Runtime/Slate/Private/Widgets/Input/SSearchBox.cpp
+++ b/Engine/Source/Runtime/Slate/Private/Widgets/Input/SSearchBox.cpp
@@ -7,8 +7,7 @@ const double SSearchBox::FilterDelayAfterTyping = 0.25f;
 
 void SSearchBox::Construct( const FArguments& InArgs )
 {
-	// Allow style to be optionally overridden, but fallback to default if not specified
-	const FSearchBoxStyle* InStyle = InArgs._Style.IsSet() ? InArgs._Style.GetValue() : &FCoreStyle::Get().GetWidgetStyle<FSearchBoxStyle>("SearchBox");
+	check(InArgs._Style);
 
 	bIsActiveTimerRegistered = false;
 	OnSearchDelegate = InArgs._OnSearch;
@@ -16,11 +15,11 @@ void SSearchBox::Construct( const FArguments& InArgs )
 	OnTextCommittedDelegate = InArgs._OnTextCommitted;
 	DelayChangeNotificationsWhileTyping = InArgs._DelayChangeNotificationsWhileTyping;
 
-	InactiveFont = InStyle->TextBoxStyle.Font;
-	ActiveFont = InStyle->ActiveFontInfo;
+	InactiveFont = InArgs._Style->TextBoxStyle.Font;
+	ActiveFont = InArgs._Style->ActiveFontInfo;
 
 	SEditableTextBox::Construct( SEditableTextBox::FArguments()
-		.Style( &InStyle->TextBoxStyle )
+		.Style( &InArgs._Style->TextBoxStyle )
 		.Font( this, &SSearchBox::GetWidgetFont )
 		.Text( InArgs._InitialText )
 		.HintText( InArgs._HintText )
@@ -33,14 +32,14 @@ void SSearchBox::Construct( const FArguments& InArgs )
 	);
 
 	// If we want to have the buttons appear to the left of the text box we have to insert the slots instead of add them
-	int32 SlotIndex = InStyle->bLeftAlignButtons ? 0 : Box->NumSlots();
+	int32 SlotIndex = InArgs._Style->bLeftAlignButtons ? 0 : Box->NumSlots();
 
 	// If a search delegate was bound, add a previous and next button
 	if (OnSearchDelegate.IsBound())
 	{
 		Box->InsertSlot(SlotIndex++)
 		.AutoWidth()
-		.Padding(InStyle->ImagePadding)
+		.Padding(InArgs._Style->ImagePadding)
 		.HAlign(HAlign_Center)
 		.VAlign(VAlign_Center)
 		[
@@ -54,13 +53,13 @@ void SSearchBox::Construct( const FArguments& InArgs )
 			.IsFocusable(false)
 			[
 				SNew(SImage)
-				.Image( &InStyle->UpArrowImage )
+				.Image( &InArgs._Style->UpArrowImage )
 				.ColorAndOpacity( FSlateColor::UseForeground() )
 			]
 		];
 		Box->InsertSlot(SlotIndex++)
 		.AutoWidth()
-		.Padding(InStyle->ImagePadding)
+		.Padding(InArgs._Style->ImagePadding)
 		.HAlign(HAlign_Center)
 		.VAlign(VAlign_Center)
 		[
@@ -74,7 +73,7 @@ void SSearchBox::Construct( const FArguments& InArgs )
 			.IsFocusable(false)
 			[
 				SNew(SImage)
-				.Image( &InStyle->DownArrowImage )
+				.Image( &InArgs._Style->DownArrowImage )
 				.ColorAndOpacity( FSlateColor::UseForeground() )
 			]
 		];
@@ -83,20 +82,20 @@ void SSearchBox::Construct( const FArguments& InArgs )
 	// Add a search glass image so that the user knows this text box is for searching
 	Box->InsertSlot(SlotIndex++)
 	.AutoWidth()
-	.Padding(InStyle->ImagePadding)
+	.Padding(InArgs._Style->ImagePadding)
 	.HAlign(HAlign_Center)
 	.VAlign(VAlign_Center)
 	[
 		SNew(SImage)
 		.Visibility(this, &SSearchBox::GetSearchGlassVisibility)
-		.Image( &InStyle->GlassImage )
+		.Image( &InArgs._Style->GlassImage )
 		.ColorAndOpacity( FSlateColor::UseForeground() )
 	];
 
 	// Add an X to clear the search whenever there is some text typed into it
 	Box->InsertSlot(SlotIndex++)
 	.AutoWidth()
-	.Padding(InStyle->ImagePadding)
+	.Padding(InArgs._Style->ImagePadding)
 	.HAlign(HAlign_Center)
 	.VAlign(VAlign_Center)
 	[
@@ -113,7 +112,7 @@ void SSearchBox::Construct( const FArguments& InArgs )
 		.IsFocusable(false)
 		[
 			SNew(SImage)
-			.Image( &InStyle->ClearImage )
+			.Image( &InArgs._Style->ClearImage )
 			.ColorAndOpacity( FSlateColor::UseForeground() )
 		]
 	];
diff --git a/Engine/Source/Runtime/Slate/Public/Widgets/Input/SSearchBox.h b/Engine/Source/Runtime/Slate/Public/Widgets/Input/SSearchBox.h
--- a/Engine/Source/Runtime/Slate/Public/Widgets/Input/SSearchBox.h
+++ b/Engine/Source/Runtime/Slate/Public/Widgets/Input/SSearchBox.h
@@ -21,7 +21,7 @@ public:
 	DECLARE_DELEGATE_OneParam(FOnSearch, SSearchBox::SearchDirection);
 
 	SLATE_BEGIN_ARGS(SSearchBox)
-		: _Style()
+		: _Style( &FCoreStyle::Get().GetWidgetStyle<FSearchBoxStyle>("SearchBox") )
 		, _HintText( LOCTEXT("SearchHint", "Search") )
 		, _InitialText()
 		, _OnTextChanged()
@@ -32,7 +32,7 @@ public:
 	{ }
 
 		/** Style used to draw this search box */
-		SLATE_ARGUMENT( TOptional<const FSearchBoxStyle*>, Style )
+		SLATE_STYLE_ARGUMENT( FSearchBoxStyle, Style )
 
 		/** The text displayed in the SearchBox when no text has been entered */
 		SLATE_ATTRIBUTE( FText, HintText )
