Slate - The responsive grid control no longer crashes when no rows have been filled out.

diff --git a/Engine/Source/Runtime/Slate/Private/Widgets/Layout/SResponsiveGridPanel.cpp b/Engine/Source/Runtime/Slate/Private/Widgets/Layout/SResponsiveGridPanel.cpp
--- a/Engine/Source/Runtime/Slate/Private/Widgets/Layout/SResponsiveGridPanel.cpp
+++ b/Engine/Source/Runtime/Slate/Private/Widgets/Layout/SResponsiveGridPanel.cpp
@@ -94,6 +94,12 @@ int32 SResponsiveGridPanel::OnPaint( const FPaintArgs& Args, const FGeometry& Al
 
 void SResponsiveGridPanel::OnArrangeChildren( const FGeometry& AllottedGeometry, FArrangedChildren& ArrangedChildren ) const
 {
+	// Don't attempt to array anything if we don't have any slots allocated.
+	if ( Slots.Num() == 0 )
+	{
+		return;
+	}
+
 	// PREPARE PHASE
 	// Prepare some data for arranging children.
 	// FinalColumns will be populated with column sizes that include the stretched column sizes.
@@ -274,17 +280,21 @@ void SResponsiveGridPanel::CacheDesiredSize()
 	ComputeDesiredCellSizes(PreviousWidth, Columns, Rows, RowToSlot);
 
 	TotalDesiredSizes = FVector2D::ZeroVector;
-	for (int ColId = 0; ColId < Columns.Num(); ++ColId)
-	{
-		TotalDesiredSizes.X += Columns[ColId];
-	}
-	TotalDesiredSizes.X += (ColumnGutter * 2) * (TotalColumns - 1);
 
-	for (int RowId = 0; RowId < Rows.Num(); ++RowId)
+	if ( Slots.Num() > 0 )
 	{
-		TotalDesiredSizes.Y += Rows[RowId];
+		for ( int ColId = 0; ColId < Columns.Num(); ++ColId )
+		{
+			TotalDesiredSizes.X += Columns[ColId];
+		}
+		TotalDesiredSizes.X += ( ColumnGutter * 2 ) * ( TotalColumns - 1 );
+
+		for ( int RowId = 0; RowId < Rows.Num(); ++RowId )
+		{
+			TotalDesiredSizes.Y += Rows[RowId];
+		}
+		TotalDesiredSizes.Y += ( RowGutter * 2 )  * ( Slots[Slots.Num() - 1].RowParam );
 	}
-	TotalDesiredSizes.Y += (RowGutter * 2)  * (Slots[Slots.Num() - 1].RowParam);
 
 	SPanel::CacheDesiredSize();
 }
