Slate - If the combobox's popup list has cached an empty size, and the list is refreshed, it wont be able to compute a new desired size.  The Vertical Box will first check if there's an intersection and ignore all 0 size elements that don't intersect with it's assumed desired 0 height.  Because of this, a vertical box with an auto height, containing a ListView that is 0 size, and is pending an item refresh will never paint that child and let it attempt to fill out its children.  Current solution is to use an SBox instead of a SVerticalBox, as it has the same powers to limit desired height, but doesn't suffer from testing child desired size problems to determine if they should be ticked.

diff --git a/Engine/Source/Runtime/Slate/Public/Widgets/Input/SComboBox.h b/Engine/Source/Runtime/Slate/Public/Widgets/Input/SComboBox.h
--- a/Engine/Source/Runtime/Slate/Public/Widgets/Input/SComboBox.h
+++ b/Engine/Source/Runtime/Slate/Public/Widgets/Input/SComboBox.h
@@ -152,11 +152,8 @@ public:
 		OptionsSource = InArgs._OptionsSource;
 
 		TSharedRef<SWidget> MenuContent = 
-			SNew(SVerticalBox)
-
-			+ SVerticalBox::Slot()
-			.AutoHeight()
-			.MaxHeight(InArgs._MaxListHeight)
+			SNew(SBox)
+			.MaxDesiredHeight(InArgs._MaxListHeight)
 			[
 				SAssignNew( this->ComboListView, SComboListType )
 				.ListItemsSource( InArgs._OptionsSource )
