Cleaned a couple leftover commented-out lines

diff --git a/Engine/Source/Runtime/Slate/Private/Widgets/Views/STableViewBase.cpp b/Engine/Source/Runtime/Slate/Private/Widgets/Views/STableViewBase.cpp
--- a/Engine/Source/Runtime/Slate/Private/Widgets/Views/STableViewBase.cpp
+++ b/Engine/Source/Runtime/Slate/Private/Widgets/Views/STableViewBase.cpp
@@ -12,8 +12,7 @@ namespace ListConstants
 void STableViewBase::ConstructChildren( const TAttribute<float>& InItemWidth, const TAttribute<float>& InItemHeight, const TAttribute<EListItemAlignment>& InItemAlignment, const TSharedPtr<SHeaderRow>& InHeaderRow, const TSharedPtr<SScrollBar>& InScrollBar )
 {
 	bItemsNeedRefresh = true;
-	//RegisterActiveTimer(0.f, FWidgetActiveTimerDelegate::CreateSP(this, &STableViewBase::EnsureTickToRefresh));
-
+	
 	HeaderRow = InHeaderRow;
 
 	// If the user provided a scrollbar, we do not need to make one of our own.
@@ -548,7 +547,6 @@ bool STableViewBase::IsUserScrolling() const
 
 void STableViewBase::RequestListRefresh()
 {
-	//bItemsNeedRefresh = true;
 	if (!bItemsNeedRefresh)
 	{
 		bItemsNeedRefresh = true;
