Fixed swapped Ctrl and Cmd keys in Mac menus caused by CL #2365131

diff --git a/Engine/Source/Runtime/Slate/Private/Framework/MultiBox/Mac/MacMenu.cpp b/Engine/Source/Runtime/Slate/Private/Framework/MultiBox/Mac/MacMenu.cpp
--- a/Engine/Source/Runtime/Slate/Private/Framework/MultiBox/Mac/MacMenu.cpp
+++ b/Engine/Source/Runtime/Slate/Private/Framework/MultiBox/Mac/MacMenu.cpp
@@ -430,7 +430,7 @@ NSString* FSlateMacMenu::GetMenuItemKeyEquivalent(const TSharedRef<const class F
 		*OutModifiers = 0;
 		if (Gesture->NeedsControl())
 		{
-			*OutModifiers |= NSCommandKeyMask;
+			*OutModifiers |= NSControlKeyMask;
 		}
 		if (Gesture->NeedsShift())
 		{
@@ -442,7 +442,7 @@ NSString* FSlateMacMenu::GetMenuItemKeyEquivalent(const TSharedRef<const class F
 		}
 		if (Gesture->NeedsCommand())
 		{
-			*OutModifiers |= NSControlKeyMask;
+			*OutModifiers |= NSCommandKeyMask;
 		}
 
 		FString KeyString = Gesture->GetKeyText().ToString().ToLower();
