Slate/UMG - Buttons now have support for Pressed/Released events.  These events are now also exposed in UMG as well.

diff --git a/Engine/Source/Runtime/Slate/Private/Widgets/Input/SButton.cpp b/Engine/Source/Runtime/Slate/Private/Widgets/Input/SButton.cpp
--- a/Engine/Source/Runtime/Slate/Private/Widgets/Input/SButton.cpp
+++ b/Engine/Source/Runtime/Slate/Private/Widgets/Input/SButton.cpp
@@ -52,6 +52,8 @@ void SButton::Construct( const FArguments& InArgs )
 	bIsFocusable = InArgs._IsFocusable;
 
 	OnClicked = InArgs._OnClicked;
+	OnPressed = InArgs._OnPressed;
+	OnReleased = InArgs._OnReleased;
 
 	ClickMethod = InArgs._ClickMethod;
 	TouchMethod = InArgs._TouchMethod;
@@ -95,7 +97,7 @@ bool SButton::SupportsKeyboardFocus() const
 
 void SButton::OnFocusLost( const FFocusEvent& InFocusEvent )
 {
-	bIsPressed = false;
+	Release();
 }
 
 FReply SButton::OnKeyDown( const FGeometry& MyGeometry, const FKeyEvent& InKeyEvent )
@@ -103,9 +105,7 @@ FReply SButton::OnKeyDown( const FGeometry& MyGeometry, const FKeyEvent& InKeyEv
 	FReply Reply = FReply::Unhandled();
 	if (IsEnabled() && (InKeyEvent.GetKey() == EKeys::Enter || InKeyEvent.GetKey() == EKeys::SpaceBar || InKeyEvent.GetKey() == EKeys::Gamepad_FaceButton_Bottom))
 	{
-		bIsPressed = true;
-
-		PlayPressedSound();
+		Press();
 
 		if (PressMethod == EButtonPressMethod::ButtonPress)
 		{
@@ -136,9 +136,11 @@ FReply SButton::OnKeyUp(const FGeometry& MyGeometry, const FKeyEvent& InKeyEvent
 	if (IsEnabled() && (InKeyEvent.GetKey() == EKeys::Enter || InKeyEvent.GetKey() == EKeys::SpaceBar || InKeyEvent.GetKey() == EKeys::Gamepad_FaceButton_Bottom))
 	{
 		bool bWasPressed = bIsPressed;
-		bIsPressed = false;
+
+		Release();
+
 		//@Todo Slate: This should check focus, however we don't have that API yet, will be easier when focus is unified.
-		if (PressMethod == EButtonPressMethod::ButtonRelease || (PressMethod == EButtonPressMethod::DownAndUp && bWasPressed))
+		if ( PressMethod == EButtonPressMethod::ButtonRelease || ( PressMethod == EButtonPressMethod::DownAndUp && bWasPressed ) )
 		{
 			//execute our "OnClicked" delegate, and get the reply
 			Reply = OnClicked.IsBound() ? OnClicked.Execute() : FReply::Handled();
@@ -161,12 +163,7 @@ FReply SButton::OnMouseButtonDown( const FGeometry& MyGeometry, const FPointerEv
 	FReply Reply = FReply::Unhandled();
 	if (IsEnabled() && (MouseEvent.GetEffectingButton() == EKeys::LeftMouseButton || MouseEvent.IsTouchEvent()))
 	{
-		bIsPressed = true;
-
-		if ( IsEnabled() )
-		{
-			PlayPressedSound();
-		}
+		Press();
 		
 		if( ClickMethod == EButtonClickMethod::MouseDown )
 		{
@@ -203,7 +200,7 @@ FReply SButton::OnMouseButtonUp( const FGeometry& MyGeometry, const FPointerEven
 	FReply Reply = FReply::Unhandled();
 	if (IsEnabled() && (MouseEvent.GetEffectingButton() == EKeys::LeftMouseButton || MouseEvent.IsTouchEvent()))
 	{
-		bIsPressed = false;
+		Release();
 
 		if( ClickMethod == EButtonClickMethod::MouseDown )
 		{
@@ -253,7 +250,7 @@ FReply SButton::OnMouseMove( const FGeometry& MyGeometry, const FPointerEvent& M
 	const float SlateDragStartDistance = FSlateApplication::Get().GetDragTriggerDistance();
 	if ( IsPreciseTapOrClick(MouseEvent) && MouseEvent.GetCursorDelta().SizeSquared() > ( SlateDragStartDistance*SlateDragStartDistance ) )
 	{
-		bIsPressed = false;
+		Release();
 	}
 	return FReply::Unhandled();
 }
@@ -276,8 +273,29 @@ void SButton::OnMouseLeave( const FPointerEvent& MouseEvent )
 	// If we're setup to click on mouse-down, then we never capture the mouse and may not receive a
 	// mouse up event, so we need to make sure our pressed state is reset properly here
 	if ( ClickMethod == EButtonClickMethod::MouseDown || IsPreciseTapOrClick(MouseEvent) )
+	{
+		Release();
+	}
+}
+
+void SButton::Press()
+{
+	if ( !bIsPressed )
+	{
+		bIsPressed = true;
+
+		PlayPressedSound();
+
+		OnPressed.ExecuteIfBound();
+	}
+}
+
+void SButton::Release()
+{
+	if ( bIsPressed )
 	{
 		bIsPressed = false;
+		OnReleased.ExecuteIfBound();
 	}
 }
 
diff --git a/Engine/Source/Runtime/Slate/Public/Widgets/Input/SButton.h b/Engine/Source/Runtime/Slate/Public/Widgets/Input/SButton.h
--- a/Engine/Source/Runtime/Slate/Public/Widgets/Input/SButton.h
+++ b/Engine/Source/Runtime/Slate/Public/Widgets/Input/SButton.h
@@ -53,6 +53,12 @@ public:
 		/** Called when the button is clicked */
 		SLATE_EVENT( FOnClicked, OnClicked )
 
+		/** Called when the button is pressed */
+		SLATE_EVENT( FSimpleDelegate, OnPressed )
+
+		/** Called when the button is released */
+		SLATE_EVENT( FSimpleDelegate, OnReleased )
+
 		/** Sets the rules to use for determining whether the button was clicked.  This is an advanced setting and generally should be left as the default. */
 		SLATE_ARGUMENT( EButtonClickMethod::Type, ClickMethod )
 
@@ -163,6 +169,12 @@ protected:
 	/** The delegate to execute when the button is clicked */
 	FOnClicked OnClicked;
 
+	/** The delegate to execute when the button is pressed */
+	FSimpleDelegate OnPressed;
+
+	/** The delegate to execute when the button is released */
+	FSimpleDelegate OnReleased;
+
 	/** Style resource for the button */
 	const FButtonStyle* Style;
 
@@ -185,6 +197,12 @@ protected:
 	/** Can this button be focused? */
 	bool bIsFocusable;
 
+	/** Press the button */
+	void Press();
+
+	/** Release the button */
+	void Release();
+
 	/** Utility function to determine if the incoming mouse event is for a precise tap or click */
 	bool IsPreciseTapOrClick(const FPointerEvent& MouseEvent) const;
 
diff --git a/Engine/Source/Runtime/UMG/Private/Components/Button.cpp b/Engine/Source/Runtime/UMG/Private/Components/Button.cpp
--- a/Engine/Source/Runtime/UMG/Private/Components/Button.cpp
+++ b/Engine/Source/Runtime/UMG/Private/Components/Button.cpp
@@ -32,10 +32,14 @@ void UButton::ReleaseSlateResources(bool bReleaseChildren)
 TSharedRef<SWidget> UButton::RebuildWidget()
 {
 	MyButton = SNew(SButton)
+		.OnClicked(BIND_UOBJECT_DELEGATE(FOnClicked, SlateHandleClicked))
+		.OnPressed(BIND_UOBJECT_DELEGATE(FSimpleDelegate, SlateHandlePressed))
+		.OnReleased(BIND_UOBJECT_DELEGATE(FSimpleDelegate, SlateHandleReleased))
 		.ButtonStyle(&WidgetStyle)
 		.ClickMethod(ClickMethod)
 		.TouchMethod(TouchMethod)
-		.IsFocusable(IsFocusable);
+		.IsFocusable(IsFocusable)
+		;
 
 	if ( GetChildrenCount() > 0 )
 	{
@@ -51,8 +55,6 @@ void UButton::SynchronizeProperties()
 
 	MyButton->SetColorAndOpacity( ColorAndOpacity );
 	MyButton->SetBorderBackgroundColor( BackgroundColor );
-	
-	MyButton->SetOnClicked(BIND_UOBJECT_DELEGATE(FOnClicked, SlateHandleClicked));
 }
 
 UClass* UButton::GetSlotClass() const
@@ -145,6 +147,16 @@ FReply UButton::SlateHandleClicked()
 	return FReply::Handled();
 }
 
+void UButton::SlateHandlePressed()
+{
+	OnPressed.Broadcast();
+}
+
+void UButton::SlateHandleReleased()
+{
+	OnReleased.Broadcast();
+}
+
 #if WITH_EDITOR
 
 const FSlateBrush* UButton::GetEditorIcon()
diff --git a/Engine/Source/Runtime/UMG/Public/Components/Button.h b/Engine/Source/Runtime/UMG/Public/Components/Button.h
--- a/Engine/Source/Runtime/UMG/Public/Components/Button.h
+++ b/Engine/Source/Runtime/UMG/Public/Components/Button.h
@@ -7,6 +7,8 @@
 #include "Button.generated.h"
 
 DECLARE_DYNAMIC_MULTICAST_DELEGATE(FOnButtonClickedEvent);
+DECLARE_DYNAMIC_MULTICAST_DELEGATE(FOnButtonPressedEvent);
+DECLARE_DYNAMIC_MULTICAST_DELEGATE(FOnButtonReleasedEvent);
 
 /**
  * The button is a click-able primitive widget to enable basic interaction, you
@@ -56,6 +58,14 @@ public:
 	UPROPERTY(BlueprintAssignable, Category="Button|Event")
 	FOnButtonClickedEvent OnClicked;
 
+	/** Called when the button is pressed */
+	UPROPERTY(BlueprintAssignable, Category="Button|Event")
+	FOnButtonPressedEvent OnPressed;
+
+	/** Called when the button is released */
+	UPROPERTY(BlueprintAssignable, Category="Button|Event")
+	FOnButtonReleasedEvent OnReleased;
+
 public:
 	
 	/** Sets the color multiplier for the button content */
@@ -104,6 +114,8 @@ protected:
 protected:
 	/** Handle the actual click event from slate and forward it on */
 	FReply SlateHandleClicked();
+	void SlateHandlePressed();
+	void SlateHandleReleased();
 
 protected:
 	// UWidget interface
