Fixed inline editable text boxes not restoring focus when they are closed

UE-7201 - Crash occurs as user is able to leave text fields in editable state

Closing an inline editable text would remove the editable text box from the widget hierarchy, but not specify where the focus should go. This resulted in list/tile views partially receiving focus, allowing the keyboard to interact with them in undesirable ways.

This change forces the inline editable text box to return focus to whatever had focus when it was opened. For list views, this returns focus to the selected item. For blueprints, this returns focus to the selected node.

diff --git a/Engine/Source/Runtime/Slate/Private/Widgets/Text/SInlineEditableTextBlock.cpp b/Engine/Source/Runtime/Slate/Private/Widgets/Text/SInlineEditableTextBlock.cpp
--- a/Engine/Source/Runtime/Slate/Private/Widgets/Text/SInlineEditableTextBlock.cpp
+++ b/Engine/Source/Runtime/Slate/Private/Widgets/Text/SInlineEditableTextBlock.cpp
@@ -83,6 +83,7 @@ void SInlineEditableTextBlock::EnterEditingMode()
 					TextBox.ToSharedRef()
 				];
 
+			WidgetToFocus = FSlateApplication::Get().GetKeyboardFocusedWidget();
 			FSlateApplication::Get().SetKeyboardFocus(TextBox, EFocusCause::SetDirectly);
 
 			TextBlock->SetVisibility(EVisibility::Collapsed);
@@ -99,6 +100,17 @@ void SInlineEditableTextBlock::ExitEditingMode()
 
 	// Clear the error so it will vanish.
 	TextBox->SetError( FText::GetEmpty() );
+
+	// Restore the original widget focus
+	TSharedPtr<SWidget> WidgetToFocusPin = WidgetToFocus.Pin();
+	if(WidgetToFocusPin.IsValid())
+	{
+		FSlateApplication::Get().SetKeyboardFocus(WidgetToFocusPin, EFocusCause::SetDirectly);
+	}
+	else
+	{
+		FSlateApplication::Get().ClearKeyboardFocus(EFocusCause::SetDirectly);
+	}
 }
 
 bool SInlineEditableTextBlock::IsInEditMode() const
diff --git a/Engine/Source/Runtime/Slate/Public/Widgets/Text/SInlineEditableTextBlock.h b/Engine/Source/Runtime/Slate/Public/Widgets/Text/SInlineEditableTextBlock.h
--- a/Engine/Source/Runtime/Slate/Public/Widgets/Text/SInlineEditableTextBlock.h
+++ b/Engine/Source/Runtime/Slate/Public/Widgets/Text/SInlineEditableTextBlock.h
@@ -156,6 +156,9 @@ protected:
 
 	/** Attribute to look up if the widget is read-only */
 	TAttribute< bool > bIsReadOnly;
+	
+	/** Widget to focus when we finish editing */
+	TWeakPtr<SWidget> WidgetToFocus;
 
 private:
 	/** Active timer to trigger entry into edit mode after a delay */
