Made text blocks better respect their justification method when the visible size is smaller than the desired size

A text block with a visible size smaller than the desired size of the text layout would previously always favor showing the left side of the text layout. This lead to some odd looking results when using center or right justification, as it was possible for the text to be completely hidden or badly aligned.

This change has the text block "auto-scroll" the text layout to try and stick to the most important part of the text layout (the center when using center justification, and the right when using right justification). This produces much better results for long asset names in the Content Browser (which uses center justified text).

diff --git a/Engine/Source/Runtime/Slate/Private/Widgets/Text/TextBlockLayout.cpp b/Engine/Source/Runtime/Slate/Private/Widgets/Text/TextBlockLayout.cpp
--- a/Engine/Source/Runtime/Slate/Private/Widgets/Text/TextBlockLayout.cpp
+++ b/Engine/Source/Runtime/Slate/Private/Widgets/Text/TextBlockLayout.cpp
@@ -98,8 +98,33 @@ int32 FTextBlockLayout::OnPaint(const FPaintArgs& InPaintArgs, const FGeometry&
 {
 	CachedSize = InAllottedGeometry.Size;
 
+	// Text blocks don't have scroll bars, so when the visible region is smaller than the desired size, 
+	// we attempt to auto-scroll to keep the view of the text aligned with the current justification method
+	FVector2D AutoScrollValue = FVector2D::ZeroVector; // Scroll to the left
+	if(TextLayout->GetJustification() != ETextJustify::Left)
+	{
+		const float ActualWidth = TextLayout->GetSize().X;
+		const float VisibleWidth = InAllottedGeometry.Size.X;
+		if(VisibleWidth < ActualWidth)
+		{
+			switch(TextLayout->GetJustification())
+			{
+			case ETextJustify::Center:
+				AutoScrollValue.X = (ActualWidth - VisibleWidth) * 0.5f; // Scroll to the center
+				break;
+
+			case ETextJustify::Right:
+				AutoScrollValue.X = (ActualWidth - VisibleWidth); // Scroll to the right
+				break;
+
+			default:
+				break;
+			}
+		}
+	}
+
 	TextLayout->SetScale(InAllottedGeometry.Scale);
-	TextLayout->SetVisibleRegion(InAllottedGeometry.Size, FVector2D::ZeroVector);
+	TextLayout->SetVisibleRegion(InAllottedGeometry.Size, AutoScrollValue);
 
 	TextLayout->UpdateIfNeeded();
 
