Slate - Removing the ContentPadding call in SComboBox, it doesn't actually do anything because it's overridden by the user's parameter.

diff --git a/Engine/Source/Runtime/Slate/Private/Widgets/Input/SComboButton.cpp b/Engine/Source/Runtime/Slate/Private/Widgets/Input/SComboButton.cpp
--- a/Engine/Source/Runtime/Slate/Private/Widgets/Input/SComboButton.cpp
+++ b/Engine/Source/Runtime/Slate/Private/Widgets/Input/SComboButton.cpp
@@ -25,7 +25,6 @@ void SComboButton::Construct( const FArguments& InArgs )
 		.Method(InArgs._Method)
 		[
 			SNew( SButton )
-			.ContentPadding( FMargin( 1, 0 ) )
 			.ButtonStyle( OurButtonStyle )
 			.ClickMethod( EButtonClickMethod::MouseDown )
 			.OnClicked( this, &SComboButton::OnButtonClicked )
