fix for crash when utilizing the virtual keyboard with a second widget on iOS
UE-7433
#ios

diff --git a/Engine/Source/Runtime/Slate/Private/Framework/Text/IOS/IOSPlatformTextField.cpp b/Engine/Source/Runtime/Slate/Private/Framework/Text/IOS/IOSPlatformTextField.cpp
--- a/Engine/Source/Runtime/Slate/Private/Framework/Text/IOS/IOSPlatformTextField.cpp
+++ b/Engine/Source/Runtime/Slate/Private/Framework/Text/IOS/IOSPlatformTextField.cpp
@@ -101,6 +101,15 @@ void FIOSPlatformTextField::ShowVirtualKeyboard(bool bShow, TSharedPtr<IVirtualK
 		UITextField* AlertTextField = [alertView textFieldAtIndex: 0];
 		TextWidget->SetTextFromVirtualKeyboard(FText::FromString(AlertTextField.text));
 	}
+    
+    FIOSAsyncTask* AsyncTask = [[FIOSAsyncTask alloc] init];
+    AsyncTask.GameThreadCallback = ^ bool(void)
+    {
+        // clear the TextWidget
+        TextWidget = nullptr;
+        return true;
+    };
+    [AsyncTask FinishedTask];
 }
 
 @end
