Apply normal logic when calculating position of zero-sized popups.

- This prevent menu windows from flickering (which is very noticeable under Linux) - submenu size may be still zero for the first frame.
- Solves UE-5161.

diff --git a/Engine/Source/Runtime/Slate/Private/Framework/Application/SlateApplication.cpp b/Engine/Source/Runtime/Slate/Private/Framework/Application/SlateApplication.cpp
--- a/Engine/Source/Runtime/Slate/Private/Framework/Application/SlateApplication.cpp
+++ b/Engine/Source/Runtime/Slate/Private/Framework/Application/SlateApplication.cpp
@@ -3150,12 +3150,6 @@ float FSlateApplication::GetCursorRadius() const
 
 FVector2D FSlateApplication::CalculatePopupWindowPosition( const FSlateRect& InAnchor, const FVector2D& InSize, const EOrientation Orientation ) const
 {
-	// Do nothing if this window has no size
-	if (InSize == FVector2D::ZeroVector)
-	{
-		return FVector2D(InAnchor.Left, InAnchor.Top);
-	}
-
 	FVector2D CalculatedPopUpWindowPosition( 0, 0 );
 
 	FPlatformRect AnchorRect;
