SViewport can now be set as active, registering an indefinite active timer. Used when the game viewport is registered with FSlateApplication to guarantee the existence of an active timer during a game

diff --git a/Engine/Source/Runtime/Slate/Private/Framework/Application/SlateApplication.cpp b/Engine/Source/Runtime/Slate/Private/Framework/Application/SlateApplication.cpp
--- a/Engine/Source/Runtime/Slate/Private/Framework/Application/SlateApplication.cpp
+++ b/Engine/Source/Runtime/Slate/Private/Framework/Application/SlateApplication.cpp
@@ -1837,8 +1837,9 @@ void FSlateApplication::InvalidateAllViewports()
 
 void FSlateApplication::RegisterGameViewport( TSharedRef<SViewport> InViewport )
 {
+	InViewport->SetActive(true);
 	GameViewportWidget = InViewport;
-
+	
 	FWidgetPath PathToViewport;
 	// If we cannot find the window it could have been destroyed.
 	if (FSlateWindowHelper::FindPathToWidget(SlateWindows, InViewport, PathToViewport, EVisibility::All))
@@ -1859,6 +1860,11 @@ void FSlateApplication::RegisterGameViewport( TSharedRef<SViewport> InViewport )
 void FSlateApplication::UnregisterGameViewport()
 {
 	ResetToDefaultInputSettings();
+
+	if (GameViewportWidget.IsValid())
+	{
+		GameViewportWidget.Pin()->SetActive(false);
+	}
 	GameViewportWidget.Reset();
 }
 
diff --git a/Engine/Source/Runtime/Slate/Private/Widgets/SViewport.cpp b/Engine/Source/Runtime/Slate/Private/Widgets/SViewport.cpp
--- a/Engine/Source/Runtime/Slate/Private/Widgets/SViewport.cpp
+++ b/Engine/Source/Runtime/Slate/Private/Widgets/SViewport.cpp
@@ -35,6 +35,22 @@ void SViewport::Construct( const FArguments& InArgs )
 	];
 }
 
+void SViewport::SetActive(bool bActive)
+{
+	if (bActive && !ActiveTimerHandle.IsValid())
+	{
+		ActiveTimerHandle = RegisterActiveTimer(0.f, FWidgetActiveTimerDelegate::CreateSP(this, &SViewport::EnsureTick));
+	}
+	else if (!bActive && ActiveTimerHandle.IsValid())
+	{
+		UnRegisterActiveTimer(ActiveTimerHandle.Pin().ToSharedRef());
+	}
+}
+
+EActiveTimerReturnType SViewport::EnsureTick(double InCurrentTime, float InDeltaTime)
+{
+	return EActiveTimerReturnType::Continue;
+}
 
 int32 SViewport::OnPaint( const FPaintArgs& Args, const FGeometry& AllottedGeometry, const FSlateRect& MyClippingRect, FSlateWindowElementList& OutDrawElements, int32 LayerId, const FWidgetStyle& InWidgetStyle, bool bParentEnabled ) const
 {
diff --git a/Engine/Source/Runtime/Slate/Public/Widgets/SViewport.h b/Engine/Source/Runtime/Slate/Public/Widgets/SViewport.h
--- a/Engine/Source/Runtime/Slate/Public/Widgets/SViewport.h
+++ b/Engine/Source/Runtime/Slate/Public/Widgets/SViewport.h
@@ -133,6 +133,13 @@ public:
 		WidgetToFocusOnActivate.Reset();
 	}
 
+	/** 
+	 * Sets whether this viewport is active. 
+	 * While active, a persistent Active Timer is registered and a Slate tick/paint pass is guaranteed every frame.
+	 * @param bActive Whether to set the viewport as active
+	 */
+	void SetActive(bool bActive);
+
 public:
 
 	// SWidget interface
@@ -180,12 +187,17 @@ private:
 	}
 
 protected:
+	/** Empty active timer meant to ensure a tick/paint pass while this viewport is active */
+	EActiveTimerReturnType EnsureTick(double InCurrentTime, float InDeltaTime);
 
 	/** Interface to the rendering and I/O implementation of the viewport. */
 	TWeakPtr<ISlateViewport> ViewportInterface;
 	
 private:
 
+	/** The handle to the active EnsureTick() timer */
+	TWeakPtr<FActiveTimerHandle> ActiveTimerHandle;
+
 	/** Whether or not to show the disabled effect when this viewport is disabled. */
 	TAttribute<bool> ShowDisabledEffect;
 
