Slate - Fixing Tab focus traversal on the SEditableText widget, it now correctly calls the baseclass to determine if it should handle navigation for the tab key.

diff --git a/Engine/Source/Runtime/Slate/Private/Widgets/Input/SEditableText.cpp b/Engine/Source/Runtime/Slate/Private/Widgets/Input/SEditableText.cpp
--- a/Engine/Source/Runtime/Slate/Private/Widgets/Input/SEditableText.cpp
+++ b/Engine/Source/Runtime/Slate/Private/Widgets/Input/SEditableText.cpp
@@ -1641,6 +1641,12 @@ FReply SEditableText::OnKeyDown( const FGeometry& MyGeometry, const FKeyEvent& I
 		Reply = FReply::Handled();
 	}
 
+	// Allow the base class to handle it to determine if navigation needs to be triggered.
+	if ( !Reply.IsEventHandled() )
+	{
+		Reply = SLeafWidget::OnKeyDown(MyGeometry, InKeyEvent);
+	}
+
 	return Reply;
 }
 
