Slate - Removing the change to handle navigation events in textboxes, broke lots of stuff.  Need to make the viewport less greedy and handle this at a higher level.

diff --git a/Engine/Source/Runtime/Slate/Private/Widgets/Input/SEditableText.cpp b/Engine/Source/Runtime/Slate/Private/Widgets/Input/SEditableText.cpp
--- a/Engine/Source/Runtime/Slate/Private/Widgets/Input/SEditableText.cpp
+++ b/Engine/Source/Runtime/Slate/Private/Widgets/Input/SEditableText.cpp
@@ -1641,12 +1641,6 @@ FReply SEditableText::OnKeyDown( const FGeometry& MyGeometry, const FKeyEvent& I
 		Reply = FReply::Handled();
 	}
 
-	// Allow the base class to handle it to determine if navigation needs to be triggered.
-	if ( !Reply.IsEventHandled() )
-	{
-		Reply = SLeafWidget::OnKeyDown(MyGeometry, InKeyEvent);
-	}
-
 	return Reply;
 }
 
