Fixed context menus on editable texts being unusable due to focus loss

UE-7906 - Right clicking in a text field removes focus from it which makes context menu options not function

The change to set bFocusImmediately to true when summoning the context menus for editable texts meant that we received some focus changed events before we'd set the context menu window pointer (as the call to summon the menu hadn't yet returned).

We now track some extra state to say that we're pending a summon and ignore the focus events in this case (as we do when a context menu is full open).

diff --git a/Engine/Source/Runtime/Slate/Private/Widgets/Input/SEditableText.cpp b/Engine/Source/Runtime/Slate/Private/Widgets/Input/SEditableText.cpp
--- a/Engine/Source/Runtime/Slate/Private/Widgets/Input/SEditableText.cpp
+++ b/Engine/Source/Runtime/Slate/Private/Widgets/Input/SEditableText.cpp
@@ -1499,7 +1499,7 @@ bool SEditableText::SupportsKeyboardFocus() const
 FReply SEditableText::OnFocusReceived( const FGeometry& MyGeometry, const FFocusEvent& InFocusEvent )
 {
 	// Skip the rest of the focus received code if it's due to the context menu closing
-	if ( !ContextMenuWindow.IsValid() )
+	if ( !ActiveContextMenu.IsValid() )
 	{
 		RegisterActiveTimer(0.f, FWidgetActiveTimerDelegate::CreateSP(this, &SEditableText::AnimateSpringsWhileFocused));
 
@@ -1565,7 +1565,7 @@ FReply SEditableText::OnFocusReceived( const FGeometry& MyGeometry, const FFocus
 void SEditableText::OnFocusLost( const FFocusEvent& InFocusEvent )
 {
 	// Skip the focus lost code if it's due to the context menu opening
-	if ( !ContextMenuWindow.IsValid() )
+	if ( !ActiveContextMenu.IsValid() )
 	{
 		// Place the caret at the front of the text
 		if (FSlateApplication::Get().IsRunningAtTargetFrameRate())
@@ -1980,14 +1980,20 @@ void SEditableText::SummonContextMenu( const FVector2D& InLocation )
 
 #undef LOCTEXT_NAMESPACE
 
-	const bool bFocusImmediately = true;
-	TSharedPtr< SWindow > ContextMenuWindowPinned = FSlateApplication::Get().PushMenu( SharedThis( this ), MenuBuilder.MakeWidget(), InLocation, FPopupTransitionEffect( FPopupTransitionEffect::ContextMenu ), bFocusImmediately );
-	ContextMenuWindow = ContextMenuWindowPinned;
+	ActiveContextMenu.PrepareToSummon();
 
+	const bool bFocusImmediately = true;
+	TSharedPtr< SWindow > ContextMenuWindow = FSlateApplication::Get().PushMenu( SharedThis( this ), MenuBuilder.MakeWidget(), InLocation, FPopupTransitionEffect( FPopupTransitionEffect::ContextMenu ), bFocusImmediately );
+	
 	// Make sure the window is valid.  It's possible for the parent to already be in the destroy queue, for example if the editable text was configured to dismiss it's window during OnTextCommitted.
-	if( ContextMenuWindowPinned.IsValid() )
+	if( ContextMenuWindow.IsValid() )
 	{
-		ContextMenuWindowPinned->SetOnWindowClosed( FOnWindowClosed::CreateSP( this, &SEditableText::OnWindowClosed ) );
+		ContextMenuWindow->SetOnWindowClosed( FOnWindowClosed::CreateSP( this, &SEditableText::OnWindowClosed ) );
+		ActiveContextMenu.SummonSucceeded( ContextMenuWindow.ToSharedRef() );
+	}
+	else
+	{
+		ActiveContextMenu.SummonFailed();
 	}
 }
 
@@ -2056,7 +2062,7 @@ void SEditableText::SetCaretPosition( int32 Position )
 bool SEditableText::ShouldAppearFocused() const
 {
 	// We'll draw with the 'focused' look if we're either focused or we have a context menu summoned
-	return HasKeyboardFocus() || ContextMenuWindow.IsValid();
+	return HasKeyboardFocus() || ActiveContextMenu.IsValid();
 }
 
 bool SEditableText::DoesClipboardHaveAnyText() const 
@@ -2113,6 +2119,9 @@ void SEditableText::RestartSelectionTargetAnimation()
 
 void SEditableText::OnWindowClosed(const TSharedRef<SWindow>&)
 {
+	// Note: We don't reset the ActiveContextMenu here, as Slate hasn't yet finished processing window focus events, and we need 
+	// to know that the window is still available for OnFocusReceived and OnFocusLost even though it's about to be destroyed
+
 	// Give this focus when the context menu has been dismissed
 	FSlateApplication::Get().SetKeyboardFocus( AsShared(), EFocusCause::OtherWidgetLostFocus );
 }
diff --git a/Engine/Source/Runtime/Slate/Private/Widgets/Text/SMultiLineEditableText.cpp b/Engine/Source/Runtime/Slate/Private/Widgets/Text/SMultiLineEditableText.cpp
--- a/Engine/Source/Runtime/Slate/Private/Widgets/Text/SMultiLineEditableText.cpp
+++ b/Engine/Source/Runtime/Slate/Private/Widgets/Text/SMultiLineEditableText.cpp
@@ -470,7 +470,7 @@ void SMultiLineEditableText::OnVScrollBarMoved(const float InScrollOffsetFractio
 FReply SMultiLineEditableText::OnFocusReceived( const FGeometry& MyGeometry, const FFocusEvent& InFocusEvent )
 {
 	// Skip the focus received code if it's due to the context menu closing
-	if ( !ContextMenuWindow.IsValid() )
+	if ( !ActiveContextMenu.IsValid() )
 	{
 		FSlateApplication& SlateApplication = FSlateApplication::Get();
 		if (FPlatformMisc::GetRequiresVirtualKeyboard())
@@ -505,7 +505,7 @@ FReply SMultiLineEditableText::OnFocusReceived( const FGeometry& MyGeometry, con
 void SMultiLineEditableText::OnFocusLost( const FFocusEvent& InFocusEvent )
 {
 	// Skip the focus lost code if it's due to the context menu opening
-	if (!ContextMenuWindow.IsValid())
+	if (!ActiveContextMenu.IsValid())
 	{
 		FSlateApplication& SlateApplication = FSlateApplication::Get();
 		if (FPlatformMisc::GetRequiresVirtualKeyboard())
@@ -2079,19 +2079,28 @@ void SMultiLineEditableText::SummonContextMenu(const FVector2D& InLocation)
 
 #undef LOCTEXT_NAMESPACE
 
+	ActiveContextMenu.PrepareToSummon();
+
 	const bool bFocusImmediately = true;
-	TSharedPtr< SWindow > ContextMenuWindowPinned = FSlateApplication::Get().PushMenu(SharedThis(this), MenuBuilder.MakeWidget(), InLocation, FPopupTransitionEffect(FPopupTransitionEffect::ContextMenu), bFocusImmediately);
-	ContextMenuWindow = ContextMenuWindowPinned;
+	TSharedPtr< SWindow > ContextMenuWindow = FSlateApplication::Get().PushMenu(SharedThis(this), MenuBuilder.MakeWidget(), InLocation, FPopupTransitionEffect(FPopupTransitionEffect::ContextMenu), bFocusImmediately);
 
 	// Make sure the window is valid.  It's possible for the parent to already be in the destroy queue, for example if the editable text was configured to dismiss it's window during OnTextCommitted.
-	if (ContextMenuWindowPinned.IsValid())
+	if (ContextMenuWindow.IsValid())
 	{
-		ContextMenuWindowPinned->SetOnWindowClosed(FOnWindowClosed::CreateSP(this, &SMultiLineEditableText::OnWindowClosed));
+		ContextMenuWindow->SetOnWindowClosed(FOnWindowClosed::CreateSP(this, &SMultiLineEditableText::OnWindowClosed));
+		ActiveContextMenu.SummonSucceeded(ContextMenuWindow.ToSharedRef());
+	}
+	else
+	{
+		ActiveContextMenu.SummonFailed();
 	}
 }
 
 void SMultiLineEditableText::OnWindowClosed(const TSharedRef<SWindow>&)
 {
+	// Note: We don't reset the ActiveContextMenu here, as Slate hasn't yet finished processing window focus events, and we need 
+	// to know that the window is still available for OnFocusReceived and OnFocusLost even though it's about to be destroyed
+
 	// Give this focus when the context menu has been dismissed
 	FSlateApplication::Get().SetKeyboardFocus(AsShared(), EFocusCause::OtherWidgetLostFocus);
 }
@@ -2124,7 +2133,7 @@ void SMultiLineEditableText::Tick( const FGeometry& AllottedGeometry, const doub
 
 	TextLayout->SetScale( AllottedGeometry.Scale );
 
-	const bool bShouldAppearFocused = HasKeyboardFocus() || ContextMenuWindow.IsValid();
+	const bool bShouldAppearFocused = HasKeyboardFocus() || ActiveContextMenu.IsValid();
 
 	if (!bShouldAppearFocused)
 	{
diff --git a/Engine/Source/Runtime/Slate/Public/Widgets/Input/SEditableText.h b/Engine/Source/Runtime/Slate/Public/Widgets/Input/SEditableText.h
--- a/Engine/Source/Runtime/Slate/Public/Widgets/Input/SEditableText.h
+++ b/Engine/Source/Runtime/Slate/Public/Widgets/Input/SEditableText.h
@@ -592,8 +592,8 @@ private:
 	/** Undo state that will be pushed if text is actually changed between calls to StartChangingText() and FinishChangingText() */
 	FUndoState StateBeforeChangingText;
 
-	/** Weak pointer to context menu window that's currently open, if there is one */
-	TWeakPtr< SWindow > ContextMenuWindow;
+	/** Information about any active context menu widgets */
+	FActiveTextEditContextMenu ActiveContextMenu;
 
 	/** Prevents the editabletext from being smaller than desired in certain cases (e.g. when it is empty) */
 	TAttribute<float> MinDesiredWidth;
diff --git a/Engine/Source/Runtime/Slate/Public/Widgets/Text/ITextEditorWidget.h b/Engine/Source/Runtime/Slate/Public/Widgets/Text/ITextEditorWidget.h
--- a/Engine/Source/Runtime/Slate/Public/Widgets/Text/ITextEditorWidget.h
+++ b/Engine/Source/Runtime/Slate/Public/Widgets/Text/ITextEditorWidget.h
@@ -102,6 +102,57 @@ private:
 	float GeometryScale;
 };
 
+/** Manages the state for an active context menu */
+class FActiveTextEditContextMenu
+{
+public:
+	FActiveTextEditContextMenu()
+		: bIsPendingSummon(false)
+		, ActiveMenuWindow()
+	{
+	}
+
+	/** Check to see whether this context is valid (either pending or active) */
+	bool IsValid() const
+	{
+		return bIsPendingSummon || ActiveMenuWindow.IsValid();
+	}
+
+	/** Called to reset the active context menu state */
+	void Reset()
+	{
+		bIsPendingSummon = false;
+		ActiveMenuWindow.Reset();
+	}
+
+	/** Called before you summon your context menu */
+	void PrepareToSummon()
+	{
+		bIsPendingSummon = true;
+		ActiveMenuWindow.Reset();
+	}
+
+	/** Called when you've successfully summoned your context menu */
+	void SummonSucceeded(const TSharedRef<SWindow>& InActiveMenuWindow)
+	{
+		bIsPendingSummon = false;
+		ActiveMenuWindow = InActiveMenuWindow;
+	}
+
+	/** Called if your context menu summon fails */
+	void SummonFailed()
+	{
+		bIsPendingSummon = false;
+		ActiveMenuWindow.Reset();
+	}
+
+private:
+	/** True if we are pending the summon of a context menu, but don't yet have an active window pointer */
+	bool bIsPendingSummon;
+
+	/** Handle to the active context menu (if any) */
+	TWeakPtr<SWindow> ActiveMenuWindow;
+};
 
 class SLATE_API ITextEditorWidget 
 {
diff --git a/Engine/Source/Runtime/Slate/Public/Widgets/Text/SMultiLineEditableText.h b/Engine/Source/Runtime/Slate/Public/Widgets/Text/SMultiLineEditableText.h
--- a/Engine/Source/Runtime/Slate/Public/Widgets/Text/SMultiLineEditableText.h
+++ b/Engine/Source/Runtime/Slate/Public/Widgets/Text/SMultiLineEditableText.h
@@ -648,8 +648,8 @@ private:
 	/** Menu extender for right-click context menu */
 	TSharedPtr<FExtender> MenuExtender;
 
-	/** Weak pointer to context menu window that's currently open, if there is one */
-	TWeakPtr< SWindow > ContextMenuWindow;
+	/** Information about any active context menu widgets */
+	FActiveTextEditContextMenu ActiveContextMenu;
 
 	/** Implemented context object for text input method systems. */
 	TSharedPtr<FTextInputMethodContext> TextInputMethodContext;
