Fixed some places where SEditableTextBox was incorrectly resolving some attributes down to their values when forwarding them on to SEditableText

diff --git a/Engine/Source/Runtime/Slate/Private/Widgets/Input/SEditableTextBox.cpp b/Engine/Source/Runtime/Slate/Private/Widgets/Input/SEditableTextBox.cpp
--- a/Engine/Source/Runtime/Slate/Private/Widgets/Input/SEditableTextBox.cpp
+++ b/Engine/Source/Runtime/Slate/Private/Widgets/Input/SEditableTextBox.cpp
@@ -42,9 +42,9 @@ void SEditableTextBox::Construct( const FArguments& InArgs )
 					.IsReadOnly( InArgs._IsReadOnly )
 					.IsPassword( InArgs._IsPassword )
 					.IsCaretMovedWhenGainFocus( InArgs._IsCaretMovedWhenGainFocus )
-					.SelectAllTextWhenFocused( InArgs._SelectAllTextWhenFocused.Get() )
-					.RevertTextOnEscape( InArgs._RevertTextOnEscape.Get() )
-					.ClearKeyboardFocusOnCommit( InArgs._ClearKeyboardFocusOnCommit.Get() )
+					.SelectAllTextWhenFocused( InArgs._SelectAllTextWhenFocused )
+					.RevertTextOnEscape( InArgs._RevertTextOnEscape )
+					.ClearKeyboardFocusOnCommit( InArgs._ClearKeyboardFocusOnCommit )
 					.OnTextChanged( InArgs._OnTextChanged )
 					.OnTextCommitted( InArgs._OnTextCommitted )
 					.MinDesiredWidth( InArgs._MinDesiredWidth )
