Fixed issue where copy from multiline editable text widgets contains the wrong newline sequence in Windows.

diff --git a/Engine/Source/Runtime/Slate/Private/Framework/Text/TextLayout.cpp b/Engine/Source/Runtime/Slate/Private/Framework/Text/TextLayout.cpp
--- a/Engine/Source/Runtime/Slate/Private/Framework/Text/TextLayout.cpp
+++ b/Engine/Source/Runtime/Slate/Private/Framework/Text/TextLayout.cpp
@@ -1442,18 +1442,20 @@ void FTextLayout::GetAsTextAndOffsets(FString* const OutDisplayText, FTextOffset
 		OutTextOffsetLocations->OffsetData.Reserve(LineModels.Num());
 	}
 
+	const int32 LineTerminatorLength = FCString::Strlen(LINE_TERMINATOR);
+
 	for (int32 LineModelIndex = 0; LineModelIndex < LineModels.Num(); LineModelIndex++)
 	{
 		const FLineModel& LineModel = LineModels[LineModelIndex];
 
-		// Append \n to the end of the previous line
+		// Append line terminator to the end of the previous line
 		if (LineModelIndex > 0)
 		{
 			if (OutDisplayText)
 			{
-				OutDisplayText->AppendChar('\n');
+				*OutDisplayText += LINE_TERMINATOR;
 			}
-			++DisplayTextLength;
+			DisplayTextLength += LineTerminatorLength;
 		}
 
 		int32 LineLength = 0;
@@ -1542,7 +1544,7 @@ void FTextLayout::GetSelectionAsText(FString& DisplayText, const FTextSelection&
 
 				if (LineIndex != SelectionEndLineIndex)
 				{
-					DisplayText.AppendChar('\n');
+					DisplayText += LINE_TERMINATOR;
 				}
 			}
 		}
