Fixed UE-5573 (Tabs do not render when moving after returning to editor on Mac)

diff --git a/Engine/Source/Runtime/Slate/Private/Framework/Docking/SDockingTabWell.cpp b/Engine/Source/Runtime/Slate/Private/Framework/Docking/SDockingTabWell.cpp
--- a/Engine/Source/Runtime/Slate/Private/Framework/Docking/SDockingTabWell.cpp
+++ b/Engine/Source/Runtime/Slate/Private/Framework/Docking/SDockingTabWell.cpp
@@ -265,6 +265,11 @@ FReply SDockingTabWell::StartDraggingTab( TSharedRef<SDockTab> TabToStartDraggin
 		ForegroundTabIndex = INDEX_NONE;
 		ParentTabStackPtr.Pin()->OnTabRemoved(TabToStartDragging->GetLayoutIdentifier());
 
+#if PLATFORM_MAC
+		// On Mac we need to activate the app as we may be dragging a window that is set to be invisible if the app is inactive
+		FPlatformMisc::ActivateApplication();
+#endif
+
 		// Start dragging.
 		TSharedRef<FDockingDragOperation> DragDropOperation =
 			FDockingDragOperation::New(
