[INTEGRATE] Change 2414001 by [REDACTED_PERSON][REDACTED_HANDLE] on 2015/01/21 12:46:19

	[UE-7779] Components selected in the world are scrolled to in the component tree
	- Fixed bug in SListView virtualization where the bottom row widget was not generated after a scroll was requested to an item at the top of the list

diff --git a/Engine/Source/Editor/LevelEditor/Private/SLevelEditor.cpp b/Engine/Source/Editor/LevelEditor/Private/SLevelEditor.cpp
--- a/Engine/Source/Editor/LevelEditor/Private/SLevelEditor.cpp
+++ b/Engine/Source/Editor/LevelEditor/Private/SLevelEditor.cpp
@@ -679,28 +679,6 @@ private:
 		{
 			GEditor->SelectNone(true, true, false);
 			GEditor->SelectActor(Actor, true, true, true);
-		//	auto Actor = GetSelectedActor();
-		//	if (Actor)
-		//	{
-		//		// Enable the selection guard to prevent OnEditorSelectionChanged() from altering the contents of the SCSTreeWidget
-		//		TGuardValue<bool> SelectionGuard(bSelectionGuard, true);
-
-		//		// Update the editor's component selection
-		//		TArray<UObject*> DetailsObjects;
-		//		USelection* SelectedComponents = GEditor->GetSelectedComponents();
-		//		DetailsObjects.Add(Actor);
-
-		//		SelectedComponents->BeginBatchSelectOperation();
-		//		SelectedComponents->DeselectAll();
-
-		//		DetailsView->SetObjects(DetailsObjects, false);
-
-		//		SelectedComponents->EndBatchSelectOperation();
-
-		//		GUnrealEd->SetActorSelectionFlags(Actor);
-		//		GUnrealEd->UpdatePivotLocationForSelection(true);
-		//		GEditor->RedrawLevelEditingViewports();
-		//	}
 		}
 	}
 
@@ -808,6 +786,7 @@ private:
 			auto SCSTreeNode = SCSEditor->GetNodeFromActorComponent(Component, false);
 			if (SCSTreeNode.IsValid() && SCSTreeNode->GetComponentTemplate())
 			{
+				SCSTreeWidget->RequestScrollIntoView(SCSTreeNode);
 				SCSTreeWidget->SetItemSelection(SCSTreeNode, true);
 
 				auto ComponentTemplate = SCSTreeNode->GetComponentTemplate();
diff --git a/Engine/Source/Runtime/Slate/Public/Widgets/Views/SListView.h b/Engine/Source/Runtime/Slate/Public/Widgets/Views/SListView.h
--- a/Engine/Source/Runtime/Slate/Public/Widgets/Views/SListView.h
+++ b/Engine/Source/Runtime/Slate/Public/Widgets/Views/SListView.h
@@ -793,7 +793,7 @@ public:
 
 			// Generate widgets assuming scenario a.
 			bool bGeneratedEnoughForSmoothScrolling = false;
-			bool AtEndOfList = false;
+			bool bAtEndOfList = false;
 
 			const float LayoutScaleMultiplier = MyGeometry.GetAccumulatedLayoutTransform().GetScale();
 			for( int32 ItemIndex = StartIndex; !bGeneratedEnoughForSmoothScrolling && ItemIndex < SourceItems->Num(); ++ItemIndex )
@@ -802,20 +802,22 @@ public:
 
 				const float ItemHeight = GenerateWidgetForItem(CurItem, ItemIndex, StartIndex, LayoutScaleMultiplier);
 
-				const bool IsFirstItem = ItemIndex == StartIndex;
+				const bool bIsFirstItem = ItemIndex == StartIndex;
 
-				if (IsFirstItem)
+				if (bIsFirstItem)
 				{
 					FirstItemHeight = ItemHeight;
 				}
 
 				// Track the number of items in the view, including fractions.
-				if (IsFirstItem)
+				if (bIsFirstItem)
 				{
-					ItemsInView += 1.0f - FMath::Fractional(ScrollOffset);
+					// The first item may not be fully visible (but cannot exceed 1)
+					ItemsInView += 1.0f - FMath::Max(FMath::Fractional(ScrollOffset), 0.0f);
 				}
 				else if (ViewHeightUsedSoFar + ItemHeight > MyGeometry.Size.Y)
 				{
+					// The last item may not be fully visible either
 					ItemsInView += (MyGeometry.Size.Y - ViewHeightUsedSoFar) / ItemHeight;
 				}
 				else
@@ -825,13 +827,13 @@ public:
 
 				HeightGeneratedSoFar += ItemHeight;
 
-				ViewHeightUsedSoFar += (IsFirstItem)
-					? ItemHeight * (1.0f - FMath::Fractional(ScrollOffset))
+				ViewHeightUsedSoFar += (bIsFirstItem)
+					? ItemHeight * ItemsInView	// For the first item, ItemsInView <= 1.0f
 					: ItemHeight;
 
 				if (ItemIndex >= SourceItems->Num()-1)
 				{
-					AtEndOfList = true;
+					bAtEndOfList = true;
 				}
 
 				if (ViewHeightUsedSoFar > MyGeometry.Size.Y )
@@ -843,7 +845,7 @@ public:
 			// Handle scenario b.
 			// We may have stopped because we got to the end of the items.
 			// But we may still have space to fill!
-			if (AtEndOfList && ViewHeightUsedSoFar < MyGeometry.Size.Y)
+			if (bAtEndOfList && ViewHeightUsedSoFar < MyGeometry.Size.Y)
 			{
 				float NewScrollOffsetForBackfill = StartIndex + (HeightGeneratedSoFar - MyGeometry.Size.Y) / FirstItemHeight;
 
@@ -864,10 +866,10 @@ public:
 					HeightGeneratedSoFar += ItemHeight;
 				}
 
-				return FReGenerateResults(NewScrollOffsetForBackfill, HeightGeneratedSoFar, ItemsSource->Num() - NewScrollOffsetForBackfill, AtEndOfList);
+				return FReGenerateResults(NewScrollOffsetForBackfill, HeightGeneratedSoFar, ItemsSource->Num() - NewScrollOffsetForBackfill, bAtEndOfList);
 			}
 
-			return FReGenerateResults(ScrollOffset, HeightGeneratedSoFar, ItemsInView, AtEndOfList);
+			return FReGenerateResults(ScrollOffset, HeightGeneratedSoFar, ItemsInView, bAtEndOfList);
 		}
 
 		return FReGenerateResults(0.0f, 0.0f, 0.0f, false);
