[INTEGRATE] Change 2414428 by [REDACTED_PERSON][REDACTED_HANDLE] on 2015/01/21 16:37:09

	Fixed crash with custom tool-tip content and SToolTip::GetContentWidget()

diff --git a/Engine/Source/Runtime/Slate/Private/Widgets/SToolTip.cpp b/Engine/Source/Runtime/Slate/Private/Widgets/SToolTip.cpp
--- a/Engine/Source/Runtime/Slate/Private/Widgets/SToolTip.cpp
+++ b/Engine/Source/Runtime/Slate/Private/Widgets/SToolTip.cpp
@@ -39,6 +39,8 @@ void SToolTip::SetContentWidget(const TSharedRef<SWidget>& InContentWidget)
 	TSharedPtr< SWidget > PinnedWidgetContent = WidgetContent.Pin();
 	if( PinnedWidgetContent.IsValid() )
 	{
+		ToolTipContent = PinnedWidgetContent;
+
 		// Tool-tip with entirely custom content.  We'll create a border with some padding (as customized by the user), then
 		// embed their custom widget right inside the border.  This tool-tip currently has a different styling than tool-tips
 		// that only contain text.
@@ -48,13 +50,20 @@ void SToolTip::SetContentWidget(const TSharedRef<SWidget>& InContentWidget)
 			.BorderImage(BorderImage)
 			.Padding(TextMargin)
 			[
-				PinnedWidgetContent.ToSharedRef()
+				ToolTipContent.ToSharedRef()
 			]
 		];
 
 	}
 	else
 	{
+		ToolTipContent =
+			SNew( STextBlock )
+			.Text( TextContent )
+			.Font( FCoreStyle::Get().GetFontStyle( "ToolTip.LargerFont" ) )
+			.ColorAndOpacity( FLinearColor::Black )
+			.WrapTextAt_Static( &SToolTip::GetToolTipWrapWidth );
+
 		// Text-only tool-tip.  This tool-tip currently has a different styling than tool-tips with custom content.  We always want basic
 		// text tool-tips to look consistent.
 		ChildSlot
@@ -63,11 +72,7 @@ void SToolTip::SetContentWidget(const TSharedRef<SWidget>& InContentWidget)
 			.BorderImage( FCoreStyle::Get().GetBrush("ToolTip.BrightBackground") )
 			.Padding(FMargin(11.0f))
 			[
-				SNew(STextBlock)
-				.Text(TextContent)
-				.Font(FCoreStyle::Get().GetFontStyle("ToolTip.LargerFont"))
-				.ColorAndOpacity( FLinearColor::Black )
-				.WrapTextAt_Static(&SToolTip::GetToolTipWrapWidth)
+				ToolTipContent.ToSharedRef()
 			]
 		];
 	}
