[INTEGRATE] Change 2415739 by [REDACTED_PERSON][REDACTED_HANDLE] on 2015/01/22 15:23:27

	The SListView no longer emits a SelectionChanged event if the list is cleared, but nothing was selected.

diff --git a/Engine/Source/Runtime/Slate/Public/Widgets/Views/SListView.h b/Engine/Source/Runtime/Slate/Public/Widgets/Views/SListView.h
--- a/Engine/Source/Runtime/Slate/Public/Widgets/Views/SListView.h
+++ b/Engine/Source/Runtime/Slate/Public/Widgets/Views/SListView.h
@@ -1007,6 +1007,11 @@ public:
 			return;
 		}
 
+		if ( SelectedItems.Num() == 0 )
+		{
+			return;
+		}
+
 		Private_ClearSelection();
 		Private_SignalSelectionChanged(ESelectInfo::Direct);
 	}
