Fixed a crash caused by an image run being left as the only run in a text layout

UE-7717 - Crash when deleting an image added via an undo in editable rich-text

FTextLayout::RemoveAt was trying to ensure that a line never had zero runs by preventing the removal of the last run on a line, however, this can cause a crash if the last run on a line is not a text-based run.

This change makes FTextLayout::RemoveAt always remove a run when it becomes empty, but then re-add a new default text run if the removal caused the line to have zero runs.

diff --git a/Engine/Source/Runtime/Slate/Private/Framework/Text/TextLayout.cpp b/Engine/Source/Runtime/Slate/Private/Framework/Text/TextLayout.cpp
--- a/Engine/Source/Runtime/Slate/Private/Framework/Text/TextLayout.cpp
+++ b/Engine/Source/Runtime/Slate/Private/Framework/Text/TextLayout.cpp
@@ -1316,14 +1316,14 @@ bool FTextLayout::RemoveAt( const FTextLocation& Location, int32 Count )
 			const int32 IntersectedLength = IntersectedRangeToRemove.Len();
 			if (RunLength == IntersectedLength)
 			{
-				// The text for this entire run has been removed
-				if (LineModel.Runs.Num() > 1)
-				{
-					LineModel.Runs.RemoveAt(RunIndex);
-				}
-				else
+				// The text for this entire run has been removed, so remove this run
+				LineModel.Runs.RemoveAt(RunIndex);
+
+				// Every line needs at least one run - if we just removed the last run for this line, add a new default text run with a zero range
+				if (LineModel.Runs.Num() == 0)
 				{
-					RunModel.SetTextRange(FTextRange(0, 0));
+					TSharedRef<IRun> NewTextRun = CreateDefaultTextRun(LineModel.Text, FTextRange(0, 0));
+					LineModel.Runs.Add(NewTextRun);
 				}
 			}
 			else if (RunRange.BeginIndex > RemoveTextRange.BeginIndex)
