Merging using UE4-To-UE4-LauncherDev, CL 2424232, 2424451, 2424469, re-applied 2424198 (was incomplete)

diff --git a/Engine/Source/Runtime/Online/HTTP/Private/Windows/WindowsPlatformHttp.cpp b/Engine/Source/Runtime/Online/HTTP/Private/Windows/WindowsPlatformHttp.cpp
--- a/Engine/Source/Runtime/Online/HTTP/Private/Windows/WindowsPlatformHttp.cpp
+++ b/Engine/Source/Runtime/Online/HTTP/Private/Windows/WindowsPlatformHttp.cpp
@@ -6,7 +6,7 @@
 #include "Curl/CurlHttp.h"
 #include "Curl/CurlHttpManager.h"
 
-bool bUseCurl = false;
+bool bUseCurl = true;
 
 void FWindowsPlatformHttp::Init()
 {
@@ -22,9 +22,9 @@ void FWindowsPlatformHttp::Init()
 	// allow override on command line
 	FString HttpMode;
 	if (FParse::Value(FCommandLine::Get(), TEXT("HTTP="), HttpMode) &&
-		(HttpMode.Equals(TEXT("Curl"), ESearchCase::IgnoreCase) || HttpMode.Equals(TEXT("LibCurl"), ESearchCase::IgnoreCase)))
+		(HttpMode.Equals(TEXT("WinInet"), ESearchCase::IgnoreCase)))
 	{
-		bUseCurl = true;
+		bUseCurl = false;
 	}
 
 #if WITH_LIBCURL
diff --git a/Engine/Source/Runtime/Slate/Private/Framework/Application/SlateApplication.cpp b/Engine/Source/Runtime/Slate/Private/Framework/Application/SlateApplication.cpp
--- a/Engine/Source/Runtime/Slate/Private/Framework/Application/SlateApplication.cpp
+++ b/Engine/Source/Runtime/Slate/Private/Framework/Application/SlateApplication.cpp
@@ -355,7 +355,7 @@ TAutoConsoleVariable<int32> SkipSecondPrepass(
 /** Whether Slate should go to sleep when there are no active timers and the user is idle */
 TAutoConsoleVariable<int32> AllowSlateToSleep(
 	TEXT("Slate.AllowSlateToSleep"),
-	false,
+	true,
 	TEXT("Whether Slate should go to sleep when there are no active timers and the user is idle"));
 
 /** The amount of time that must pass without any user action before Slate is put to sleep (provided that there are no active timers). */
diff --git a/Engine/Source/Runtime/Slate/Private/Widgets/Notifications/SProgressBar.cpp b/Engine/Source/Runtime/Slate/Private/Widgets/Notifications/SProgressBar.cpp
--- a/Engine/Source/Runtime/Slate/Private/Widgets/Notifications/SProgressBar.cpp
+++ b/Engine/Source/Runtime/Slate/Private/Widgets/Notifications/SProgressBar.cpp
@@ -7,8 +7,8 @@ void SProgressBar::Construct( const FArguments& InArgs )
 {
 	check(InArgs._Style);
 
-	MarqueeAnimation = FCurveSequence(0.0f, 0.5f);
-	
+	MarqueeOffset = 0.0f;
+
 	Style = InArgs._Style;
 
 	SetPercent(InArgs._Percent);
@@ -20,25 +20,16 @@ void SProgressBar::Construct( const FArguments& InArgs )
 	
 	FillColorAndOpacity = InArgs._FillColorAndOpacity;
 	BorderPadding = InArgs._BorderPadding;
+
+	CurrentTickRate = 0.0f;
+	MinimumTickRate = InArgs._RefreshRate;
+
+	ActiveTimerHandle = RegisterActiveTimer(CurrentTickRate, FWidgetActiveTimerDelegate::CreateSP(this, &SProgressBar::ActiveTick));
 }
 
 void SProgressBar::SetPercent(TAttribute< TOptional<float> > InPercent)
 {
 	Percent = InPercent;
-
-	if ( InPercent.Get().IsSet() )
-	{
-		// No need for the marquee animation
-		MarqueeAnimation.JumpToEnd();
-	}
-	else
-	{
-		// Kick off the marquee animation
-		if ( !MarqueeAnimation.IsPlaying() )
-		{
-			MarqueeAnimation.Play(this->AsShared(), true);
-		}
-	}
 }
 
 void SProgressBar::SetStyle(const FProgressBarStyle* InStyle)
@@ -229,7 +220,7 @@ int32 SProgressBar::OnPaint( const FPaintArgs& Args, const FGeometry& AllottedGe
 		const FSlateBrush* CurrentMarqueeImage = GetMarqueeImage();
 		
 		// Draw Marquee
-		const float MarqueeAnimOffset = CurrentMarqueeImage->ImageSize.X * MarqueeAnimation.GetLerp();
+		const float MarqueeAnimOffset = CurrentMarqueeImage->ImageSize.X * MarqueeOffset;
 		const float MarqueeImageSize = CurrentMarqueeImage->ImageSize.X;
 
 		FSlateDrawElement::MakeBox(
@@ -258,3 +249,38 @@ FVector2D SProgressBar::ComputeDesiredSize( float ) const
 {
 	return GetMarqueeImage()->ImageSize;
 }
+
+void SProgressBar::SetActiveTimerTickRate(float TickRate)
+{
+	if (CurrentTickRate != TickRate || !ActiveTimerHandle.IsValid())
+	{
+		CurrentTickRate = TickRate;
+
+		TSharedPtr<FActiveTimerHandle> SharedActiveTimerHandle = ActiveTimerHandle.Pin();
+		if (SharedActiveTimerHandle.IsValid())
+		{
+			UnRegisterActiveTimer(SharedActiveTimerHandle.ToSharedRef());
+		}
+
+		ActiveTimerHandle = RegisterActiveTimer(TickRate, FWidgetActiveTimerDelegate::CreateSP(this, &SProgressBar::ActiveTick));
+	}
+}
+
+EActiveTimerReturnType SProgressBar::ActiveTick(double InCurrentTime, float InDeltaTime)
+{
+	MarqueeOffset = InCurrentTime - FMath::FloorToDouble(InCurrentTime);
+	
+	TOptional<float> PrecentFracton = Percent.Get();
+	if (PrecentFracton.IsSet())
+	{
+		SetActiveTimerTickRate(MinimumTickRate);
+	}
+	else
+	{
+		SetActiveTimerTickRate(0.0f);
+	}
+
+	return EActiveTimerReturnType::Continue;
+}
+
+
diff --git a/Engine/Source/Runtime/Slate/Public/Widgets/Notifications/SProgressBar.h b/Engine/Source/Runtime/Slate/Public/Widgets/Notifications/SProgressBar.h
--- a/Engine/Source/Runtime/Slate/Public/Widgets/Notifications/SProgressBar.h
+++ b/Engine/Source/Runtime/Slate/Public/Widgets/Notifications/SProgressBar.h
@@ -39,6 +39,7 @@ public:
 		, _BackgroundImage(nullptr)
 		, _FillImage(nullptr)
 		, _MarqueeImage(nullptr)
+		, _RefreshRate(2.0f)
 		{}
 
 		/** Style used for the progress bar */
@@ -65,6 +66,9 @@ public:
 		/** The brush to use as the marquee image */
 		SLATE_ARGUMENT(const FSlateBrush*, MarqueeImage)
 
+		/** Rate at which this widget is ticked when sleeping in seconds */
+		SLATE_ARGUMENT(float, RefreshRate)
+
 	SLATE_END_ARGS()
 
 	/**
@@ -103,6 +107,13 @@ public:
 	void SetMarqueeImage(const FSlateBrush* InMarqueeImage);
 	
 private:
+
+	/** Controls the speed at which the widget is ticked when in slate sleep mode */
+	void SetActiveTimerTickRate(float TickRate);
+
+	/** Widgets active tick */
+	EActiveTimerReturnType ActiveTick(double InCurrentTime, float InDeltaTime);
+
 	/** Gets the current background image. */
 	const FSlateBrush* GetBackgroundImage() const;
 	/** Gets the current fill image */
@@ -135,7 +146,16 @@ private:
 	/** Border Padding */
 	TAttribute<FVector2D> BorderPadding;
 
-	/** Curve sequence to drive progress bar animation */
-	FCurveSequence MarqueeAnimation;
+	/** Value to drive progress bar animation */
+	float MarqueeOffset;
+
+	/** Reference to the widgets current active timer */
+	TWeakPtr<FActiveTimerHandle> ActiveTimerHandle;
+
+	/** Rate at which the widget is currently ticked when slate sleep mode is active */
+	float CurrentTickRate;
+
+	/** The slowest that this widget can tick when in slate sleep mode */
+	float MinimumTickRate;
 };
 
