Merging UE4-LauncherDev to UE4, CL 2424586

diff --git a/Engine/Source/Runtime/Slate/Private/Framework/Text/RichTextMarkupProcessing.cpp b/Engine/Source/Runtime/Slate/Private/Framework/Text/RichTextMarkupProcessing.cpp
--- a/Engine/Source/Runtime/Slate/Private/Framework/Text/RichTextMarkupProcessing.cpp
+++ b/Engine/Source/Runtime/Slate/Private/Framework/Text/RichTextMarkupProcessing.cpp
@@ -31,7 +31,7 @@ void FDefaultRichTextMarkupParser::Process(TArray<FTextLineParseResults>& Result
 void FDefaultRichTextMarkupParser::ParseLineRanges(const FString& Input, const TArray<FTextRange>& LineRanges, TArray<FTextLineParseResults>& LineParseResultsArray) const
 {
 	// Special regular expression pattern for matching rich text markup elements. IE: <ElementName AttributeName="AttributeValue">Content</>
-	FRegexPattern ElementRegexPattern( TEXT("<([\\w\\d\\.]+)((?: (?:[\\w\\d\\.]+=(?>\".*?\")))+)?(?:(?:/>)|(?:>(.*?)</>))") );
+	FRegexPattern ElementRegexPattern( TEXT("<([\\w\\d\\.-]+)((?: (?:[\\w\\d\\.-]+=(?>\".*?\")))+)?(?:(?:/>)|(?:>(.*?)</>))") );
 	FRegexMatcher ElementRegexMatcher(ElementRegexPattern, Input);
 
 	// Parse line ranges, creating line parse results and run parse results.
