Added compile time checks to ensure non-UObject pointer types are incompatible with SListView

diff --git a/Engine/Source/Editor/DeviceProfileEditor/Private/SDeviceProfileCreateProfilePanel.h b/Engine/Source/Editor/DeviceProfileEditor/Private/SDeviceProfileCreateProfilePanel.h
--- a/Engine/Source/Editor/DeviceProfileEditor/Private/SDeviceProfileCreateProfilePanel.h
+++ b/Engine/Source/Editor/DeviceProfileEditor/Private/SDeviceProfileCreateProfilePanel.h
@@ -2,6 +2,7 @@
 
 #pragma once
 
+#include "DeviceProfiles/DeviceProfile.h"
 
 #define LOCTEXT_NAMESPACE "DeviceProfileCreateProfilePanel"
 
diff --git a/Engine/Source/Runtime/Slate/Public/Framework/Views/TableViewTypeTraits.h b/Engine/Source/Runtime/Slate/Public/Framework/Views/TableViewTypeTraits.h
--- a/Engine/Source/Runtime/Slate/Public/Framework/Views/TableViewTypeTraits.h
+++ b/Engine/Source/Runtime/Slate/Public/Framework/Views/TableViewTypeTraits.h
@@ -4,25 +4,27 @@
 
 
 /**
- * Lists/Trees only work with pointer types.
+ * Lists/Trees only work with shared pointer types, and UObjbectBase*.
  * Type traits to ensure that the user does not accidentally make a List/Tree of value types.
  */
-template<typename T> struct TIsValidListItem											{ enum { Value = false }; };
+template<typename T, typename Enable = void> struct TIsValidListItem					{ enum { Value = false }; };
 template<typename T> struct TIsValidListItem< TSharedRef<T, ESPMode::NotThreadSafe> >	{ enum { Value = true }; };
 template<typename T> struct TIsValidListItem< TSharedRef<T, ESPMode::ThreadSafe> >		{ enum { Value = true }; };
 template<typename T> struct TIsValidListItem< TSharedPtr<T, ESPMode::NotThreadSafe> >	{ enum { Value = true }; };
 template<typename T> struct TIsValidListItem< TSharedPtr<T, ESPMode::ThreadSafe> >		{ enum { Value = true }; };
-template<typename T> struct TIsValidListItem< T* >										{ enum { Value = true }; };
+template<typename T> struct TIsValidListItem< T*, typename TEnableIf<TPointerIsConvertibleFromTo<T, UObjectBase>::Value>::Type >
+ 																						{ enum { Value = true }; };
 template<typename T> struct TIsValidListItem< TWeakObjectPtr<T> >						{ enum { Value = true }; };
+					 
 
 /**
  * Furthermore, ListViews of TSharedPtr<> work differently from lists of UObject*.
  * ListTypeTraits provide the specialized functionality such as pointer testing, resetting,
  * and optional serialization for UObject garbage collection.
  */
-template <typename T> struct TListTypeTraits
+template <typename T, typename Enable=void> struct TListTypeTraits
 {
-	static_assert(TIsValidListItem<T>::Value, "Item type T must be a pointer or a TSharedPtr.");
+	static_assert(TIsValidListItem<T>::Value, "Item type T must be a UObjectBase pointer, TSharedRef, or TSharedPtr.");
 };
 
 
@@ -202,7 +204,7 @@ public:
  * In addition to testing and setting the pointers to null, Lists of UObjects
  * will serialize the objects they are holding onto.
  */
-template <typename T> struct TListTypeTraits< T* >
+template <typename T> struct TListTypeTraits< T*, typename TEnableIf<TPointerIsConvertibleFromTo<T, UObjectBase>::Value>::Type >
 {
 public:
 	typedef T* NullableType;
diff --git a/Engine/Source/Runtime/Slate/Public/Widgets/Views/STableRow.h b/Engine/Source/Runtime/Slate/Public/Widgets/Views/STableRow.h
--- a/Engine/Source/Runtime/Slate/Public/Widgets/Views/STableRow.h
+++ b/Engine/Source/Runtime/Slate/Public/Widgets/Views/STableRow.h
@@ -69,7 +69,7 @@ DECLARE_DELEGATE_RetVal_OneParam(FReply, FOnTableRowDrop, FDragDropEvent const&)
 template<typename ItemType>
 class STableRow : public ITableRow, public SBorder
 {
-	static_assert(TIsValidListItem<ItemType>::Value, "Item type T must be a pointer or a TSharedPtr.");
+	static_assert(TIsValidListItem<ItemType>::Value, "Item type T must be UObjectBase*, TSharedRef<>, or TSharedPtr<>.");
 
 public:
 	/** Delegate signature for querying whether this FDragDropEvent will be handled by the drop target of type ItemType. */
