[INTEGRATE] Change 2428851 by [REDACTED_PERSON][REDACTED_HANDLE] on 2015/02/02 17:39:35

	Fixed the "draw attention" animation playing when focusing tabs that already have focus

	UE-8742 - Regression 4.7: Every time you click a mode button, the App title tab flashes a few times

	This feature was originally intended to draw attention to a tab when it became active (for example, when clicking the button to find an asset in the Content Browser while using the details panel), however it also triggered when the tab being invoked was already active, causing the flashing seen when using the "Modes" UI in the editor.

diff --git a/Engine/Source/Runtime/Slate/Private/Framework/Docking/TabManager.cpp b/Engine/Source/Runtime/Slate/Private/Framework/Docking/TabManager.cpp
--- a/Engine/Source/Runtime/Slate/Private/Framework/Docking/TabManager.cpp
+++ b/Engine/Source/Runtime/Slate/Private/Framework/Docking/TabManager.cpp
@@ -920,7 +920,11 @@ TSharedRef<SDockTab> FTabManager::InvokeTab_Internal( const FTabId& TabId )
 
 		if ( ExistingTab.IsValid() )
 		{
-			DrawAttention( ExistingTab.ToSharedRef() );
+			if ( !ExistingTab->IsActive() )
+			{
+				// Draw attention to this tab if it didn't already have focus
+				DrawAttention( ExistingTab.ToSharedRef() );
+			}
 			return ExistingTab.ToSharedRef();
 		}
 	}
