Process Mac menu bar update asynchronously to avoid processing of the remaining queued events in the middle of handling the mouse up event.

diff --git a/Engine/Source/Runtime/Slate/Private/Framework/MultiBox/Mac/MacMenu.cpp b/Engine/Source/Runtime/Slate/Private/Framework/MultiBox/Mac/MacMenu.cpp
--- a/Engine/Source/Runtime/Slate/Private/Framework/MultiBox/Mac/MacMenu.cpp
+++ b/Engine/Source/Runtime/Slate/Private/Framework/MultiBox/Mac/MacMenu.cpp
@@ -152,7 +152,7 @@ void FSlateMacMenu::UpdateWithMultiBox(const TSharedPtr< FMultiBox > MultiBox)
 		}
 
 		FPlatformMisc::bChachedMacMenuStateNeedsUpdate = true;
-	});
+	}, NSDefaultRunLoopMode, false);
 }
 
 void FSlateMacMenu::UpdateMenu(FMacMenu* Menu)
