[INTEGRATE] Change 2434737 by [REDACTED_PERSON][REDACTED_HANDLE] on 2015/02/05 22:21:32

	Editor: Enable automatic indentation for multibox menu items missing icons in sections where at least one item contains an icon

diff --git a/Engine/Source/Runtime/Slate/Private/Framework/MultiBox/SMenuEntryBlock.cpp b/Engine/Source/Runtime/Slate/Private/Framework/MultiBox/SMenuEntryBlock.cpp
--- a/Engine/Source/Runtime/Slate/Private/Framework/MultiBox/SMenuEntryBlock.cpp
+++ b/Engine/Source/Runtime/Slate/Private/Framework/MultiBox/SMenuEntryBlock.cpp
@@ -332,8 +332,8 @@ TSharedRef< SWidget > SMenuEntryBlock::BuildMenuEntryWidget( const FMenuEntryBui
 		else
 		{
 			// This will silently pad the offending items
-			// 	IconWidget = SNew(SSpacer)
-			// 		.Size(FVector2D(MultiBoxConstants::MenuIconSize, MultiBoxConstants::MenuIconSize));
+			IconWidget = SNew(SSpacer)
+				.Size(FVector2D(MultiBoxConstants::MenuIconSize, MultiBoxConstants::MenuIconSize));
 		}
 	}
 
