[INTEGRATE] Change 2435443 by [REDACTED_PERSON][REDACTED_HANDLE] on 2015/02/06 10:32:36

	Slate - The ComboButton now properly exposes the underlying event for the menu opening and closing.

diff --git a/Engine/Source/Runtime/Slate/Private/Widgets/Input/SComboButton.cpp b/Engine/Source/Runtime/Slate/Private/Widgets/Input/SComboButton.cpp
--- a/Engine/Source/Runtime/Slate/Private/Widgets/Input/SComboButton.cpp
+++ b/Engine/Source/Runtime/Slate/Private/Widgets/Input/SComboButton.cpp
@@ -23,6 +23,7 @@ void SComboButton::Construct( const FArguments& InArgs )
 	SMenuAnchor::Construct( SMenuAnchor::FArguments()
 		.Placement(InArgs._MenuPlacement)
 		.Method(InArgs._Method)
+		.OnMenuOpenChanged(InArgs._OnMenuOpenChanged)
 		[
 			SNew( SButton )
 			.ButtonStyle( OurButtonStyle )
diff --git a/Engine/Source/Runtime/Slate/Public/Widgets/Input/SComboButton.h b/Engine/Source/Runtime/Slate/Public/Widgets/Input/SComboButton.h
--- a/Engine/Source/Runtime/Slate/Public/Widgets/Input/SComboButton.h
+++ b/Engine/Source/Runtime/Slate/Public/Widgets/Input/SComboButton.h
@@ -39,6 +39,7 @@ public:
 		
 		/** Sets an event handler to generate a widget dynamically when the menu is needed. */
 		SLATE_EVENT( FOnGetContent, OnGetMenuContent )
+		SLATE_EVENT( FOnIsOpenChanged, OnMenuOpenChanged )
 		
 		SLATE_EVENT( FOnComboBoxOpened, OnComboBoxOpened )
 		SLATE_ARGUMENT( bool, IsFocusable )
