[INTEGRATE] Change 2436564 by [REDACTED_PERSON][REDACTED_HANDLE] on 2015/02/06 22:19:45

	Backing out menu padding change CL# 2434737 for now, as not all items get padded appropriately

diff --git a/Engine/Source/Runtime/Slate/Private/Framework/MultiBox/SMenuEntryBlock.cpp b/Engine/Source/Runtime/Slate/Private/Framework/MultiBox/SMenuEntryBlock.cpp
--- a/Engine/Source/Runtime/Slate/Private/Framework/MultiBox/SMenuEntryBlock.cpp
+++ b/Engine/Source/Runtime/Slate/Private/Framework/MultiBox/SMenuEntryBlock.cpp
@@ -332,8 +332,8 @@ TSharedRef< SWidget > SMenuEntryBlock::BuildMenuEntryWidget( const FMenuEntryBui
 		else
 		{
 			// This will silently pad the offending items
-			IconWidget = SNew(SSpacer)
-				.Size(FVector2D(MultiBoxConstants::MenuIconSize, MultiBoxConstants::MenuIconSize));
+			// 	IconWidget = SNew(SSpacer)
+			// 		.Size(FVector2D(MultiBoxConstants::MenuIconSize, MultiBoxConstants::MenuIconSize));
 		}
 	}
 
