[INTEGRATE] Change 2439338 by [REDACTED_PERSON][REDACTED_HANDLE] on 2015/02/09 23:26:35

	Slate - Allowing the STextComboBox to be styled.

diff --git a/Engine/Source/Runtime/Slate/Private/Widgets/Input/STextComboBox.cpp b/Engine/Source/Runtime/Slate/Private/Widgets/Input/STextComboBox.cpp
--- a/Engine/Source/Runtime/Slate/Private/Widgets/Input/STextComboBox.cpp
+++ b/Engine/Source/Runtime/Slate/Private/Widgets/Input/STextComboBox.cpp
@@ -14,6 +14,8 @@ void STextComboBox::Construct( const FArguments& InArgs )
 	this->ChildSlot
 	[
 		SAssignNew(StringCombo, SComboBox< TSharedPtr<FString> > )
+		.ComboBoxStyle(InArgs._ComboBoxStyle)
+		.ButtonStyle(InArgs._ButtonStyle)
 		.OptionsSource(InArgs._OptionsSource)
 		.OnGenerateWidget(this, &STextComboBox::MakeItemWidget)
 		.OnSelectionChanged(this, &STextComboBox::OnSelectionChanged)
diff --git a/Engine/Source/Runtime/Slate/Public/Widgets/Input/STextComboBox.h b/Engine/Source/Runtime/Slate/Public/Widgets/Input/STextComboBox.h
--- a/Engine/Source/Runtime/Slate/Public/Widgets/Input/STextComboBox.h
+++ b/Engine/Source/Runtime/Slate/Public/Widgets/Input/STextComboBox.h
@@ -12,12 +12,19 @@ public:
 	DECLARE_DELEGATE_RetVal_OneParam( FString, FGetTextComboLabel, TSharedPtr<FString> );
 	typedef TSlateDelegates< TSharedPtr<FString> >::FOnSelectionChanged FOnTextSelectionChanged;
 
-	SLATE_BEGIN_ARGS( STextComboBox ) 
-		: _ColorAndOpacity( FSlateColor::UseForeground() )
+	SLATE_BEGIN_ARGS( STextComboBox )
+		: _ComboBoxStyle(&FCoreStyle::Get().GetWidgetStyle< FComboBoxStyle >("ComboBox"))
+		, _ButtonStyle(nullptr)
+		, _ColorAndOpacity( FSlateColor::UseForeground() )
 		, _ContentPadding(FMargin(4.0, 2.0))
 		, _OnGetTextLabelForItem()
 		{}
 
+		SLATE_STYLE_ARGUMENT(FComboBoxStyle, ComboBoxStyle)
+
+		/** The visual style of the button part of the combo box (overrides ComboBoxStyle) */
+		SLATE_STYLE_ARGUMENT(FButtonStyle, ButtonStyle)
+
 		/** Selection of strings to pick from */
 		SLATE_ARGUMENT( TArray< TSharedPtr<FString> >*, OptionsSource )
 
