UMG/Slate - Adding windowless rendering for drag and drop, there's a new drag & drop operation you should inherit from for games.  FGameDragDropOperation, can be used for any single window game drag/drop operation you'd like to implement.  All drag/drop operations in UMG use this operation type.

diff --git a/Engine/Source/Runtime/Slate/Private/Framework/Application/SlateApplication.cpp b/Engine/Source/Runtime/Slate/Private/Framework/Application/SlateApplication.cpp
--- a/Engine/Source/Runtime/Slate/Private/Framework/Application/SlateApplication.cpp
+++ b/Engine/Source/Runtime/Slate/Private/Framework/Application/SlateApplication.cpp
@@ -965,6 +965,32 @@ void FSlateApplication::DrawWindowAndChildren( const TSharedRef<SWindow>& Window
 				FWidgetStyle(),
 				WindowToDraw->IsEnabled() );
 
+			// Draw drag drop operation if it's windowless.
+			if ( IsDragDropping() && DragDropContent->IsWindowlessOperation() )
+			{
+				TSharedPtr<SWindow> DragDropWindow = DragDropWindowPtr.Pin();
+				if ( DragDropWindow.IsValid() && DragDropWindow == WindowToDraw )
+				{
+					TSharedPtr<SWidget> DecoratorWidget = DragDropContent->GetDefaultDecorator();
+					if ( DecoratorWidget.IsValid() && DecoratorWidget->GetVisibility().IsVisible() )
+					{
+						DecoratorWidget->SetVisibility(EVisibility::HitTestInvisible);
+						DecoratorWidget->SlatePrepass();
+
+						FVector2D DragDropContentInWindowSpace = WindowToDraw->GetWindowGeometryInScreen().AbsoluteToLocal(DragDropContent->GetDecoratorPosition());
+						const FGeometry DragDropContentGeometry = FGeometry::MakeRoot(DecoratorWidget->GetDesiredSize(), FSlateLayoutTransform(DragDropContentInWindowSpace));
+
+						DecoratorWidget->Paint(
+							FPaintArgs(WindowToDraw, *WindowToDraw->GetHittestGrid(), WindowToDraw->GetPositionInScreen(), GetCurrentTime(), GetDeltaTime()),
+							DragDropContentGeometry, WindowToDraw->GetClippingRectangleInWindow(),
+							WindowElementList,
+							++MaxLayerId,
+							FWidgetStyle(),
+							WindowToDraw->IsEnabled());
+					}
+				}
+			}
+
 			// Draw Software Cursor
 			TSharedPtr<SWindow> CursorWindow = CursorWindowPtr.Pin();
 			if (CursorWindow.IsValid() && WindowToDraw == CursorWindow)
@@ -4522,6 +4548,17 @@ bool FSlateApplication::ProcessMouseMoveEvent( FPointerEvent& MouseEvent, bool b
 		FScopedSwitchWorldHack SwitchWorld( WidgetsUnderCursor );
 		DragDropContent->OnDragged( DragDropEvent );
 
+		// Update the window we're under for rendering the drag drop operation if
+		// it's a windowless drag drop operation.
+		if ( WidgetsUnderCursor.IsValid() )
+		{
+			DragDropWindowPtr = WidgetsUnderCursor.GetWindow();
+		}
+		else
+		{
+			DragDropWindowPtr = nullptr;
+		}
+
 		// check the drag-drop operation for a cursor switch (on Windows, the OS thinks the mouse is
 		// captured so we wont get QueryCursor calls for drag/drops internal to the Slate application)
 		FCursorReply CursorResult = DragDropContent->OnCursorQuery();
@@ -4536,6 +4573,10 @@ bool FSlateApplication::ProcessMouseMoveEvent( FPointerEvent& MouseEvent, bool b
 			PlatformApplication->Cursor->SetType( EMouseCursor::Default );
 		}
 	}
+	else
+	{
+		DragDropWindowPtr = nullptr;
+	}
 
 	WidgetsUnderCursorLastEvent = FWeakWidgetPath( WidgetsUnderCursor );
 
diff --git a/Engine/Source/Runtime/Slate/Public/Framework/Application/SlateApplication.h b/Engine/Source/Runtime/Slate/Public/Framework/Application/SlateApplication.h
--- a/Engine/Source/Runtime/Slate/Public/Framework/Application/SlateApplication.h
+++ b/Engine/Source/Runtime/Slate/Public/Framework/Application/SlateApplication.h
@@ -1401,6 +1401,9 @@ private:
 	/** When not null, the content of the current drag drop operation. */
 	TSharedPtr< FDragDropOperation > DragDropContent;
 
+	/** The window the drag drop content is over. */
+	TWeakPtr< SWindow > DragDropWindowPtr;
+
 	/** Whether or not we are requesting that we leave debugging mode after the tick is complete */
 	bool bRequestLeaveDebugMode;
 	/** Whether or not we need to leave debug mode for single stepping */
diff --git a/Engine/Source/Runtime/SlateCore/Public/Input/DragAndDrop.h b/Engine/Source/Runtime/SlateCore/Public/Input/DragAndDrop.h
--- a/Engine/Source/Runtime/SlateCore/Public/Input/DragAndDrop.h
+++ b/Engine/Source/Runtime/SlateCore/Public/Input/DragAndDrop.h
@@ -21,7 +21,7 @@ public:
 	/**
 	 * Default constructor.
 	 */
-	FDragDropOperation( ) { }
+	FDragDropOperation();
 
 	/**
 	 * Destructor.
@@ -52,15 +52,24 @@ public:
 	/** Allows drag/drop operations to override the current cursor */
 	virtual FCursorReply OnCursorQuery();
 
-	/** Gets the widget that will serve as the decorator unless overridden. If you do not override, you will have no decorator */
+	/**
+	 * Gets the widget that will serve as the decorator unless overridden. 
+	 * If you do not override, you will have no decorator
+	 */
 	virtual TSharedPtr<SWidget> GetDefaultDecorator() const { return TSharedPtr<SWidget>(); }
+
+	/** Gets the desired position of the decorator in absolute space. */
+	virtual FVector2D GetDecoratorPosition() const { return FVector2D(0, 0); }
 	
 	/** Alters the visibility of the window */
-	void SetDecoratorVisibility(bool bVisible);
+	virtual void SetDecoratorVisibility(bool bVisible);
 
 	/** Is this drag Drop operation going to interact with applications outside of Slate */
 	virtual bool IsExternalOperation() const { return false; }
 
+	/** @return true if the operation does not create a decorator window.  Used for game drag drop operations. */
+	virtual bool IsWindowlessOperation() const { return bCreateNewWindow == false; }
+
 	/** 
 	 * Sets the cursor to override the drag drop operations cursor with so that a 
 	 * control can give temporary feedback, for example - a slashed circle telling 
@@ -76,9 +85,14 @@ protected:
 	virtual void Construct();
 
 	/**
-	 * Destroys the cursor decorator.
+	 * Creates a window for the cursor decorator.
+	 */
+	void CreateCursorDecoratorWindow();
+
+	/**
+	 * Destroys the cursor decorator window.
 	 */
-	void DestroyCursorDecoratorWindow( );
+	void DestroyCursorDecoratorWindow();
 
 	/**
 	 * Checks whether this drag and drop operation can cast safely to the specified type.
@@ -90,6 +104,12 @@ protected:
 
 protected:
 
+	/**
+	 * Create a new window for the decorator content.  If this is false,
+	 * SlateApplication will draw the decorator directly using the 
+	 */
+	bool bCreateNewWindow;
+
 	/** The window that owns the decorator widget */
 	TSharedPtr<SWindow> CursorDecoratorWindow;
 
@@ -216,4 +236,21 @@ private:
 };
 
 
+class SLATECORE_API FGameDragDropOperation : public FDragDropOperation
+{
+public:
+	DRAG_DROP_OPERATOR_TYPE(FGameDragDropOperation, FDragDropOperation)
+
+	/**
+	* Default constructor.
+	*/
+	FGameDragDropOperation();
+
+	virtual FVector2D GetDecoratorPosition() const override;
+
+protected:
+	/** The absolute position of the decorator. */
+	FVector2D DecoratorPosition;
+};
+
 #include "DragAndDrop.inl"
diff --git a/Engine/Source/Runtime/UMG/Private/Slate/UMGDragDropOp.cpp b/Engine/Source/Runtime/UMG/Private/Slate/UMGDragDropOp.cpp
--- a/Engine/Source/Runtime/UMG/Private/Slate/UMGDragDropOp.cpp
+++ b/Engine/Source/Runtime/UMG/Private/Slate/UMGDragDropOp.cpp
@@ -17,6 +17,11 @@ void FUMGDragDropOp::AddReferencedObjects(FReferenceCollector& Collector)
 	Collector.AddReferencedObject(DragOperation);
 }
 
+void FUMGDragDropOp::Construct()
+{
+
+}
+
 void FUMGDragDropOp::OnDrop( bool bDropWasHandled, const FPointerEvent& MouseEvent )
 {
 	if ( DragOperation )
@@ -96,11 +101,12 @@ void FUMGDragDropOp::OnDragged( const class FDragDropEvent& DragDropEvent )
 	{
 		float T = DeltaTime / AnimationTime;
 		FVector2D LerpPosition = ( Position - StartingScreenPos ) * T;
-		CursorDecoratorWindow->MoveWindowTo(StartingScreenPos + LerpPosition);
+		
+		DecoratorPosition = StartingScreenPos + LerpPosition;
 	}
 	else
 	{
-		CursorDecoratorWindow->MoveWindowTo(Position);
+		DecoratorPosition = Position;
 	}
 }
 
diff --git a/Engine/Source/Runtime/UMG/Public/Slate/UMGDragDropOp.h b/Engine/Source/Runtime/UMG/Public/Slate/UMGDragDropOp.h
--- a/Engine/Source/Runtime/UMG/Public/Slate/UMGDragDropOp.h
+++ b/Engine/Source/Runtime/UMG/Public/Slate/UMGDragDropOp.h
@@ -8,16 +8,18 @@
  * This is the drag/drop class used for UMG, all UMG drag drop operations utilize this operation.
  * It supports moving a UObject payload and using a UWidget decorator.
  */
-class FUMGDragDropOp : public FDragDropOperation, public FGCObject
+class FUMGDragDropOp : public FGameDragDropOperation, public FGCObject
 {
 public:
-	DRAG_DROP_OPERATOR_TYPE(FUMGDragDropOp, FDragDropOperation)
+	DRAG_DROP_OPERATOR_TYPE(FUMGDragDropOp, FGameDragDropOperation)
 	
 	static TSharedRef<FUMGDragDropOp> New(UDragDropOperation* Operation, const FVector2D &CursorPosition, const FVector2D &ScreenPositionOfNode, TSharedPtr<SObjectWidget> SourceUserWidget);
 
 	FUMGDragDropOp();
 
+	// Begin FGCObject
 	virtual void AddReferencedObjects(FReferenceCollector& Collector) override;
+	// End FGCObject
 
 	virtual void OnDrop( bool bDropWasHandled, const FPointerEvent& MouseEvent ) override;
 
@@ -27,6 +29,9 @@ public:
 
 	UDragDropOperation* GetOperation() const { return DragOperation; }
 
+protected:
+	virtual void Construct() override;
+
 private:
 
 	// Raw pointer to the drag operation, kept alive by AddReferencedObjects.
