When running in PIE all users focus is set to the main window.
There were compoiants of multi controller games not having all controller input recieved.

diff --git a/Engine/Source/Runtime/Slate/Private/Framework/Application/SlateApplication.cpp b/Engine/Source/Runtime/Slate/Private/Framework/Application/SlateApplication.cpp
--- a/Engine/Source/Runtime/Slate/Private/Framework/Application/SlateApplication.cpp
+++ b/Engine/Source/Runtime/Slate/Private/Framework/Application/SlateApplication.cpp
@@ -1908,7 +1908,7 @@ void FSlateApplication::RegisterGameViewport( TSharedRef<SViewport> InViewport )
 	// If we cannot find the window it could have been destroyed.
 	if (FSlateWindowHelper::FindPathToWidget(SlateWindows, InViewport, PathToViewport, EVisibility::All))
 	{
-		FReply Reply = FReply::Handled().SetUserFocus(InViewport, EFocusCause::SetDirectly);
+		FReply Reply = FReply::Handled().SetUserFocus(InViewport, EFocusCause::SetDirectly, true);
 	
 		// Set keyboard focus on the actual OS window for the top level Slate window in the viewport path
 		// This is needed because some OS messages are only sent to the window with keyboard focus
