[INTEGRATE] Change 2445115 by [REDACTED_PERSON][REDACTED_HANDLE] on 2015/02/13 14:12:21

	UE-9756 - UMG - Now checking if there is a registered factory for a tab before attempting to spawn it in the blueprint editor.

diff --git a/Engine/Source/Editor/Kismet/Private/BlueprintEditor.cpp b/Engine/Source/Editor/Kismet/Private/BlueprintEditor.cpp
--- a/Engine/Source/Editor/Kismet/Private/BlueprintEditor.cpp
+++ b/Engine/Source/Editor/Kismet/Private/BlueprintEditor.cpp
@@ -786,14 +786,7 @@ void FBlueprintEditor::OnSelectionUpdated(const TArray<FSCSEditorTreeNodePtrType
 					InspectorTitle = FText::FromString(Title);
 					bShowComponents = false;
 
-					// Show the details panel if it doesn't exist.
-					TabManager->InvokeTab(FBlueprintEditorTabs::DetailsID);
-
-					TSharedPtr<SDockTab> OwnerTab = Inspector->GetOwnerTab();
-					if ( OwnerTab.IsValid() )
-					{
-						OwnerTab->FlashTab();
-					}
+					TryInvokingDetailsTab();
 				}
 				else
 				{
@@ -1987,7 +1980,7 @@ void FBlueprintEditor::SetupViewForBlueprintEditingMode()
 
 	// Make sure the inspector is always on top
 	//@TODO: This is necessary right now because of a bug in restoring layouts not remembering which tab is on top (to get it right initially), but do we want this behavior always?
-	TabManager->InvokeTab(FBlueprintEditorTabs::DetailsID);
+	TryInvokingDetailsTab();
 
 	UBlueprint* Blueprint = GetBlueprintObj();
 	if ((Blueprint != nullptr) && (Blueprint->Status == EBlueprintStatus::BS_Error))
@@ -2582,14 +2575,7 @@ void FBlueprintEditor::EditGlobalOptions_Clicked()
 		// Show details for the Blueprint instance we're editing
 		Inspector->ShowDetailsForSingleObject(Blueprint);
 
-		// Show the details panel if it doesn't exist.
-		TabManager->InvokeTab(FBlueprintEditorTabs::DetailsID);
-
-		TSharedPtr<SDockTab> OwnerTab = Inspector->GetOwnerTab();
-		if ( OwnerTab.IsValid() )
-		{
-			OwnerTab->FlashTab();
-		}
+		TryInvokingDetailsTab();
 	}
 }
 
@@ -6305,14 +6291,7 @@ void FBlueprintEditor::StartEditingDefaults(bool bAutoFocus, bool bForceRefresh)
 
 					Inspector->ShowDetailsForSingleObject(DefaultObject, Options);
 
-					// Show the details panel if it doesn't exist.
-					TabManager->InvokeTab(FBlueprintEditorTabs::DetailsID);
-
-					TSharedPtr<SDockTab> OwnerTab = Inspector->GetOwnerTab();
-					if ( OwnerTab.IsValid() )
-					{
-						OwnerTab->FlashTab();
-					}
+					TryInvokingDetailsTab();
 				}
 			}
 		}
@@ -6345,7 +6324,7 @@ void FBlueprintEditor::RefreshStandAloneDefaultsEditor()
 void FBlueprintEditor::RenameNewlyAddedAction(FName InActionName)
 {
 	TabManager->InvokeTab(FBlueprintEditorTabs::MyBlueprintID);
-	TabManager->InvokeTab(FBlueprintEditorTabs::DetailsID);
+	TryInvokingDetailsTab(/*Flash*/false);
 
 	if (MyBlueprintWidget.IsValid())
 	{
@@ -7549,10 +7528,26 @@ void FBlueprintEditor::SaveAsset_Execute()
 	IBlueprintEditor::SaveAsset_Execute();
 }
 
+void FBlueprintEditor::TryInvokingDetailsTab(bool bFlash)
+{
+	if ( TabManager->CanSpawnTab(FBlueprintEditorTabs::DetailsID) )
+	{
+		// Show the details panel if it doesn't exist.
+		TabManager->InvokeTab(FBlueprintEditorTabs::DetailsID);
+
+		if ( bFlash )
+		{
+			TSharedPtr<SDockTab> OwnerTab = Inspector->GetOwnerTab();
+			if ( OwnerTab.IsValid() )
+			{
+				OwnerTab->FlashTab();
+			}
+		}
+	}
+}
+
 /////////////////////////////////////////////////////
 /////////////////////////////////////////////////////
 /////////////////////////////////////////////////////
 
-
 #undef LOCTEXT_NAMESPACE
-
diff --git a/Engine/Source/Editor/Kismet/Public/BlueprintEditor.h b/Engine/Source/Editor/Kismet/Public/BlueprintEditor.h
--- a/Engine/Source/Editor/Kismet/Public/BlueprintEditor.h
+++ b/Engine/Source/Editor/Kismet/Public/BlueprintEditor.h
@@ -884,6 +884,9 @@ protected:
 	/** Used to get the apperance of a specific graph, GetGraphAppearance() uses the currently focused graph. */
 	FGraphAppearanceInfo GetGraphAppearance(class UEdGraph* InGraph) const;
 
+	/** Attempts to invoke the details tab if it's currently possible to. */
+	void TryInvokingDetailsTab(bool bFlash = true);
+
 private:
 
 	/** Returns true if modules can be recompiled */
diff --git a/Engine/Source/Runtime/Slate/Private/Framework/Docking/TabManager.cpp b/Engine/Source/Runtime/Slate/Private/Framework/Docking/TabManager.cpp
--- a/Engine/Source/Runtime/Slate/Private/Framework/Docking/TabManager.cpp
+++ b/Engine/Source/Runtime/Slate/Private/Framework/Docking/TabManager.cpp
@@ -761,7 +761,6 @@ void FTabManager::PopulateTabSpawnerMenu_Helper( FMenuBuilder& PopulateMe, FPopu
 	}
 }
 
-
 void FTabManager::MakeSpawnerMenuEntry( FMenuBuilder &PopulateMe, const TSharedPtr<FTabSpawnerEntry> &SpawnerNode ) 
 {
 	auto CanExecuteMenuEntry = [](TAttribute<ETabSpawnerMenuType::Type> InMenuType) -> bool
@@ -863,7 +862,6 @@ void FTabManager::DrawAttention( const TSharedRef<SDockTab>& TabToHighlight )
 	TabToHighlight->FlashTab();
 }
 
-
 void FTabManager::InsertNewDocumentTab( FName PlaceholderId, ESearchPreference::Type SearchPreference, const TSharedRef<SDockTab>& UnmanagedTab )
 {
 	InsertDocumentTab(PlaceholderId, SearchPreference, UnmanagedTab, true);
@@ -874,7 +872,6 @@ void FTabManager::RestoreDocumentTab( FName PlaceholderId, ESearchPreference::Ty
 	InsertDocumentTab(PlaceholderId, SearchPreference, UnmanagedTab, false);
 }
 
-
 TSharedRef<SDockTab> FTabManager::InvokeTab( const FTabId& TabId )
 {
 	TSharedRef<SDockTab> NewTab = InvokeTab_Internal( TabId );
@@ -1039,7 +1036,6 @@ FTabManager::FTabManager( const TSharedPtr<SDockTab>& InOwnerTab, const TSharedR
 	LocalWorkspaceMenuRoot = FWorkspaceItem::NewGroup(LOCTEXT("LocalWorkspaceRoot", "Local Workspace Root"));
 }
 
-
 TSharedRef<SDockingArea> FTabManager::RestoreArea( const TSharedRef<FArea>& AreaToRestore, const TSharedPtr<SWindow>& InParentWindow, const bool bEmbedTitleAreaContent  )
 {
 	TSharedRef<SDockingNode> RestoredNode = RestoreArea_Helper( AreaToRestore, InParentWindow, bEmbedTitleAreaContent );	
@@ -1153,8 +1149,6 @@ TSharedRef<SDockingNode> FTabManager::RestoreArea_Helper( const TSharedRef<FLayo
 			RestoreSplitterContent( NodeAsArea.ToSharedRef(), NewDockAreaWidget.ToSharedRef(), ParentWindow );
 		}
 		
-		
-
 		return NewDockAreaWidget.ToSharedRef();
 	}
 	else if ( NodeAsSplitter.IsValid() ) 
@@ -1175,8 +1169,6 @@ TSharedRef<SDockingNode> FTabManager::RestoreArea_Helper( const TSharedRef<FLayo
 	}
 }
 
-
-
 void FTabManager::RestoreSplitterContent( const TSharedRef<FSplitter>& SplitterNode, const TSharedRef<SDockingSplitter>& SplitterWidget, const TSharedPtr<SWindow>& ParentWindow )
 {
 	// Restore the contents of this splitter.
@@ -1190,6 +1182,11 @@ void FTabManager::RestoreSplitterContent( const TSharedRef<FSplitter>& SplitterN
 	}
 }
 
+bool FTabManager::CanSpawnTab(FName TabId)
+{
+	TSharedPtr<FTabSpawnerEntry> Spawner = FindTabSpawnerFor(TabId);
+	return Spawner.IsValid();
+}
 
 bool FTabManager::IsValidTabForSpawning( const FTab& SomeTab ) const
 {
@@ -1198,7 +1195,6 @@ bool FTabManager::IsValidTabForSpawning( const FTab& SomeTab ) const
 	return ( !NomadSpawner || !NomadSpawner->Get().IsSoleTabInstanceSpawned() );
 }
 
-
 TSharedRef<SDockTab> FTabManager::SpawnTab( const FTabId& TabId, const TSharedPtr<SWindow>& ParentWindow )
 {
 	TSharedPtr<SDockTab> NewTabWidget;
@@ -1242,8 +1238,6 @@ TSharedRef<SDockTab> FTabManager::SpawnTab( const FTabId& TabId, const TSharedPt
 	return NewTabWidget.ToSharedRef();
 }
 
-
-
 TSharedPtr<SDockTab> FTabManager::FindExistingLiveTab( const FTabId& TabId ) const
 {
 	for ( int32 AreaIndex = 0; AreaIndex < DockAreas.Num(); ++AreaIndex )
@@ -1265,7 +1259,6 @@ TSharedPtr<SDockTab> FTabManager::FindExistingLiveTab( const FTabId& TabId ) con
 	return TSharedPtr<SDockTab>();
 }
 
-
 TSharedPtr<class SDockingTabStack> FTabManager::FindTabInLiveAreas( const FTabMatcher& TabMatcher ) const
 {
 	for ( int32 AreaIndex = 0; AreaIndex < DockAreas.Num(); ++AreaIndex )
diff --git a/Engine/Source/Runtime/Slate/Public/Framework/Docking/TabManager.h b/Engine/Source/Runtime/Slate/Public/Framework/Docking/TabManager.h
--- a/Engine/Source/Runtime/Slate/Public/Framework/Docking/TabManager.h
+++ b/Engine/Source/Runtime/Slate/Public/Framework/Docking/TabManager.h
@@ -628,6 +628,9 @@ class SLATE_API FTabManager : public TSharedFromThis<FTabManager>
 		/** Clears all categories in the local workspace menu */
 		void ClearLocalWorkspaceMenuCategories();
 
+		/** @return true if the tab has a factory registered for it that allows it to be spawned. */
+		bool CanSpawnTab( FName TabId );
+
 	protected:
 		void InvokeTabForMenu( FName TabId );
 
