Slate - No longer using the SelectorFocusedBrush when creating the drop zone effect on table rows.  This color is normally clear on applications not using the focus ring, like the editor, and so it never shows up.  Better to just use the widget color and a perdictive color like the brush tint.

diff --git a/Engine/Source/Runtime/Slate/Public/Widgets/Views/STableRow.h b/Engine/Source/Runtime/Slate/Public/Widgets/Views/STableRow.h
--- a/Engine/Source/Runtime/Slate/Public/Widgets/Views/STableRow.h
+++ b/Engine/Source/Runtime/Slate/Public/Widgets/Views/STableRow.h
@@ -223,7 +223,7 @@ public:
 				DropIndicatorBrush,
 				MyClippingRect,
 				ESlateDrawEffect::None,
-				Style->SelectorFocusedBrush.GetTint(InWidgetStyle) * InWidgetStyle.GetColorAndOpacityTint()
+				InWidgetStyle.GetColorAndOpacityTint()
 			);
 		}
 
