Mobile - The Virtual Joystick no longer takes takes keyboard focus.  This prevents the problem of it sealing focus from the game viewport and resetting the touch state of fingers that have been pressed outside of the joystick.

diff --git a/Engine/Source/Runtime/Slate/Private/Widgets/Input/SVirtualJoystick.cpp b/Engine/Source/Runtime/Slate/Private/Widgets/Input/SVirtualJoystick.cpp
--- a/Engine/Source/Runtime/Slate/Private/Widgets/Input/SVirtualJoystick.cpp
+++ b/Engine/Source/Runtime/Slate/Private/Widgets/Input/SVirtualJoystick.cpp
@@ -149,6 +149,10 @@ FVector2D SVirtualJoystick::ComputeDesiredSize( float ) const
 	return FVector2D(100, 100);
 }
 
+bool SVirtualJoystick::SupportsKeyboardFocus() const
+{
+	return false;
+}
 
 FReply SVirtualJoystick::OnTouchStarted(const FGeometry& MyGeometry, const FPointerEvent& Event)
 {
@@ -228,12 +232,12 @@ FReply SVirtualJoystick::OnTouchMoved(const FGeometry& MyGeometry, const FPointe
 
 FReply SVirtualJoystick::OnTouchEnded(const FGeometry& MyGeometry, const FPointerEvent& Event)
 {
-	for (int32 ControlIndex = 0; ControlIndex < Controls.Num(); ControlIndex++)
+	for ( int32 ControlIndex = 0; ControlIndex < Controls.Num(); ControlIndex++ )
 	{
 		FControlInfo& Control = Controls[ControlIndex];
 
 		// is this control the one captured to this pointer?
-		if (Control.CapturedPointerIndex == Event.GetPointerIndex())
+		if ( Control.CapturedPointerIndex == Event.GetPointerIndex() )
 		{
 			// release and center the joystick
 			Control.ThumbPosition = FVector2D(0, 0);
@@ -243,7 +247,7 @@ FReply SVirtualJoystick::OnTouchEnded(const FGeometry& MyGeometry, const FPointe
 			Control.bSendOneMoreEvent = true;
 
 			// Pass event as unhandled if time is too short
-			if (Control.bNeedUpdatedCenter)
+			if ( Control.bNeedUpdatedCenter )
 			{
 				Control.bNeedUpdatedCenter = false;
 				return FReply::Unhandled();
@@ -463,3 +467,32 @@ void SVirtualJoystick::SetJoystickVisibility(const bool bInVisible, const bool b
 
 	bVisible = bInVisible;
 }
+
+void SVirtualJoystick::AlignBoxIntoScreen(FVector2D& Position, const FVector2D& Size, const FVector2D& ScreenSize)
+{
+	if ( Size.X > ScreenSize.X || Size.Y > ScreenSize.Y )
+	{
+		return;
+	}
+
+	// Align box to fit into screen
+	if ( Position.X - Size.X * 0.5f < 0.f )
+	{
+		Position.X = Size.X * 0.5f;
+	}
+
+	if ( Position.X + Size.X * 0.5f > ScreenSize.X )
+	{
+		Position.X = ScreenSize.X - Size.X * 0.5f;
+	}
+
+	if ( Position.Y - Size.Y * 0.5f < 0.f )
+	{
+		Position.Y = Size.Y * 0.5f;
+	}
+
+	if ( Position.Y + Size.Y * 0.5f > ScreenSize.Y )
+	{
+		Position.Y = ScreenSize.Y - Size.Y * 0.5f;
+	}
+}
diff --git a/Engine/Source/Runtime/Slate/Public/Widgets/Input/SVirtualJoystick.h b/Engine/Source/Runtime/Slate/Public/Widgets/Input/SVirtualJoystick.h
--- a/Engine/Source/Runtime/Slate/Public/Widgets/Input/SVirtualJoystick.h
+++ b/Engine/Source/Runtime/Slate/Public/Widgets/Input/SVirtualJoystick.h
@@ -132,41 +132,11 @@ public:
 
 	virtual void Tick( const FGeometry& AllottedGeometry, const double InCurrentTime, const float InDeltaTime ) override;
 
-	virtual bool SupportsKeyboardFocus() const override
-	{
-		return true;
-	}
-
-	void AlignBoxIntoScreen( FVector2D& Position, const FVector2D& Size, const FVector2D& ScreenSize )
-	{
-		if (Size.X > ScreenSize.X || Size.Y > ScreenSize.Y)
-		{
-			return;
-		}
-
-		// Align box to fit into screen
-		if (Position.X - Size.X * 0.5f < 0.f )
-		{
-			Position.X = Size.X * 0.5f;
-		}
-
-		if (Position.X + Size.X * 0.5f > ScreenSize.X )
-		{
-			Position.X = ScreenSize.X - Size.X * 0.5f;
-		}
-
-		if (Position.Y - Size.Y * 0.5f < 0.f)
-		{
-			Position.Y = Size.Y * 0.5f;
-		}
-
-		if (Position.Y + Size.Y * 0.5f > ScreenSize.Y )
-		{
-			Position.Y = ScreenSize.Y - Size.Y * 0.5f;
-		}
-	}
+	virtual bool SupportsKeyboardFocus() const override;
 
 private:
+	void AlignBoxIntoScreen(FVector2D& Position, const FVector2D& Size, const FVector2D& ScreenSize);
+
 	/**
 	 * Process a touch event (on movement and possibly on initial touch)
 	 *
