Slate - we no longer return the wrap at size for editable text controls desired size when the text isn't that large yet.

diff --git a/Engine/Source/Runtime/Slate/Private/Widgets/Text/SMultiLineEditableText.cpp b/Engine/Source/Runtime/Slate/Private/Widgets/Text/SMultiLineEditableText.cpp
--- a/Engine/Source/Runtime/Slate/Private/Widgets/Text/SMultiLineEditableText.cpp
+++ b/Engine/Source/Runtime/Slate/Private/Widgets/Text/SMultiLineEditableText.cpp
@@ -2276,13 +2276,12 @@ FVector2D SMultiLineEditableText::ComputeDesiredSize( float ) const
 
 	// The layouts current margin size. We should not report a size smaller then the margins.
 	const FMargin Margin = TextLayout->GetMargin();
-	const FVector2D TextLayoutSize = TextLayout->GetSize();
+	const FVector2D TextLayoutSize = TextLayout->GetSize() + FVector2D(CaretWidth, 0);
 	const float WrappingWidth = WrapTextAt.Get();
 
 	// If a wrapping width has been provided that should be reported as the desired width.
-	float DesiredWidth = WrappingWidth > 0 ? WrappingWidth : TextLayoutSize.X;
+	float DesiredWidth = WrappingWidth > 0 ? FMath::Min(WrappingWidth, TextLayoutSize.X) : TextLayoutSize.X;
 	DesiredWidth = FMath::Max(Margin.GetTotalSpaceAlong<Orient_Horizontal>(), DesiredWidth);
-	DesiredWidth += CaretWidth;
 
 	float DesiredHeight = FMath::Max(Margin.GetTotalSpaceAlong<Orient_Vertical>(), TextLayoutSize.Y);
 	DesiredHeight = FMath::Max(DesiredHeight, FontMaxCharHeight);
