Slate - The ML editable text now returns the text size for the desired size always.

diff --git a/Engine/Source/Runtime/Slate/Private/Widgets/Text/SMultiLineEditableText.cpp b/Engine/Source/Runtime/Slate/Private/Widgets/Text/SMultiLineEditableText.cpp
--- a/Engine/Source/Runtime/Slate/Private/Widgets/Text/SMultiLineEditableText.cpp
+++ b/Engine/Source/Runtime/Slate/Private/Widgets/Text/SMultiLineEditableText.cpp
@@ -2276,11 +2276,11 @@ FVector2D SMultiLineEditableText::ComputeDesiredSize( float ) const
 
 	// The layouts current margin size. We should not report a size smaller then the margins.
 	const FMargin Margin = TextLayout->GetMargin();
-	const FVector2D TextLayoutSize = TextLayout->GetSize() + FVector2D(CaretWidth, 0);
+	const FVector2D TextLayoutSize = TextLayout->GetSize();
 	const float WrappingWidth = WrapTextAt.Get();
 
 	// If a wrapping width has been provided that should be reported as the desired width.
-	float DesiredWidth = WrappingWidth > 0 ? FMath::Min(WrappingWidth, TextLayoutSize.X) : TextLayoutSize.X;
+	float DesiredWidth = TextLayoutSize.X + CaretWidth;
 	DesiredWidth = FMath::Max(Margin.GetTotalSpaceAlong<Orient_Horizontal>(), DesiredWidth);
 
 	float DesiredHeight = FMath::Max(Margin.GetTotalSpaceAlong<Orient_Vertical>(), TextLayoutSize.Y);
