Converted Slate, SlateCore, and SlateViewer to exclusively use FText with Slate

UETOOL-213 - Minimize Slate FString -> FText conversion (remove SLATE_TEXT_ATTRIBUTE)

diff --git a/Engine/Source/Runtime/Slate/Public/Widgets/Input/SEditableComboBox.h b/Engine/Source/Runtime/Slate/Public/Widgets/Input/SEditableComboBox.h
--- a/Engine/Source/Runtime/Slate/Public/Widgets/Input/SEditableComboBox.h
+++ b/Engine/Source/Runtime/Slate/Public/Widgets/Input/SEditableComboBox.h
@@ -170,7 +170,7 @@ public:
 										.ForegroundColor(FSlateColor::UseForeground())
 										.ButtonStyle( FCoreStyle::Get(), "NoBorder" )
 										.OnClicked(this, &SEditableComboBox::HandleAcceptButtonClicked)
-										.ToolTipText(LOCTEXT("AcceptButtonTooltip", "Accept").ToString())
+										.ToolTipText(LOCTEXT("AcceptButtonTooltip", "Accept"))
 										.VAlign(VAlign_Center)
 										.Content()
 										[
