Fixed SMultiLineEditableText swallowing new-lines when pasting text with ModiferKeyForNewLine set

SMultiLineEditableText::InsertTextAtCursorImpl was using SMultiLineEditableText::OnEnter to insert new-lines in the pasted text, however if you've got ModiferKeyForNewLine set, that modifier must be active for OnEnter to add the new-line. This caused SMultiLineEditableText to potentially strip the new-lines from pasted text.

SMultiLineEditableText::InsertTextAtCursorImpl now calls SMultiLineEditableText::InsertNewLineAtCursorImpl which is just the shared implementation logic taken from SMultiLineEditableText::OnEnter, skipping the modifier key check.

diff --git a/Engine/Source/Runtime/Slate/Private/Widgets/Text/SMultiLineEditableText.cpp b/Engine/Source/Runtime/Slate/Private/Widgets/Text/SMultiLineEditableText.cpp
--- a/Engine/Source/Runtime/Slate/Private/Widgets/Text/SMultiLineEditableText.cpp
+++ b/Engine/Source/Runtime/Slate/Private/Widgets/Text/SMultiLineEditableText.cpp
@@ -1718,22 +1718,7 @@ void SMultiLineEditableText::OnEnter()
 {
 	if ( FSlateApplication::Get().GetModifierKeys().AreModifersDown(ModiferKeyForNewLine) )
 	{
-		if ( AnyTextSelected() )
-		{
-			// Delete selected text
-			DeleteSelectedText();
-		}
-
-		const FTextLocation CursorInteractionPosition = CursorInfo.GetCursorInteractionLocation();
-		if ( TextLayout->SplitLineAt(CursorInteractionPosition) )
-		{
-			// Adjust the cursor position to be at the beginning of the new line
-			const FTextLocation NewCursorPosition = FTextLocation(CursorInteractionPosition.GetLineIndex() + 1, 0);
-			CursorInfo.SetCursorLocationAndCalculateAlignment(TextLayout, NewCursorPosition);
-		}
-
-		ClearSelection();
-		UpdateCursorHighlight();
+		InsertNewLineAtCursorImpl();
 	}
 	else
 	{
@@ -1877,7 +1862,7 @@ void SMultiLineEditableText::InsertTextAtCursorImpl(const FString& InString)
 		{
 			if (*Character == '\n')
 			{
-				OnEnter();
+				InsertNewLineAtCursorImpl();
 			}
 			else
 			{
@@ -1887,6 +1872,26 @@ void SMultiLineEditableText::InsertTextAtCursorImpl(const FString& InString)
 	}
 }
 
+void SMultiLineEditableText::InsertNewLineAtCursorImpl()
+{
+	if ( AnyTextSelected() )
+	{
+		// Delete selected text
+		DeleteSelectedText();
+	}
+
+	const FTextLocation CursorInteractionPosition = CursorInfo.GetCursorInteractionLocation();
+	if ( TextLayout->SplitLineAt(CursorInteractionPosition) )
+	{
+		// Adjust the cursor position to be at the beginning of the new line
+		const FTextLocation NewCursorPosition = FTextLocation(CursorInteractionPosition.GetLineIndex() + 1, 0);
+		CursorInfo.SetCursorLocationAndCalculateAlignment(TextLayout, NewCursorPosition);
+	}
+
+	ClearSelection();
+	UpdateCursorHighlight();
+}
+
 bool SMultiLineEditableText::CanExecuteUndo() const
 {
 	if(IsReadOnly.Get() || TextInputMethodContext->IsComposing)
diff --git a/Engine/Source/Runtime/Slate/Public/Widgets/Text/SMultiLineEditableText.h b/Engine/Source/Runtime/Slate/Public/Widgets/Text/SMultiLineEditableText.h
--- a/Engine/Source/Runtime/Slate/Public/Widgets/Text/SMultiLineEditableText.h
+++ b/Engine/Source/Runtime/Slate/Public/Widgets/Text/SMultiLineEditableText.h
@@ -466,6 +466,9 @@ private:
 	/** Insert the given text at the current cursor position, correctly taking into account new line characters */
 	void InsertTextAtCursorImpl(const FString& InString);
 
+	/** Insert a new-line at the current cursor position */
+	void InsertNewLineAtCursorImpl();
+
 	/**
 	 * Given a location and a Direction to offset, return a new location.
 	 *
