SNumericEntryBox - Use the font from the editable text style if no font is spcified in the arguments.

diff --git a/Engine/Source/Runtime/Slate/Public/Widgets/Input/SNumericEntryBox.h b/Engine/Source/Runtime/Slate/Public/Widgets/Input/SNumericEntryBox.h
--- a/Engine/Source/Runtime/Slate/Public/Widgets/Input/SNumericEntryBox.h
+++ b/Engine/Source/Runtime/Slate/Public/Widgets/Input/SNumericEntryBox.h
@@ -31,7 +31,6 @@ public:
 		, _BorderForegroundColor(FCoreStyle::Get().GetSlateColor("InvertedForeground"))
 		, _BorderBackgroundColor(FLinearColor::White)
 		, _UndeterminedString( SNumericEntryBox<NumericType>::DefaultUndeterminedString )
-		, _Font( FCoreStyle::Get().GetFontStyle( TEXT("NormalFont") ) )
 		, _AllowSpin(false)
 		, _Delta(0)
 		, _MinValue(TNumericLimits<NumericType>::Lowest())
@@ -118,7 +117,7 @@ public:
 		{
 			SAssignNew( SpinBox, SSpinBox<NumericType> )
 				.Style( FCoreStyle::Get(), "NumericEntrySpinBox" )
-				.Font( InArgs._Font )
+				.Font( InArgs._Font.IsSet() ? InArgs._Font : InArgs._EditableTextBoxStyle->Font )
 				.ContentPadding( TextMargin )
 				.Value( this, &SNumericEntryBox<NumericType>::OnGetValueForSpinBox )
 				.Delta( InArgs._Delta )
@@ -139,7 +138,7 @@ public:
 		SAssignNew( EditableText, SEditableText )
 			.Text( this, &SNumericEntryBox<NumericType>::OnGetValueForTextBox )
 			.Visibility( bAllowSpin ? EVisibility::Collapsed : EVisibility::Visible )
-			.Font( InArgs._Font )
+			.Font( InArgs._Font.IsSet() ? InArgs._Font : InArgs._EditableTextBoxStyle->Font )
 			.SelectAllTextWhenFocused( true )
 			.ClearKeyboardFocusOnCommit( false )
 			.OnTextChanged( this, &SNumericEntryBox<NumericType>::OnTextChanged  )
