Multi-line editable text now includes the first piece of trailing whitespace in its reported width when being wrapped

UE-10652 - Make SMultiLineEditableText limit its width correctly when using WrapTextAt

This prevents it showing a scrollbar in the cases where a piece of text only fits on a line because we're ignoring its trailing whitepsace.

diff --git a/Engine/Source/Runtime/Slate/Private/Framework/Text/TextLayout.cpp b/Engine/Source/Runtime/Slate/Private/Framework/Text/TextLayout.cpp
--- a/Engine/Source/Runtime/Slate/Private/Framework/Text/TextLayout.cpp
+++ b/Engine/Source/Runtime/Slate/Private/Framework/Text/TextLayout.cpp
@@ -10,6 +10,7 @@ FTextLayout::FBreakCandidate FTextLayout::CreateBreakCandidate( int32& OutRunInd
 	int16 MaxBelowBaseline = 0;
 	FVector2D BreakSize( ForceInitToZero );
 	FVector2D BreakSizeWithoutTrailingWhitespace( ForceInitToZero );
+	float FirstTrailingWhitespaceCharWidth = 0.0f;
 	int32 WhitespaceStopIndex = CurrentBreak;
 	uint8 Kerning = 0;
 
@@ -59,7 +60,20 @@ FTextLayout::FBreakCandidate FTextLayout::CreateBreakCandidate( int32& OutRunInd
 		{
 			// This slice contains trailing whitespace, measure the text size, then add on the whitespace size
 			SliceSize = SliceSizeWithoutTrailingWhitespace = Run.Measure( BeginIndex, WhitespaceStopIndex, Scale );
-			SliceSize.X += Run.Measure( WhitespaceStopIndex, StopIndex, Scale ).X;
+			const float WhitespaceWidth = Run.Measure( WhitespaceStopIndex, StopIndex, Scale ).X;
+			SliceSize.X += WhitespaceWidth;
+
+			// We also need to measure the width of the first piece of trailing whitespace
+			if ( WhitespaceStopIndex + 1 == StopIndex )
+			{
+				// Only have one piece of whitespace
+				FirstTrailingWhitespaceCharWidth = WhitespaceWidth;
+			}
+			else
+			{
+				// Deliberately use the run version of Measure as we don't want the run model to cache this measurement since it may be out of order and break the binary search
+				FirstTrailingWhitespaceCharWidth = Run.GetRun()->Measure( WhitespaceStopIndex, WhitespaceStopIndex + 1, Scale ).X;
+			}
 		}
 		else
 		{
@@ -97,6 +111,7 @@ FTextLayout::FBreakCandidate FTextLayout::CreateBreakCandidate( int32& OutRunInd
 	BreakCandidate.TrimmedSize = BreakSizeWithoutTrailingWhitespace;
 	BreakCandidate.ActualRange = FTextRange( PreviousBreak, CurrentBreak );
 	BreakCandidate.TrimmedRange = FTextRange( PreviousBreak, WhitespaceStopIndex );
+	BreakCandidate.FirstTrailingWhitespaceCharWidth = FirstTrailingWhitespaceCharWidth;
 	BreakCandidate.MaxAboveBaseline = MaxAboveBaseline;
 	BreakCandidate.MaxBelowBaseline = MaxBelowBaseline;
 	BreakCandidate.Kerning = Kerning;
@@ -108,7 +123,7 @@ FTextLayout::FBreakCandidate FTextLayout::CreateBreakCandidate( int32& OutRunInd
 	return BreakCandidate;
 }
 
-void FTextLayout::CreateLineViewBlocks( int32 LineModelIndex, const int32 StopIndex, int32& OutRunIndex, int32& OutRendererIndex, int32& OutPreviousBlockEnd, TArray< TSharedRef< ILayoutBlock > >& OutSoftLine )
+void FTextLayout::CreateLineViewBlocks( int32 LineModelIndex, const int32 StopIndex, const float WrappedLineWidth, int32& OutRunIndex, int32& OutRendererIndex, int32& OutPreviousBlockEnd, TArray< TSharedRef< ILayoutBlock > >& OutSoftLine )
 {
 	const FLineModel& LineModel = LineModels[ LineModelIndex ];
 
@@ -219,7 +234,7 @@ void FTextLayout::CreateLineViewBlocks( int32 LineModelIndex, const int32 StopIn
 	FVector2D LineSize( ForceInitToZero );
 	
 	// Use a negative scroll offset since positive scrolling moves things negatively in screen space
-	FVector2D CurrentOffset( Margin.Left - ScrollOffset.X, Margin.Top + DrawSize.Y - ScrollOffset.Y );
+	FVector2D CurrentOffset( Margin.Left - ScrollOffset.X, Margin.Top + TextLayoutSize.Height - ScrollOffset.Y );
 
 	if ( OutSoftLine.Num() > 0 )
 	{
@@ -254,8 +269,9 @@ void FTextLayout::CreateLineViewBlocks( int32 LineModelIndex, const int32 StopIn
 		LineViews.Add( LineView );
 	}
 
-	DrawSize.X = FMath::Max( DrawSize.X, LineSize.X ); //DrawSize.X is the size of the longest line + the Margin
-	DrawSize.Y += LineSize.Y; //DrawSize.Y is the total height of all lines
+	TextLayoutSize.DrawWidth = FMath::Max( TextLayoutSize.DrawWidth, LineSize.X ); // DrawWidth is the size of the longest line + the Margin
+	TextLayoutSize.WrappedWidth = FMath::Max( TextLayoutSize.WrappedWidth, (StopIndex == INDEX_NONE) ? LineSize.X : WrappedLineWidth ); // WrappedWidth is the size of the longest line + the Margin + any trailing whitespace width
+	TextLayoutSize.Height += LineSize.Y; // Height is the total height of all lines
 }
 
 void FTextLayout::JustifyLayout()
@@ -265,7 +281,7 @@ void FTextLayout::JustifyLayout()
 		return;
 	}
 
-	const float LayoutWidthNoMargin = FMath::Max(DrawSize.X, ViewSize.X * Scale) - (Margin.GetTotalSpaceAlong<Orient_Horizontal>() * Scale);
+	const float LayoutWidthNoMargin = FMath::Max(TextLayoutSize.DrawWidth, ViewSize.X * Scale) - (Margin.GetTotalSpaceAlong<Orient_Horizontal>() * Scale);
 
 	for (FLineView& LineView : LineViews)
 	{
@@ -308,9 +324,12 @@ void FTextLayout::FlowLayout()
 		FlowLineLayout(LineModelIndex, WrappingDrawWidth, SoftLine);
 	}
 
-	//Add on the margins to the DrawSize
-	DrawSize.X += Margin.GetTotalSpaceAlong<Orient_Horizontal>() * Scale;
-	DrawSize.Y += Margin.GetTotalSpaceAlong<Orient_Vertical>() * Scale;
+	// Add on the margins to the layout size
+	const float MarginWidth = Margin.GetTotalSpaceAlong<Orient_Horizontal>() * Scale;
+	const float MarginHeight = Margin.GetTotalSpaceAlong<Orient_Vertical>() * Scale;
+	TextLayoutSize.DrawWidth += MarginWidth;
+	TextLayoutSize.WrappedWidth += MarginWidth;
+	TextLayoutSize.Height += MarginHeight;
 }
 
 void FTextLayout::FlowLineLayout(const int32 LineModelIndex, const float WrappingDrawWidth, TArray<TSharedRef<ILayoutBlock>>& SoftLine)
@@ -333,7 +352,7 @@ void FTextLayout::FlowLineLayout(const int32 LineModelIndex, const float Wrappin
 	if (!IsWrapping || LineModel.BreakCandidates.Num() == 0 )
 	{
 		//Then iterate over all of it's runs
-		CreateLineViewBlocks( LineModelIndex, INDEX_NONE, /*OUT*/CurrentRunIndex, /*OUT*/CurrentRendererIndex, /*OUT*/PreviousBlockEnd, SoftLine );
+		CreateLineViewBlocks( LineModelIndex, INDEX_NONE, 0.0f, /*OUT*/CurrentRunIndex, /*OUT*/CurrentRendererIndex, /*OUT*/PreviousBlockEnd, SoftLine );
 		check( CurrentRunIndex == LineModel.Runs.Num() );
 		CurrentWidth = 0;
 		SoftLine.Empty();
@@ -355,8 +374,23 @@ void FTextLayout::FlowLineLayout(const int32 LineModelIndex, const float Wrappin
 				const bool IsFirstBreak = BreakIndex == 0;
 
 				const FBreakCandidate& FinalBreakOnSoftLine = ( !IsFirstBreak && !IsFirstBreakOnSoftLine && !BreakWithoutTrailingWhitespaceDoesFit ) ? LineModel.BreakCandidates[ --BreakIndex ] : Break;
+				
+				// We want the wrapped line width to contain the first piece of trailing whitespace for a line, however we only do this if we have trailing whitespace
+				// otherwise very long non-breaking words can cause the wrapped line width to expand beyond the desired wrap width
+				float WrappedLineWidth = CurrentWidth;
+				if ( BreakWithoutTrailingWhitespaceDoesFit && !IsLastBreak )
+				{
+					// This break has trailing whitespace
+					WrappedLineWidth += ( FinalBreakOnSoftLine.TrimmedSize.X + FinalBreakOnSoftLine.FirstTrailingWhitespaceCharWidth );
+				}
+				else
+				{
+					// This break is either longer than the wrapping point or the last break on this line, so make sure and clamp the line size to the given wrapping width
+					WrappedLineWidth += FinalBreakOnSoftLine.ActualSize.X;
+					WrappedLineWidth = FMath::Min(WrappedLineWidth, WrappingDrawWidth);
+				}
 
-				CreateLineViewBlocks( LineModelIndex, FinalBreakOnSoftLine.ActualRange.EndIndex, /*OUT*/CurrentRunIndex, /*OUT*/CurrentRendererIndex, /*OUT*/PreviousBlockEnd, SoftLine );
+				CreateLineViewBlocks( LineModelIndex, FinalBreakOnSoftLine.ActualRange.EndIndex, WrappedLineWidth, /*OUT*/CurrentRunIndex, /*OUT*/CurrentRendererIndex, /*OUT*/PreviousBlockEnd, SoftLine );
 
 				if ( CurrentRunIndex < LineModel.Runs.Num() && FinalBreakOnSoftLine.ActualRange.EndIndex == LineModel.Runs[ CurrentRunIndex ].GetTextRange().EndIndex )
 				{
@@ -495,7 +529,7 @@ void FTextLayout::BeginLineLayout(FLineModel& LineModel)
 
 void FTextLayout::ClearView()
 {
-	DrawSize = FVector2D::ZeroVector;
+	TextLayoutSize = FTextLayoutSize();
 	LineViews.Empty();
 }
 
@@ -566,7 +600,7 @@ FTextLayout::FTextLayout()
 	, Margin()
 	, Justification( ETextJustify::Left )
 	, LineHeightPercentage( 1.0f )
-	, DrawSize( ForceInitToZero )
+	, TextLayoutSize()
 	, ViewSize( ForceInitToZero )
 	, ScrollOffset( ForceInitToZero )
 	, LineBreakIterator() // Initialized in FTextLayout::CreateWrappingCache if no custom iterator is provided
@@ -1671,8 +1705,11 @@ void FTextLayout::SetMargin( const FMargin& InMargin )
 			}
 		}
 
-		DrawSize.X += ( Margin.GetTotalSpaceAlong<Orient_Horizontal>() - PreviousMargin.GetTotalSpaceAlong<Orient_Horizontal>() ) * Scale;
-		DrawSize.Y += ( Margin.GetTotalSpaceAlong<Orient_Vertical>() - PreviousMargin.GetTotalSpaceAlong<Orient_Vertical>() ) * Scale;
+		const float MarginDeltaWidth = ( Margin.GetTotalSpaceAlong<Orient_Horizontal>() - PreviousMargin.GetTotalSpaceAlong<Orient_Horizontal>() ) * Scale;
+		const float MarginDeltaHeight = ( Margin.GetTotalSpaceAlong<Orient_Vertical>() - PreviousMargin.GetTotalSpaceAlong<Orient_Vertical>() ) * Scale;
+		TextLayoutSize.DrawWidth += MarginDeltaWidth;
+		TextLayoutSize.WrappedWidth += MarginDeltaWidth;
+		TextLayoutSize.Height += MarginDeltaHeight;
 	}
 }
 
@@ -1747,12 +1784,17 @@ float FTextLayout::GetWrappingWidth() const
 
 FVector2D FTextLayout::GetDrawSize() const
 {
-	return DrawSize;
+	return TextLayoutSize.GetDrawSize();
+}
+
+FVector2D FTextLayout::GetWrappedSize() const
+{
+	return TextLayoutSize.GetWrappedSize() * ( 1 / Scale );
 }
 
 FVector2D FTextLayout::GetSize() const
 {
-	return DrawSize * ( 1 / Scale );
+	return TextLayoutSize.GetDrawSize() * ( 1 / Scale );
 }
 
 const TArray< FTextLayout::FLineModel >& FTextLayout::GetLineModels() const
diff --git a/Engine/Source/Runtime/Slate/Private/Widgets/Text/SMultiLineEditableText.cpp b/Engine/Source/Runtime/Slate/Private/Widgets/Text/SMultiLineEditableText.cpp
--- a/Engine/Source/Runtime/Slate/Private/Widgets/Text/SMultiLineEditableText.cpp
+++ b/Engine/Source/Runtime/Slate/Private/Widgets/Text/SMultiLineEditableText.cpp
@@ -2279,15 +2279,16 @@ FVector2D SMultiLineEditableText::ComputeDesiredSize( float ) const
 	const float FontMaxCharHeight = FTextEditHelper::GetFontHeight(TextStyle.Font);
 	const float CaretWidth = FTextEditHelper::CalculateCaretWidth(FontMaxCharHeight);
 
-	// The layouts current margin size. We should not report a size smaller then the margins.
-	const FMargin Margin = TextLayout->GetMargin();
-	const FVector2D TextLayoutSize = TextLayout->GetSize();
+	// If a wrapping width has been provided, then we need to report the wrapped size as the desired width
+	// (rather than the actual text layout size as that can have non-breaking lines that extend beyond the wrap width)
+	// Note: We don't do this when auto-wrapping as it would cause a feedback loop in the Slate sizing logic
 	const float WrappingWidth = WrapTextAt.Get();
+	const FVector2D TextLayoutBaseSize = WrappingWidth > 0 ? TextLayout->GetWrappedSize() : TextLayout->GetSize();
+	const FVector2D TextLayoutSize = TextLayoutBaseSize + FVector2D(CaretWidth, 0);
 
-	// If a wrapping width has been provided that should be reported as the desired width.
-	float DesiredWidth = TextLayoutSize.X + CaretWidth;
-	DesiredWidth = FMath::Max(Margin.GetTotalSpaceAlong<Orient_Horizontal>(), DesiredWidth);
-
+	// The layouts current margin size. We should not report a size smaller then the margins.
+	const FMargin Margin = TextLayout->GetMargin();
+	float DesiredWidth = FMath::Max(Margin.GetTotalSpaceAlong<Orient_Horizontal>(), TextLayoutSize.X);
 	float DesiredHeight = FMath::Max(Margin.GetTotalSpaceAlong<Orient_Vertical>(), TextLayoutSize.Y);
 	DesiredHeight = FMath::Max(DesiredHeight, FontMaxCharHeight);
 	
diff --git a/Engine/Source/Runtime/Slate/Public/Framework/Text/TextLayout.h b/Engine/Source/Runtime/Slate/Public/Framework/Text/TextLayout.h
--- a/Engine/Source/Runtime/Slate/Public/Framework/Text/TextLayout.h
+++ b/Engine/Source/Runtime/Slate/Public/Framework/Text/TextLayout.h
@@ -151,6 +151,8 @@ public:
 		FVector2D ActualSize;
 		/** Measured size exclusive of trailing whitespace, as used to perform wrapping on a word boundary */
 		FVector2D TrimmedSize;
+		/** If this break candidate has trailing whitespace, this is the width of the first character of the trailing whitespace */
+		float FirstTrailingWhitespaceCharWidth;
 
 		int16 MaxAboveBaseline;
 		int16 MaxBelowBaseline;
@@ -280,6 +282,7 @@ public:
 
 	FVector2D GetSize() const;
 	FVector2D GetDrawSize() const;
+	FVector2D GetWrappedSize() const;
 
 	float GetWrappingWidth() const;
 	void SetWrappingWidth( float Value );
@@ -451,7 +454,7 @@ private:
 
 	void JustifyLayout();
 
-	void CreateLineViewBlocks( int32 LineModelIndex, const int32 StopIndex, int32& OutRunIndex, int32& OutRendererIndex, int32& OutPreviousBlockEnd, TArray< TSharedRef< ILayoutBlock > >& OutSoftLine );
+	void CreateLineViewBlocks( int32 LineModelIndex, const int32 StopIndex, const float WrappedLineWidth, int32& OutRunIndex, int32& OutRendererIndex, int32& OutPreviousBlockEnd, TArray< TSharedRef< ILayoutBlock > >& OutSoftLine );
 
 	FBreakCandidate CreateBreakCandidate( int32& OutRunIndex, FLineModel& Line, int32 PreviousBreak, int32 CurrentBreak );
 
@@ -470,6 +473,35 @@ protected:
 		};
 	};
 
+	struct FTextLayoutSize
+	{
+		FTextLayoutSize()
+			: DrawWidth(0.0f)
+			, WrappedWidth(0.0f)
+			, Height(0.0f)
+		{
+		}
+
+		FVector2D GetDrawSize() const
+		{
+			return FVector2D(DrawWidth, Height);
+		}
+
+		FVector2D GetWrappedSize() const
+		{
+			return FVector2D(WrappedWidth, Height);
+		}
+
+		/** Width of the text layout, including any lines which extend beyond the wrapping boundaries (eg, lines with lots of trailing whitespace, or lines with no break candidates) */
+		float DrawWidth;
+
+		/** Width of the text layout after the text has been wrapped, and including the first piece of trailing whitespace for any given soft-wrapped line */
+		float WrappedWidth;
+
+		/** Height of the text layout */
+		float Height;
+	};
+
 	/** The models for the lines of text. A LineModel represents a single string with no manual breaks. */
 	TArray< FLineModel > LineModels;
 
@@ -495,7 +527,7 @@ protected:
 	float LineHeightPercentage;
 
 	/** The final size of the text layout on screen. */
-	FVector2D DrawSize;
+	FTextLayoutSize TextLayoutSize;
 
 	/** The size of the text layout that can actually be seen from the parent widget */
 	FVector2D ViewSize;
