Enabled active timers for multi-line editable text to make sure that cursor animation and auto-scrolling continue to work

diff --git a/Engine/Source/Runtime/Slate/Private/Widgets/Text/SMultiLineEditableText.cpp b/Engine/Source/Runtime/Slate/Private/Widgets/Text/SMultiLineEditableText.cpp
--- a/Engine/Source/Runtime/Slate/Private/Widgets/Text/SMultiLineEditableText.cpp
+++ b/Engine/Source/Runtime/Slate/Private/Widgets/Text/SMultiLineEditableText.cpp
@@ -467,11 +467,32 @@ void SMultiLineEditableText::OnVScrollBarMoved(const float InScrollOffsetFractio
 	OnVScrollBarUserScrolled.ExecuteIfBound(InScrollOffsetFraction);
 }
 
+void SMultiLineEditableText::EnsureActiveTick()
+{
+	TSharedPtr<FActiveTimerHandle> ActiveTickTimerPin = ActiveTickTimer.Pin();
+	if(ActiveTickTimerPin.IsValid())
+	{
+		return;
+	}
+
+	auto DoActiveTick = [this](double InCurrentTime, float InDeltaTime) -> EActiveTimerReturnType
+	{
+		// Continue if we still have focus, otherwise treat as a fire-and-forget Tick() request
+		const bool bShouldAppearFocused = HasKeyboardFocus() || ActiveContextMenu.IsValid();
+		return (bShouldAppearFocused) ? EActiveTimerReturnType::Continue : EActiveTimerReturnType::Stop;
+	};
+
+	ActiveTickTimer = RegisterActiveTimer(0.5f, FWidgetActiveTimerDelegate::CreateLambda(DoActiveTick));
+}
+
 FReply SMultiLineEditableText::OnFocusReceived( const FGeometry& MyGeometry, const FFocusEvent& InFocusEvent )
 {
 	// Skip the focus received code if it's due to the context menu closing
 	if ( !ActiveContextMenu.IsValid() )
 	{
+		// We need to Tick() while we have focus to keep some things up-to-date
+		EnsureActiveTick();
+
 		FSlateApplication& SlateApplication = FSlateApplication::Get();
 		if (FPlatformMisc::GetRequiresVirtualKeyboard())
 		{
@@ -487,7 +508,6 @@ FReply SMultiLineEditableText::OnFocusReceived( const FGeometry& MyGeometry, con
 			}
 		}
 
-
 		UpdateCursorHighlight();
 
 		// UpdateCursorHighlight always tries to scroll to the cursor, but we don't want that to happen when we 
@@ -1008,6 +1028,7 @@ FReply SMultiLineEditableText::MoveCursor( FMoveCursor Args )
 void SMultiLineEditableText::UpdateCursorHighlight()
 {
 	PositionToScrollIntoView = FScrollInfo(CursorInfo.GetCursorInteractionLocation(), CursorInfo.GetCursorAlignment());
+	EnsureActiveTick();
 
 	RemoveCursorHighlight();
 
@@ -1344,6 +1365,7 @@ void SMultiLineEditableText::ScrollTo(const FTextLocation& NewLocation)
 		if (NewLocation.GetOffset() <= Line.Text->Len())
 		{
 			PositionToScrollIntoView = FScrollInfo(NewLocation, ECursorAlignment::Left);
+			EnsureActiveTick();
 		}
 	}
 }
diff --git a/Engine/Source/Runtime/Slate/Public/Widgets/Text/SMultiLineEditableText.h b/Engine/Source/Runtime/Slate/Public/Widgets/Text/SMultiLineEditableText.h
--- a/Engine/Source/Runtime/Slate/Public/Widgets/Text/SMultiLineEditableText.h
+++ b/Engine/Source/Runtime/Slate/Public/Widgets/Text/SMultiLineEditableText.h
@@ -548,6 +548,12 @@ private:
 	void OnHScrollBarMoved(const float InScrollOffsetFraction);
 	void OnVScrollBarMoved(const float InScrollOffsetFraction);
 
+	/**
+	 * Ensure that we will get a Tick() soon (either due to having active focus, or something having changed progmatically and requiring an update)
+	 * Does nothing if the active tick timer is already enabled
+	 */
+	void EnsureActiveTick();
+
 private:
 
 	/** The text displayed in this text block */
@@ -665,6 +671,9 @@ private:
 
 	/** The optional modifier key necessary to create a newline when typing into the editor. */
 	EModifierKey::Type ModiferKeyForNewLine;
+
+	/** The timer that is actively driving this widget to Tick() even when Slate is idle */
+	TWeakPtr<FActiveTimerHandle> ActiveTickTimer;
 };
 
 
