Fixed clearing editable text causing a commit reason of default instead of cleared. Causes text boxes that delay notification to not clear when requested.

diff --git a/Engine/Source/Runtime/Slate/Private/Widgets/Input/SSearchBox.cpp b/Engine/Source/Runtime/Slate/Private/Widgets/Input/SSearchBox.cpp
--- a/Engine/Source/Runtime/Slate/Private/Widgets/Input/SSearchBox.cpp
+++ b/Engine/Source/Runtime/Slate/Private/Widgets/Input/SSearchBox.cpp
@@ -182,7 +182,7 @@ FReply SSearchBox::OnClearSearch()
 {
 	// We clear the focus to commit any unset values as the search box is typically used for filtering
 	// and the widget could get immediately destroyed before committing its value.
-	FSlateApplication::Get().ClearKeyboardFocus( EFocusCause::SetDirectly );
+	FSlateApplication::Get().ClearKeyboardFocus( EFocusCause::Cleared );
 
 	this->SetText( FText::GetEmpty() );
 	return FReply::Handled().SetUserFocus(EditableText.ToSharedRef(), EFocusCause::SetDirectly);
