Fix shadowed variables

diff --git a/Engine/Plugins/Developer/GitSourceControl/Source/GitSourceControl/Private/GitSourceControlUtils.cpp b/Engine/Plugins/Developer/GitSourceControl/Source/GitSourceControl/Private/GitSourceControlUtils.cpp
--- a/Engine/Plugins/Developer/GitSourceControl/Source/GitSourceControl/Private/GitSourceControlUtils.cpp
+++ b/Engine/Plugins/Developer/GitSourceControl/Source/GitSourceControl/Private/GitSourceControlUtils.cpp
@@ -235,10 +235,12 @@ bool RunCommit(const FString& InPathToGitBinary, const FString& InRepositoryRoot
 	{
 		// Batch files up so we dont exceed command-line limits
 		int32 FileCount = 0;
-		TArray<FString> FilesInBatch;
-		for(int32 FileIndex = 0; FileIndex < GitSourceControlConstants::MaxFilesPerBatch; FileIndex++, FileCount++)
 		{
-			FilesInBatch.Add(InFiles[FileCount]);
+			TArray<FString> FilesInBatch;
+			for(int32 FileIndex = 0; FileIndex < GitSourceControlConstants::MaxFilesPerBatch; FileIndex++, FileCount++)
+			{
+				FilesInBatch.Add(InFiles[FileCount]);
+			}
 		}
 		// First batch is a simple "git commit" command with only the first files
 		bResult &= RunCommandInternal(TEXT("commit"), InPathToGitBinary, InRepositoryRoot, InParameters, InFiles, OutResults, OutErrorMessages);
@@ -432,9 +434,9 @@ bool RunUpdateStatus(const FString& InPathToGitBinary, const FString& InReposito
 		}
 		else
 		{
-			TArray<FString> Group;
-			Group.Add(File);
-			GroupOfFiles.Add(Path, Group);
+			TArray<FString> NewGroup;
+			NewGroup.Add(File);
+			GroupOfFiles.Add(Path, NewGroup);
 		}
 	}
 
diff --git a/Engine/Source/Developer/AutomationWindow/Private/SAutomationTestItem.cpp b/Engine/Source/Developer/AutomationWindow/Private/SAutomationTestItem.cpp
--- a/Engine/Source/Developer/AutomationWindow/Private/SAutomationTestItem.cpp
+++ b/Engine/Source/Developer/AutomationWindow/Private/SAutomationTestItem.cpp
@@ -73,18 +73,18 @@ void SAutomationHistoryCell::RebuildContentArea()
 
 	// Create an overview of the previous results in icon form.
 	const TArray<TSharedPtr<FAutomationHistoryItem>> History = HistoryItem->GetHistory();
-	for (const TSharedPtr<FAutomationHistoryItem> HistoryItem : History)
+	for (const TSharedPtr<FAutomationHistoryItem> AutomationHistoryItem : History)
 	{
 		FFormatNamedArguments Args;
-		Args.Add(TEXT("Date"), FText::AsDateTime(HistoryItem->RunDate));
+		Args.Add(TEXT("Date"), FText::AsDateTime(AutomationHistoryItem->RunDate));
 
 		const FSlateBrush* ResultIcon = nullptr;
-		if (HistoryItem->RunResult == FAutomationHistoryItem::EAutomationHistoryResult::Errors)
+		if (AutomationHistoryItem->RunResult == FAutomationHistoryItem::EAutomationHistoryResult::Errors)
 		{
 			Args.Add(TEXT("Result"), LOCTEXT("HasErrors", "had errors"));
 			ResultIcon = FEditorStyle::GetBrush("Automation.Fail");
 		}
-		else if (HistoryItem->RunResult == FAutomationHistoryItem::EAutomationHistoryResult::Warnings)
+		else if (AutomationHistoryItem->RunResult == FAutomationHistoryItem::EAutomationHistoryResult::Warnings)
 		{
 			Args.Add(TEXT("Result"), LOCTEXT("HasWarnings", "had warnings"));
 			ResultIcon = FEditorStyle::GetBrush("Automation.Warning");
diff --git a/Engine/Source/Developer/Profiler/Private/ProfilerSession.cpp b/Engine/Source/Developer/Profiler/Private/ProfilerSession.cpp
--- a/Engine/Source/Developer/Profiler/Private/ProfilerSession.cpp
+++ b/Engine/Source/Developer/Profiler/Private/ProfilerSession.cpp
@@ -1442,9 +1442,9 @@ void FRawProfilerSession::ProcessMemoryOperations( const TMap<int64, FStatPacket
 
 							// Create a callstack.
 							TArray<FName> StatsBasedCallstack;
-							for( const auto& RawName : StackState->Stack )
+							for( const auto& StackRawName : StackState->Stack )
 							{
-								StatsBasedCallstack.Add(RawName);
+								StatsBasedCallstack.Add(StackRawName);
 							}
 
 							// Add a new allocation.
@@ -1467,9 +1467,9 @@ void FRawProfilerSession::ProcessMemoryOperations( const TMap<int64, FStatPacket
 
 							// Create a callstack.
 							TArray<FName> StatsBasedCallstack;
-							for( const auto& RawName : StackState->Stack )
+							for( const auto& StackRawName : StackState->Stack )
 							{
-								StatsBasedCallstack.Add(RawName);
+								StatsBasedCallstack.Add(StackRawName);
 							}
 
 							// Add a new free.
diff --git a/Engine/Source/Editor/EditorStyle/Private/SlateEditorStyle.cpp b/Engine/Source/Editor/EditorStyle/Private/SlateEditorStyle.cpp
--- a/Engine/Source/Editor/EditorStyle/Private/SlateEditorStyle.cpp
+++ b/Engine/Source/Editor/EditorStyle/Private/SlateEditorStyle.cpp
@@ -396,7 +396,7 @@ void FSlateEditorStyle::FStyle::SetupGeneralStyles()
 			FLinearColor Hovered;
 			FLinearColor Pressed;
 
-			ButtonColor(const FName& Name, const FLinearColor& Color) : Name(Name)
+			ButtonColor(const FName& InName, const FLinearColor& Color) : Name(InName)
 			{
 				Normal = Color * 0.8f;
 				Normal.A = Color.A;
@@ -2888,14 +2888,14 @@ void FSlateEditorStyle::FStyle::SetupTutorialStyles()
 				.SetCheckedPressedImage(BOX_BRUSH("Tutorials/FlatColorSquare", FVector2D(1.0f, 1.0f), FMargin(1), HoverColor))
 				);
 
-			const FButtonStyle Button = FButtonStyle()
+			const FButtonStyle TutorialButton = FButtonStyle()
 				.SetNormal(BOX_BRUSH("Tutorials/FlatColorSquare", FVector2D(1.0f, 1.0f), FMargin(1), NormalColor))
 				.SetHovered(BOX_BRUSH("Tutorials/FlatColorSquare", FVector2D(1.0f, 1.0f), FMargin(1), HoverColor))
 				.SetPressed(BOX_BRUSH("Tutorials/FlatColorSquare", FVector2D(1.0f, 1.0f), FMargin(1), SelectedColor))
 				.SetDisabled(BOX_BRUSH("Tutorials/FlatColorSquare", FVector2D(1.0f, 1.0f), FMargin(1), DisabledColor))
 				.SetNormalPadding(FMargin(2, 2, 2, 2))
 				.SetPressedPadding(FMargin(2, 3, 2, 1));
-			Set("TutorialEditableText.Toolbar.Button", Button);
+			Set("TutorialEditableText.Toolbar.Button", TutorialButton);
 
 			const FComboButtonStyle ComboButton = FComboButtonStyle()
 				.SetButtonStyle(Button)
diff --git a/Engine/Source/Editor/NewsFeed/Private/Implementation/CdnNewsFeedTitleFile.cpp b/Engine/Source/Editor/NewsFeed/Private/Implementation/CdnNewsFeedTitleFile.cpp
--- a/Engine/Source/Editor/NewsFeed/Private/Implementation/CdnNewsFeedTitleFile.cpp
+++ b/Engine/Source/Editor/NewsFeed/Private/Implementation/CdnNewsFeedTitleFile.cpp
@@ -370,9 +370,11 @@ void FCdnNewsFeedTitleFile::ReadFile_HttpRequestComplete(FHttpRequestPtr HttpReq
 	FString ResponseStr, ErrorStr;
 
 	// Cloud file being operated on
-	FCloudFile* CloudFile = GetCloudFile(PendingRequest->FileName, true);
-	CloudFile->AsyncState = EOnlineAsyncTaskState::Failed;
-	CloudFile->Data.Empty();
+	{
+		FCloudFile* CloudFile = GetCloudFile(PendingRequest->FileName, true);
+		CloudFile->AsyncState = EOnlineAsyncTaskState::Failed;
+		CloudFile->Data.Empty();
+	}
 	
 	if (bSucceeded && 
 		HttpResponse.IsValid())
diff --git a/Engine/Source/Runtime/Core/Private/Apple/ApplePlatformSymbolication.cpp b/Engine/Source/Runtime/Core/Private/Apple/ApplePlatformSymbolication.cpp
--- a/Engine/Source/Runtime/Core/Private/Apple/ApplePlatformSymbolication.cpp
+++ b/Engine/Source/Runtime/Core/Private/Apple/ApplePlatformSymbolication.cpp
@@ -272,7 +272,7 @@ bool FApplePlatformSymbolication::LoadSymbolDatabaseForBinary(FString SourceFold
 	return bOK;
 }
 
-bool FApplePlatformSymbolication::SaveSymbolDatabaseForBinary(FString TargetFolder, FString Name, FApplePlatformSymbolDatabase& Database)
+bool FApplePlatformSymbolication::SaveSymbolDatabaseForBinary(FString TargetFolder, FString InName, FApplePlatformSymbolDatabase& Database)
 {
 	__block bool bOK = true;
 	if(GAllowApplePlatformSymbolication)
@@ -360,7 +360,7 @@ bool FApplePlatformSymbolication::SaveSymbolDatabaseForBinary(FString TargetFold
 	
 	if(bOK)
 	{
-		bOK = FGenericPlatformSymbolication::SaveSymbolDatabaseForBinary(TargetFolder, Name, Database.GenericDB);
+		bOK = FGenericPlatformSymbolication::SaveSymbolDatabaseForBinary(TargetFolder, InName, Database.GenericDB);
 	}
 	
 	return bOK;
diff --git a/Engine/Source/Runtime/Core/Private/HAL/FileManagerGeneric.cpp b/Engine/Source/Runtime/Core/Private/HAL/FileManagerGeneric.cpp
--- a/Engine/Source/Runtime/Core/Private/HAL/FileManagerGeneric.cpp
+++ b/Engine/Source/Runtime/Core/Private/HAL/FileManagerGeneric.cpp
@@ -460,8 +460,8 @@ FString FFileManagerGeneric::DefaultConvertToRelativePath( const TCHAR* Filename
 	FPaths::NormalizeFilename(RelativePath);
 
 	// See whether it is a relative path.
-	FString RootDir( FPlatformMisc::RootDir() );
-	FPaths::NormalizeFilename(RootDir);
+	FString RootDirectory( FPlatformMisc::RootDir() );
+	FPaths::NormalizeFilename(RootDirectory);
 
 	//the default relative directory it to the app root which is 3 directories up from the starting directory
 	int32 NumberOfDirectoriesToGoUp = 3;
@@ -470,26 +470,26 @@ FString FFileManagerGeneric::DefaultConvertToRelativePath( const TCHAR* Filename
 	int32 CurrentSlashPosition;
 
 	//while we haven't run out of parent directories until we which a drive name
-	while( ( CurrentSlashPosition = RootDir.Find( TEXT("/"), ESearchCase::CaseSensitive, ESearchDir::FromEnd ) ) != INDEX_NONE )
+	while( ( CurrentSlashPosition = RootDirectory.Find( TEXT("/"), ESearchCase::CaseSensitive, ESearchDir::FromEnd ) ) != INDEX_NONE )
 	{
-		if( RelativePath.StartsWith( RootDir ) )
+		if( RelativePath.StartsWith( RootDirectory ) )
 		{
 			FString BinariesDir = FString(FPlatformProcess::BaseDir());
 			FPaths::MakePathRelativeTo( RelativePath, *BinariesDir );
 			break;
 		}
-		int32 PositionOfNextSlash = RootDir.Find( TEXT("/"), ESearchCase::CaseSensitive, ESearchDir::FromEnd, CurrentSlashPosition );
+		int32 PositionOfNextSlash = RootDirectory.Find( TEXT("/"), ESearchCase::CaseSensitive, ESearchDir::FromEnd, CurrentSlashPosition );
 		//if there is another slash to find
 		if( PositionOfNextSlash != INDEX_NONE )
 		{
 			//move up a directory and on an extra .. TEXT("/")
 			// the +1 from "InStr" moves to include the "\" at the end of the directory name
 			NumberOfDirectoriesToGoUp++;
-			RootDir = RootDir.Left( PositionOfNextSlash + 1 );
+			RootDirectory = RootDirectory.Left( PositionOfNextSlash + 1 );
 		}
 		else
 		{
-			RootDir.Empty();
+			RootDirectory.Empty();
 		}
 	}
 	return RelativePath;
diff --git a/Engine/Source/Runtime/CoreUObject/Private/Serialization/BulkData.cpp b/Engine/Source/Runtime/CoreUObject/Private/Serialization/BulkData.cpp
--- a/Engine/Source/Runtime/CoreUObject/Private/Serialization/BulkData.cpp
+++ b/Engine/Source/Runtime/CoreUObject/Private/Serialization/BulkData.cpp
@@ -650,13 +650,13 @@ void FUntypedBulkData::StartSerializingBulkData(FArchive& Ar, UObject* Owner, in
 			BulkDataAsync = FMemory::Realloc(BulkDataAsync, GetBulkDataSize());
 		}
 
-		FArchive* Ar = IFileManager::Get().CreateFileReader(*Filename, FILEREAD_Silent);
-		checkf(Ar != NULL, TEXT("Attempted to load bulk data from an invalid filename '%s'."), *Filename);
+		FArchive* FileReaderAr = IFileManager::Get().CreateFileReader(*Filename, FILEREAD_Silent);
+		checkf(FileReaderAr != NULL, TEXT("Attempted to load bulk data from an invalid filename '%s'."), *Filename);
 
 		// Seek to the beginning of the bulk data in the file.
-		Ar->Seek(BulkDataOffsetInFile);
-		SerializeBulkData(*Ar, BulkDataAsync);
-		delete Ar;
+		FileReaderAr->Seek(BulkDataOffsetInFile);
+		SerializeBulkData(*FileReaderAr, BulkDataAsync);
+		delete FileReaderAr;
 
 		return true;
 	});
diff --git a/Engine/Source/Runtime/CoreUObject/Private/UObject/Class.cpp b/Engine/Source/Runtime/CoreUObject/Private/UObject/Class.cpp
--- a/Engine/Source/Runtime/CoreUObject/Private/UObject/Class.cpp
+++ b/Engine/Source/Runtime/CoreUObject/Private/UObject/Class.cpp
@@ -800,12 +800,12 @@ void UStruct::SerializeTaggedProperties(FArchive& Ar, uint8* Data, UStruct* Defa
 				break;
 			}
 
-			auto CanSerializeFromStructWithDifferentName = [](const FArchive& Ar, const FPropertyTag& PropertyTag, const UStructProperty* StructProperty)
+			auto CanSerializeFromStructWithDifferentName = [](const FArchive& InAr, const FPropertyTag& PropertyTag, const UStructProperty* StructProperty)
 			{
-				if (Ar.UE4Ver() < VER_UE4_STRUCT_GUID_IN_PROPERTY_TAG)
+				if (InAr.UE4Ver() < VER_UE4_STRUCT_GUID_IN_PROPERTY_TAG)
 				{
 					// Old Implementation
-					return !StructProperty->UseBinaryOrNativeSerialization(Ar);
+					return !StructProperty->UseBinaryOrNativeSerialization(InAr);
 				}
 				return PropertyTag.StructGuid.IsValid() && StructProperty && StructProperty->Struct && (PropertyTag.StructGuid == StructProperty->Struct->GetCustomGuid());
 			};
diff --git a/Engine/Source/Runtime/ShaderCore/Private/ShaderCache.cpp b/Engine/Source/Runtime/ShaderCore/Private/ShaderCache.cpp
--- a/Engine/Source/Runtime/ShaderCore/Private/ShaderCache.cpp
+++ b/Engine/Source/Runtime/ShaderCore/Private/ShaderCache.cpp
@@ -1084,7 +1084,7 @@ void FShaderCache::SetShaderSamplerTextures( FRHICommandList& RHICmdList, FShade
 
 void FShaderCache::PreDrawShader(FRHICommandList& RHICmdList, FShaderCacheBoundState const& Shader, TSet<int32> const& DrawStates)
 {
-	FBoundShaderStateRHIRef BoundShaderState = BoundShaderStates.FindRef( Shader );
+	FBoundShaderStateRHIRef ShaderBoundState = BoundShaderStates.FindRef( Shader );
 	{
 		uint32 VertexBufferSize = 0;
 		for( auto VertexDec : Shader.VertexDeclaration )
@@ -1178,7 +1178,7 @@ void FShaderCache::PreDrawShader(FRHICommandList& RHICmdList, FShaderCacheBoundS
 			
 			RHICmdList.SetRenderTargets(NewNumRenderTargets, RenderTargets, bDepthStencilTarget ? &DepthStencilTarget : nullptr, 0, nullptr);
 			
-			if ( !IsValidRef(BoundShaderState) )
+			if ( !IsValidRef(ShaderBoundState) )
 			{
 				FVertexShaderRHIRef VertexShader = Shader.VertexShader.bActive ? CachedVertexShaders.FindRef(Shader.VertexShader) : nullptr;
 				FPixelShaderRHIRef PixelShader = Shader.PixelShader.bActive ? CachedPixelShaders.FindRef(Shader.PixelShader) : nullptr;
@@ -1200,7 +1200,7 @@ void FShaderCache::PreDrawShader(FRHICommandList& RHICmdList, FShaderCacheBoundS
 					
 					if ( bOK )
 					{
-						BoundShaderState = RHICreateBoundShaderState( VertexDeclaration,
+						ShaderBoundState = RHICreateBoundShaderState( VertexDeclaration,
 																	 VertexShader,
 																	 HullShader,
 																	 DomainShader,
@@ -1210,9 +1210,9 @@ void FShaderCache::PreDrawShader(FRHICommandList& RHICmdList, FShaderCacheBoundS
 				}
 			}
 			
-			if ( IsValidRef( BoundShaderState ) )
+			if ( IsValidRef( ShaderBoundState ) )
 			{
-				RHICmdList.SetBoundShaderState(BoundShaderState);
+				RHICmdList.SetBoundShaderState(ShaderBoundState);
 			}
 			else
 			{
@@ -1264,7 +1264,7 @@ void FShaderCache::PreDrawShader(FRHICommandList& RHICmdList, FShaderCacheBoundS
 			}
 		}
 		
-		if( IsValidRef( BoundShaderState ) && DrawStates.Num() )
+		if( IsValidRef( ShaderBoundState ) && DrawStates.Num() )
 		{
 			if ( Shader.VertexShader.bActive )
 			{
diff --git a/Engine/Source/Runtime/Slate/Private/Widgets/Text/SMultiLineEditableText.cpp b/Engine/Source/Runtime/Slate/Private/Widgets/Text/SMultiLineEditableText.cpp
--- a/Engine/Source/Runtime/Slate/Private/Widgets/Text/SMultiLineEditableText.cpp
+++ b/Engine/Source/Runtime/Slate/Private/Widgets/Text/SMultiLineEditableText.cpp
@@ -8,7 +8,7 @@
 #include "PlainTextLayoutMarshaller.h"
 #include "GenericCommands.h"
 
-void SMultiLineEditableText::FCursorInfo::SetCursorLocationAndCalculateAlignment(const TSharedPtr<FTextLayout>& TextLayout, const FTextLocation& InCursorPosition)
+void SMultiLineEditableText::FCursorInfo::SetCursorLocationAndCalculateAlignment(const TSharedPtr<FTextLayout>& InTextLayout, const FTextLocation& InCursorPosition)
 {
 	FTextLocation NewCursorPosition = InCursorPosition;
 	ECursorAlignment NewAlignment = ECursorAlignment::Left;
@@ -17,9 +17,9 @@ void SMultiLineEditableText::FCursorInfo::SetCursorLocationAndCalculateAlignment
 	const int32 CursorOffset = InCursorPosition.GetOffset();
 
 	// A CursorOffset of zero could mark the end of an empty line, but we don't need to adjust the cursor for an empty line
-	if (TextLayout.IsValid() && CursorOffset > 0)
+	if (InTextLayout.IsValid() && CursorOffset > 0)
 	{
-		const TArray< FTextLayout::FLineModel >& Lines = TextLayout->GetLineModels();
+		const TArray< FTextLayout::FLineModel >& Lines = InTextLayout->GetLineModels();
 		const FTextLayout::FLineModel& Line = Lines[CursorLineIndex];
 		if (Line.Text->Len() == CursorOffset)
 		{
@@ -2309,9 +2309,9 @@ FVector2D SMultiLineEditableText::ComputeDesiredSize( float ) const
 	const FVector2D TextLayoutSize = TextLayoutBaseSize + FVector2D(CaretWidth, 0);
 
 	// The layouts current margin size. We should not report a size smaller then the margins.
-	const FMargin Margin = TextLayout->GetMargin();
-	float DesiredWidth = FMath::Max(Margin.GetTotalSpaceAlong<Orient_Horizontal>(), TextLayoutSize.X);
-	float DesiredHeight = FMath::Max(Margin.GetTotalSpaceAlong<Orient_Vertical>(), TextLayoutSize.Y);
+	const FMargin TextLayoutMargin = TextLayout->GetMargin();
+	float DesiredWidth = FMath::Max(TextLayoutMargin.GetTotalSpaceAlong<Orient_Horizontal>(), TextLayoutSize.X);
+	float DesiredHeight = FMath::Max(TextLayoutMargin.GetTotalSpaceAlong<Orient_Vertical>(), TextLayoutSize.Y);
 	DesiredHeight = FMath::Max(DesiredHeight, FontMaxCharHeight);
 	
 	return FVector2D(DesiredWidth, DesiredHeight);
diff --git a/Engine/Source/Runtime/SlateCore/Private/Widgets/SWindow.cpp b/Engine/Source/Runtime/SlateCore/Private/Widgets/SWindow.cpp
--- a/Engine/Source/Runtime/SlateCore/Private/Widgets/SWindow.cpp
+++ b/Engine/Source/Runtime/SlateCore/Private/Widgets/SWindow.cpp
@@ -608,8 +608,8 @@ FGeometry SWindow::GetWindowGeometryInWindow() const
 	// FGeometry expects Size in Local space, but our size is stored in screen space (same as window space + screen offset).
 	// So we need to transform Size into the window's local space for FGeometry.
 	FSlateLayoutTransform LocalToWindow = GetLocalToWindowTransform();
-	FVector2D ViewportSize = GetViewportSize();
-	return FGeometry::MakeRoot( TransformVector(Inverse(LocalToWindow), ViewportSize), LocalToWindow );
+	FVector2D ViewSize = GetViewportSize();
+	return FGeometry::MakeRoot( TransformVector(Inverse(LocalToWindow), ViewSize), LocalToWindow );
 }
 
 FSlateLayoutTransform SWindow::GetLocalToScreenTransform() const
@@ -674,8 +674,8 @@ FVector2D SWindow::GetClientSizeInScreen() const
 
 FSlateRect SWindow::GetClippingRectangleInWindow() const
 {
-	FVector2D ViewportSize = GetViewportSize();
-	return FSlateRect( 0, 0, ViewportSize.X, ViewportSize.Y );
+	FVector2D ViewSize = GetViewportSize();
+	return FSlateRect( 0, 0, ViewSize.X, ViewSize.Y );
 }
 
 
