Compile fix.

diff --git a/Engine/Source/Runtime/Slate/Private/Framework/Application/SlateApplication.cpp b/Engine/Source/Runtime/Slate/Private/Framework/Application/SlateApplication.cpp
--- a/Engine/Source/Runtime/Slate/Private/Framework/Application/SlateApplication.cpp
+++ b/Engine/Source/Runtime/Slate/Private/Framework/Application/SlateApplication.cpp
@@ -2183,7 +2183,7 @@ bool FSlateApplication::SetUserFocus(const uint32 InUserIndex, const FWidgetPath
 		FScopedSwitchWorldHack SwitchWorld(OldFocusedWidgetPath.Window.Pin());
 
 		// Let previously-focused widget know that it's losing focus
-		OldFocusedWidget->OnFocusLost(FFocusEvent(InCause));
+		OldFocusedWidget->OnFocusLost(FFocusEvent(InCause, InUserIndex));
 	}
 
 	if (bReflectorShowingFocus)
@@ -2201,7 +2201,7 @@ bool FSlateApplication::SetUserFocus(const uint32 InUserIndex, const FWidgetPath
 		
 		const FArrangedWidget& WidgetToFocus = NewFocusedWidgetPath.Widgets.Last();
 
-		FReply Reply = NewFocusedWidget->OnFocusReceived(WidgetToFocus.Geometry, FFocusEvent(InCause));
+		FReply Reply = NewFocusedWidget->OnFocusReceived(WidgetToFocus.Geometry, FFocusEvent(InCause, InUserIndex));
 		if (Reply.IsEventHandled())
 		{
 			ProcessReply(InFocusPath, Reply, nullptr, nullptr);
