PS4 virtual keyboard support.

diff --git a/Engine/Source/Runtime/Slate/Private/Framework/Application/SlateApplication.cpp b/Engine/Source/Runtime/Slate/Private/Framework/Application/SlateApplication.cpp
--- a/Engine/Source/Runtime/Slate/Private/Framework/Application/SlateApplication.cpp
+++ b/Engine/Source/Runtime/Slate/Private/Framework/Application/SlateApplication.cpp
@@ -3297,14 +3297,14 @@ const FSlateBrush* FSlateApplication::GetAppIcon() const
 }
 
 
-void FSlateApplication::ShowVirtualKeyboard( bool bShow, TSharedPtr<IVirtualKeyboardEntry> TextEntryWidget )
+void FSlateApplication::ShowVirtualKeyboard( bool bShow, int32 UserIndex, TSharedPtr<IVirtualKeyboardEntry> TextEntryWidget )
 {
 	if(SlateTextField == nullptr)
 	{
 		SlateTextField = new FPlatformTextField();
 	}
 
-	SlateTextField->ShowVirtualKeyboard(bShow, TextEntryWidget);
+	SlateTextField->ShowVirtualKeyboard(bShow, UserIndex, TextEntryWidget);
 }
 
 FSlateRect FSlateApplication::GetPreferredWorkArea() const
diff --git a/Engine/Source/Runtime/Slate/Private/Framework/Text/Android/AndroidPlatformTextField.cpp b/Engine/Source/Runtime/Slate/Private/Framework/Text/Android/AndroidPlatformTextField.cpp
--- a/Engine/Source/Runtime/Slate/Private/Framework/Text/Android/AndroidPlatformTextField.cpp
+++ b/Engine/Source/Runtime/Slate/Private/Framework/Text/Android/AndroidPlatformTextField.cpp
@@ -17,7 +17,7 @@
 #define TYPE_TEXT_FLAG_NO_SUGGESTIONS		0x00080000
 
 
-void FAndroidPlatformTextField::ShowVirtualKeyboard(bool bShow, TSharedPtr<IVirtualKeyboardEntry> TextEntryWidget)
+void FAndroidPlatformTextField::ShowVirtualKeyboard(bool bShow, int32 UserIndex, TSharedPtr<IVirtualKeyboardEntry> TextEntryWidget)
 {
 	if(bShow)
 	{
diff --git a/Engine/Source/Runtime/Slate/Private/Framework/Text/IOS/IOSPlatformTextField.cpp b/Engine/Source/Runtime/Slate/Private/Framework/Text/IOS/IOSPlatformTextField.cpp
--- a/Engine/Source/Runtime/Slate/Private/Framework/Text/IOS/IOSPlatformTextField.cpp
+++ b/Engine/Source/Runtime/Slate/Private/Framework/Text/IOS/IOSPlatformTextField.cpp
@@ -17,7 +17,7 @@ FIOSPlatformTextField::~FIOSPlatformTextField()
 	}
 }
 
-void FIOSPlatformTextField::ShowVirtualKeyboard(bool bShow, TSharedPtr<IVirtualKeyboardEntry> TextEntryWidget)
+void FIOSPlatformTextField::ShowVirtualKeyboard(bool bShow, int32 UserIndex, TSharedPtr<IVirtualKeyboardEntry> TextEntryWidget)
 {
 	if(TextField == nullptr)
 	{
diff --git a/Engine/Source/Runtime/Slate/Private/Widgets/Input/SEditableText.cpp b/Engine/Source/Runtime/Slate/Private/Widgets/Input/SEditableText.cpp
--- a/Engine/Source/Runtime/Slate/Private/Widgets/Input/SEditableText.cpp
+++ b/Engine/Source/Runtime/Slate/Private/Widgets/Input/SEditableText.cpp
@@ -1541,7 +1541,7 @@ FReply SEditableText::OnFocusReceived( const FGeometry& MyGeometry, const FFocus
 		if (FPlatformMisc::GetRequiresVirtualKeyboard())
 		{
 			// @TODO: Create ITextInputMethodSystem derivations for mobile
-			SlateApplication.ShowVirtualKeyboard(true, SharedThis(this));
+			SlateApplication.ShowVirtualKeyboard(true, InFocusEvent.GetUser(), SharedThis(this));
 		}
 		else
 		{
@@ -1576,7 +1576,7 @@ void SEditableText::OnFocusLost( const FFocusEvent& InFocusEvent )
 		FSlateApplication& SlateApplication = FSlateApplication::Get();
 		if (FPlatformMisc::GetRequiresVirtualKeyboard())
 		{
-			SlateApplication.ShowVirtualKeyboard(false);
+			SlateApplication.ShowVirtualKeyboard(false, InFocusEvent.GetUser());
 		}
 		else
 		{
diff --git a/Engine/Source/Runtime/Slate/Private/Widgets/Input/SVirtualKeyboardEntry.cpp b/Engine/Source/Runtime/Slate/Private/Widgets/Input/SVirtualKeyboardEntry.cpp
--- a/Engine/Source/Runtime/Slate/Private/Widgets/Input/SVirtualKeyboardEntry.cpp
+++ b/Engine/Source/Runtime/Slate/Private/Widgets/Input/SVirtualKeyboardEntry.cpp
@@ -252,7 +252,7 @@ FReply SVirtualKeyboardEntry::OnFocusReceived( const FGeometry& MyGeometry, cons
 
 	int32 CaretPosition = EditedText.ToString().Len();
 	FSlateApplication& CurrentApp = FSlateApplication::Get();
-	CurrentApp.ShowVirtualKeyboard(true, SharedThis(this));
+	CurrentApp.ShowVirtualKeyboard(true, InFocusEvent.GetUser(), SharedThis(this));
 
 	return FReply::Handled();
 }
@@ -283,7 +283,7 @@ void SVirtualKeyboardEntry::OnFocusLost( const FFocusEvent& InFocusEvent )
 	}
 
 	FSlateApplication& CurrentApp = FSlateApplication::Get();
-	CurrentApp.ShowVirtualKeyboard(false);
+	CurrentApp.ShowVirtualKeyboard(false, InFocusEvent.GetUser());
 
 	OnTextCommitted.ExecuteIfBound( EditedText, TextAction );
 }
diff --git a/Engine/Source/Runtime/Slate/Private/Widgets/Text/SMultiLineEditableText.cpp b/Engine/Source/Runtime/Slate/Private/Widgets/Text/SMultiLineEditableText.cpp
--- a/Engine/Source/Runtime/Slate/Private/Widgets/Text/SMultiLineEditableText.cpp
+++ b/Engine/Source/Runtime/Slate/Private/Widgets/Text/SMultiLineEditableText.cpp
@@ -497,7 +497,7 @@ FReply SMultiLineEditableText::OnFocusReceived( const FGeometry& MyGeometry, con
 		if (FPlatformMisc::GetRequiresVirtualKeyboard())
 		{
 			// @TODO: Create ITextInputMethodSystem derivations for mobile
-			SlateApplication.ShowVirtualKeyboard(true, SharedThis(this));
+			SlateApplication.ShowVirtualKeyboard(true, InFocusEvent.GetUser(), SharedThis(this));
 		}
 		else
 		{
@@ -530,7 +530,7 @@ void SMultiLineEditableText::OnFocusLost( const FFocusEvent& InFocusEvent )
 		FSlateApplication& SlateApplication = FSlateApplication::Get();
 		if (FPlatformMisc::GetRequiresVirtualKeyboard())
 		{
-			SlateApplication.ShowVirtualKeyboard(false);
+			SlateApplication.ShowVirtualKeyboard(false, InFocusEvent.GetUser());
 		}
 		else
 		{
diff --git a/Engine/Source/Runtime/Slate/Public/Framework/Application/IPlatformTextField.h b/Engine/Source/Runtime/Slate/Public/Framework/Application/IPlatformTextField.h
--- a/Engine/Source/Runtime/Slate/Public/Framework/Application/IPlatformTextField.h
+++ b/Engine/Source/Runtime/Slate/Public/Framework/Application/IPlatformTextField.h
@@ -9,7 +9,7 @@ class IPlatformTextField
 public:
 	virtual ~IPlatformTextField() {};
 
-	virtual void ShowVirtualKeyboard(bool bShow, TSharedPtr<IVirtualKeyboardEntry> TextEntryWidget) = 0;
+	virtual void ShowVirtualKeyboard(bool bShow, int32 UserIndex, TSharedPtr<IVirtualKeyboardEntry> TextEntryWidget) = 0;
 
 private:
 
diff --git a/Engine/Source/Runtime/Slate/Public/Framework/Application/SlateApplication.h b/Engine/Source/Runtime/Slate/Public/Framework/Application/SlateApplication.h
--- a/Engine/Source/Runtime/Slate/Public/Framework/Application/SlateApplication.h
+++ b/Engine/Source/Runtime/Slate/Public/Framework/Application/SlateApplication.h
@@ -625,7 +625,7 @@ public:
 	 * @param bShow	true to show the keyboard, false to hide it
 	 * @param TextEntryWidget The widget that will receive the input from the virtual keyboard
 	 */
-	void ShowVirtualKeyboard( bool bShow, TSharedPtr<IVirtualKeyboardEntry> TextEntryWidget = nullptr );
+	void ShowVirtualKeyboard( bool bShow, int32 UserIndex, TSharedPtr<IVirtualKeyboardEntry> TextEntryWidget = nullptr );
 
 	/** Get the work area that has the largest intersection with the specified rectangle */
 	FSlateRect GetWorkArea( const FSlateRect& InRect ) const;
diff --git a/Engine/Source/Runtime/Slate/Public/Framework/Text/Android/AndroidPlatformTextField.h b/Engine/Source/Runtime/Slate/Public/Framework/Text/Android/AndroidPlatformTextField.h
--- a/Engine/Source/Runtime/Slate/Public/Framework/Text/Android/AndroidPlatformTextField.h
+++ b/Engine/Source/Runtime/Slate/Public/Framework/Text/Android/AndroidPlatformTextField.h
@@ -7,7 +7,7 @@
 class FAndroidPlatformTextField : public IPlatformTextField
 {
 public:
-	virtual void ShowVirtualKeyboard(bool bShow, TSharedPtr<IVirtualKeyboardEntry> TextEntryWidget) override;
+	virtual void ShowVirtualKeyboard(bool bShow, int32 UserIndex, TSharedPtr<IVirtualKeyboardEntry> TextEntryWidget) override;
 
 private:
 //	SlateTextField* TextField;
diff --git a/Engine/Source/Runtime/Slate/Public/Framework/Text/GenericPlatformTextField.h b/Engine/Source/Runtime/Slate/Public/Framework/Text/GenericPlatformTextField.h
--- a/Engine/Source/Runtime/Slate/Public/Framework/Text/GenericPlatformTextField.h
+++ b/Engine/Source/Runtime/Slate/Public/Framework/Text/GenericPlatformTextField.h
@@ -7,7 +7,7 @@
 class FGenericPlatformTextField : public IPlatformTextField
 {
 public:
-	virtual void ShowVirtualKeyboard(bool bShow, TSharedPtr<IVirtualKeyboardEntry> TextEntryWidget) override {};
+	virtual void ShowVirtualKeyboard(bool bShow, int32 UserIndex, TSharedPtr<IVirtualKeyboardEntry> TextEntryWidget) override {};
 
 };
 
diff --git a/Engine/Source/Runtime/Slate/Public/Framework/Text/IOS/IOSPlatformTextField.h b/Engine/Source/Runtime/Slate/Public/Framework/Text/IOS/IOSPlatformTextField.h
--- a/Engine/Source/Runtime/Slate/Public/Framework/Text/IOS/IOSPlatformTextField.h
+++ b/Engine/Source/Runtime/Slate/Public/Framework/Text/IOS/IOSPlatformTextField.h
@@ -14,7 +14,7 @@ public:
 	FIOSPlatformTextField();
 	virtual ~FIOSPlatformTextField();
 
-	virtual void ShowVirtualKeyboard(bool bShow, TSharedPtr<IVirtualKeyboardEntry> TextEntryWidget) override;
+	virtual void ShowVirtualKeyboard(bool bShow, int32 UserIndex, TSharedPtr<IVirtualKeyboardEntry> TextEntryWidget) override;
 
 private:
 	SlateTextField* TextField;
diff --git a/Engine/Source/Runtime/Slate/Public/SlateBasics.h b/Engine/Source/Runtime/Slate/Public/SlateBasics.h
--- a/Engine/Source/Runtime/Slate/Public/SlateBasics.h
+++ b/Engine/Source/Runtime/Slate/Public/SlateBasics.h
@@ -17,6 +17,8 @@
 	#include "IOS/IOSPlatformTextField.h"
 #elif PLATFORM_ANDROID
 	#include "Android/AndroidPlatformTextField.h"
+#elif PLATFORM_PS4
+	#include "PS4/PS4PlatformTextField.h"
 #else
 	#include "GenericPlatformTextField.h"
 #endif
