Added a scoped lock for GCachedMenuState access in MacMenu

diff --git a/Engine/Source/Runtime/Slate/Private/Framework/MultiBox/Mac/MacMenu.cpp b/Engine/Source/Runtime/Slate/Private/Framework/MultiBox/Mac/MacMenu.cpp
--- a/Engine/Source/Runtime/Slate/Private/Framework/MultiBox/Mac/MacMenu.cpp
+++ b/Engine/Source/Runtime/Slate/Private/Framework/MultiBox/Mac/MacMenu.cpp
@@ -27,6 +27,7 @@ struct FMacMenuItemState
 };
 
 static TMap<FMacMenu*, TSharedPtr<TArray<FMacMenuItemState>>> GCachedMenuState;
+static FCriticalSection GCachedMenuStateCS;
 
 @interface FMacMenuItem : NSMenuItem
 @property (assign) TSharedPtr<const FMenuEntryBlock> MenuEntryBlock;
@@ -60,6 +61,7 @@ static TMap<FMacMenu*, TSharedPtr<TArray<FMacMenuItemState>>> GCachedMenuState;
 	self = [super initWithTitle:@""];
 	[self setDelegate:self];
 	self.MenuEntryBlock = Block;
+	FScopeLock Lock(&GCachedMenuStateCS);
 	GCachedMenuState.Add(self, TSharedPtr<TArray<FMacMenuItemState>>(new TArray<FMacMenuItemState>()));
 	return self;
 }
@@ -84,6 +86,8 @@ static TMap<FMacMenu*, TSharedPtr<TArray<FMacMenuItemState>>> GCachedMenuState;
 void FSlateMacMenu::UpdateWithMultiBox(const TSharedPtr< FMultiBox > MultiBox)
 {
 	MainThreadCall(^{
+		FScopeLock Lock(&GCachedMenuStateCS);
+
 		if (!FPlatformMisc::UpdateCachedMacMenuState)
 		{
 			FPlatformMisc::UpdateCachedMacMenuState = UpdateCachedState;
@@ -158,6 +162,8 @@ void FSlateMacMenu::UpdateWithMultiBox(const TSharedPtr< FMultiBox > MultiBox)
 void FSlateMacMenu::UpdateMenu(FMacMenu* Menu)
 {
 	MainThreadCall(^{
+		FScopeLock Lock(&GCachedMenuStateCS);
+
 		FText WindowLabel = NSLOCTEXT("MainMenu", "WindowMenu", "Window");
 		const bool bIsWindowMenu = (WindowLabel.ToString().Compare(FString([Menu title])) == 0);
 
@@ -285,6 +291,8 @@ void FSlateMacMenu::UpdateCachedState()
 
 	if (bShouldUpdate)
 	{
+		FScopeLock Lock(&GCachedMenuStateCS);
+
 		for (TMap<FMacMenu*, TSharedPtr<TArray<FMacMenuItemState>>>::TIterator It(GCachedMenuState); It; ++It)
 		{
 			FMacMenu* Menu = It.Key();
