UE-10875
Possible to have the toolbar icons disappear when adjusting the size of the editor

diff --git a/Engine/Source/Runtime/Slate/Private/Framework/MultiBox/SClippingHorizontalBox.cpp b/Engine/Source/Runtime/Slate/Private/Framework/MultiBox/SClippingHorizontalBox.cpp
--- a/Engine/Source/Runtime/Slate/Private/Framework/MultiBox/SClippingHorizontalBox.cpp
+++ b/Engine/Source/Runtime/Slate/Private/Framework/MultiBox/SClippingHorizontalBox.cpp
@@ -42,12 +42,13 @@ void SClippingHorizontalBox::OnArrangeChildren( const FGeometry& AllottedGeometr
 		// Right align the wrap button
 		FArrangedWidget& ArrangedButton = ArrangedChildren[ArrangedChildren.Num() - 1];
 		ArrangedButton.Geometry = AllottedGeometry.MakeChild(ArrangedButton.Geometry.Size, FSlateLayoutTransform(AllottedGeometry.Size - ArrangedButton.Geometry.Size));
+		const int32 WrapButtonXPosition = FMath::TruncToInt(ArrangedButton.Geometry.AbsolutePosition.X);
 
 		// Further remove any children that the wrap button overlaps with
 		for (int32 ChildIdx = IndexClippedAt - 1; ChildIdx >= 0; --ChildIdx)
 		{
 			const FArrangedWidget& CurWidget = ArrangedChildren[ChildIdx];
-			if (FMath::TruncToInt(CurWidget.Geometry.AbsolutePosition.X + CurWidget.Geometry.Size.X * CurWidget.Geometry.Scale) > FMath::TruncToInt(ArrangedButton.Geometry.AbsolutePosition.X))
+			if (FMath::TruncToInt(CurWidget.Geometry.AbsolutePosition.X + CurWidget.Geometry.Size.X * CurWidget.Geometry.Scale) > WrapButtonXPosition)
 			{
 				ArrangedChildren.Remove(ChildIdx);
 			}
