Slate - Fixing Above Menu anchor spawn point to correctly offset directly above the menu anchor.  Additionally adding Above Centered as a spawn location.

diff --git a/Engine/Source/Runtime/Slate/Private/Widgets/Input/SMenuAnchor.cpp b/Engine/Source/Runtime/Slate/Private/Widgets/Input/SMenuAnchor.cpp
--- a/Engine/Source/Runtime/Slate/Private/Widgets/Input/SMenuAnchor.cpp
+++ b/Engine/Source/Runtime/Slate/Private/Widgets/Input/SMenuAnchor.cpp
@@ -24,7 +24,10 @@ static FVector2D GetMenuOffsetForPlacement(const FGeometry& AllottedGeometry, EM
 			return FVector2D(AllottedGeometry.GetLocalSize().X, 0.0f);
 			break;
 		case MenuPlacement_AboveAnchor:
-			return FVector2D(0.0f, -AllottedGeometry.GetLocalSize().Y - PopupSizeLocalSpace.Y);
+			return FVector2D(0.0f, -PopupSizeLocalSpace.Y);
+			break;
+		case MenuPlacement_CenteredAboveAnchor:
+			return FVector2D(-((PopupSizeLocalSpace.X / 2) - (AllottedGeometry.GetLocalSize().X / 2)), -PopupSizeLocalSpace.Y);
 			break;
 		case MenuPlacement_MenuLeft:
 			return FVector2D(-AllottedGeometry.GetLocalSize().X - PopupSizeLocalSpace.X, 0.0f);
diff --git a/Engine/Source/Runtime/SlateCore/Public/Types/SlateEnums.h b/Engine/Source/Runtime/SlateCore/Public/Types/SlateEnums.h
--- a/Engine/Source/Runtime/SlateCore/Public/Types/SlateEnums.h
+++ b/Engine/Source/Runtime/SlateCore/Public/Types/SlateEnums.h
@@ -166,9 +166,12 @@ enum EMenuPlacement
 	/** Place the menu to the right of the anchor */
 	MenuPlacement_MenuRight UMETA(DisplayName="Right"),
 
-	/** Place the menu immediately above the anchor, not transition effect */
+	/** Place the menu immediately above the anchor, no transition effect */
 	MenuPlacement_AboveAnchor UMETA(DisplayName="Above"),
 
+	/** Place the menu immediately centered above the anchor, no transition effect */
+	MenuPlacement_CenteredAboveAnchor UMETA(DisplayName="Centered Above"),
+
 	/** Place the menu to the left of the anchor */
 	MenuPlacement_MenuLeft UMETA(DisplayName="Left"),
 };
