Slate : fixed right-mouse-button scroll being double-applied in SScrollBox.

diff --git a/Engine/Source/Runtime/Slate/Private/Widgets/Layout/SScrollBox.cpp b/Engine/Source/Runtime/Slate/Private/Widgets/Layout/SScrollBox.cpp
--- a/Engine/Source/Runtime/Slate/Private/Widgets/Layout/SScrollBox.cpp
+++ b/Engine/Source/Runtime/Slate/Private/Widgets/Layout/SScrollBox.cpp
@@ -440,7 +440,8 @@ EActiveTimerReturnType SScrollBox::UpdateInertialScroll(double InCurrentTime, fl
 	if (bIsScrolling)
 	{
 		const float ScrollVelocity = InertialScrollManager.GetScrollVelocity();
-		if (ScrollVelocity != 0.f)
+		// Do not apply inertial scrolling while the user is actively scrolling via RMB.
+		if (ScrollVelocity != 0.f && !IsRightClickScrolling())
 		{
 			ScrollBy(CachedGeometry, ScrollVelocity * InDeltaTime, true);
 		}
