Fixed UI_COMMAND macro implementation so that it requires use of the LOCTEXT_NAMESPACE macro in the same file it is invoked from. Added necessary LOCTEXT_NAMESPACE definitions to files missing them using UI_COMMAND.

diff --git a/Engine/Plugins/2D/Paper2D/Source/Paper2DEditor/Private/FlipbookEditor/FlipbookEditorCommands.cpp b/Engine/Plugins/2D/Paper2D/Source/Paper2DEditor/Private/FlipbookEditor/FlipbookEditorCommands.cpp
--- a/Engine/Plugins/2D/Paper2D/Source/Paper2DEditor/Private/FlipbookEditor/FlipbookEditorCommands.cpp
+++ b/Engine/Plugins/2D/Paper2D/Source/Paper2DEditor/Private/FlipbookEditor/FlipbookEditorCommands.cpp
@@ -6,6 +6,8 @@
 //////////////////////////////////////////////////////////////////////////
 // FFlipbookEditorCommands
 
+#define LOCTEXT_NAMESPACE ""
+
 void FFlipbookEditorCommands::RegisterCommands()
 {
 	UI_COMMAND(SetShowGrid, "Grid", "Displays the viewport grid.", EUserInterfaceActionType::ToggleButton, FInputGesture());
@@ -23,3 +25,5 @@ void FFlipbookEditorCommands::RegisterCommands()
 	UI_COMMAND(EditSpriteFrame, "Edit Sprite", "Opens the sprite for this key frame in the Sprite Editor.", EUserInterfaceActionType::Button, FInputGesture());
 	UI_COMMAND(ShowInContentBrowser, "Show in Content Browser", "Shows the sprite for this key frame in the Content Browser.", EUserInterfaceActionType::Button, FInputGesture());
 }
+
+#undef LOCTEXT_NAMESPACE
\ No newline at end of file
diff --git a/Engine/Plugins/2D/Paper2D/Source/Paper2DEditor/Private/PaperEditorCommands.cpp b/Engine/Plugins/2D/Paper2D/Source/Paper2DEditor/Private/PaperEditorCommands.cpp
--- a/Engine/Plugins/2D/Paper2D/Source/Paper2DEditor/Private/PaperEditorCommands.cpp
+++ b/Engine/Plugins/2D/Paper2D/Source/Paper2DEditor/Private/PaperEditorCommands.cpp
@@ -6,8 +6,10 @@
 //////////////////////////////////////////////////////////////////////////
 // FPaperEditorCommands
 
+#define LOCTEXT_NAMESPACE ""
+
 void FPaperEditorCommands::RegisterCommands()
 {
 }
 	
-	
\ No newline at end of file
+#undef LOCTEXT_NAMESPACE
\ No newline at end of file
diff --git a/Engine/Plugins/2D/Paper2D/Source/Paper2DEditor/Private/SpriteEditor/SpriteEditorCommands.cpp b/Engine/Plugins/2D/Paper2D/Source/Paper2DEditor/Private/SpriteEditor/SpriteEditorCommands.cpp
--- a/Engine/Plugins/2D/Paper2D/Source/Paper2DEditor/Private/SpriteEditor/SpriteEditorCommands.cpp
+++ b/Engine/Plugins/2D/Paper2D/Source/Paper2DEditor/Private/SpriteEditor/SpriteEditorCommands.cpp
@@ -6,6 +6,8 @@
 //////////////////////////////////////////////////////////////////////////
 // FSpriteEditorCommands
 
+#define LOCTEXT_NAMESPACE ""
+
 void FSpriteEditorCommands::RegisterCommands()
 {
 	// Show toggles
@@ -36,3 +38,5 @@ void FSpriteEditorCommands::RegisterCommands()
 	UI_COMMAND(ToggleAddPolygonMode, "Add Polygon", "Adds a new polygon.", EUserInterfaceActionType::ToggleButton, FInputGesture());
 	UI_COMMAND(SnapAllVertices, "Snap to pixel grid", "Snaps all vertices to the pixel grid.", EUserInterfaceActionType::Button, FInputGesture());
 }
+
+#undef LOCTEXT_NAMESPACE
\ No newline at end of file
diff --git a/Engine/Plugins/2D/Paper2D/Source/Paper2DEditor/Private/TileMapEditing/TileMapEditorCommands.cpp b/Engine/Plugins/2D/Paper2D/Source/Paper2DEditor/Private/TileMapEditing/TileMapEditorCommands.cpp
--- a/Engine/Plugins/2D/Paper2D/Source/Paper2DEditor/Private/TileMapEditing/TileMapEditorCommands.cpp
+++ b/Engine/Plugins/2D/Paper2D/Source/Paper2DEditor/Private/TileMapEditing/TileMapEditorCommands.cpp
@@ -6,6 +6,8 @@
 //////////////////////////////////////////////////////////////////////////
 // FTileMapEditorCommands
 
+#define LOCTEXT_NAMESPACE ""
+
 void FTileMapEditorCommands::RegisterCommands()
 {
 	UI_COMMAND(EnterTileMapEditMode, "Enable Tile Map Mode", "Enables Tile Map editing mode", EUserInterfaceActionType::Button, FInputGesture());
@@ -37,3 +39,5 @@ void FTileMapEditorCommands::RegisterCommands()
 	UI_COMMAND(MoveLayerUp, "Move layer up", "Move layer up", EUserInterfaceActionType::Button, FInputGesture());
 	UI_COMMAND(MoveLayerDown, "Move layer down", "Move layer down", EUserInterfaceActionType::Button, FInputGesture());
 }
+
+#undef LOCTEXT_NAMESPACE
\ No newline at end of file
diff --git a/Engine/Plugins/Media/MediaPlayerEditor/Source/MediaPlayerEditor/Private/Models/MediaPlayerEditorCommands.cpp b/Engine/Plugins/Media/MediaPlayerEditor/Source/MediaPlayerEditor/Private/Models/MediaPlayerEditorCommands.cpp
--- a/Engine/Plugins/Media/MediaPlayerEditor/Source/MediaPlayerEditor/Private/Models/MediaPlayerEditorCommands.cpp
+++ b/Engine/Plugins/Media/MediaPlayerEditor/Source/MediaPlayerEditor/Private/Models/MediaPlayerEditorCommands.cpp
@@ -2,6 +2,7 @@
 
 #include "MediaPlayerEditorPrivatePCH.h"
 
+#define LOCTEXT_NAMESPACE ""
 
 void FMediaPlayerEditorCommands::RegisterCommands()
 {
@@ -11,3 +12,5 @@ void FMediaPlayerEditorCommands::RegisterCommands()
 	UI_COMMAND(ReverseMedia, "Reverse", "Reverses media playback", EUserInterfaceActionType::Button, FInputGesture());
 	UI_COMMAND(RewindMedia, "Rewind", "Rewinds the media to the beginning", EUserInterfaceActionType::Button, FInputGesture());
 }
+
+#undef LOCTEXT_NAMESPACE
\ No newline at end of file
diff --git a/Engine/Source/Developer/DeviceManager/Private/Models/DeviceDetailsCommands.h b/Engine/Source/Developer/DeviceManager/Private/Models/DeviceDetailsCommands.h
--- a/Engine/Source/Developer/DeviceManager/Private/Models/DeviceDetailsCommands.h
+++ b/Engine/Source/Developer/DeviceManager/Private/Models/DeviceDetailsCommands.h
@@ -2,6 +2,7 @@
 
 #pragma once
 
+#define LOCTEXT_NAMESPACE ""
 
 /**
  * The device details commands
@@ -60,3 +61,5 @@ public:
 	TSharedPtr<FUICommandInfo> PowerOn;
 	TSharedPtr<FUICommandInfo> Reboot;
 };
+
+#undef LOCTEXT_NAMESPACE
\ No newline at end of file
diff --git a/Engine/Source/Developer/ProjectLauncher/Private/Models/ProjectLauncherCommands.h b/Engine/Source/Developer/ProjectLauncher/Private/Models/ProjectLauncherCommands.h
--- a/Engine/Source/Developer/ProjectLauncher/Private/Models/ProjectLauncherCommands.h
+++ b/Engine/Source/Developer/ProjectLauncher/Private/Models/ProjectLauncherCommands.h
@@ -2,6 +2,7 @@
 
 #pragma once
 
+#define LOCTEXT_NAMESPACE ""
 
 /**
  * The project launcher UI commands
@@ -58,3 +59,5 @@ public:
 
 	TSharedPtr<FUICommandInfo> DeleteProfile;
 };
+
+#undef LOCTEXT_NAMESPACE
\ No newline at end of file
diff --git a/Engine/Source/Developer/SessionFrontend/Private/Models/SessionConsoleCommands.h b/Engine/Source/Developer/SessionFrontend/Private/Models/SessionConsoleCommands.h
--- a/Engine/Source/Developer/SessionFrontend/Private/Models/SessionConsoleCommands.h
+++ b/Engine/Source/Developer/SessionFrontend/Private/Models/SessionConsoleCommands.h
@@ -2,6 +2,7 @@
 
 #pragma once
 
+#define LOCTEXT_NAMESPACE ""
 
 /**
  * The device details commands.
@@ -37,3 +38,5 @@ public:
 	TSharedPtr<FUICommandInfo> SessionCopy;
 	TSharedPtr<FUICommandInfo> SessionSave;
 };
+
+#undef LOCTEXT_NAMESPACE
\ No newline at end of file
diff --git a/Engine/Source/Editor/Cascade/Private/CascadeActions.cpp b/Engine/Source/Editor/Cascade/Private/CascadeActions.cpp
--- a/Engine/Source/Editor/Cascade/Private/CascadeActions.cpp
+++ b/Engine/Source/Editor/Cascade/Private/CascadeActions.cpp
@@ -4,6 +4,8 @@
 #include "CascadeActions.h"
 #include "Cascade.h"
 
+#define LOCTEXT_NAMESPACE ""
+
 /** UI_COMMAND takes long for the compile to optimize */
 PRAGMA_DISABLE_OPTIMIZATION
 void FCascadeCommands::RegisterCommands()
@@ -75,4 +77,6 @@ void FCascadeCommands::RegisterCommands()
 	UI_COMMAND(RemoveDuplicateModules, "Remove Duplicate Modules", "Remove Duplicate Modules", EUserInterfaceActionType::Button, FInputGesture());
 }
 
-PRAGMA_ENABLE_OPTIMIZATION
\ No newline at end of file
+PRAGMA_ENABLE_OPTIMIZATION
+
+#undef LOCTEXT_NAMESPACE
\ No newline at end of file
diff --git a/Engine/Source/Editor/DistCurveEditor/Private/CurveEditorActions.cpp b/Engine/Source/Editor/DistCurveEditor/Private/CurveEditorActions.cpp
--- a/Engine/Source/Editor/DistCurveEditor/Private/CurveEditorActions.cpp
+++ b/Engine/Source/Editor/DistCurveEditor/Private/CurveEditorActions.cpp
@@ -3,6 +3,7 @@
 #include "DistCurveEditorModule.h"
 #include "CurveEditorActions.h"
 
+#define LOCTEXT_NAMESPACE ""
 
 void FCurveEditorCommands::RegisterCommands()
 {
@@ -35,3 +36,5 @@ void FCurveEditorCommands::RegisterCommands()
 	UI_COMMAND(CreateTab, "Create", "Create Tab", EUserInterfaceActionType::Button, FInputGesture());
 	UI_COMMAND(DeleteTab, "Delete", "Delete Tab", EUserInterfaceActionType::Button, FInputGesture());
 }
+
+#undef LOCTEXT_NAMESPACE
\ No newline at end of file
diff --git a/Engine/Source/Editor/FoliageEdit/Private/FoliageEditActions.cpp b/Engine/Source/Editor/FoliageEdit/Private/FoliageEditActions.cpp
--- a/Engine/Source/Editor/FoliageEdit/Private/FoliageEditActions.cpp
+++ b/Engine/Source/Editor/FoliageEdit/Private/FoliageEditActions.cpp
@@ -6,6 +6,8 @@
 
 #include "FoliageEditActions.h"
 
+#define LOCTEXT_NAMESPACE ""
+
 void FFoliageEditCommands::RegisterCommands()
 {
 	UI_COMMAND( SetPaint, "Paint", "Paint", EUserInterfaceActionType::ToggleButton, FInputGesture() );
@@ -24,3 +26,4 @@ void FFoliageEditCommands::RegisterCommands()
 	UI_COMMAND( DeselectAllInstances, "Deselect All Instances", "Deselect all instances of foliage type.", EUserInterfaceActionType::Button, FInputGesture());
 }
 
+#undef LOCTEXT_NAMESPACE
\ No newline at end of file
diff --git a/Engine/Source/Editor/GraphEditor/Private/GraphEditorActions.cpp b/Engine/Source/Editor/GraphEditor/Private/GraphEditorActions.cpp
--- a/Engine/Source/Editor/GraphEditor/Private/GraphEditorActions.cpp
+++ b/Engine/Source/Editor/GraphEditor/Private/GraphEditorActions.cpp
@@ -3,6 +3,7 @@
 #include "GraphEditorCommon.h"
 #include "GraphEditorActions.h"
 
+#define LOCTEXT_NAMESPACE ""
 
 void FGraphEditorCommandsImpl::RegisterCommands()
 {
@@ -93,3 +94,5 @@ void FGraphEditorCommands::Unregister()
 {
 	return FGraphEditorCommandsImpl::Unregister();
 }
+
+#undef LOCTEXT_NAMESPACE
\ No newline at end of file
diff --git a/Engine/Source/Editor/Kismet/Private/BlueprintEditorCommands.cpp b/Engine/Source/Editor/Kismet/Private/BlueprintEditorCommands.cpp
--- a/Engine/Source/Editor/Kismet/Private/BlueprintEditorCommands.cpp
+++ b/Engine/Source/Editor/Kismet/Private/BlueprintEditorCommands.cpp
@@ -11,6 +11,8 @@
 
 struct FBlueprintPaletteListBuilder;
 
+#define LOCTEXT_NAMESPACE ""
+
 /** UI_COMMAND takes long for the compile to optimize */
 PRAGMA_DISABLE_OPTIMIZATION
 void FBlueprintEditorCommands::RegisterCommands()
@@ -406,3 +408,5 @@ void FBlueprintSpawnNodeCommands::GetGraphActionByGesture(FInputGesture& InGestu
 		}
 	}
 }
+
+#undef LOCTEXT_NAMESPACE
\ No newline at end of file
diff --git a/Engine/Source/Editor/LandscapeEditor/Private/LandscapeEditorCommands.cpp b/Engine/Source/Editor/LandscapeEditor/Private/LandscapeEditorCommands.cpp
--- a/Engine/Source/Editor/LandscapeEditor/Private/LandscapeEditorCommands.cpp
+++ b/Engine/Source/Editor/LandscapeEditor/Private/LandscapeEditorCommands.cpp
@@ -4,6 +4,8 @@
 #include "LandscapeEditorPrivatePCH.h"
 #include "LandscapeEditorCommands.h"
 
+#define LOCTEXT_NAMESPACE ""
+
 void FLandscapeEditorCommands::RegisterCommands()
 {
 	UI_COMMAND(ManageMode, "Mode - Manage", "", EUserInterfaceActionType::RadioButton, FInputGesture());
@@ -84,3 +86,5 @@ void FLandscapeEditorCommands::RegisterCommands()
 	UI_COMMAND(ViewModeLayerDebug, "Layer Debug", "", EUserInterfaceActionType::RadioButton, FInputGesture());
 	UI_COMMAND(ViewModeWireframeOnTop, "Wireframe on Top", "", EUserInterfaceActionType::RadioButton, FInputGesture());
 }
+
+#undef LOCTEXT_NAMESPACE
\ No newline at end of file
diff --git a/Engine/Source/Editor/Layers/Private/LayerCollectionViewCommands.h b/Engine/Source/Editor/Layers/Private/LayerCollectionViewCommands.h
--- a/Engine/Source/Editor/Layers/Private/LayerCollectionViewCommands.h
+++ b/Engine/Source/Editor/Layers/Private/LayerCollectionViewCommands.h
@@ -2,6 +2,8 @@
 
 #pragma once
 
+#define LOCTEXT_NAMESPACE ""
+
 /**
  * The set of commands supported by the LayersView
  */
@@ -74,4 +76,4 @@ public:
 	TSharedPtr< FUICommandInfo > RequestRenameLayer;
 };
 
-
+#undef LOCTEXT_NAMESPACE
\ No newline at end of file
diff --git a/Engine/Source/Editor/Levels/Private/LevelCollectionViewCommands.h b/Engine/Source/Editor/Levels/Private/LevelCollectionViewCommands.h
--- a/Engine/Source/Editor/Levels/Private/LevelCollectionViewCommands.h
+++ b/Engine/Source/Editor/Levels/Private/LevelCollectionViewCommands.h
@@ -2,6 +2,8 @@
 
 #pragma once
 
+#define LOCTEXT_NAMESPACE ""
+
 /** The set of commands supported by the LevelsView */
 class FLevelsViewCommands : public TCommands<FLevelsViewCommands>
 {
@@ -246,3 +248,5 @@ public:
 	/** Unlocks all read-only levels */
 	TSharedPtr< FUICommandInfo > UnlockReadOnlyLevels;
 };
+
+#undef LOCTEXT_NAMESPACE
\ No newline at end of file
diff --git a/Engine/Source/Editor/MaterialEditor/Private/MaterialEditorActions.cpp b/Engine/Source/Editor/MaterialEditor/Private/MaterialEditorActions.cpp
--- a/Engine/Source/Editor/MaterialEditor/Private/MaterialEditorActions.cpp
+++ b/Engine/Source/Editor/MaterialEditor/Private/MaterialEditorActions.cpp
@@ -6,6 +6,7 @@
 #include "MaterialEditorActions.h"
 #include "SMaterialEditorViewport.h"
 
+#define LOCTEXT_NAMESPACE ""
 
 void FMaterialEditorCommands::RegisterCommands()
 {
@@ -170,3 +171,4 @@ const TSharedPtr<const FInputGesture> FMaterialEditorSpawnNodeCommands::GetGestu
 	return TSharedPtr< const FInputGesture >();
 }
 
+#undef LOCTEXT_NAMESPACE
\ No newline at end of file
diff --git a/Engine/Source/Editor/Matinee/Private/MatineeActions.cpp b/Engine/Source/Editor/Matinee/Private/MatineeActions.cpp
--- a/Engine/Source/Editor/Matinee/Private/MatineeActions.cpp
+++ b/Engine/Source/Editor/Matinee/Private/MatineeActions.cpp
@@ -4,6 +4,8 @@
 #include "MatineeActions.h"
 #include "Matinee.h"
 
+#define LOCTEXT_NAMESPACE ""
+
 /** UI_COMMAND takes long for the compile to optimize */
 PRAGMA_DISABLE_OPTIMIZATION
 void FMatineeCommands::RegisterCommands()
@@ -187,4 +189,6 @@ void FMatineeCommands::RegisterCommands()
 	UI_COMMAND(ChangeKeyInterpModeConstant,"Change Key Interp Mode Constant","",EUserInterfaceActionType::Button,FInputGesture(EModifierKey::None, EKeys::Five));
 	UI_COMMAND(DeleteSelection, "Delete Selection", "", EUserInterfaceActionType::Button, FInputGesture(EModifierKey::None, EKeys::Platform_Delete));
 }
-PRAGMA_ENABLE_OPTIMIZATION
\ No newline at end of file
+PRAGMA_ENABLE_OPTIMIZATION
+
+#undef LOCTEXT_NAMESPACE
\ No newline at end of file
diff --git a/Engine/Source/Editor/PListEditor/Public/SPlistEditor.h b/Engine/Source/Editor/PListEditor/Public/SPlistEditor.h
--- a/Engine/Source/Editor/PListEditor/Public/SPlistEditor.h
+++ b/Engine/Source/Editor/PListEditor/Public/SPlistEditor.h
@@ -11,6 +11,7 @@
 #include "../Private/PListNodeFile.h"
 #include "../Private/PListNodeString.h"
 
+#define LOCTEXT_NAMESPACE ""
 
 // Type of notification to spawn
 enum ENTF_Types
@@ -261,3 +262,5 @@ private:
 	/** Helper function to search through nodes to find a specific node's parent */
 	bool FindParent(const TSharedPtr<IPListNode>& InChildNode, TSharedPtr<IPListNode>& OutFoundNode) const;
 };
+
+#undef LOCTEXT_NAMESPACE
\ No newline at end of file
diff --git a/Engine/Source/Editor/Persona/Private/AnimViewportLODCommands.cpp b/Engine/Source/Editor/Persona/Private/AnimViewportLODCommands.cpp
--- a/Engine/Source/Editor/Persona/Private/AnimViewportLODCommands.cpp
+++ b/Engine/Source/Editor/Persona/Private/AnimViewportLODCommands.cpp
@@ -3,6 +3,8 @@
 #include "PersonaPrivatePCH.h"
 #include "AnimViewportLODCommands.h"
 
+#define LOCTEXT_NAMESPACE ""
+
 void FAnimViewportLODCommands::RegisterCommands()
 {
 	UI_COMMAND( LODAuto, "LOD Auto", "Automatically select LOD", EUserInterfaceActionType::RadioButton, FInputGesture() );
@@ -12,3 +14,4 @@ void FAnimViewportLODCommands::RegisterCommands()
 	UI_COMMAND( LOD3, "LOD 3", "Force select LOD 3", EUserInterfaceActionType::RadioButton, FInputGesture() );
 }
 
+#undef LOCTEXT_NAMESPACE
\ No newline at end of file
diff --git a/Engine/Source/Editor/Persona/Private/AnimViewportMenuCommands.cpp b/Engine/Source/Editor/Persona/Private/AnimViewportMenuCommands.cpp
--- a/Engine/Source/Editor/Persona/Private/AnimViewportMenuCommands.cpp
+++ b/Engine/Source/Editor/Persona/Private/AnimViewportMenuCommands.cpp
@@ -3,6 +3,8 @@
 #include "PersonaPrivatePCH.h"
 #include "AnimViewportMenuCommands.h"
 
+#define LOCTEXT_NAMESPACE ""
+
 void FAnimViewportMenuCommands::RegisterCommands()
 {
 	UI_COMMAND( Auto, "Auto", "Preview what is active on the current window", EUserInterfaceActionType::RadioButton, FInputGesture() );
@@ -19,3 +21,4 @@ void FAnimViewportMenuCommands::RegisterCommands()
 	UI_COMMAND( AnimSetDrawUVs, "UV", "Toggles display of the mesh's UVs for the specified channel.", EUserInterfaceActionType::ToggleButton, FInputGesture() );
 }
 
+#undef LOCTEXT_NAMESPACE
\ No newline at end of file
diff --git a/Engine/Source/Editor/Persona/Private/AnimViewportPlaybackCommands.cpp b/Engine/Source/Editor/Persona/Private/AnimViewportPlaybackCommands.cpp
--- a/Engine/Source/Editor/Persona/Private/AnimViewportPlaybackCommands.cpp
+++ b/Engine/Source/Editor/Persona/Private/AnimViewportPlaybackCommands.cpp
@@ -4,6 +4,8 @@
 #include "AnimViewportPlaybackCommands.h"
 #include "SAnimationEditorViewport.h"
 
+#define LOCTEXT_NAMESPACE ""
+
 FAnimViewportPlaybackCommands::FAnimViewportPlaybackCommands() : TCommands<FAnimViewportPlaybackCommands>
 	(
 	TEXT("AnimViewportPlayback"), // Context name for fast lookup
@@ -26,3 +28,4 @@ void FAnimViewportPlaybackCommands::RegisterCommands()
 	UI_COMMAND( PlaybackSpeedCommands[EAnimationPlaybackSpeeds::TenTimes],	"x10.0", "Set the animation playback speed to ten times the normal speed", EUserInterfaceActionType::RadioButton, FInputGesture() );
 }
 
+#undef LOCTEXT_NAMESPACE
\ No newline at end of file
diff --git a/Engine/Source/Editor/Persona/Private/AnimViewportShowCommands.cpp b/Engine/Source/Editor/Persona/Private/AnimViewportShowCommands.cpp
--- a/Engine/Source/Editor/Persona/Private/AnimViewportShowCommands.cpp
+++ b/Engine/Source/Editor/Persona/Private/AnimViewportShowCommands.cpp
@@ -3,6 +3,8 @@
 #include "PersonaPrivatePCH.h"
 #include "AnimViewportShowCommands.h"
 
+#define LOCTEXT_NAMESPACE ""
+
 void FAnimViewportShowCommands::RegisterCommands()
 {
 	UI_COMMAND( ToggleGrid, "Grid", "Display Grid", EUserInterfaceActionType::ToggleButton, FInputGesture() );
@@ -61,3 +63,5 @@ void FAnimViewportShowCommands::RegisterCommands()
 
 #endif// #if WITH_APEX_CLOTHING
 }
+
+#undef LOCTEXT_NAMESPACE
\ No newline at end of file
diff --git a/Engine/Source/Editor/Persona/Private/PersonaCommands.cpp b/Engine/Source/Editor/Persona/Private/PersonaCommands.cpp
--- a/Engine/Source/Editor/Persona/Private/PersonaCommands.cpp
+++ b/Engine/Source/Editor/Persona/Private/PersonaCommands.cpp
@@ -4,6 +4,8 @@
 #include "PersonaPrivatePCH.h"
 #include "PersonaCommands.h"
 
+#define LOCTEXT_NAMESPACE ""
+
 void FPersonaCommands::RegisterCommands()
 {
 	// save menu
@@ -34,3 +36,5 @@ void FPersonaCommands::RegisterCommands()
 
 	UI_COMMAND( UpdateSkeletonRefPose, "Update Skeleton RefPose", "Update Skeleton ref pose based on current preview mesh", EUserInterfaceActionType::Button, FInputGesture() );
 }
+
+#undef LOCTEXT_NAMESPACE
\ No newline at end of file
diff --git a/Engine/Source/Editor/Persona/Private/SkeletonTreeCommands.cpp b/Engine/Source/Editor/Persona/Private/SkeletonTreeCommands.cpp
--- a/Engine/Source/Editor/Persona/Private/SkeletonTreeCommands.cpp
+++ b/Engine/Source/Editor/Persona/Private/SkeletonTreeCommands.cpp
@@ -3,6 +3,8 @@
 #include "PersonaPrivatePCH.h"
 #include "SkeletonTreeCommands.h"
 
+#define LOCTEXT_NAMESPACE ""
+
 void FSkeletonTreeCommands::RegisterCommands()
 {
 	UI_COMMAND( ShowAllBones, "Show All Bones", "Show every bone in the skeleton", EUserInterfaceActionType::RadioButton, FInputGesture() );
@@ -28,3 +30,5 @@ void FSkeletonTreeCommands::RegisterCommands()
 	UI_COMMAND( ShowSkeletonSockets, "Show Skeleton Sockets", "Show sockets that are in the skeleton only", EUserInterfaceActionType::RadioButton, FInputGesture() );
 	UI_COMMAND( HideSockets, "Hide Sockets", "Show no sockets", EUserInterfaceActionType::RadioButton, FInputGesture() );
 }
+
+#undef LOCTEXT_NAMESPACE
\ No newline at end of file
diff --git a/Engine/Source/Editor/PhAT/Private/PhATActions.cpp b/Engine/Source/Editor/PhAT/Private/PhATActions.cpp
--- a/Engine/Source/Editor/PhAT/Private/PhATActions.cpp
+++ b/Engine/Source/Editor/PhAT/Private/PhATActions.cpp
@@ -4,6 +4,7 @@
 #include "PhATActions.h"
 #include "PhAT.h"
 
+#define LOCTEXT_NAMESPACE ""
 
 void FPhATCommands::RegisterCommands()
 {
@@ -87,3 +88,5 @@ void FPhATCommands::RegisterCommands()
 	UI_COMMAND(ToggleTwist, "Toggle Twist Constraint", "", EUserInterfaceActionType::Button, FInputGesture(EKeys::Two));
 	UI_COMMAND(FocusOnSelection, "Focus the viewport on the current selection", "", EUserInterfaceActionType::Button, FInputGesture(EKeys::F));
 }
+
+#undef LOCTEXT_NAMESPACE
\ No newline at end of file
diff --git a/Engine/Source/Editor/ReferenceViewer/Private/ReferenceViewerActions.cpp b/Engine/Source/Editor/ReferenceViewer/Private/ReferenceViewerActions.cpp
--- a/Engine/Source/Editor/ReferenceViewer/Private/ReferenceViewerActions.cpp
+++ b/Engine/Source/Editor/ReferenceViewer/Private/ReferenceViewerActions.cpp
@@ -6,6 +6,8 @@
 //////////////////////////////////////////////////////////////////////////
 // FReferenceViewerActions
 
+#define LOCTEXT_NAMESPACE ""
+
 void FReferenceViewerActions::RegisterCommands()
 {
 	UI_COMMAND(OpenSelectedInAssetEditor, "Edit Asset", "Opens the selected asset in the relevent editor.", EUserInterfaceActionType::Button, FInputGesture());
@@ -18,3 +20,5 @@ void FReferenceViewerActions::RegisterCommands()
 	UI_COMMAND(MakePrivateCollectionWithReferencedAssets, "Private", "Private. This collection is only visible to you.", EUserInterfaceActionType::Button, FInputGesture());
 	UI_COMMAND(MakeSharedCollectionWithReferencedAssets, "Shared", "Shared. This collection is visible to everyone.", EUserInterfaceActionType::Button, FInputGesture());
 }
+
+#undef LOCTEXT_NAMESPACE
\ No newline at end of file
diff --git a/Engine/Source/Editor/Sequencer/Private/SequencerCommands.cpp b/Engine/Source/Editor/Sequencer/Private/SequencerCommands.cpp
--- a/Engine/Source/Editor/Sequencer/Private/SequencerCommands.cpp
+++ b/Engine/Source/Editor/Sequencer/Private/SequencerCommands.cpp
@@ -3,6 +3,8 @@
 #include "SequencerPrivatePCH.h"
 #include "SequencerCommands.h"
 
+#define LOCTEXT_NAMESPACE ""
+
 void FSequencerCommands::RegisterCommands()
 {
 	UI_COMMAND( SetKey, "Set Key", "Sets a key at the current time for the selected actor.", EUserInterfaceActionType::Button, FInputGesture(EKeys::K) );
@@ -29,3 +31,5 @@ void FSequencerCommands::RegisterCommands()
 	UI_COMMAND( SetSelectedCurveVisibility, "Selected Curves", "Show only selected curves in the curve editor.", EUserInterfaceActionType::RadioButton, FInputGesture() );
 	UI_COMMAND( SetAnimatedCurveVisibility, "Animated Curves", "Show only curves with keyframes in the curve editor.", EUserInterfaceActionType::RadioButton, FInputGesture() );
 }
+
+#undef LOCTEXT_NAMESPACE
\ No newline at end of file
diff --git a/Engine/Source/Editor/StaticMeshEditor/Private/StaticMeshEditorActions.cpp b/Engine/Source/Editor/StaticMeshEditor/Private/StaticMeshEditorActions.cpp
--- a/Engine/Source/Editor/StaticMeshEditor/Private/StaticMeshEditorActions.cpp
+++ b/Engine/Source/Editor/StaticMeshEditor/Private/StaticMeshEditorActions.cpp
@@ -5,6 +5,8 @@
 #include "StaticMeshEditorActions.h"
 #include "SStaticMeshEditorViewport.h"
 
+#define LOCTEXT_NAMESPACE ""
+
 void FStaticMeshEditorCommands::RegisterCommands()
 {
 	UI_COMMAND( SetShowWireframe, "Wireframe", "Toggles the viewmode of the Preview Pane between a lit view and a wireframe view.", EUserInterfaceActionType::ToggleButton, FInputGesture() );
@@ -45,3 +47,4 @@ void FStaticMeshEditorCommands::RegisterCommands()
 	UI_COMMAND( SaveGeneratedLODs, "Save Generated LODs", "Saves the generated LODs in the mesh package.", EUserInterfaceActionType::Button, FInputGesture());
 }
 
+#undef LOCTEXT_NAMESPACE
\ No newline at end of file
diff --git a/Engine/Source/Editor/TextureEditor/Private/Models/TextureEditorCommands.cpp b/Engine/Source/Editor/TextureEditor/Private/Models/TextureEditorCommands.cpp
--- a/Engine/Source/Editor/TextureEditor/Private/Models/TextureEditorCommands.cpp
+++ b/Engine/Source/Editor/TextureEditor/Private/Models/TextureEditorCommands.cpp
@@ -2,6 +2,7 @@
 
 #include "TextureEditorPrivatePCH.h"
 
+#define LOCTEXT_NAMESPACE ""
 
 void FTextureEditorCommands::RegisterCommands()
 {
@@ -21,3 +22,5 @@ void FTextureEditorCommands::RegisterCommands()
 	UI_COMMAND(Reimport, "Reimport", "Reimports the texture from file", EUserInterfaceActionType::Button, FInputGesture());
 	UI_COMMAND(Settings, "Settings...", "Opens the settings for the texture editor", EUserInterfaceActionType::Button, FInputGesture());
 }
+
+#undef LOCTEXT_NAMESPACE
\ No newline at end of file
diff --git a/Engine/Source/Editor/UMGEditor/Private/Designer/DesignerCommands.cpp b/Engine/Source/Editor/UMGEditor/Private/Designer/DesignerCommands.cpp
--- a/Engine/Source/Editor/UMGEditor/Private/Designer/DesignerCommands.cpp
+++ b/Engine/Source/Editor/UMGEditor/Private/Designer/DesignerCommands.cpp
@@ -3,6 +3,8 @@
 #include "UMGEditorPrivatePCH.h"
 #include "DesignerCommands.h"
 
+#define LOCTEXT_NAMESPACE ""
+
 void FDesignerCommands::RegisterCommands()
 {
 	UI_COMMAND( LayoutTransform, "Layout Transform Mode", "Adjust widget layout transform", EUserInterfaceActionType::ToggleButton, FInputGesture(EKeys::W) );
@@ -11,3 +13,5 @@ void FDesignerCommands::RegisterCommands()
 	UI_COMMAND( LocationGridSnap, "Grid Snap", "Enables or disables snapping to the grid when dragging objects around", EUserInterfaceActionType::ToggleButton, FInputGesture() );
 	UI_COMMAND( RotationGridSnap, "Rotation Snap", "Enables or disables snapping objects to a rotation grid", EUserInterfaceActionType::ToggleButton, FInputGesture() );
 }
+
+#undef LOCTEXT_NAMESPACE
\ No newline at end of file
diff --git a/Engine/Source/Editor/UnrealEd/Private/EditorViewportCommands.cpp b/Engine/Source/Editor/UnrealEd/Private/EditorViewportCommands.cpp
--- a/Engine/Source/Editor/UnrealEd/Private/EditorViewportCommands.cpp
+++ b/Engine/Source/Editor/UnrealEd/Private/EditorViewportCommands.cpp
@@ -4,7 +4,7 @@
 #include "UnrealEd.h"
 #include "EditorViewportCommands.h"
 
-
+#define LOCTEXT_NAMESPACE ""
 
 void FEditorViewportCommands::RegisterCommands()
 {
@@ -75,7 +75,6 @@ void FEditorViewportCommands::RegisterCommands()
 	UI_COMMAND( FixedExposure2p, "Fixed Exposure: +2", "Set the fixed exposure to 2", EUserInterfaceActionType::RadioButton, FInputGesture() );
 	UI_COMMAND( FixedExposure3p, "Fixed Exposure: +3", "Set the fixed exposure to 3", EUserInterfaceActionType::RadioButton, FInputGesture() );
 	UI_COMMAND( FixedExposure4p, "Fixed Exposure: +4", "Set the fixed exposure to 4", EUserInterfaceActionType::RadioButton, FInputGesture() );
-
-
 }
 
+#undef LOCTEXT_NAMESPACE
\ No newline at end of file
diff --git a/Engine/Source/Editor/UnrealEd/Private/RichCurveEditorCommands.cpp b/Engine/Source/Editor/UnrealEd/Private/RichCurveEditorCommands.cpp
--- a/Engine/Source/Editor/UnrealEd/Private/RichCurveEditorCommands.cpp
+++ b/Engine/Source/Editor/UnrealEd/Private/RichCurveEditorCommands.cpp
@@ -3,6 +3,8 @@
 #include "UnrealEd.h"
 #include "RichCurveEditorCommands.h"
 
+#define LOCTEXT_NAMESPACE ""
+
 void FRichCurveEditorCommands::RegisterCommands()
 {
 	UI_COMMAND(ZoomToFitHorizontal, "Fit Horizontal", "Zoom to Fit - Horizontal", EUserInterfaceActionType::Button, FInputGesture());
@@ -18,3 +20,5 @@ void FRichCurveEditorCommands::RegisterCommands()
 	UI_COMMAND(InterpolationCubicUser, "User", "Cubic interpolation - User flat tangents", EUserInterfaceActionType::ToggleButton, FInputGesture());
 	UI_COMMAND(InterpolationCubicBreak, "Break", "Cubic interpolation - User broken tangents", EUserInterfaceActionType::ToggleButton, FInputGesture());
 }
+
+#undef LOCTEXT_NAMESPACE
\ No newline at end of file
diff --git a/Engine/Source/Editor/UnrealEd/Private/SoundCueGraphEditorCommands.cpp b/Engine/Source/Editor/UnrealEd/Private/SoundCueGraphEditorCommands.cpp
--- a/Engine/Source/Editor/UnrealEd/Private/SoundCueGraphEditorCommands.cpp
+++ b/Engine/Source/Editor/UnrealEd/Private/SoundCueGraphEditorCommands.cpp
@@ -3,6 +3,8 @@
 #include "UnrealEd.h"
 #include "SoundCueGraphEditorCommands.h"
 
+#define LOCTEXT_NAMESPACE ""
+
 void FSoundCueGraphEditorCommands::RegisterCommands()
 {
 	UI_COMMAND(PlayCue, "Play Cue", "Plays the SoundCue", EUserInterfaceActionType::Button, FInputGesture());
@@ -14,3 +16,5 @@ void FSoundCueGraphEditorCommands::RegisterCommands()
 	UI_COMMAND(AddInput, "Add Input", "Adds an input to the node", EUserInterfaceActionType::Button, FInputGesture());
 	UI_COMMAND(DeleteInput, "Delete Input", "Removes an input from the node", EUserInterfaceActionType::Button, FInputGesture());
 }
+
+#undef LOCTEXT_NAMESPACE
\ No newline at end of file
diff --git a/Engine/Source/Editor/UnrealEd/Private/Toolkits/AssetEditorCommonCommands.cpp b/Engine/Source/Editor/UnrealEd/Private/Toolkits/AssetEditorCommonCommands.cpp
--- a/Engine/Source/Editor/UnrealEd/Private/Toolkits/AssetEditorCommonCommands.cpp
+++ b/Engine/Source/Editor/UnrealEd/Private/Toolkits/AssetEditorCommonCommands.cpp
@@ -3,7 +3,7 @@
 #include "UnrealEd.h"
 #include "AssetEditorCommonCommands.h"
 
-
+#define LOCTEXT_NAMESPACE ""
 
 void FAssetEditorCommonCommands::RegisterCommands()
 {
@@ -12,3 +12,5 @@ void FAssetEditorCommonCommands::RegisterCommands()
 	UI_COMMAND( SwitchToStandaloneEditor, "Switch to Standalone Editor", "Closes the level-centric asset editor and reopens it in 'standalone' mode", EUserInterfaceActionType::Button, FInputGesture() );
 	UI_COMMAND( SwitchToWorldCentricEditor, "Switch to World-Centric Editor", "Closes the standalone asset editor and reopens it in 'world-centric' mode, docked within the level editor that it was originally opened in.", EUserInterfaceActionType::Button, FInputGesture() );
 }
+
+#undef LOCTEXT_NAMESPACE
\ No newline at end of file
diff --git a/Engine/Source/Editor/UnrealEd/Private/Toolkits/GlobalEditorCommonCommands.cpp b/Engine/Source/Editor/UnrealEd/Private/Toolkits/GlobalEditorCommonCommands.cpp
--- a/Engine/Source/Editor/UnrealEd/Private/Toolkits/GlobalEditorCommonCommands.cpp
+++ b/Engine/Source/Editor/UnrealEd/Private/Toolkits/GlobalEditorCommonCommands.cpp
@@ -87,6 +87,8 @@ FReply SGlobalOpenAssetDialog::OnPreviewKeyDown(const FGeometry& MyGeometry, con
 //////////////////////////////////////////////////////////////////////////
 // FGlobalEditorCommonCommands
 
+#define LOCTEXT_NAMESPACE ""
+
 FGlobalEditorCommonCommands::FGlobalEditorCommonCommands()
 	: TCommands<FGlobalEditorCommonCommands>(TEXT("SystemWideCommands"), NSLOCTEXT("Contexts", "SystemWideCommands", "System-wide"), NAME_None, FEditorStyle::GetStyleSetName())
 {
@@ -183,4 +185,6 @@ void FGlobalEditorCommonCommands::OnSummonedConsoleCommandBox()
 
 		OutputLogModule.ToggleDebugConsoleForWindow(WindowRef, EDebugConsoleStyle::Compact, Delegates);
 	}
-}
\ No newline at end of file
+}
+
+#undef LOCTEXT_NAMESPACE
\ No newline at end of file
diff --git a/Engine/Source/Editor/WorldBrowser/Private/LevelCollectionCommands.h b/Engine/Source/Editor/WorldBrowser/Private/LevelCollectionCommands.h
--- a/Engine/Source/Editor/WorldBrowser/Private/LevelCollectionCommands.h
+++ b/Engine/Source/Editor/WorldBrowser/Private/LevelCollectionCommands.h
@@ -1,6 +1,8 @@
 // Copyright 1998-2015 Epic Games, Inc. All Rights Reserved.
 #pragma once
 
+#define LOCTEXT_NAMESPACE ""
+
 /** The set of commands supported by the WorldView */
 class FLevelCollectionCommands 
 	: public TCommands<FLevelCollectionCommands>
@@ -297,3 +299,4 @@ public:
 	TSharedPtr< FUICommandInfo > AddLandscapeLevelYPositive;
 };
 
+#undef LOCTEXT_NAMESPACE
\ No newline at end of file
diff --git a/Engine/Source/Runtime/Slate/Private/Framework/Commands/Commands.cpp b/Engine/Source/Runtime/Slate/Private/Framework/Commands/Commands.cpp
--- a/Engine/Source/Runtime/Slate/Private/Framework/Commands/Commands.cpp
+++ b/Engine/Source/Runtime/Slate/Private/Framework/Commands/Commands.cpp
@@ -2,19 +2,19 @@
 
 #include "SlatePrivatePCH.h"
 
-
 #define LOC_DEFINE_REGION
 
-void UI_COMMAND_Function( FBindingContext* This, TSharedPtr< FUICommandInfo >& OutCommand, const TCHAR* OutCommandName, const TCHAR* OutCommandNameUnderscoreTooltip, const ANSICHAR* DotOutCommandName, const TCHAR* FriendlyName, const TCHAR* InDescription, const EUserInterfaceActionType::Type CommandType, const FInputGesture& InDefaultGesture )
+void UI_COMMAND_Function( FBindingContext* This, TSharedPtr< FUICommandInfo >& OutCommand, const TCHAR* OutSubNamespace, const TCHAR* OutCommandName, const TCHAR* OutCommandNameUnderscoreTooltip, const ANSICHAR* DotOutCommandName, const TCHAR* FriendlyName, const TCHAR* InDescription, const EUserInterfaceActionType::Type CommandType, const FInputGesture& InDefaultGesture )
 {
 	static const FString UICommandsStr(TEXT("UICommands"));
+	const FString Namespace = OutSubNamespace && FCString::Strlen(OutSubNamespace) > 0 ?  UICommandsStr + TEXT(".") + OutSubNamespace : UICommandsStr;
 
 	FUICommandInfo::MakeCommandInfo(
 		This->AsShared(),
 		OutCommand,
 		OutCommandName,
-		FInternationalization::ForUseOnlyByLocMacroAndGraphNodeTextLiterals_CreateText( FriendlyName, *UICommandsStr, OutCommandName ),
-		FInternationalization::ForUseOnlyByLocMacroAndGraphNodeTextLiterals_CreateText( InDescription, *UICommandsStr, OutCommandNameUnderscoreTooltip ),
+		FInternationalization::ForUseOnlyByLocMacroAndGraphNodeTextLiterals_CreateText( FriendlyName, *Namespace, OutCommandName ),
+		FInternationalization::ForUseOnlyByLocMacroAndGraphNodeTextLiterals_CreateText( InDescription, *Namespace, OutCommandNameUnderscoreTooltip ),
 		FSlateIcon( This->GetStyleSetName(), ISlateStyle::Join( This->GetContextName(), DotOutCommandName ) ),
 		CommandType,
 		InDefaultGesture
diff --git a/Engine/Source/Runtime/Slate/Private/Framework/Commands/GenericCommands.cpp b/Engine/Source/Runtime/Slate/Private/Framework/Commands/GenericCommands.cpp
--- a/Engine/Source/Runtime/Slate/Private/Framework/Commands/GenericCommands.cpp
+++ b/Engine/Source/Runtime/Slate/Private/Framework/Commands/GenericCommands.cpp
@@ -3,6 +3,8 @@
 #include "SlatePrivatePCH.h"
 #include "GenericCommands.h"
 
+#define LOCTEXT_NAMESPACE ""
+
 void FGenericCommands::RegisterCommands()
 {
 	UI_COMMAND( Cut, "Cut", "Cut selection", EUserInterfaceActionType::Button, FInputGesture( EModifierKey::Control, EKeys::X ) )
@@ -21,3 +23,5 @@ void FGenericCommands::RegisterCommands()
 	UI_COMMAND( Rename, "Rename", "Rename current selection", EUserInterfaceActionType::Button, FInputGesture( EKeys::F2 ) )
 	UI_COMMAND( SelectAll, "Select All", "Select everything in the current scope", EUserInterfaceActionType::Button, FInputGesture( EModifierKey::Control, EKeys::A ) )
 }
+
+#undef LOCTEXT_NAMESPACE
\ No newline at end of file
diff --git a/Engine/Source/Runtime/Slate/Public/Framework/Commands/Commands.h b/Engine/Source/Runtime/Slate/Public/Framework/Commands/Commands.h
--- a/Engine/Source/Runtime/Slate/Public/Framework/Commands/Commands.h
+++ b/Engine/Source/Runtime/Slate/Public/Framework/Commands/Commands.h
@@ -11,10 +11,11 @@
 
 #define LOC_DEFINE_REGION
 
-SLATE_API void UI_COMMAND_Function( FBindingContext* This, TSharedPtr< FUICommandInfo >& OutCommand, const TCHAR* OutCommandName, const TCHAR* OutCommandNameUnderscoreTooltip, const ANSICHAR* DotOutCommandName, const TCHAR* FriendlyName, const TCHAR* InDescription, const EUserInterfaceActionType::Type CommandType, const FInputGesture& InDefaultGesture );
+SLATE_API void UI_COMMAND_Function( FBindingContext* This, TSharedPtr< FUICommandInfo >& OutCommand, const TCHAR* OutSubNamespace, const TCHAR* OutCommandName, const TCHAR* OutCommandNameUnderscoreTooltip, const ANSICHAR* DotOutCommandName, const TCHAR* FriendlyName, const TCHAR* InDescription, const EUserInterfaceActionType::Type CommandType, const FInputGesture& InDefaultGesture );
 
+// This macro requires LOCTEXT_NAMESPACE to be defined. If you don't want the command to be placed under a sub namespace, provide "" as the namespace.
 #define UI_COMMAND( CommandId, FriendlyName, InDescription, CommandType, InDefaultGesture ) \
-	UI_COMMAND_Function( this, CommandId, TEXT(#CommandId), TEXT(#CommandId) TEXT("_ToolTip"), "." #CommandId, TEXT(FriendlyName), TEXT(InDescription), CommandType, InDefaultGesture );
+	UI_COMMAND_Function( this, CommandId, TEXT(LOCTEXT_NAMESPACE), TEXT(#CommandId), TEXT(#CommandId) TEXT("_ToolTip"), "." #CommandId, TEXT(FriendlyName), TEXT(InDescription), CommandType, InDefaultGesture );
 
 #undef LOC_DEFINE_REGION
 
