TabManager : Hitting ESC will interrupt the current drag drop. NOTE: It will not restore the dragged widget to its original location.
Fixes: UE-1232.

Also, SceneOutliner no longer handles DockingDrag operations erroneously.

diff --git a/Engine/Source/Editor/SceneOutliner/Private/SOutlinerTreeView.cpp b/Engine/Source/Editor/SceneOutliner/Private/SOutlinerTreeView.cpp
--- a/Engine/Source/Editor/SceneOutliner/Private/SOutlinerTreeView.cpp
+++ b/Engine/Source/Editor/SceneOutliner/Private/SOutlinerTreeView.cpp
@@ -85,8 +85,7 @@ namespace SceneOutliner
 		// Validate now to make sure we don't doing anything we shouldn't
 		if (!DraggedObjects.ParseDrag(*DragDropEvent.GetOperation()))
 		{
-			// Return handled here to stop anything else trying to handle it - the operation is invalid as far as we're concerned
-			return FReply::Handled();
+			return FReply::Unhandled();
 		}
 
 		ValidationInfo = DropTarget.ValidateDrop(DraggedObjects, *SharedData.RepresentingWorld);
diff --git a/Engine/Source/Runtime/Slate/Private/Framework/Docking/FDockingDragOperation.cpp b/Engine/Source/Runtime/Slate/Private/Framework/Docking/FDockingDragOperation.cpp
--- a/Engine/Source/Runtime/Slate/Private/Framework/Docking/FDockingDragOperation.cpp
+++ b/Engine/Source/Runtime/Slate/Private/Framework/Docking/FDockingDragOperation.cpp
@@ -13,81 +13,12 @@ void FDockingDragOperation::OnDrop( bool bDropWasHandled, const FPointerEvent& M
 	check(CursorDecoratorWindow.IsValid());
 
 	const FVector2D WindowSize = CursorDecoratorWindow->GetSizeInScreen();
-	// Destroy the CursorDecoratorWindow by calling the base class implementation because we are relocating the content into a more permanent home.
-	FDragDropOperation::OnDrop(bDropWasHandled, MouseEvent);
 
 	TabBeingDragged->SetDraggedOverDockArea( NULL );
 
 	if (!bDropWasHandled)
 	{
-		// If we dropped the tab into an existing DockNode then it would have handled the DropEvent.
-		// We are here because that didn't happen, so make a new window with a new DockNode and drop the tab into that.
-
-		const FVector2D PositionToDrop = MouseEvent.GetScreenSpacePosition() - GetDecoratorOffsetFromCursor();
-
-		TSharedRef<FTabManager> MyTabManager = TabBeingDragged->GetTabManager();
-		
-		TSharedPtr<SWindow> NewWindowParent = MyTabManager->GetPrivateApi().GetParentWindow();
-
-		
-		TSharedRef<SWindow> NewWindow = SNew(SWindow)
-			.Title( FGlobalTabmanager::Get()->GetApplicationTitle() )
-			.AutoCenter(EAutoCenter::None)
-			.ScreenPosition( PositionToDrop )
-			// Make room for the title bar; otherwise windows will get progressive smaller whenver you float them.
-			.ClientSize( SWindow::ComputeWindowSizeForContent( WindowSize ) )
-			.CreateTitleBar(false);
-
-		TSharedPtr<SDockingTabStack> NewDockNode;
-
-		if ( TabBeingDragged->GetTabRole() == ETabRole::NomadTab )
-		{
-			TabBeingDragged->SetTabManager(FGlobalTabmanager::Get());
-		}
-
-		// Create a new dockarea
-		TSharedRef<SDockingArea> NewDockArea = 
-			SNew(SDockingArea, TabBeingDragged->GetTabManager(), FTabManager::NewPrimaryArea())
-			. ParentWindow( NewWindow )
-			. InitialContent
-			(
-				SAssignNew(NewDockNode, SDockingTabStack, FTabManager::NewStack())
-			);
-
-		if (TabBeingDragged->GetTabRole() == ETabRole::MajorTab || TabBeingDragged->GetTabRole() == ETabRole::NomadTab)
-		{
-			TSharedPtr<SWindow> RootWindow = FGlobalTabmanager::Get()->GetRootWindow();
-			if ( RootWindow.IsValid() )
-			{
-				// We have a root window, so all MajorTabs are nested under it.
-				FSlateApplication::Get().AddWindowAsNativeChild( NewWindow, RootWindow.ToSharedRef() )->SetContent(NewDockArea);
-			}
-			else
-			{
-				// App tabs get put in top-level windows. They show up on the taskbar.
-				FSlateApplication::Get().AddWindow( NewWindow )->SetContent(NewDockArea);
-			}
-		}
-		else
-		{
-			// Other tab types are placed in child windows. Their life is controlled by the top-level windows.
-			// They do not show up on the taskbar.
-
-			if ( NewWindowParent.IsValid() )
-			{
-				FSlateApplication::Get().AddWindowAsNativeChild( NewWindow, NewWindowParent.ToSharedRef() )->SetContent(NewDockArea);
-			}
-			else
-			{
-				FSlateApplication::Get().AddWindow( NewWindow )->SetContent(NewDockArea);
-			}
-		}
-
-		// Do this after the window parenting so that the window title is set correctly
-		NewDockNode->OpenTab(TabBeingDragged.ToSharedRef());
-
-		// Let every widget under this tab manager know that this tab has found a new home.
-		TabOwnerAreaOfOrigin->GetTabManager()->GetPrivateApi().OnTabRelocated( TabBeingDragged.ToSharedRef(), NewWindow );
+		DroppedOntoNothing();
 	}
 	else
 	{
@@ -103,6 +34,11 @@ void FDockingDragOperation::OnDrop( bool bDropWasHandled, const FPointerEvent& M
 
 		TabOwnerAreaOfOrigin->GetTabManager()->GetPrivateApi().OnTabRelocated( TabBeingDragged.ToSharedRef(), WindowDroppedInto );
 	}
+
+	// Destroy the CursorDecoratorWindow by calling the base class implementation because we are relocating the content into a more permanent home.
+	FDragDropOperation::OnDrop(bDropWasHandled, MouseEvent);
+
+	TabBeingDragged.Reset();
 }
 
 /** 
@@ -279,6 +215,10 @@ bool FDockingDragOperation::CanDockInNode(const TSharedRef<SDockingNode>& DockNo
 
 FDockingDragOperation::~FDockingDragOperation()
 {
+	if (TabBeingDragged.IsValid())
+	{
+		DroppedOntoNothing();
+	}
 }
 
 
@@ -351,3 +291,74 @@ FVector2D FDockingDragOperation::DesiredSizeFrom( const FVector2D& InitialTabSiz
 	return InitialTabSize * SizeCoefficient;
 }
 
+void FDockingDragOperation::DroppedOntoNothing()
+{
+	// If we dropped the tab into an existing DockNode then it would have handled the DropEvent.
+	// We are here because that didn't happen, so make a new window with a new DockNode and drop the tab into that.
+
+	const FVector2D PositionToDrop = CursorDecoratorWindow->GetPositionInScreen();
+
+	TSharedRef<FTabManager> MyTabManager = TabBeingDragged->GetTabManager();
+
+	TSharedPtr<SWindow> NewWindowParent = MyTabManager->GetPrivateApi().GetParentWindow();
+
+
+	TSharedRef<SWindow> NewWindow = SNew(SWindow)
+		.Title(FGlobalTabmanager::Get()->GetApplicationTitle())
+		.AutoCenter(EAutoCenter::None)
+		.ScreenPosition(PositionToDrop)
+		// Make room for the title bar; otherwise windows will get progressive smaller whenver you float them.
+		.ClientSize(SWindow::ComputeWindowSizeForContent(CursorDecoratorWindow->GetSizeInScreen()))
+		.CreateTitleBar(false);
+
+	TSharedPtr<SDockingTabStack> NewDockNode;
+
+	if (TabBeingDragged->GetTabRole() == ETabRole::NomadTab)
+	{
+		TabBeingDragged->SetTabManager(FGlobalTabmanager::Get());
+	}
+
+	// Create a new dockarea
+	TSharedRef<SDockingArea> NewDockArea =
+		SNew(SDockingArea, TabBeingDragged->GetTabManager(), FTabManager::NewPrimaryArea())
+		.ParentWindow(NewWindow)
+		.InitialContent
+		(
+		SAssignNew(NewDockNode, SDockingTabStack, FTabManager::NewStack())
+		);
+
+	if (TabBeingDragged->GetTabRole() == ETabRole::MajorTab || TabBeingDragged->GetTabRole() == ETabRole::NomadTab)
+	{
+		TSharedPtr<SWindow> RootWindow = FGlobalTabmanager::Get()->GetRootWindow();
+		if (RootWindow.IsValid())
+		{
+			// We have a root window, so all MajorTabs are nested under it.
+			FSlateApplication::Get().AddWindowAsNativeChild(NewWindow, RootWindow.ToSharedRef())->SetContent(NewDockArea);
+		}
+		else
+		{
+			// App tabs get put in top-level windows. They show up on the taskbar.
+			FSlateApplication::Get().AddWindow(NewWindow)->SetContent(NewDockArea);
+		}
+	}
+	else
+	{
+		// Other tab types are placed in child windows. Their life is controlled by the top-level windows.
+		// They do not show up on the taskbar.
+
+		if (NewWindowParent.IsValid())
+		{
+			FSlateApplication::Get().AddWindowAsNativeChild(NewWindow, NewWindowParent.ToSharedRef())->SetContent(NewDockArea);
+		}
+		else
+		{
+			FSlateApplication::Get().AddWindow(NewWindow)->SetContent(NewDockArea);
+		}
+	}
+
+	// Do this after the window parenting so that the window title is set correctly
+	NewDockNode->OpenTab(TabBeingDragged.ToSharedRef());
+
+	// Let every widget under this tab manager know that this tab has found a new home.
+	TabOwnerAreaOfOrigin->GetTabManager()->GetPrivateApi().OnTabRelocated(TabBeingDragged.ToSharedRef(), NewWindow);
+}
\ No newline at end of file
diff --git a/Engine/Source/Runtime/Slate/Private/Framework/Docking/FDockingDragOperation.h b/Engine/Source/Runtime/Slate/Private/Framework/Docking/FDockingDragOperation.h
--- a/Engine/Source/Runtime/Slate/Private/Framework/Docking/FDockingDragOperation.h
+++ b/Engine/Source/Runtime/Slate/Private/Framework/Docking/FDockingDragOperation.h
@@ -129,6 +129,9 @@ protected:
 	/** @return the size of the DockNode that looks good in a preview given the initial size of the tab that we grabbed. */
 	static FVector2D DesiredSizeFrom( const FVector2D& InitialTabSize );
 
+	/** The tab was dropped onto nothing or someone interrupted the drag drop operation. */
+	void DroppedOntoNothing();
+
 	/** What is actually being dragged in this operation */
 	TSharedPtr<class SDockTab> TabBeingDragged;
 
