Slate: Add FTabManager::GetOwnerTab() to return the tab that owns the specified tab manager

diff --git a/Engine/Source/Runtime/Slate/Public/Framework/Docking/TabManager.h b/Engine/Source/Runtime/Slate/Public/Framework/Docking/TabManager.h
--- a/Engine/Source/Runtime/Slate/Public/Framework/Docking/TabManager.h
+++ b/Engine/Source/Runtime/Slate/Public/Framework/Docking/TabManager.h
@@ -631,6 +631,9 @@ class SLATE_API FTabManager : public TSharedFromThis<FTabManager>
 		/** @return true if the tab has a factory registered for it that allows it to be spawned. */
 		bool CanSpawnTab( FName TabId );
 
+		/** Returns the owner tab (if it exists) */
+		TSharedPtr<SDockTab> GetOwnerTab() { return OwnerTabPtr.Pin(); }
+
 	protected:
 		void InvokeTabForMenu( FName TabId );
 
