Slate: Add tracking of the last activation time to SDockTab instances

diff --git a/Engine/Source/Runtime/Slate/Private/Framework/Docking/TabManager.cpp b/Engine/Source/Runtime/Slate/Private/Framework/Docking/TabManager.cpp
--- a/Engine/Source/Runtime/Slate/Private/Framework/Docking/TabManager.cpp
+++ b/Engine/Source/Runtime/Slate/Private/Framework/Docking/TabManager.cpp
@@ -1598,8 +1598,9 @@ void FGlobalTabmanager::SetActiveTab( const TSharedPtr<class SDockTab>& NewActiv
 	
 	TSharedPtr<SDockTab> CurrentlyActiveTab = GetActiveTab();
 
-	if (bShouldApplyChange && CurrentlyActiveTab != NewActiveTab)
+	if (bShouldApplyChange && (CurrentlyActiveTab != NewActiveTab))
 	{
+		NewActiveTab->UpdateActivationTime();
 		OnActiveTabChanged.Broadcast( CurrentlyActiveTab, NewActiveTab );
 		ActiveTabPtr = NewActiveTab;
 	}	
@@ -1763,9 +1764,11 @@ void FGlobalTabmanager::OnTabForegrounded( const TSharedPtr<SDockTab>& NewForegr
 			TSharedPtr<FTabManager> ForegroundTabManager = SubTabManagers[ForegroundedTabIndex].TabManager.Pin();
 			ForegroundTabManager->GetPrivateApi().ShowWindows();
 		}
+
+		NewForegroundTab->UpdateActivationTime();
 	}
 	
-	if ( BackgroundedTab.IsValid() )
+	if (BackgroundedTab.IsValid())
 	{
 		// Hide any child windows associated with the Major Tab that got backgrounded.
 		const int32 BackgroundedTabIndex = SubTabManagers.IndexOfByPredicate(FindByTab(BackgroundedTab.ToSharedRef()));
diff --git a/Engine/Source/Runtime/Slate/Private/Widgets/Docking/SDockTab.cpp b/Engine/Source/Runtime/Slate/Private/Widgets/Docking/SDockTab.cpp
--- a/Engine/Source/Runtime/Slate/Private/Widgets/Docking/SDockTab.cpp
+++ b/Engine/Source/Runtime/Slate/Private/Widgets/Docking/SDockTab.cpp
@@ -309,8 +309,8 @@ SDockTab::SDockTab()
 	, ContentAreaPadding( 2 )
 	, bShouldAutosize(false)
 	, TabColorScale(FLinearColor(0,0,0,0))
+	, LastActivationTime(0.0)
 {
-
 }
 
 void SDockTab::ActivateInParent(ETabActivationCause InActivationCause)
diff --git a/Engine/Source/Runtime/Slate/Public/Widgets/Docking/SDockTab.h b/Engine/Source/Runtime/Slate/Public/Widgets/Docking/SDockTab.h
--- a/Engine/Source/Runtime/Slate/Public/Widgets/Docking/SDockTab.h
+++ b/Engine/Source/Runtime/Slate/Public/Widgets/Docking/SDockTab.h
@@ -78,25 +78,21 @@ public:
 	/** Construct the widget from the declaration. */
 	void Construct( const FArguments& InArgs );
 
+	// SWidget interface
 	virtual FReply OnMouseButtonDown( const FGeometry& MyGeometry, const FPointerEvent& MouseEvent ) override;
-
 	virtual FReply OnMouseButtonDoubleClick( const FGeometry& MyGeometry, const FPointerEvent& MouseEvent ) override;
-
 	virtual FReply OnDragDetected( const FGeometry& MyGeometry, const FPointerEvent& MouseEvent  ) override;
-	
 	virtual FReply OnMouseButtonUp( const FGeometry& MyGeometry, const FPointerEvent& MouseEvent ) override;
-
 	virtual void OnDragEnter( const FGeometry& MyGeometry, const FDragDropEvent& DragDropEvent ) override;
-
 	virtual void OnDragLeave( const FDragDropEvent& DragDropEvent ) override;
-
 	virtual FReply OnDrop( const FGeometry& MyGeometry, const FDragDropEvent& DragDropEvent ) override;
-
 	virtual FReply OnTouchStarted( const FGeometry& MyGeometry, const FPointerEvent& InTouchEvent ) override;
-
 	virtual FReply OnTouchEnded( const FGeometry& MyGeometry, const FPointerEvent& InTouchEvent ) override;
+	// End of SWidget interface
 
+	// SBorder interface
 	virtual void SetContent(TSharedRef<SWidget> InContent) override;
+	// End of SBorder interface
 
 	void SetLeftContent( TSharedRef<SWidget> InContent );
 	void SetRightContent( TSharedRef<SWidget> InContent );
@@ -236,6 +232,17 @@ public:
 	 */
 	bool HasSiblingTab(const FTabId& SiblingTabId, const bool TreatIndexNoneAsWildcard = true) const;
 
+	/** Updates the 'last activated' time to the current time */
+	void UpdateActivationTime()
+	{
+		LastActivationTime = FSlateApplication::Get().GetCurrentTime();
+	}
+
+	/** Returns the time this tab was last activated */
+	double GetLastActivationTime()
+	{
+		return LastActivationTime;
+	}
 protected:
 	
 	/** Gets the dock tab stack this dockable tab resides within, if any */
@@ -353,4 +360,7 @@ protected:
 	
 	/** Widget used to show the icon on the tab */
 	TSharedPtr<SImage> IconWidget;
+	
+	/** Time this tab was last activated */
+	double LastActivationTime;
 };
