Use SWindow::Resize instead of SWindow::ReshapeWindow to update the size of autosized windows. This avoids a problem with jerky movement when dragging such windows on Mac.

diff --git a/Engine/Source/Runtime/Slate/Private/Framework/Application/SlateApplication.cpp b/Engine/Source/Runtime/Slate/Private/Framework/Application/SlateApplication.cpp
--- a/Engine/Source/Runtime/Slate/Private/Framework/Application/SlateApplication.cpp
+++ b/Engine/Source/Runtime/Slate/Private/Framework/Application/SlateApplication.cpp
@@ -903,7 +903,7 @@ void FSlateApplication::TickWindowAndChildren( TSharedRef<SWindow> WindowToTick
 
 		if (WindowToTick->IsAutosized())
 		{
-			WindowToTick->ReshapeWindow(WindowToTick->GetPositionInScreen(), WindowToTick->GetDesiredSize());
+			WindowToTick->Resize(WindowToTick->GetDesiredSize());
 		}
 
 		{
@@ -1059,7 +1059,7 @@ static void PrepassWindowAndChildren( TSharedRef<SWindow> WindowToPrepass )
 
 	if (WindowToPrepass->IsAutosized())
 	{
-		WindowToPrepass->ReshapeWindow(WindowToPrepass->GetPositionInScreen(), WindowToPrepass->GetDesiredSize());
+		WindowToPrepass->Resize(WindowToPrepass->GetDesiredSize());
 	}
 
 	for (const TSharedRef<SWindow>& ChildWindow : WindowToPrepass->GetChildWindows() )
