Fix for Editor crashes when dragging a docked window (UE-11893)

diff --git a/Engine/Source/Runtime/Slate/Private/Framework/Docking/TabManager.cpp b/Engine/Source/Runtime/Slate/Private/Framework/Docking/TabManager.cpp
--- a/Engine/Source/Runtime/Slate/Private/Framework/Docking/TabManager.cpp
+++ b/Engine/Source/Runtime/Slate/Private/Framework/Docking/TabManager.cpp
@@ -1600,7 +1600,11 @@ void FGlobalTabmanager::SetActiveTab( const TSharedPtr<class SDockTab>& NewActiv
 
 	if (bShouldApplyChange && (CurrentlyActiveTab != NewActiveTab))
 	{
-		NewActiveTab->UpdateActivationTime();
+		if (NewActiveTab.IsValid())
+		{
+			NewActiveTab->UpdateActivationTime();
+		}
+
 		OnActiveTabChanged.Broadcast( CurrentlyActiveTab, NewActiveTab );
 		ActiveTabPtr = NewActiveTab;
 	}	
