Allow tabs to be entered into text edit boxes

Higher-level code prevents this unless the edit box explicitly handles it, so for existing controls this has no effect.

diff --git a/Engine/Source/Runtime/Slate/Private/Widgets/Input/SEditableText.cpp b/Engine/Source/Runtime/Slate/Private/Widgets/Input/SEditableText.cpp
--- a/Engine/Source/Runtime/Slate/Private/Widgets/Input/SEditableText.cpp
+++ b/Engine/Source/Runtime/Slate/Private/Widgets/Input/SEditableText.cpp
@@ -646,7 +646,11 @@ void SEditableText::TypeChar( const int32 InChar )
 	// Certain characters are not allowed
 	bool bIsCharAllowed = true;
 	{
-		if( InChar <= 0x1F )
+		if( InChar == TEXT('\t') )
+		{
+			bIsCharAllowed = true;
+		}
+		else if( InChar <= 0x1F )
 		{
 			bIsCharAllowed = false;
 		}
diff --git a/Engine/Source/Runtime/Slate/Private/Widgets/Text/SMultiLineEditableText.cpp b/Engine/Source/Runtime/Slate/Private/Widgets/Text/SMultiLineEditableText.cpp
--- a/Engine/Source/Runtime/Slate/Private/Widgets/Text/SMultiLineEditableText.cpp
+++ b/Engine/Source/Runtime/Slate/Private/Widgets/Text/SMultiLineEditableText.cpp
@@ -801,7 +801,11 @@ void SMultiLineEditableText::TypeChar( const int32 Character )
 	// Certain characters are not allowed
 	bool bIsCharAllowed = true;
 	{
-		if( Character <= 0x1F )
+		if( Character == TEXT('\t') )
+		{
+			bIsCharAllowed = true;
+		}
+		else if( Character <= 0x1F )
 		{
 			bIsCharAllowed = false;
 		}
diff --git a/Engine/Source/Runtime/Slate/Public/Widgets/Text/SMultiLineEditableText.h b/Engine/Source/Runtime/Slate/Public/Widgets/Text/SMultiLineEditableText.h
--- a/Engine/Source/Runtime/Slate/Public/Widgets/Text/SMultiLineEditableText.h
+++ b/Engine/Source/Runtime/Slate/Public/Widgets/Text/SMultiLineEditableText.h
@@ -425,7 +425,7 @@ public:
 	}
 	// END IVirtualKeyboardEntry interface
 
-private:
+protected:
 	// BEGIN SWidget interface
 	virtual void Tick( const FGeometry& AllottedGeometry, const double InCurrentTime, const float InDeltaTime ) override;
 	virtual int32 OnPaint( const FPaintArgs& Args, const FGeometry& AllottedGeometry, const FSlateRect& MyClippingRect, FSlateWindowElementList& OutDrawElements, int32 LayerId, const FWidgetStyle& InWidgetStyle, bool bParentEnabled ) const override;
