Unify Slate's FInputGesture and Engine's FInputChord as FInputChord defined in Slate

diff --git a/Engine/Source/Runtime/Slate/Private/Framework/Commands/GenericCommands.cpp b/Engine/Source/Runtime/Slate/Private/Framework/Commands/GenericCommands.cpp
--- a/Engine/Source/Runtime/Slate/Private/Framework/Commands/GenericCommands.cpp
+++ b/Engine/Source/Runtime/Slate/Private/Framework/Commands/GenericCommands.cpp
@@ -7,21 +7,21 @@
 
 void FGenericCommands::RegisterCommands()
 {
-	UI_COMMAND( Cut, "Cut", "Cut selection", EUserInterfaceActionType::Button, FInputGesture( EModifierKey::Control, EKeys::X ) )
-	UI_COMMAND( Copy, "Copy", "Copy selection", EUserInterfaceActionType::Button, FInputGesture( EModifierKey::Control, EKeys::C ) )
-	UI_COMMAND( Paste, "Paste", "Paste clipboard contents", EUserInterfaceActionType::Button, FInputGesture( EModifierKey::Control, EKeys::V ) )
+	UI_COMMAND( Cut, "Cut", "Cut selection", EUserInterfaceActionType::Button, FInputChord( EModifierKey::Control, EKeys::X ) )
+	UI_COMMAND( Copy, "Copy", "Copy selection", EUserInterfaceActionType::Button, FInputChord( EModifierKey::Control, EKeys::C ) )
+	UI_COMMAND( Paste, "Paste", "Paste clipboard contents", EUserInterfaceActionType::Button, FInputChord( EModifierKey::Control, EKeys::V ) )
 #if PLATFORM_MAC
-	UI_COMMAND( Duplicate, "Duplicate", "Duplicate selection", EUserInterfaceActionType::Button, FInputGesture( EModifierKey::Command, EKeys::W ) )
+	UI_COMMAND( Duplicate, "Duplicate", "Duplicate selection", EUserInterfaceActionType::Button, FInputChord( EModifierKey::Command, EKeys::W ) )
 #else
-	UI_COMMAND( Duplicate, "Duplicate", "Duplicate selection", EUserInterfaceActionType::Button, FInputGesture( EModifierKey::Control, EKeys::W ) )
+	UI_COMMAND( Duplicate, "Duplicate", "Duplicate selection", EUserInterfaceActionType::Button, FInputChord( EModifierKey::Control, EKeys::W ) )
 #endif
-	UI_COMMAND( Undo, "Undo", "Undo last action", EUserInterfaceActionType::Button, FInputGesture( EModifierKey::Control, EKeys::Z ) )
-	UI_COMMAND( Redo, "Redo", "Redo last action", EUserInterfaceActionType::Button, FInputGesture( EModifierKey::Control, EKeys::Y ) )
+	UI_COMMAND( Undo, "Undo", "Undo last action", EUserInterfaceActionType::Button, FInputChord( EModifierKey::Control, EKeys::Z ) )
+	UI_COMMAND( Redo, "Redo", "Redo last action", EUserInterfaceActionType::Button, FInputChord( EModifierKey::Control, EKeys::Y ) )
 
-	UI_COMMAND(Delete, "Delete", "Delete current selection", EUserInterfaceActionType::Button, FInputGesture(EKeys::Platform_Delete))
+	UI_COMMAND(Delete, "Delete", "Delete current selection", EUserInterfaceActionType::Button, FInputChord(EKeys::Platform_Delete))
 	
-	UI_COMMAND( Rename, "Rename", "Rename current selection", EUserInterfaceActionType::Button, FInputGesture( EKeys::F2 ) )
-	UI_COMMAND( SelectAll, "Select All", "Select everything in the current scope", EUserInterfaceActionType::Button, FInputGesture( EModifierKey::Control, EKeys::A ) )
+	UI_COMMAND( Rename, "Rename", "Rename current selection", EUserInterfaceActionType::Button, FInputChord( EKeys::F2 ) )
+	UI_COMMAND( SelectAll, "Select All", "Select everything in the current scope", EUserInterfaceActionType::Button, FInputChord( EModifierKey::Control, EKeys::A ) )
 }
 
 #undef LOCTEXT_NAMESPACE
\ No newline at end of file
diff --git a/Engine/Source/Runtime/Slate/Private/Framework/Commands/InputBindingManager.cpp b/Engine/Source/Runtime/Slate/Private/Framework/Commands/InputBindingManager.cpp
--- a/Engine/Source/Runtime/Slate/Private/Framework/Commands/InputBindingManager.cpp
+++ b/Engine/Source/Runtime/Slate/Private/Framework/Commands/InputBindingManager.cpp
@@ -3,28 +3,28 @@
 #include "SlatePrivatePCH.h"
 
 
-/* FUserDefinedGestures helper class
+/* FUserDefinedChords helper class
  *****************************************************************************/
 
-/** An identifier for a user defined gesture */
-struct FUserDefinedGestureKey
+/** An identifier for a user defined chord */
+struct FUserDefinedChordKey
 {
-	FUserDefinedGestureKey()
+	FUserDefinedChordKey()
 	{
 	}
 
-	FUserDefinedGestureKey(const FName InBindingContext, const FName InCommandName)
+	FUserDefinedChordKey(const FName InBindingContext, const FName InCommandName)
 		: BindingContext(InBindingContext)
 		, CommandName(InCommandName)
 	{
 	}
 
-	bool operator==(const FUserDefinedGestureKey& Other) const
+	bool operator==(const FUserDefinedChordKey& Other) const
 	{
 		return BindingContext == Other.BindingContext && CommandName == Other.CommandName;
 	}
 
-	bool operator!=(const FUserDefinedGestureKey& Other) const
+	bool operator!=(const FUserDefinedChordKey& Other) const
 	{
 		return !(*this == Other);
 	}
@@ -33,61 +33,68 @@ struct FUserDefinedGestureKey
 	FName CommandName;
 };
 
-uint32 GetTypeHash( const FUserDefinedGestureKey& Key )
+uint32 GetTypeHash( const FUserDefinedChordKey& Key )
 {
 	return GetTypeHash(Key.BindingContext) ^ GetTypeHash(Key.CommandName);
 }
 
-class FUserDefinedGestures
+class FUserDefinedChords
 {
 public:
-	void LoadGestures();
-	void SaveGestures() const;
-	bool GetUserDefinedGesture( const FName BindingContext, const FName CommandName, FInputGesture& OutUserDefinedGesture ) const;
-	void SetUserDefinedGesture( const FUICommandInfo& CommandInfo );
-	/** Remove all user defined gestures */
+	void LoadChords();
+	void SaveChords() const;
+	bool GetUserDefinedChord( const FName BindingContext, const FName CommandName, FInputChord& OutUserDefinedChord ) const;
+	void SetUserDefinedChord( const FUICommandInfo& CommandInfo );
+	/** Remove all user defined chords */
 	void RemoveAll();
 private:
-	/* Mapping from a gesture key to the user defined gesture */
-	typedef TMap<FUserDefinedGestureKey, FInputGesture> FGesturesMap;
-	TSharedPtr<FGesturesMap> Gestures;
+	/* Mapping from a chord key to the user defined chord */
+	typedef TMap<FUserDefinedChordKey, FInputChord> FChordsMap;
+	TSharedPtr<FChordsMap> Chords;
 };
 
-void FUserDefinedGestures::LoadGestures()
+void FUserDefinedChords::LoadChords()
 {
-	if( !Gestures.IsValid() )
+	if( !Chords.IsValid() )
 	{
-		Gestures = MakeShareable( new FGesturesMap );
+		Chords = MakeShareable( new FChordsMap );
 
-		// First, try and load the gestures from their new location in the ini file
+		// First, try and load the chords from their new location in the ini file
 		// Failing that, try and load them from the older txt file
-		TArray<FString> GestureJsonArray;
-		if( GConfig->GetArray(TEXT("UserDefinedGestures"), TEXT("UserDefinedGestures"), GestureJsonArray, GEditorKeyBindingsIni) )
+		TArray<FString> ChordJsonArray;
+		bool bFoundChords = (GConfig->GetArray(TEXT("UserDefinedChords"), TEXT("UserDefinedChords"), ChordJsonArray, GEditorKeyBindingsIni) > 0);
+		if (!bFoundChords)
 		{
-			// This loads an array of JSON strings representing the FUserDefinedGestureKey and FInputGesture in a single JSON object
-			for(const FString& GestureJson : GestureJsonArray)
+			// Backwards compat for when it was named gestures
+			bFoundChords = (GConfig->GetArray(TEXT("UserDefinedGestures"), TEXT("UserDefinedGestures"), ChordJsonArray, GEditorKeyBindingsIni) > 0);
+		}
+
+		if(bFoundChords)
+		{
+			// This loads an array of JSON strings representing the FUserDefinedChordKey and FInputChord in a single JSON object
+			for(const FString& ChordJson : ChordJsonArray)
 			{
-				const FString UnescapedContent = FRemoteConfig::ReplaceIniSpecialCharWithChar(GestureJson).ReplaceEscapedCharWithChar();
+				const FString UnescapedContent = FRemoteConfig::ReplaceIniSpecialCharWithChar(ChordJson).ReplaceEscapedCharWithChar();
 
-				TSharedPtr<FJsonObject> GestureInfoObj;
+				TSharedPtr<FJsonObject> ChordInfoObj;
 				auto JsonReader = TJsonReaderFactory<>::Create( UnescapedContent );
-				if( FJsonSerializer::Deserialize( JsonReader, GestureInfoObj ) )
+				if( FJsonSerializer::Deserialize( JsonReader, ChordInfoObj ) )
 				{
-					const TSharedPtr<FJsonValue> BindingContextObj = GestureInfoObj->Values.FindRef( TEXT("BindingContext") );
-					const TSharedPtr<FJsonValue> CommandNameObj = GestureInfoObj->Values.FindRef( TEXT("CommandName") );
-					const TSharedPtr<FJsonValue> CtrlObj = GestureInfoObj->Values.FindRef( TEXT("Control") );
-					const TSharedPtr<FJsonValue> AltObj = GestureInfoObj->Values.FindRef( TEXT("Alt") );
-					const TSharedPtr<FJsonValue> ShiftObj = GestureInfoObj->Values.FindRef( TEXT("Shift") );
-					const TSharedPtr<FJsonValue> CmdObj = GestureInfoObj->Values.FindRef( TEXT("Command") );
-					const TSharedPtr<FJsonValue> KeyObj = GestureInfoObj->Values.FindRef( TEXT("Key") );
+					const TSharedPtr<FJsonValue> BindingContextObj = ChordInfoObj->Values.FindRef( TEXT("BindingContext") );
+					const TSharedPtr<FJsonValue> CommandNameObj = ChordInfoObj->Values.FindRef( TEXT("CommandName") );
+					const TSharedPtr<FJsonValue> CtrlObj = ChordInfoObj->Values.FindRef( TEXT("Control") );
+					const TSharedPtr<FJsonValue> AltObj = ChordInfoObj->Values.FindRef( TEXT("Alt") );
+					const TSharedPtr<FJsonValue> ShiftObj = ChordInfoObj->Values.FindRef( TEXT("Shift") );
+					const TSharedPtr<FJsonValue> CmdObj = ChordInfoObj->Values.FindRef( TEXT("Command") );
+					const TSharedPtr<FJsonValue> KeyObj = ChordInfoObj->Values.FindRef( TEXT("Key") );
 
 					const FName BindingContext = *BindingContextObj->AsString();
 					const FName CommandName = *CommandNameObj->AsString();
 
-					const FUserDefinedGestureKey GestureKey(BindingContext, CommandName);
-					FInputGesture& UserDefinedGesture = Gestures->FindOrAdd(GestureKey);
+					const FUserDefinedChordKey ChordKey(BindingContext, CommandName);
+					FInputChord& UserDefinedChord = Chords->FindOrAdd(ChordKey);
 
-					UserDefinedGesture = FInputGesture(*KeyObj->AsString(), EModifierKey::FromBools(
+					UserDefinedChord = FInputChord(*KeyObj->AsString(), EModifierKey::FromBools(
 																				CtrlObj->AsBool(), 
 																				AltObj->AsBool(), 
 																				ShiftObj->AsBool(), 
@@ -98,32 +105,38 @@ void FUserDefinedGestures::LoadGestures()
 		}
 		else
 		{
-			TSharedPtr<FJsonObject> GesturesObj;
+			TSharedPtr<FJsonObject> ChordsObj;
 
-			// This loads a JSON object containing BindingContexts, containing objects of CommandNames, containing the FInputGesture information
+			// This loads a JSON object containing BindingContexts, containing objects of CommandNames, containing the FInputChord information
 			FString Content;
-			if( GConfig->GetString(TEXT("UserDefinedGestures"), TEXT("Content"), Content, GEditorKeyBindingsIni) )
+			bool bFoundContent = (GConfig->GetArray(TEXT("UserDefinedChords"), TEXT("Content"), ChordJsonArray, GEditorKeyBindingsIni) > 0);
+			if (!bFoundContent)
+			{
+				// Backwards compat for when it was named gestures
+				bFoundChords = (GConfig->GetArray(TEXT("UserDefinedGestures"), TEXT("Content"), ChordJsonArray, GEditorKeyBindingsIni) > 0);
+			}
+			if (bFoundContent)
 			{
 				const FString UnescapedContent = FRemoteConfig::ReplaceIniSpecialCharWithChar(Content).ReplaceEscapedCharWithChar();
 
 				auto JsonReader = TJsonReaderFactory<>::Create( UnescapedContent );
-				FJsonSerializer::Deserialize( JsonReader, GesturesObj );
+				FJsonSerializer::Deserialize( JsonReader, ChordsObj );
 			}
 
-			if (!GesturesObj.IsValid())
+			if (!ChordsObj.IsValid())
 			{
-				// Gestures have not been loaded from the ini file, try reading them from the txt file now
+				// Chords have not been loaded from the ini file, try reading them from the txt file now
 				TSharedPtr<FArchive> Ar = MakeShareable( IFileManager::Get().CreateFileReader( *( FPaths::GameSavedDir() / TEXT("Preferences/EditorKeyBindings.txt") ) ) );
 				if( Ar.IsValid() )
 				{
 					auto TextReader = TJsonReaderFactory<ANSICHAR>::Create( Ar.Get() );
-					FJsonSerializer::Deserialize( TextReader, GesturesObj );
+					FJsonSerializer::Deserialize( TextReader, ChordsObj );
 				}
 			}
 
-			if (GesturesObj.IsValid())
+			if (ChordsObj.IsValid())
 			{
-				for(const auto& BindingContextInfo : GesturesObj->Values)
+				for(const auto& BindingContextInfo : ChordsObj->Values)
 				{
 					const FName BindingContext = *BindingContextInfo.Key;
 					TSharedPtr<FJsonObject> BindingContextObj = BindingContextInfo.Value->AsObject();
@@ -138,10 +151,10 @@ void FUserDefinedGestures::LoadGestures()
 						const TSharedPtr<FJsonValue> CmdObj = CommandObj->Values.FindRef( TEXT("Command") );
 						const TSharedPtr<FJsonValue> KeyObj = CommandObj->Values.FindRef( TEXT("Key") );
 
-						const FUserDefinedGestureKey GestureKey(BindingContext, CommandName);
-						FInputGesture& UserDefinedGesture = Gestures->FindOrAdd(GestureKey);
+						const FUserDefinedChordKey ChordKey(BindingContext, CommandName);
+						FInputChord& UserDefinedChord = Chords->FindOrAdd(ChordKey);
 
-						UserDefinedGesture = FInputGesture(*KeyObj->AsString(), EModifierKey::FromBools(
+						UserDefinedChord = FInputChord(*KeyObj->AsString(), EModifierKey::FromBools(
 																					CtrlObj->AsBool(), 
 																					AltObj->AsBool(), 
 																					ShiftObj->AsBool(), 
@@ -154,51 +167,53 @@ void FUserDefinedGestures::LoadGestures()
 	}
 }
 
-void FUserDefinedGestures::SaveGestures() const
+void FUserDefinedChords::SaveChords() const
 {
-	if( Gestures.IsValid() )
+	if( Chords.IsValid() )
 	{
-		FString GestureRawJsonContent;
-		TArray<FString> GestureJsonArray;
-		for(const auto& GestureInfo : *Gestures)
+		FString ChordRawJsonContent;
+		TArray<FString> ChordJsonArray;
+		for(const auto& ChordInfo : *Chords)
 		{
-			TSharedPtr<FJsonValueObject> GestureInfoValueObj = MakeShareable( new FJsonValueObject( MakeShareable( new FJsonObject ) ) );
-			TSharedPtr<FJsonObject> GestureInfoObj = GestureInfoValueObj->AsObject();
-
-			// Set the gesture values for the command
-			GestureInfoObj->Values.Add( TEXT("BindingContext"), MakeShareable( new FJsonValueString( GestureInfo.Key.BindingContext.ToString() ) ) );
-			GestureInfoObj->Values.Add( TEXT("CommandName"), MakeShareable( new FJsonValueString( GestureInfo.Key.CommandName.ToString() ) ) );
-			GestureInfoObj->Values.Add( TEXT("Control"), MakeShareable( new FJsonValueBoolean( GestureInfo.Value.NeedsControl() ) ) );
-			GestureInfoObj->Values.Add( TEXT("Alt"), MakeShareable( new FJsonValueBoolean( GestureInfo.Value.NeedsAlt() ) ) );
-			GestureInfoObj->Values.Add( TEXT("Shift"), MakeShareable( new FJsonValueBoolean( GestureInfo.Value.NeedsShift() ) ) );
-			GestureInfoObj->Values.Add( TEXT("Command"), MakeShareable( new FJsonValueBoolean( GestureInfo.Value.NeedsCommand() ) ) );
-			GestureInfoObj->Values.Add( TEXT("Key"), MakeShareable( new FJsonValueString( GestureInfo.Value.Key.ToString() ) ) );
-
-			auto JsonWriter = TJsonWriterFactory< TCHAR, TCondensedJsonPrintPolicy<TCHAR> >::Create( &GestureRawJsonContent );
-			FJsonSerializer::Serialize( GestureInfoObj.ToSharedRef(), JsonWriter );
-
-			const FString EscapedContent = FRemoteConfig::ReplaceIniCharWithSpecialChar(GestureRawJsonContent).ReplaceCharWithEscapedChar();
-			GestureJsonArray.Add(EscapedContent);
+			TSharedPtr<FJsonValueObject> ChordInfoValueObj = MakeShareable( new FJsonValueObject( MakeShareable( new FJsonObject ) ) );
+			TSharedPtr<FJsonObject> ChordInfoObj = ChordInfoValueObj->AsObject();
+
+			// Set the chord values for the command
+			ChordInfoObj->Values.Add( TEXT("BindingContext"), MakeShareable( new FJsonValueString( ChordInfo.Key.BindingContext.ToString() ) ) );
+			ChordInfoObj->Values.Add( TEXT("CommandName"), MakeShareable( new FJsonValueString( ChordInfo.Key.CommandName.ToString() ) ) );
+			ChordInfoObj->Values.Add( TEXT("Control"), MakeShareable( new FJsonValueBoolean( ChordInfo.Value.NeedsControl() ) ) );
+			ChordInfoObj->Values.Add( TEXT("Alt"), MakeShareable( new FJsonValueBoolean( ChordInfo.Value.NeedsAlt() ) ) );
+			ChordInfoObj->Values.Add( TEXT("Shift"), MakeShareable( new FJsonValueBoolean( ChordInfo.Value.NeedsShift() ) ) );
+			ChordInfoObj->Values.Add( TEXT("Command"), MakeShareable( new FJsonValueBoolean( ChordInfo.Value.NeedsCommand() ) ) );
+			ChordInfoObj->Values.Add( TEXT("Key"), MakeShareable( new FJsonValueString( ChordInfo.Value.Key.ToString() ) ) );
+
+			auto JsonWriter = TJsonWriterFactory< TCHAR, TCondensedJsonPrintPolicy<TCHAR> >::Create( &ChordRawJsonContent );
+			FJsonSerializer::Serialize( ChordInfoObj.ToSharedRef(), JsonWriter );
+
+			const FString EscapedContent = FRemoteConfig::ReplaceIniCharWithSpecialChar(ChordRawJsonContent).ReplaceCharWithEscapedChar();
+			ChordJsonArray.Add(EscapedContent);
 		}
 
-		GConfig->SetArray(TEXT("UserDefinedGestures"), TEXT("UserDefinedGestures"), GestureJsonArray, GEditorKeyBindingsIni);
+		GConfig->SetArray(TEXT("UserDefinedChords"), TEXT("UserDefinedChords"), ChordJsonArray, GEditorKeyBindingsIni);
 
-		// Clean up the old Content key (if it still exists)
-		GConfig->RemoveKey(TEXT("UserDefinedGestures"), TEXT("Content"), GEditorKeyBindingsIni);
+		// Clean up old keys (if they still exist)
+		GConfig->RemoveKey(TEXT("UserDefinedGestures"), TEXT("UserDefinedGestures"), GEditorKeyBindingsIni);
+		GConfig->RemoveKey(TEXT("UserDefinedChords"), TEXT("Content"), GEditorKeyBindingsIni);
+		GConfig->RemoveKey(TEXT("UserDefinedChords"), TEXT("Content"), GEditorKeyBindingsIni);
 	}
 }
 
-bool FUserDefinedGestures::GetUserDefinedGesture( const FName BindingContext, const FName CommandName, FInputGesture& OutUserDefinedGesture ) const
+bool FUserDefinedChords::GetUserDefinedChord( const FName BindingContext, const FName CommandName, FInputChord& OutUserDefinedChord ) const
 {
 	bool bResult = false;
 
-	if( Gestures.IsValid() )
+	if( Chords.IsValid() )
 	{
-		const FUserDefinedGestureKey GestureKey(BindingContext, CommandName);
-		const FInputGesture* const UserDefinedGesturePtr = Gestures->Find(GestureKey);
-		if( UserDefinedGesturePtr )
+		const FUserDefinedChordKey ChordKey(BindingContext, CommandName);
+		const FInputChord* const UserDefinedChordPtr = Chords->Find(ChordKey);
+		if( UserDefinedChordPtr )
 		{
-			OutUserDefinedGesture = *UserDefinedGesturePtr;
+			OutUserDefinedChord = *UserDefinedChordPtr;
 			bResult = true;
 		}
 	}
@@ -206,27 +221,27 @@ bool FUserDefinedGestures::GetUserDefinedGesture( const FName BindingContext, co
 	return bResult;
 }
 
-void FUserDefinedGestures::SetUserDefinedGesture( const FUICommandInfo& CommandInfo )
+void FUserDefinedChords::SetUserDefinedChord( const FUICommandInfo& CommandInfo )
 {
-	if( Gestures.IsValid() )
+	if( Chords.IsValid() )
 	{
 		const FName BindingContext = CommandInfo.GetBindingContext();
 		const FName CommandName = CommandInfo.GetCommandName();
 
 		// Find or create the command context
-		const FUserDefinedGestureKey GestureKey(BindingContext, CommandName);
-		FInputGesture& UserDefinedGesture = Gestures->FindOrAdd(GestureKey);
+		const FUserDefinedChordKey ChordKey(BindingContext, CommandName);
+		FInputChord& UserDefinedChord = Chords->FindOrAdd(ChordKey);
 
-		// Save an empty invalid gesture if one was not set
-		// This is an indication that the user doesn't want this bound and not to use the default gesture
-		const TSharedPtr<const FInputGesture> InputGesture = CommandInfo.GetActiveGesture();
-		UserDefinedGesture = (InputGesture.IsValid()) ? *InputGesture : FInputGesture();
+		// Save an empty invalid chord if one was not set
+		// This is an indication that the user doesn't want this bound and not to use the default chord
+		const TSharedPtr<const FInputChord> InputChord = CommandInfo.GetActiveChord();
+		UserDefinedChord = (InputChord.IsValid()) ? *InputChord : FInputChord();
 	}
 }
 
-void FUserDefinedGestures::RemoveAll()
+void FUserDefinedChords::RemoveAll()
 {
-	Gestures = MakeShareable( new FGesturesMap );
+	Chords = MakeShareable( new FChordsMap );
 }
 
 
@@ -249,88 +264,88 @@ FInputBindingManager& FInputBindingManager::Get()
  *****************************************************************************/
 
 /**
- * Gets the user defined gesture (if any) from the provided command name
+ * Gets the user defined chord (if any) from the provided command name
  * 
  * @param InBindingContext	The context in which the command is active
- * @param InCommandName		The name of the command to get the gesture from
+ * @param InCommandName		The name of the command to get the chord from
  */
-bool FInputBindingManager::GetUserDefinedGesture( const FName InBindingContext, const FName InCommandName, FInputGesture& OutUserDefinedGesture )
+bool FInputBindingManager::GetUserDefinedChord( const FName InBindingContext, const FName InCommandName, FInputChord& OutUserDefinedChord )
 {
-	if( !UserDefinedGestures.IsValid() )
+	if( !UserDefinedChords.IsValid() )
 	{
-		UserDefinedGestures = MakeShareable( new FUserDefinedGestures );
-		UserDefinedGestures->LoadGestures();
+		UserDefinedChords = MakeShareable( new FUserDefinedChords );
+		UserDefinedChords->LoadChords();
 	}
 
-	return UserDefinedGestures->GetUserDefinedGesture( InBindingContext, InCommandName, OutUserDefinedGesture );
+	return UserDefinedChords->GetUserDefinedChord( InBindingContext, InCommandName, OutUserDefinedChord );
 }
 
-void FInputBindingManager::CheckForDuplicateDefaultGestures( const FBindingContext& InBindingContext, TSharedPtr<FUICommandInfo> InCommandInfo ) const
+void FInputBindingManager::CheckForDuplicateDefaultChords( const FBindingContext& InBindingContext, TSharedPtr<FUICommandInfo> InCommandInfo ) const
 {
 	const bool bCheckDefault = true;
-	TSharedPtr<FUICommandInfo> ExistingInfo = GetCommandInfoFromInputGesture( InBindingContext.GetContextName(), InCommandInfo->DefaultGesture, bCheckDefault );
+	TSharedPtr<FUICommandInfo> ExistingInfo = GetCommandInfoFromInputChord( InBindingContext.GetContextName(), InCommandInfo->DefaultChord, bCheckDefault );
 	if( ExistingInfo.IsValid() )
 	{
 		if( ExistingInfo->CommandName != InCommandInfo->CommandName )
 		{
 			// Two different commands with the same name in the same context or parent context
-			UE_LOG(LogSlate, Fatal, TEXT("The command '%s.%s' has the same default gesture as '%s.%s' [%s]"), 
+			UE_LOG(LogSlate, Fatal, TEXT("The command '%s.%s' has the same default chord as '%s.%s' [%s]"), 
 				*InCommandInfo->BindingContext.ToString(),
 				*InCommandInfo->CommandName.ToString(), 
 				*ExistingInfo->BindingContext.ToString(),
 				*ExistingInfo->CommandName.ToString(), 
-				*InCommandInfo->DefaultGesture.GetInputText().ToString() );
+				*InCommandInfo->DefaultChord.GetInputText().ToString() );
 		}
 	}
 }
 
 
-void FInputBindingManager::NotifyActiveGestureChanged( const FUICommandInfo& CommandInfo )
+void FInputBindingManager::NotifyActiveChordChanged( const FUICommandInfo& CommandInfo )
 {
 	FContextEntry& ContextEntry = ContextMap.FindChecked( CommandInfo.GetBindingContext() );
 
 	// Slow but doesn't happen frequently
-	for( FGestureMap::TIterator It( ContextEntry.GestureToCommandInfoMap ); It; ++It )
+	for( FChordMap::TIterator It( ContextEntry.ChordToCommandInfoMap ); It; ++It )
 	{
-		// Remove the currently active gesture from the map if one exists
+		// Remove the currently active chord from the map if one exists
 		if( It.Value() == CommandInfo.GetCommandName() )
 		{
 			It.RemoveCurrent();
-			// There should only be one active gesture
+			// There should only be one active chord
 			break;
 		}
 	}
 
-	if( CommandInfo.GetActiveGesture()->IsValidGesture() )
+	if( CommandInfo.GetActiveChord()->IsValidChord() )
 	{
-		checkSlow( !ContextEntry.GestureToCommandInfoMap.Contains( *CommandInfo.GetActiveGesture() ) )
-		ContextEntry.GestureToCommandInfoMap.Add( *CommandInfo.GetActiveGesture(), CommandInfo.GetCommandName() );
+		checkSlow( !ContextEntry.ChordToCommandInfoMap.Contains( *CommandInfo.GetActiveChord() ) )
+		ContextEntry.ChordToCommandInfoMap.Add( *CommandInfo.GetActiveChord(), CommandInfo.GetCommandName() );
 	}
 	
 
-	// The user defined gestures should have already been created
-	check( UserDefinedGestures.IsValid() );
+	// The user defined chords should have already been created
+	check( UserDefinedChords.IsValid() );
 
-	UserDefinedGestures->SetUserDefinedGesture( CommandInfo );
+	UserDefinedChords->SetUserDefinedChord( CommandInfo );
 
-	// Broadcast the gesture event when a new one is added
-	OnUserDefinedGestureChanged.Broadcast(CommandInfo);
+	// Broadcast the chord event when a new one is added
+	OnUserDefinedChordChanged.Broadcast(CommandInfo);
 }
 
 void FInputBindingManager::SaveInputBindings()
 {
-	if( UserDefinedGestures.IsValid() )
+	if( UserDefinedChords.IsValid() )
 	{
-		UserDefinedGestures->SaveGestures();
+		UserDefinedChords->SaveChords();
 	}
 }
 
-void FInputBindingManager::RemoveUserDefinedGestures()
+void FInputBindingManager::RemoveUserDefinedChords()
 {
-	if( UserDefinedGestures.IsValid() )
+	if( UserDefinedChords.IsValid() )
 	{
-		UserDefinedGestures->RemoveAll();
-		UserDefinedGestures->SaveGestures();
+		UserDefinedChords->RemoveAll();
+		UserDefinedChords->SaveChords();
 	}
 }
 
@@ -346,8 +361,8 @@ void FInputBindingManager::CreateInputCommand( const TSharedRef<FBindingContext>
 	// The command name should be valid
 	check( InCommandInfo->CommandName != NAME_None );
 
-	// Should not have already created a gesture for this command
-	check( !InCommandInfo->ActiveGesture->IsValidGesture() );
+	// Should not have already created a chord for this command
+	check( !InCommandInfo->ActiveChord->IsValidChord() );
 	
 	const FName ContextName = InBindingContext->GetContextName();
 
@@ -370,9 +385,9 @@ void FInputBindingManager::CreateInputCommand( const TSharedRef<FBindingContext>
 		ParentToChildMap.AddUnique( InBindingContext->GetContextParent(), InBindingContext->GetContextName() );
 	}
 
-	if( InCommandInfo->DefaultGesture.IsValidGesture() )
+	if( InCommandInfo->DefaultChord.IsValidChord() )
 	{
-		CheckForDuplicateDefaultGestures( *InBindingContext, InCommandInfo );
+		CheckForDuplicateDefaultChords( *InBindingContext, InCommandInfo );
 	}
 	
 	{
@@ -383,62 +398,62 @@ void FInputBindingManager::CreateInputCommand( const TSharedRef<FBindingContext>
 	// Add the command info to the list of known infos.  It can only exist once.
 	CommandInfoMap.Add( InCommandInfo->CommandName, InCommandInfo );
 
-	// See if there are user defined gestures for this command
-	FInputGesture UserDefinedGesture;
-	bool bFoundUserDefinedGesture = GetUserDefinedGesture( ContextName, InCommandInfo->CommandName, UserDefinedGesture );
+	// See if there are user defined chords for this command
+	FInputChord UserDefinedChord;
+	bool bFoundUserDefinedChord = GetUserDefinedChord( ContextName, InCommandInfo->CommandName, UserDefinedChord );
 
 	
-	if( !bFoundUserDefinedGesture && InCommandInfo->DefaultGesture.IsValidGesture() )
+	if( !bFoundUserDefinedChord && InCommandInfo->DefaultChord.IsValidChord() )
 	{
-		// Find any existing command with the same gesture 
-		// This is for inconsistency between default and user defined gesture.  We need to make sure that if default gestures are changed to a gesture that a user set to a different command, that the default gesture doesn't replace
-		// the existing commands gesture. Note: Duplicate default gestures are found above in CheckForDuplicateDefaultGestures
-		FName ExisingCommand = ContextEntry.GestureToCommandInfoMap.FindRef( InCommandInfo->DefaultGesture );
+		// Find any existing command with the same chord 
+		// This is for inconsistency between default and user defined chord.  We need to make sure that if default chords are changed to a chord that a user set to a different command, that the default chord doesn't replace
+		// the existing commands chord. Note: Duplicate default chords are found above in CheckForDuplicateDefaultChords
+		FName ExisingCommand = ContextEntry.ChordToCommandInfoMap.FindRef( InCommandInfo->DefaultChord );
 
 		if( ExisingCommand == NAME_None )
 		{
-			// No existing command has a user defined gesture and no user defined gesture is available for this command 
-			TSharedRef<FInputGesture> NewGesture = MakeShareable( new FInputGesture( InCommandInfo->DefaultGesture ) );
-			InCommandInfo->ActiveGesture = NewGesture;
+			// No existing command has a user defined chord and no user defined chord is available for this command 
+			TSharedRef<FInputChord> NewChord = MakeShareable( new FInputChord( InCommandInfo->DefaultChord ) );
+			InCommandInfo->ActiveChord = NewChord;
 		}
 
 	}
-	else if( bFoundUserDefinedGesture )
+	else if( bFoundUserDefinedChord )
 	{
-		// Find any existing command with the same gesture 
-		// This is for inconsistency between default and user defined gesture.  We need to make sure that if default gestures are changed to a gesture that a user set to a different command, that the default gesture doesn't replace
-		// the existing commands gesture.
-		FName ExisingCommandName = ContextEntry.GestureToCommandInfoMap.FindRef( UserDefinedGesture );
+		// Find any existing command with the same chord 
+		// This is for inconsistency between default and user defined chord.  We need to make sure that if default chords are changed to a chord that a user set to a different command, that the default chord doesn't replace
+		// the existing commands chord.
+		FName ExisingCommandName = ContextEntry.ChordToCommandInfoMap.FindRef( UserDefinedChord );
 
 		if( ExisingCommandName != NAME_None )
 		{
-			// Get the command with using the same gesture
+			// Get the command with using the same chord
 			TSharedPtr<FUICommandInfo> ExistingInfo = CommandInfoMap.FindRef( ExisingCommandName );
-			if( *ExistingInfo->ActiveGesture != ExistingInfo->DefaultGesture )
+			if( *ExistingInfo->ActiveChord != ExistingInfo->DefaultChord )
 			{
-				// two user defined gestures are the same within a context.  If the keybinding editor was used this wont happen so this must have been directly a modified user setting file
-				UE_LOG(LogSlate, Error, TEXT("Duplicate user defined gestures found: [%s,%s].  Gesture for %s being removed"), *InCommandInfo->GetLabel().ToString(), *ExistingInfo->GetLabel().ToString(), *ExistingInfo->GetLabel().ToString() );
+				// two user defined chords are the same within a context.  If the keybinding editor was used this wont happen so this must have been directly a modified user setting file
+				UE_LOG(LogSlate, Error, TEXT("Duplicate user defined chords found: [%s,%s].  Chord for %s being removed"), *InCommandInfo->GetLabel().ToString(), *ExistingInfo->GetLabel().ToString(), *ExistingInfo->GetLabel().ToString() );
 			}
-			ContextEntry.GestureToCommandInfoMap.Remove( *ExistingInfo->ActiveGesture );
-			// Remove the existing gesture. 
-			ExistingInfo->ActiveGesture = MakeShareable( new FInputGesture() );
+			ContextEntry.ChordToCommandInfoMap.Remove( *ExistingInfo->ActiveChord );
+			// Remove the existing chord. 
+			ExistingInfo->ActiveChord = MakeShareable( new FInputChord() );
 		
 		}
 
-		TSharedRef<FInputGesture> NewGesture = MakeShareable( new FInputGesture( UserDefinedGesture ) );
-		// Set the active gesture on the command info
-		InCommandInfo->ActiveGesture = NewGesture;
+		TSharedRef<FInputChord> NewChord = MakeShareable( new FInputChord( UserDefinedChord ) );
+		// Set the active chord on the command info
+		InCommandInfo->ActiveChord = NewChord;
 	}
 
-	// If the active gesture is valid, map the gesture to the map for fast lookup when processing bindings
-	if( InCommandInfo->ActiveGesture->IsValidGesture() )
+	// If the active chord is valid, map the chord to the map for fast lookup when processing bindings
+	if( InCommandInfo->ActiveChord->IsValidChord() )
 	{
-		checkSlow( !ContextEntry.GestureToCommandInfoMap.Contains( *InCommandInfo->GetActiveGesture() ) );
-		ContextEntry.GestureToCommandInfoMap.Add( *InCommandInfo->GetActiveGesture(), InCommandInfo->GetCommandName() );
+		checkSlow( !ContextEntry.ChordToCommandInfoMap.Contains( *InCommandInfo->GetActiveChord() ) );
+		ContextEntry.ChordToCommandInfoMap.Add( *InCommandInfo->GetActiveChord(), InCommandInfo->GetCommandName() );
 	}
 }
 
-const TSharedPtr<FUICommandInfo> FInputBindingManager::FindCommandInContext( const FName InBindingContext, const FInputGesture& InGesture, bool bCheckDefault ) const
+const TSharedPtr<FUICommandInfo> FInputBindingManager::FindCommandInContext( const FName InBindingContext, const FInputChord& InChord, bool bCheckDefault ) const
 {
 	const FContextEntry& ContextEntry = ContextMap.FindRef( InBindingContext );
 	
@@ -450,7 +465,7 @@ const TSharedPtr<FUICommandInfo> FInputBindingManager::FindCommandInContext( con
 		for( FCommandInfoMap::TConstIterator It(InfoMap); It && !FoundCommand.IsValid(); ++It )
 		{
 			const FUICommandInfo& CommandInfo = *It.Value();
-			if( CommandInfo.DefaultGesture == InGesture )
+			if( CommandInfo.DefaultChord == InChord )
 			{
 				FoundCommand = It.Value();
 			}	
@@ -458,8 +473,8 @@ const TSharedPtr<FUICommandInfo> FInputBindingManager::FindCommandInContext( con
 	}
 	else
 	{
-		// faster lookup for active gestures
-		FName CommandName = ContextEntry.GestureToCommandInfoMap.FindRef( InGesture );
+		// faster lookup for active chords
+		FName CommandName = ContextEntry.ChordToCommandInfoMap.FindRef( InChord );
 		if( CommandName != NAME_None )
 		{
 			FoundCommand = ContextEntry.CommandInfoMap.FindChecked( CommandName );
@@ -489,7 +504,7 @@ void FInputBindingManager::GetAllChildContexts( const FName InBindingContext, TA
 	}
 }
 
-const TSharedPtr<FUICommandInfo> FInputBindingManager::GetCommandInfoFromInputGesture( const FName InBindingContext, const FInputGesture& InGesture, bool bCheckDefault ) const
+const TSharedPtr<FUICommandInfo> FInputBindingManager::GetCommandInfoFromInputChord( const FName InBindingContext, const FInputChord& InChord, bool bCheckDefault ) const
 {
 	TSharedPtr<FUICommandInfo> FoundCommand = NULL;
 
@@ -498,7 +513,7 @@ const TSharedPtr<FUICommandInfo> FInputBindingManager::GetCommandInfoFromInputGe
 	{
 		const FContextEntry& ContextEntry = ContextMap.FindRef( CurrentContext );
 
-		FoundCommand = FindCommandInContext( CurrentContext, InGesture, bCheckDefault );
+		FoundCommand = FindCommandInContext( CurrentContext, InChord, bCheckDefault );
 
 		CurrentContext = ContextEntry.BindingContext->GetContextParent();
 	}
@@ -510,7 +525,7 @@ const TSharedPtr<FUICommandInfo> FInputBindingManager::GetCommandInfoFromInputGe
 
 		for( int32 ContextIndex = 0; ContextIndex < Children.Num() && !FoundCommand.IsValid(); ++ContextIndex )
 		{
-			FoundCommand = FindCommandInContext( Children[ContextIndex], InGesture, bCheckDefault );
+			FoundCommand = FindCommandInContext( Children[ContextIndex], InChord, bCheckDefault );
 		}
 	}
 	
diff --git a/Engine/Source/Runtime/Slate/Private/Framework/Commands/InputChord.cpp b/Engine/Source/Runtime/Slate/Private/Framework/Commands/InputChord.cpp
new file mode 100644
--- /dev/null
+++ b/Engine/Source/Runtime/Slate/Private/Framework/Commands/InputChord.cpp
@@ -0,0 +1,112 @@
+// Copyright 1998-2015 Epic Games, Inc. All Rights Reserved.
+
+#include "SlatePrivatePCH.h"
+
+#define LOCTEXT_NAMESPACE "FInputChord"
+
+/* FInputGesture interface
+ *****************************************************************************/
+
+/**
+ * Returns the friendly, localized string name of this key binding
+ * @todo Slate: Got to be a better way to do this
+ */
+FText FInputChord::GetInputText( ) const
+{
+#if PLATFORM_MAC
+    const FText CommandText = LOCTEXT("KeyName_Control", "Ctrl");
+    const FText ControlText = LOCTEXT("KeyName_Command", "Cmd");
+#else
+    const FText ControlText = LOCTEXT("KeyName_Control", "Ctrl");
+    const FText CommandText = LOCTEXT("KeyName_Command", "Cmd"); 
+#endif
+    const FText AltText = LOCTEXT("KeyName_Alt", "Alt");
+    const FText ShiftText = LOCTEXT("KeyName_Shift", "Shift");
+    
+	const FText AppenderText = LOCTEXT("ModAppender", "+");
+
+	FFormatNamedArguments Args;
+	int32 ModCount = 0;
+
+    if (bCtrl)
+    {
+		Args.Add(FString::Printf(TEXT("Mod%d"),++ModCount), ControlText);
+    }
+    if (bCmd)
+    {
+		Args.Add(FString::Printf(TEXT("Mod%d"),++ModCount), CommandText);
+    }
+    if (bAlt)
+    {
+		Args.Add(FString::Printf(TEXT("Mod%d"),++ModCount), AltText);
+    }
+    if (bShift)
+    {
+		Args.Add(FString::Printf(TEXT("Mod%d"),++ModCount), ShiftText);
+    }
+
+	for (int32 i = 1; i <= 4; ++i)
+	{
+		if (i > ModCount)
+		{
+			Args.Add(FString::Printf(TEXT("Mod%d"), i), FText::GetEmpty());
+			Args.Add(FString::Printf(TEXT("Appender%d"), i), FText::GetEmpty());
+		}
+		else
+		{
+			Args.Add(FString::Printf(TEXT("Appender%d"), i), AppenderText);
+		}
+
+	}
+
+	Args.Add(TEXT("Key"), GetKeyText());
+
+	return FText::Format(LOCTEXT("FourModifiers", "{Mod1}{Appender1}{Mod2}{Appender2}{Mod3}{Appender3}{Mod4}{Appender4}{Key}"), Args);
+}
+
+
+FText FInputChord::GetKeyText( ) const
+{
+	FText OutString; // = KeyGetDisplayName(Key);
+
+	if (Key.IsValid() && !Key.IsModifierKey())
+	{
+		OutString = Key.GetDisplayName();
+	}
+
+	return OutString;
+}
+
+FInputChord::ERelationshipType FInputChord::GetRelationship( const FInputChord& OtherChord ) const
+{
+	ERelationshipType Relationship = None;
+
+	if (Key == OtherChord.Key)
+	{
+		if ((bAlt == OtherChord.bAlt) &&
+			(bCtrl == OtherChord.bCtrl) &&
+			(bShift == OtherChord.bShift) &&
+			(bCmd == OtherChord.bCmd))
+		{
+			Relationship = Same;
+		}
+		else if ((bAlt || !OtherChord.bAlt) &&
+				(bCtrl || !OtherChord.bCtrl) &&
+				(bShift || !OtherChord.bShift) &&
+				(bCmd || !OtherChord.bCmd))
+		{
+			Relationship = Masks;
+		}
+		else if ((!bAlt || OtherChord.bAlt) &&
+				(!bCtrl || OtherChord.bCtrl) &&
+				(!bShift || OtherChord.bShift) &&
+				(!bCmd || OtherChord.bCmd))
+		{
+			Relationship = Masked;
+		}
+	}
+
+	return Relationship;
+}
+
+#undef LOCTEXT_NAMESPACE
diff --git a/Engine/Source/Runtime/Slate/Private/Framework/Commands/UICommandInfo.cpp b/Engine/Source/Runtime/Slate/Private/Framework/Commands/UICommandInfo.cpp
--- a/Engine/Source/Runtime/Slate/Private/Framework/Commands/UICommandInfo.cpp
+++ b/Engine/Source/Runtime/Slate/Private/Framework/Commands/UICommandInfo.cpp
@@ -24,9 +24,9 @@ FUICommandInfoDecl::FUICommandInfoDecl( const TSharedRef<FBindingContext>& InCon
 	Info->Description = InDesc;
 }
 
-FUICommandInfoDecl& FUICommandInfoDecl::DefaultGesture( const FInputGesture& InDefaultGesture )
+FUICommandInfoDecl& FUICommandInfoDecl::DefaultChord( const FInputChord& InDefaultChord )
 {
-	Info->DefaultGesture = InDefaultGesture;
+	Info->DefaultChord = InDefaultChord;
 	return *this;
 }
 FUICommandInfoDecl& FUICommandInfoDecl::UserInterfaceType( EUserInterfaceActionType::Type InType )
@@ -63,11 +63,11 @@ FUICommandInfoDecl::operator TSharedRef<FUICommandInfo>() const
 
 const FText FUICommandInfo::GetInputText() const
 {	
-	return ActiveGesture->GetInputText();
+	return ActiveChord->GetInputText();
 }
 
 
-void FUICommandInfo::MakeCommandInfo( const TSharedRef<class FBindingContext>& InContext, TSharedPtr< FUICommandInfo >& OutCommand, const FName InCommandName, const FText& InCommandLabel, const FText& InCommandDesc, const FSlateIcon& InIcon, const EUserInterfaceActionType::Type InUserInterfaceType, const FInputGesture& InDefaultGesture )
+void FUICommandInfo::MakeCommandInfo( const TSharedRef<class FBindingContext>& InContext, TSharedPtr< FUICommandInfo >& OutCommand, const FName InCommandName, const FText& InCommandLabel, const FText& InCommandDesc, const FSlateIcon& InIcon, const EUserInterfaceActionType::Type InUserInterfaceType, const FInputChord& InDefaultChord )
 {
 	ensureMsgf( !InCommandLabel.IsEmpty(), TEXT("Command labels cannot be empty") );
 
@@ -77,24 +77,24 @@ void FUICommandInfo::MakeCommandInfo( const TSharedRef<class FBindingContext>& I
 	OutCommand->Description = InCommandDesc;
 	OutCommand->Icon = InIcon;
 	OutCommand->UserInterfaceType = InUserInterfaceType;
-	OutCommand->DefaultGesture = InDefaultGesture;
+	OutCommand->DefaultChord = InDefaultChord;
 	FInputBindingManager::Get().CreateInputCommand( InContext, OutCommand.ToSharedRef() );
 }
 
-void FUICommandInfo::SetActiveGesture( const FInputGesture& NewGesture )
+void FUICommandInfo::SetActiveChord( const FInputChord& NewChord )
 {
-	ActiveGesture->Set( NewGesture );
+	ActiveChord->Set( NewChord );
 
-	// Set the user defined gesture for this command so it can be saved to disk later
-	FInputBindingManager::Get().NotifyActiveGestureChanged( *this );
+	// Set the user defined chord for this command so it can be saved to disk later
+	FInputBindingManager::Get().NotifyActiveChordChanged( *this );
 }
 
-void FUICommandInfo::RemoveActiveGesture()
+void FUICommandInfo::RemoveActiveChord()
 {
-	// Gesture already exists
-	// Reset the other gesture that has the same binding
-	ActiveGesture = MakeShareable( new FInputGesture() );
-	FInputBindingManager::Get().NotifyActiveGestureChanged( *this );
+	// Chord already exists
+	// Reset the other chord that has the same binding
+	ActiveChord = MakeShareable( new FInputChord() );
+	FInputBindingManager::Get().NotifyActiveChordChanged( *this );
 }
 
 TSharedRef<SToolTip> FUICommandInfo::MakeTooltip( const TAttribute<FText>& InText, const TAttribute< EVisibility >& InToolTipVisibility ) const
diff --git a/Engine/Source/Runtime/Slate/Private/Framework/Commands/UICommandList.cpp b/Engine/Source/Runtime/Slate/Private/Framework/Commands/UICommandList.cpp
--- a/Engine/Source/Runtime/Slate/Private/Framework/Commands/UICommandList.cpp
+++ b/Engine/Source/Runtime/Slate/Private/Framework/Commands/UICommandList.cpp
@@ -161,9 +161,9 @@ bool FUICommandList::ConditionalProcessCommandBindings( const FKey Key, bool bCt
 {
 	if ( !bRepeat && !FSlateApplication::Get().IsDragDropping() )
 	{
-		FInputGesture CheckGesture( Key, EModifierKey::FromBools(bCtrl, bAlt, bShift, bCmd) );
+		FInputChord CheckChord( Key, EModifierKey::FromBools(bCtrl, bAlt, bShift, bCmd) );
 
-		if( CheckGesture.IsValidGesture() )
+		if( CheckChord.IsValidChord() )
 		{
 			TSet<FName> AllContextsToCheck;
 			GatherContextsForList(AllContextsToCheck);
@@ -172,13 +172,13 @@ bool FUICommandList::ConditionalProcessCommandBindings( const FKey Key, bool bCt
 			{
 				FName Context = *It;
 	
-				// Only active gestures process commands
+				// Only active chords process commands
 				const bool bCheckDefault = false;
 
-				// Check to see if there is any command in the context activated by the gesture
-				TSharedPtr<FUICommandInfo> Command = FInputBindingManager::Get().FindCommandInContext( Context, CheckGesture, bCheckDefault );
+				// Check to see if there is any command in the context activated by the chord
+				TSharedPtr<FUICommandInfo> Command = FInputBindingManager::Get().FindCommandInContext( Context, CheckChord, bCheckDefault );
 
-				if( Command.IsValid() && ensure( *Command->GetActiveGesture() == CheckGesture ) )
+				if( Command.IsValid() && ensure( *Command->GetActiveChord() == CheckChord ) )
 				{
 					// Find the bound action for this command
 					const FUIAction* Action = GetActionForCommand(Command);
@@ -194,7 +194,7 @@ bool FUICommandList::ConditionalProcessCommandBindings( const FKey Key, bool bCt
 						}
 						else
 						{
-							// An action wasn't bound to the command but a gesture was found or an action was found but cant be executed
+							// An action wasn't bound to the command but a chord was found or an action was found but cant be executed
 							return false;
 						}
 					}
diff --git a/Engine/Source/Runtime/Slate/Private/Framework/MultiBox/Mac/MacMenu.cpp b/Engine/Source/Runtime/Slate/Private/Framework/MultiBox/Mac/MacMenu.cpp
--- a/Engine/Source/Runtime/Slate/Private/Framework/MultiBox/Mac/MacMenu.cpp
+++ b/Engine/Source/Runtime/Slate/Private/Framework/MultiBox/Mac/MacMenu.cpp
@@ -433,27 +433,27 @@ NSString* FSlateMacMenu::GetMenuItemKeyEquivalent(const TSharedRef<const class F
 {
 	if (Block->GetAction().IsValid())
 	{
-		const TSharedRef<const FInputGesture>& Gesture = Block->GetAction()->GetActiveGesture();
+		const TSharedRef<const FInputChord>& Chord = Block->GetAction()->GetActiveChord();
 
 		*OutModifiers = 0;
-		if (Gesture->NeedsControl())
+		if (Chord->NeedsControl())
 		{
 			*OutModifiers |= NSControlKeyMask;
 		}
-		if (Gesture->NeedsShift())
+		if (Chord->NeedsShift())
 		{
 			*OutModifiers |= NSShiftKeyMask;
 		}
-		if (Gesture->NeedsAlt())
+		if (Chord->NeedsAlt())
 		{
 			*OutModifiers |= NSAlternateKeyMask;
 		}
-		if (Gesture->NeedsCommand())
+		if (Chord->NeedsCommand())
 		{
 			*OutModifiers |= NSCommandKeyMask;
 		}
 
-		FString KeyString = Gesture->GetKeyText().ToString().ToLower();
+		FString KeyString = Chord->GetKeyText().ToString().ToLower();
 		return KeyString.GetNSString();
 	}
 	return @"";
diff --git a/Engine/Source/Runtime/Slate/Private/Framework/MultiBox/SToolBarButtonBlock.cpp b/Engine/Source/Runtime/Slate/Private/Framework/MultiBox/SToolBarButtonBlock.cpp
--- a/Engine/Source/Runtime/Slate/Private/Framework/MultiBox/SToolBarButtonBlock.cpp
+++ b/Engine/Source/Runtime/Slate/Private/Framework/MultiBox/SToolBarButtonBlock.cpp
@@ -109,7 +109,7 @@ void SToolBarButtonBlock::BuildMultiBlockWidget(const ISlateStyle* StyleSet, con
 		{
 			TSharedPtr<const FUICommandInfo> CommandPtr = Command.Pin();
 
-			if( CommandPtr.IsValid() && CommandPtr->GetActiveGesture()->IsValidGesture() )
+			if( CommandPtr.IsValid() && CommandPtr->GetActiveChord()->IsValidChord() )
 			{
 				FFormatNamedArguments Args;
 				Args.Add( TEXT("ToolTipDescription"), ToolTip.Get() );
diff --git a/Engine/Source/Runtime/Slate/Public/Framework/Commands/Commands.h b/Engine/Source/Runtime/Slate/Public/Framework/Commands/Commands.h
--- a/Engine/Source/Runtime/Slate/Public/Framework/Commands/Commands.h
+++ b/Engine/Source/Runtime/Slate/Public/Framework/Commands/Commands.h
@@ -11,11 +11,11 @@
 
 #define LOC_DEFINE_REGION
 
-SLATE_API void UI_COMMAND_Function( FBindingContext* This, TSharedPtr< FUICommandInfo >& OutCommand, const TCHAR* OutSubNamespace, const TCHAR* OutCommandName, const TCHAR* OutCommandNameUnderscoreTooltip, const ANSICHAR* DotOutCommandName, const TCHAR* FriendlyName, const TCHAR* InDescription, const EUserInterfaceActionType::Type CommandType, const FInputGesture& InDefaultGesture );
+SLATE_API void UI_COMMAND_Function( FBindingContext* This, TSharedPtr< FUICommandInfo >& OutCommand, const TCHAR* OutSubNamespace, const TCHAR* OutCommandName, const TCHAR* OutCommandNameUnderscoreTooltip, const ANSICHAR* DotOutCommandName, const TCHAR* FriendlyName, const TCHAR* InDescription, const EUserInterfaceActionType::Type CommandType, const FInputChord& InDefaultChord );
 
 // This macro requires LOCTEXT_NAMESPACE to be defined. If you don't want the command to be placed under a sub namespace, provide "" as the namespace.
-#define UI_COMMAND( CommandId, FriendlyName, InDescription, CommandType, InDefaultGesture ) \
-	UI_COMMAND_Function( this, CommandId, TEXT(LOCTEXT_NAMESPACE), TEXT(#CommandId), TEXT(#CommandId) TEXT("_ToolTip"), "." #CommandId, TEXT(FriendlyName), TEXT(InDescription), CommandType, InDefaultGesture );
+#define UI_COMMAND( CommandId, FriendlyName, InDescription, CommandType, InDefaultChord ) \
+	UI_COMMAND_Function( this, CommandId, TEXT(LOCTEXT_NAMESPACE), TEXT(#CommandId), TEXT(#CommandId) TEXT("_ToolTip"), "." #CommandId, TEXT(FriendlyName), TEXT(InDescription), CommandType, InDefaultChord );
 
 #undef LOC_DEFINE_REGION
 
diff --git a/Engine/Source/Runtime/Slate/Public/Framework/Commands/InputBindingManager.h b/Engine/Source/Runtime/Slate/Public/Framework/Commands/InputBindingManager.h
--- a/Engine/Source/Runtime/Slate/Public/Framework/Commands/InputBindingManager.h
+++ b/Engine/Source/Runtime/Slate/Public/Framework/Commands/InputBindingManager.h
@@ -4,12 +4,16 @@
 
 
 typedef TMap<FName, TSharedPtr<FUICommandInfo>> FCommandInfoMap;
-typedef TMap<FInputGesture, FName> FGestureMap;
+typedef TMap<FInputChord, FName> FChordMap;
+DEPRECATED(4.8, "Use FChordMap instead of FGestureMap")
+typedef FChordMap FGestureMap;
 
 
-/** Delegate for alerting subscribers the input manager records a user-defined gesture */
-DECLARE_MULTICAST_DELEGATE_OneParam(FOnUserDefinedGestureChanged, const FUICommandInfo& );
+/** Delegate for alerting subscribers the input manager records a user-defined chord */
+DECLARE_MULTICAST_DELEGATE_OneParam(FOnUserDefinedChordChanged, const FUICommandInfo& );
 
+DEPRECATED(4.8, "Use FOnUserDefinedChordChanged instead of FOnUserDefinedGestureChanged")
+typedef FOnUserDefinedChordChanged FOnUserDefinedGestureChanged;
 
 /**
  * Manager responsible for creating and processing input bindings.                  
@@ -47,29 +51,31 @@ public:
 	 * @param InBindingContext		The context where the command is valid
 	 * @param InUserInterfaceAction	The user interface action to create the input command from
 	 *								Will automatically bind the user interface action to the command specified 
-	 *								by the action so when the command's gesture is pressed, the action is executed
+	 *								by the action so when the command's chord is pressed, the action is executed
 	 */
 	void CreateInputCommand( const TSharedRef<FBindingContext>& InBindingContext, TSharedRef<FUICommandInfo> InUICommandInfo );
 
 	/**
-	 * Returns a command info that is has the same active gesture as the provided gesture and is in the same binding context or parent context
+	 * Returns a command info that is has the same active chord as the provided chord and is in the same binding context or parent context
 	 *
 	 * @param	InBindingContext		The context in which the command is valid
-	 * @param	InGesture				The gesture to match against commands
-	 * @param	bCheckDefault			Whether or not to check the default gesture.  Will check the active gesture if false
-	 * @return	A pointer to the command info which has the InGesture as its active gesture or null if one cannot be found
+	 * @param	InChord					The chord to match against commands
+	 * @param	bCheckDefault			Whether or not to check the default chord.  Will check the active chord if false
+	 * @return	A pointer to the command info which has the InChord as its active chord or null if one cannot be found
 	 */
-	const TSharedPtr<FUICommandInfo> GetCommandInfoFromInputGesture( const FName InBindingContext, const FInputGesture& InGesture, bool bCheckDefault ) const;
+	DEPRECATED(4.8, "Use GetCommandInfoFromInputChord instead of GetCommandInfoFromInputGesture")
+	const TSharedPtr<FUICommandInfo> GetCommandInfoFromInputGesture( const FName InBindingContext, const FInputChord& InGesture, bool bCheckDefault ) const;
+	const TSharedPtr<FUICommandInfo> GetCommandInfoFromInputChord( const FName InBindingContext, const FInputChord& InChord, bool bCheckDefault ) const;
 
 	/** 
-	 * Finds the command in the provided context which uses the provided input gesture
+	 * Finds the command in the provided context which uses the provided input chord
 	 * 
 	 * @param InBindingContext	The binding context name
-	 * @param InGesture			The gesture to check against when looking for commands
-	 * @param bCheckDefault		Whether or not to check the default gesture of commands instead of active gestures
-	 * @param OutActiveGestureIndex	The index into the commands active gesture array where the passed in gesture was matched.  INDEX_NONE if bCheckDefault is true or nothing was found.
+	 * @param InChord			The chord to check against when looking for commands
+	 * @param bCheckDefault		Whether or not to check the default chord of commands instead of active chords
+	 * @param OutActiveChordIndex	The index into the commands active chord array where the passed in chord was matched.  INDEX_NONE if bCheckDefault is true or nothing was found.
 	 */
-	const TSharedPtr<FUICommandInfo> FindCommandInContext( const FName InBindingContext, const FInputGesture& InGesture, bool bCheckDefault ) const;
+	const TSharedPtr<FUICommandInfo> FindCommandInContext( const FName InBindingContext, const FInputChord& InChord, bool bCheckDefault ) const;
 
 	/** 
 	 * Finds the command in the provided context which has the provided name 
@@ -80,21 +86,25 @@ public:
 	const TSharedPtr<FUICommandInfo> FindCommandInContext( const FName InBindingContext, const FName CommandName ) const;
 
 	/**
-	 * Called when the active gesture is changed on a command
+	 * Called when the active chord is changed on a command
 	 *
-	 * @param CommandInfo	The command that had the active gesture changed. 
+	 * @param CommandInfo	The command that had the active chord changed. 
 	 */
-	virtual void NotifyActiveGestureChanged( const FUICommandInfo& CommandInfo );
+	DEPRECATED(4.8, "Use NotifyActiveChordChanged instead of NotifyActiveGestureChanged")
+	void NotifyActiveGestureChanged( const FUICommandInfo& CommandInfo );
+	virtual void NotifyActiveChordChanged( const FUICommandInfo& CommandInfo );
 	
 	/**
-	 * Saves the user defined gestures to a json file
+	 * Saves the user defined chords to a json file
 	 */
 	void SaveInputBindings();
 
 	/**
-	 * Removes any user defined gestures
+	 * Removes any user defined chords
 	 */
-	void RemoveUserDefinedGestures();
+	DEPRECATED(4.8, "Use RemoveUserDefinedChords instead of RemoveUserDefinedGestures")
+	void RemoveUserDefinedGestures() { RemoveUserDefinedChords(); }
+	void RemoveUserDefinedChords();
 
 	/**
 	 * Returns all known command infos for a given binding context
@@ -104,27 +114,37 @@ public:
 	 */
 	void GetCommandInfosFromContext( const FName InBindingContext, TArray< TSharedPtr<FUICommandInfo> >& OutCommandInfos ) const;
 
-	/** Registers a delegate to be called when a user-defined gesture is edited */
-	FDelegateHandle RegisterUserDefinedGestureChanged(const FOnUserDefinedGestureChanged::FDelegate& Delegate)
+	/** Registers a delegate to be called when a user-defined chord is edited */
+	DEPRECATED(4.8, "Use RegisterUserDefinedChordChanged instead of RegisterUserDefinedGestureChanged")
+	FDelegateHandle RegisterUserDefinedGestureChanged(const FOnUserDefinedChordChanged::FDelegate& Delegate)
 	{
-		return OnUserDefinedGestureChanged.Add(Delegate);
+		return RegisterUserDefinedChordChanged(Delegate);
+	}
+	FDelegateHandle RegisterUserDefinedChordChanged(const FOnUserDefinedChordChanged::FDelegate& Delegate)
+	{
+		return OnUserDefinedChordChanged.Add(Delegate);
 	}
 
-	/** Unregisters a delegate to be called when a user-defined gesture is edited */
-	DELEGATE_DEPRECATED("This UnregisterUserDefinedGestureChanged overload is deprecated - please remove delegates using the FDelegateHandle returned by the RegisterUserDefinedGestureChanged function.")
-	void UnregisterUserDefinedGestureChanged(const FOnUserDefinedGestureChanged::FDelegate& Delegate)
+	/** Unregisters a delegate to be called when a user-defined chord is edited */
+	DELEGATE_DEPRECATED("This UnregisterUserDefinedGestureChanged overload is deprecated - please remove delegates using the FDelegateHandle returned by the RegisterUserDefinedChordChanged function.")
+	void UnregisterUserDefinedGestureChanged(const FOnUserDefinedChordChanged::FDelegate& Delegate)
 	{
 		PRAGMA_DISABLE_DEPRECATION_WARNINGS
 
-		OnUserDefinedGestureChanged.Remove(Delegate);
+		OnUserDefinedChordChanged.Remove(Delegate);
 
 		PRAGMA_ENABLE_DEPRECATION_WARNINGS
 	}
 
-	/** Unregisters a delegate to be called when a user-defined gesture is edited */
+	/** Unregisters a delegate to be called when a user-defined chord is edited */
+	DEPRECATED(4.8, "Use RegisterUserDefinedChordChanged instead of RegisterUserDefinedGestureChanged")
 	void UnregisterUserDefinedGestureChanged(FDelegateHandle DelegateHandle)
 	{
-		OnUserDefinedGestureChanged.Remove(DelegateHandle);
+		UnregisterUserDefinedChordChanged(DelegateHandle);
+	}
+	void UnregisterUserDefinedChordChanged(FDelegateHandle DelegateHandle)
+	{
+		OnUserDefinedChordChanged.Remove(DelegateHandle);
 	}
 
 private:
@@ -135,17 +155,17 @@ private:
 	FInputBindingManager() { }
 
 	/**
-	 * Gets the user defined gesture (if any) from the provided command name
+	 * Gets the user defined chord (if any) from the provided command name
 	 * 
 	 * @param InBindingContext	The context in which the command is active
-	 * @param InCommandName		The name of the command to get the gesture from
+	 * @param InCommandName		The name of the command to get the chord from
 	 */
-	bool GetUserDefinedGesture( const FName InBindingContext, const FName InCommandName, FInputGesture& OutUserDefinedGesture );
+	bool GetUserDefinedChord( const FName InBindingContext, const FName InCommandName, FInputChord& OutUserDefinedChord );
 
 	/**
-	 *	Checks a binding context for duplicate gestures 
+	 *	Checks a binding context for duplicate chords 
 	 */
-	void CheckForDuplicateDefaultGestures( const FBindingContext& InBindingContext, TSharedPtr<FUICommandInfo> InCommandInfo ) const;
+	void CheckForDuplicateDefaultChords( const FBindingContext& InBindingContext, TSharedPtr<FUICommandInfo> InCommandInfo ) const;
 
 	/**
 	 * Recursively finds all child contexts of the provided binding context.  
@@ -162,8 +182,8 @@ private:
 		/** A list of commands associated with the context */
 		FCommandInfoMap CommandInfoMap;
 
-		/** Gesture to command info map */
-		FGestureMap GestureToCommandInfoMap;
+		/** Chord to command info map */
+		FChordMap ChordToCommandInfoMap;
 
 		/** The binding context for this entry*/
 		TSharedPtr< FBindingContext > BindingContext;
@@ -175,9 +195,9 @@ private:
 	/** A mapping of contexts to their child contexts */
 	TMultiMap< FName, FName > ParentToChildMap;
 
-	/** User defined gesture overrides for commands */
-	TSharedPtr< class FUserDefinedGestures > UserDefinedGestures;
+	/** User defined chord overrides for commands */
+	TSharedPtr< class FUserDefinedChords > UserDefinedChords;
 
-	/** Delegate called when a user-defined gesture is edited */
-	FOnUserDefinedGestureChanged OnUserDefinedGestureChanged;
+	/** Delegate called when a user-defined chord is edited */
+	FOnUserDefinedChordChanged OnUserDefinedChordChanged;
 };
diff --git a/Engine/Source/Runtime/Slate/Public/Framework/Commands/UICommandInfo.h b/Engine/Source/Runtime/Slate/Public/Framework/Commands/UICommandInfo.h
--- a/Engine/Source/Runtime/Slate/Public/Framework/Commands/UICommandInfo.h
+++ b/Engine/Source/Runtime/Slate/Public/Framework/Commands/UICommandInfo.h
@@ -4,7 +4,7 @@
 
 #include "Layout/Visibility.h"
 #include "Misc/Attribute.h"
-#include "Framework/Commands/InputGesture.h"
+#include "Framework/Commands/InputChord.h"
 #include "Textures/SlateIcon.h"
 
 
@@ -40,7 +40,9 @@ class SLATE_API FUICommandInfoDecl
 
 public:
 
-	FUICommandInfoDecl& DefaultGesture( const FInputGesture& InDefaultGesture );
+	DEPRECATED(4.8, "Use DefaultChord instead of DefaultGesture")
+	FUICommandInfoDecl& DefaultGesture( const FInputChord& InDefaultGesture );
+	FUICommandInfoDecl& DefaultChord( const FInputChord& InDefaultChord );
 	FUICommandInfoDecl& UserInterfaceType( EUserInterfaceActionType::Type InType );
 	FUICommandInfoDecl& Icon( const FSlateIcon& InIcon );
 	FUICommandInfoDecl& Description( const FText& InDesc );
@@ -157,29 +159,32 @@ public:
 	 * @param InBindingContext The name of the binding context to use.
 	 */
 	FUICommandInfo( const FName InBindingContext )
-		: ActiveGesture( new FInputGesture )
-		, DefaultGesture( EKeys::Invalid, EModifierKey::None )
+		: ActiveChord( new FInputChord )
+		, DefaultChord( EKeys::Invalid, EModifierKey::None )
 		, BindingContext( InBindingContext )
 		, UserInterfaceType( EUserInterfaceActionType::Button )
 	{ }
 
 	/**
-	 * Returns the friendly, localized string name of the gesture that is required to perform the command
+	 * Returns the friendly, localized string name of the chord that is required to perform the command
 	 *
-	 * @param	bAllGestures	If true a comma separated string with all secondary active gestures will be returned.  Otherwise just the main gesture string is returned
-	 * @return	Localized friendly text for the gesture
+	 * @return	Localized friendly text for the chord
 	 */
 	const FText GetInputText() const;
 
 	/**
-	 * @return	Returns the active gesture for this command
+	 * @return	Returns the active chord for this command
 	 */
-	const TSharedRef<const FInputGesture> GetActiveGesture() const { return ActiveGesture; }
+	DEPRECATED(4.8, "Use GetActiveChord instead of GetActiveGesture")
+	const TSharedRef<const FInputChord> GetActiveGesture() const { return GetActiveChord(); }
+	const TSharedRef<const FInputChord> GetActiveChord() const { return ActiveChord; }
 
-	const FInputGesture& GetDefaultGesture() const { return DefaultGesture; }
+	DEPRECATED(4.8, "Use GetDefaultChord instead of GetDefaultGesture")
+	const FInputChord& GetDefaultGesture() const { return GetDefaultChord(); }
+	const FInputChord& GetDefaultChord() const { return DefaultChord; }
 
 	/** Utility function to make an FUICommandInfo */
-	static void MakeCommandInfo( const TSharedRef<class FBindingContext>& InContext, TSharedPtr< FUICommandInfo >& OutCommand, const FName InCommandName, const FText& InCommandLabel, const FText& InCommandDesc, const FSlateIcon& InIcon, const EUserInterfaceActionType::Type InUserInterfaceType, const FInputGesture& InDefaultGesture );
+	static void MakeCommandInfo( const TSharedRef<class FBindingContext>& InContext, TSharedPtr< FUICommandInfo >& OutCommand, const FName InCommandName, const FText& InCommandLabel, const FText& InCommandDesc, const FSlateIcon& InIcon, const EUserInterfaceActionType::Type InUserInterfaceType, const FInputChord& InDefaultChord );
 
 	/** @return The display label for this command */
 	const FText& GetLabel() const { return Label; }
@@ -199,11 +204,15 @@ public:
 	/** @return The name of the context where the command is valid */
 	FName GetBindingContext() const { return BindingContext; }
 
-	/** Sets the new active gesture for this command */
-	void SetActiveGesture( const FInputGesture& NewGesture );
+	/** Sets the new active chord for this command */
+	DEPRECATED(4.8, "Use SetActiveChord instead of SetActiveGesture")
+	void SetActiveGesture( const FInputChord& NewGesture ) { SetActiveChord(NewGesture); }
+	void SetActiveChord( const FInputChord& NewChord );
 
-	/** Removes the active gesture from this command */
-	void RemoveActiveGesture();
+	/** Removes the active chord from this command */
+	DEPRECATED(4.8, "Use RemoveActiveChord instead of RemoveActiveGesture")
+	void RemoveActiveGesture() { RemoveActiveChord(); }
+	void RemoveActiveChord();
 
 	/** 
 	 * Makes a tooltip for this command.
@@ -215,7 +224,7 @@ public:
 private:
 
 	/** Input command that executes this action */
-	TSharedRef<FInputGesture> ActiveGesture;
+	TSharedRef<FInputChord> ActiveChord;
 
 	/** Default display name of the command */
 	FText Label;
@@ -223,8 +232,8 @@ private:
 	/** Localized help text for this command */
 	FText Description;
 
-	/** The default input gesture for this command (can be invalid) */
-	FInputGesture DefaultGesture;
+	/** The default input chord for this command (can be invalid) */
+	FInputChord DefaultChord;
 
 	/** Brush name for icon to use in tool bars and menu items to represent this command */
 	FSlateIcon Icon;
diff --git a/Engine/Source/Runtime/Slate/Public/SlateBasics.h b/Engine/Source/Runtime/Slate/Public/SlateBasics.h
--- a/Engine/Source/Runtime/Slate/Public/SlateBasics.h
+++ b/Engine/Source/Runtime/Slate/Public/SlateBasics.h
@@ -31,7 +31,7 @@
 #include "SlateIcon.h"
 
 // Commands
-#include "InputGesture.h"
+#include "InputChord.h"
 #include "UIAction.h"
 #include "UICommandInfo.h"
 #include "InputBindingManager.h"
