Replace renamed InputGesture.h with one that provides a warning message and redirects to InputChord.h

diff --git a/Engine/Source/Runtime/Slate/Public/Framework/Commands/InputGesture.h b/Engine/Source/Runtime/Slate/Public/Framework/Commands/InputGesture.h
new file mode 100644
--- /dev/null
+++ b/Engine/Source/Runtime/Slate/Public/Framework/Commands/InputGesture.h
@@ -0,0 +1,9 @@
+// Copyright 1998-2015 Epic Games, Inc. All Rights Reserved.
+
+#ifdef _MSC_VER
+    #pragma message(__FILE__"(9): warning : #include InputChord.h instead of InputGesture.h")
+#else
+    #pragma message("#include InputChord.h instead of InputGesture.h")
+#endif
+
+#include "InputChord.h"
\ No newline at end of file
