Added support for hint text on multi-line editable text

UE-5280 - Hint Text implementation is missing from SMultiLineEditableTextBox

We had the API to let you set this, but we were never using it.

We now create a FTextBlockLayout to overlay the hint text on top of the multi-line editable text when the ML editable text itself is empty. This is only created when we know we're capable of showing hint text.

diff --git a/Engine/Source/Runtime/Slate/Private/Widgets/Input/SMultiLineEditableTextBox.cpp b/Engine/Source/Runtime/Slate/Private/Widgets/Input/SMultiLineEditableTextBox.cpp
--- a/Engine/Source/Runtime/Slate/Private/Widgets/Input/SMultiLineEditableTextBox.cpp
+++ b/Engine/Source/Runtime/Slate/Private/Widgets/Input/SMultiLineEditableTextBox.cpp
@@ -90,6 +90,7 @@ void SMultiLineEditableTextBox::Construct( const FArguments& InArgs )
 				[
 					SAssignNew( EditableText, SMultiLineEditableText )
 					.Text( InArgs._Text )
+					.HintText( InArgs._HintText )
 					.TextStyle( InArgs._TextStyle )
 					.Marshaller( InArgs._Marshaller )
 					.Font( Font )
@@ -140,17 +141,16 @@ void SMultiLineEditableTextBox::Construct( const FArguments& InArgs )
 
 }
 
-
-/**
- * Sets the text string currently being edited 
- *
- * @param  InNewText  The new text string
- */
 void SMultiLineEditableTextBox::SetText( const TAttribute< FText >& InNewText )
 {
 	EditableText->SetText( InNewText );
 }
 
+void SMultiLineEditableTextBox::SetHintText( const TAttribute< FText >& InHintText )
+{
+	EditableText->SetHintText( InHintText );
+}
+
 void SMultiLineEditableTextBox::SetError( const FText& InError )
 {
 	SetError( InError.ToString() );
diff --git a/Engine/Source/Runtime/Slate/Private/Widgets/Text/SMultiLineEditableText.cpp b/Engine/Source/Runtime/Slate/Private/Widgets/Text/SMultiLineEditableText.cpp
--- a/Engine/Source/Runtime/Slate/Private/Widgets/Text/SMultiLineEditableText.cpp
+++ b/Engine/Source/Runtime/Slate/Private/Widgets/Text/SMultiLineEditableText.cpp
@@ -6,6 +6,7 @@
 
 #include "TextEditHelper.h"
 #include "PlainTextLayoutMarshaller.h"
+#include "TextBlockLayout.h"
 #include "GenericCommands.h"
 
 void SMultiLineEditableText::FCursorInfo::SetCursorLocationAndCalculateAlignment(const TSharedPtr<FTextLayout>& InTextLayout, const FTextLocation& InCursorPosition)
@@ -300,6 +301,14 @@ void SMultiLineEditableText::Construct( const FArguments& InArgs )
 
 	TextLayout->UpdateIfNeeded();
 
+	// If we have hint text that is either non-empty or bound to a delegate, we'll also need to make the hint text layout
+	HintText = InArgs._HintText;
+	if (HintText.IsBound() || !HintText.Get(FText::GetEmpty()).IsEmpty())
+	{
+		HintTextStyle = MakeShareable(new FTextBlockStyle(TextStyle));
+		HintTextLayout = FTextBlockLayout::Create(*HintTextStyle, Marshaller.ToSharedRef(), nullptr);
+	}
+
 	// Map UI commands to delegates which are called when the command should be executed
 	UICommandList->MapAction(FGenericCommands::Get().Undo,
 		FExecuteAction::CreateSP(this, &SMultiLineEditableText::Undo),
@@ -363,6 +372,23 @@ void SMultiLineEditableText::SetText(const TAttribute< FText >& InText)
 	}
 }
 
+void SMultiLineEditableText::SetHintText(const TAttribute< FText >& InHintText)
+{
+	HintText = InHintText;
+
+	// If we have hint text that is either non-empty or bound to a delegate, we'll also need to make the hint text layout
+	if (HintText.IsBound() || !HintText.Get(FText::GetEmpty()).IsEmpty())
+	{
+		HintTextStyle = MakeShareable(new FTextBlockStyle(TextStyle));
+		HintTextLayout = FTextBlockLayout::Create(*HintTextStyle, Marshaller.ToSharedRef(), nullptr);
+	}
+	else
+	{
+		HintTextStyle.Reset();
+		HintTextLayout.Reset();
+	}
+}
+
 void SMultiLineEditableText::SetTextFromVirtualKeyboard(const FText& InNewText)
 {
 	// Only set the text if the text attribute doesn't have a getter binding (otherwise it would be blown away).
@@ -2271,6 +2297,21 @@ int32 SMultiLineEditableText::OnPaint( const FPaintArgs& Args, const FGeometry&
 	// Note: This is done here rather than in Tick(), because Tick() doesn't get called while resizing windows, but OnPaint() does
 	CachedSize = AllottedGeometry.Size;
 
+	// Only paint the hint text layout if we don't have any text set
+	if(TextLayout->IsEmpty() && HintTextLayout.IsValid())
+	{
+		// This should always be set when HintTextLayout is
+		check(HintTextStyle.IsValid());
+
+		const FLinearColor ThisColorAndOpacity = TextStyle.ColorAndOpacity.GetColor(InWidgetStyle);
+
+		// Make sure the hint text is the correct color before we paint it
+		HintTextStyle->ColorAndOpacity = FLinearColor(ThisColorAndOpacity.R, ThisColorAndOpacity.G, ThisColorAndOpacity.B, 0.35f);
+		HintTextLayout->OverrideTextStyle(*HintTextStyle);
+
+		LayerId = HintTextLayout->OnPaint( Args, AllottedGeometry, MyClippingRect, OutDrawElements, LayerId, InWidgetStyle, ShouldBeEnabled(bParentEnabled) );
+	}
+
 	LayerId = TextLayout->OnPaint( Args, AllottedGeometry, MyClippingRect, OutDrawElements, LayerId, InWidgetStyle, ShouldBeEnabled( bParentEnabled ) );
 
 	return LayerId;
@@ -2300,35 +2341,67 @@ void SMultiLineEditableText::CacheDesiredSize(float LayoutScaleMultiplier)
 	SWidget::CacheDesiredSize(LayoutScaleMultiplier);
 }
 
-FVector2D SMultiLineEditableText::ComputeDesiredSize( float ) const
+FVector2D SMultiLineEditableText::ComputeDesiredSize( float LayoutScaleMultiplier ) const
 {
 	const float FontMaxCharHeight = FTextEditHelper::GetFontHeight(TextStyle.Font);
 	const float CaretWidth = FTextEditHelper::CalculateCaretWidth(FontMaxCharHeight);
 
-	// If a wrapping width has been provided, then we need to report the wrapped size as the desired width
-	// (rather than the actual text layout size as that can have non-breaking lines that extend beyond the wrap width)
-	// Note: We don't do this when auto-wrapping as it would cause a feedback loop in the Slate sizing logic
 	const float WrappingWidth = WrapTextAt.Get();
-	const FVector2D TextLayoutBaseSize = WrappingWidth > 0 ? TextLayout->GetWrappedSize() : TextLayout->GetSize();
-	const FVector2D TextLayoutSize = TextLayoutBaseSize + FVector2D(CaretWidth, 0);
-
-	// The layouts current margin size. We should not report a size smaller then the margins.
-	const FMargin TextLayoutMargin = TextLayout->GetMargin();
-	float DesiredWidth = FMath::Max(TextLayoutMargin.GetTotalSpaceAlong<Orient_Horizontal>(), TextLayoutSize.X);
-	float DesiredHeight = FMath::Max(TextLayoutMargin.GetTotalSpaceAlong<Orient_Vertical>(), TextLayoutSize.Y);
-	DesiredHeight = FMath::Max(DesiredHeight, FontMaxCharHeight);
+	float DesiredWidth = 0.0f;
+	float DesiredHeight = 0.0f;
+
+	// If we have hint text, make sure we include that in any size calculations
+	if(TextLayout->IsEmpty() && HintTextLayout.IsValid())
+	{
+		// This should always be set when HintTextLayout is
+		check(HintTextStyle.IsValid());
+
+		const FVector2D HintTextSize = HintTextLayout->ComputeDesiredSize(
+			FTextBlockLayout::FWidgetArgs(HintText, FText::GetEmpty(), WrapTextAt, AutoWrapText, Margin, LineHeightPercentage, Justification), 
+			LayoutScaleMultiplier, *HintTextStyle
+			);
+
+		// If a wrapping width has been provided, then we need to report that as the desired width
+		DesiredWidth = WrappingWidth > 0 ? WrappingWidth : HintTextSize.X;
+		DesiredHeight = HintTextSize.Y;
+	}
+	else
+	{
+		// If a wrapping width has been provided, then we need to report the wrapped size as the desired width
+		// (rather than the actual text layout size as that can have non-breaking lines that extend beyond the wrap width)
+		// Note: We don't do this when auto-wrapping as it would cause a feedback loop in the Slate sizing logic
+		const FVector2D TextLayoutBaseSize = WrappingWidth > 0 ? TextLayout->GetWrappedSize() : TextLayout->GetSize();
+		const FVector2D TextLayoutSize = TextLayoutBaseSize + FVector2D(CaretWidth, 0);
+
+		// The layouts current margin size. We should not report a size smaller then the margins.
+		const FMargin TextLayoutMargin = TextLayout->GetMargin();
+		DesiredWidth = FMath::Max(TextLayoutMargin.GetTotalSpaceAlong<Orient_Horizontal>(), TextLayoutSize.X);
+		DesiredHeight = FMath::Max(TextLayoutMargin.GetTotalSpaceAlong<Orient_Vertical>(), TextLayoutSize.Y);
+		DesiredHeight = FMath::Max(DesiredHeight, FontMaxCharHeight);
+	}
 	
 	return FVector2D(DesiredWidth, DesiredHeight);
 }
 
 FChildren* SMultiLineEditableText::GetChildren()
 {
-	return TextLayout->GetChildren();
+	// Only use the hint text layout if we don't have any text set
+	return (TextLayout->IsEmpty() && HintTextLayout.IsValid())
+		? HintTextLayout->GetChildren()
+		: TextLayout->GetChildren();
 }
 
 void SMultiLineEditableText::OnArrangeChildren( const FGeometry& AllottedGeometry, FArrangedChildren& ArrangedChildren ) const
 {
-	TextLayout->ArrangeChildren( AllottedGeometry, ArrangedChildren );
+	// Only arrange the hint text layout if we don't have any text set
+	if(TextLayout->IsEmpty() && HintTextLayout.IsValid())
+	{
+		HintTextLayout->ArrangeChildren( AllottedGeometry, ArrangedChildren );
+	}
+	else
+	{
+		TextLayout->ArrangeChildren( AllottedGeometry, ArrangedChildren );
+	}
 }
 
 bool SMultiLineEditableText::SupportsKeyboardFocus() const
diff --git a/Engine/Source/Runtime/Slate/Private/Widgets/Text/TextBlockLayout.cpp b/Engine/Source/Runtime/Slate/Private/Widgets/Text/TextBlockLayout.cpp
--- a/Engine/Source/Runtime/Slate/Private/Widgets/Text/TextBlockLayout.cpp
+++ b/Engine/Source/Runtime/Slate/Private/Widgets/Text/TextBlockLayout.cpp
@@ -3,7 +3,7 @@
 #include "SlatePrivatePCH.h"
 #include "TextBlockLayout.h"
 #include "SlateTextLayout.h"
-#include "BaseTextLayoutMarshaller.h"
+#include "ITextLayoutMarshaller.h"
 #include "SlateTextHighlightRunRenderer.h"
 #include "SlateStats.h"
 
@@ -11,12 +11,12 @@
 
 SLATE_DECLARE_CYCLE_COUNTER(GSlateTextBlockLayoutComputeDesiredSize, "FTextBlockLayout ComputeDesiredSize");
 
-TSharedRef<FTextBlockLayout> FTextBlockLayout::Create(FTextBlockStyle InDefaultTextStyle, TSharedRef<FBaseTextLayoutMarshaller> InMarshaller, TSharedPtr<IBreakIterator> InLineBreakPolicy)
+TSharedRef<FTextBlockLayout> FTextBlockLayout::Create(FTextBlockStyle InDefaultTextStyle, TSharedRef<ITextLayoutMarshaller> InMarshaller, TSharedPtr<IBreakIterator> InLineBreakPolicy)
 {
 	return MakeShareable(new FTextBlockLayout(MoveTemp(InDefaultTextStyle), InMarshaller, InLineBreakPolicy));
 }
 
-FTextBlockLayout::FTextBlockLayout(FTextBlockStyle InDefaultTextStyle, TSharedRef<FBaseTextLayoutMarshaller> InMarshaller, TSharedPtr<IBreakIterator> InLineBreakPolicy)
+FTextBlockLayout::FTextBlockLayout(FTextBlockStyle InDefaultTextStyle, TSharedRef<ITextLayoutMarshaller> InMarshaller, TSharedPtr<IBreakIterator> InLineBreakPolicy)
 	: TextLayout(FSlateTextLayout::Create(MoveTemp(InDefaultTextStyle)))
 	, Marshaller(InMarshaller)
 	, TextHighlighter(FSlateTextHighlightRunRenderer::Create())
@@ -125,6 +125,19 @@ void FTextBlockLayout::DirtyLayout()
 	TextLayout->DirtyLayout();
 }
 
+void FTextBlockLayout::OverrideTextStyle(const FTextBlockStyle& InTextStyle)
+{
+	// Has the style used for this text block changed?
+	if(!IsStyleUpToDate(InTextStyle))
+	{
+		TextLayout->SetDefaultTextStyle(InTextStyle);
+		
+		FString CurrentText;
+		Marshaller->GetText(CurrentText, *TextLayout);
+		UpdateTextLayout(CurrentText);
+	}
+}
+
 FChildren* FTextBlockLayout::GetChildren()
 {
 	return TextLayout->GetChildren();
@@ -135,11 +148,17 @@ void FTextBlockLayout::ArrangeChildren(const FGeometry& AllottedGeometry, FArran
 	TextLayout->ArrangeChildren(AllottedGeometry, ArrangedChildren);
 }
 
+
 void FTextBlockLayout::UpdateTextLayout(const FText& InText)
+{
+	UpdateTextLayout(InText.ToString());
+}
+
+void FTextBlockLayout::UpdateTextLayout(const FString& InText)
 {
 	Marshaller->ClearDirty();
 	TextLayout->ClearLines();
-	Marshaller->SetText(InText.ToString(), *TextLayout);
+	Marshaller->SetText(InText, *TextLayout);
 
 	TextLayout->ClearLineHighlights();
 	TextLayout->ClearRunRenderers();
diff --git a/Engine/Source/Runtime/Slate/Public/Widgets/Input/SMultiLineEditableTextBox.h b/Engine/Source/Runtime/Slate/Public/Widgets/Input/SMultiLineEditableTextBox.h
--- a/Engine/Source/Runtime/Slate/Public/Widgets/Input/SMultiLineEditableTextBox.h
+++ b/Engine/Source/Runtime/Slate/Public/Widgets/Input/SMultiLineEditableTextBox.h
@@ -176,6 +176,13 @@ public:
 	 */
 	void SetText( const TAttribute< FText >& InNewText );
 
+	/**
+	 * Sets the text that appears when there is no text in the text box
+	 *
+	 * @param  InHintText The hint text string
+	 */
+	void SetHintText( const TAttribute< FText >& InHintText );
+
 	/**
 	 * If InError is a non-empty string the TextBox will the ErrorReporting provided during construction
 	 * If no error reporting was provided, the TextBox will create a default error reporter.
diff --git a/Engine/Source/Runtime/Slate/Public/Widgets/Text/SMultiLineEditableText.h b/Engine/Source/Runtime/Slate/Public/Widgets/Text/SMultiLineEditableText.h
--- a/Engine/Source/Runtime/Slate/Public/Widgets/Text/SMultiLineEditableText.h
+++ b/Engine/Source/Runtime/Slate/Public/Widgets/Text/SMultiLineEditableText.h
@@ -10,6 +10,8 @@
 #include "ITextLayoutMarshaller.h"
 #include "SScrollBar.h"
 
+class FTextBlockLayout;
+
 /** An editable text widget that supports multiple lines and soft word-wrapping. */
 class SLATE_API SMultiLineEditableText : public SWidget, public ITextEditorWidget, public IVirtualKeyboardEntry
 {
@@ -19,6 +21,7 @@ public:
 
 	SLATE_BEGIN_ARGS( SMultiLineEditableText )
 		: _Text()
+		, _HintText()
 		, _Marshaller()
 		, _WrapTextAt( 0.0f )
 		, _AutoWrapText(false)
@@ -37,6 +40,9 @@ public:
 		/** The initial text that will appear in the widget. */
 		SLATE_ATTRIBUTE(FText, Text)
 
+		/** Hint text that appears when there is no text in the text box */
+		SLATE_ATTRIBUTE(FText, HintText)
+
 		/** The marshaller used to get/set the raw text to/from the text layout. */
 		SLATE_ARGUMENT(TSharedPtr< ITextLayoutMarshaller >, Marshaller)
 
@@ -105,6 +111,11 @@ public:
 	 */
 	void SetText(const TAttribute< FText >& InText);
 
+	/**
+	 * Sets the text that appears when there is no text in the text box
+	 */
+	void SetHintText(const TAttribute< FText >& InHintText);
+
 	virtual bool GetIsReadOnly() const override;
 	virtual void ClearSelection() override;
 
@@ -406,12 +417,12 @@ public:
 
 	virtual const FText& GetText() const override
 	{
-		return BoundText.Get();
+		return BoundText.Get(FText::GetEmpty());
 	}
 
 	virtual const FText GetHintText() const override
 	{
-		return FText();
+		return HintText.Get(FText::GetEmpty());
 	}
 
 	virtual EKeyboardType GetVirtualKeyboardType() const override
@@ -430,7 +441,7 @@ protected:
 	virtual void Tick( const FGeometry& AllottedGeometry, const double InCurrentTime, const float InDeltaTime ) override;
 	virtual int32 OnPaint( const FPaintArgs& Args, const FGeometry& AllottedGeometry, const FSlateRect& MyClippingRect, FSlateWindowElementList& OutDrawElements, int32 LayerId, const FWidgetStyle& InWidgetStyle, bool bParentEnabled ) const override;
 	virtual void CacheDesiredSize(float) override;
-	virtual FVector2D ComputeDesiredSize(float) const override;
+	virtual FVector2D ComputeDesiredSize(float LayoutScaleMultiplier) const override;
 	virtual FChildren* GetChildren() override;
 	virtual void OnArrangeChildren( const FGeometry& AllottedGeometry, FArrangedChildren& ArrangedChildren ) const override;
 	virtual bool SupportsKeyboardFocus() const override;
@@ -562,15 +573,24 @@ private:
 	/** The state of BoundText last Tick() (used to allow updates when the text is changed) */
 	FTextSnapshot BoundTextLastTick;
 
+	/** The text that appears when there is no text in the text box */
+	TAttribute<FText> HintText;
+
 	/** The marshaller used to get/set the BoundText text to/from the text layout. */
 	TSharedPtr< ITextLayoutMarshaller > Marshaller;
 
 	/** In control of the layout and wrapping of the BoundText */
 	TSharedPtr< FSlateTextLayout > TextLayout;
 
+	/** In control of the layout and wrapping of the HintText */
+	TSharedPtr< FTextBlockLayout > HintTextLayout;
+
 	/** Default style used by the TextLayout */
 	FTextBlockStyle TextStyle;
 	
+	/** Style used to draw the hint text (only valid when HintTextLayout is set) */
+	TSharedPtr< FTextBlockStyle > HintTextStyle;
+
 	/** Whether text wraps onto a new line when it's length exceeds this width; if this value is zero or negative, no wrapping occurs. */
 	TAttribute< float > WrapTextAt;
 
diff --git a/Engine/Source/Runtime/Slate/Public/Widgets/Text/TextBlockLayout.h b/Engine/Source/Runtime/Slate/Public/Widgets/Text/TextBlockLayout.h
--- a/Engine/Source/Runtime/Slate/Public/Widgets/Text/TextBlockLayout.h
+++ b/Engine/Source/Runtime/Slate/Public/Widgets/Text/TextBlockLayout.h
@@ -3,7 +3,7 @@
 #pragma once
 
 class FSlateTextLayout;
-class FBaseTextLayoutMarshaller;
+class ITextLayoutMarshaller;
 class ISlateRunRenderer;
 
 #if WITH_FANCY_TEXT
@@ -42,7 +42,7 @@ public:
 		const TAttribute<ETextJustify::Type>& Justification;
 	};
 
-	static TSharedRef<FTextBlockLayout> Create(FTextBlockStyle InDefaultTextStyle, TSharedRef<FBaseTextLayoutMarshaller> InMarshaller, TSharedPtr<IBreakIterator> InLineBreakPolicy);
+	static TSharedRef<FTextBlockLayout> Create(FTextBlockStyle InDefaultTextStyle, TSharedRef<ITextLayoutMarshaller> InMarshaller, TSharedPtr<IBreakIterator> InLineBreakPolicy);
 
 	/**
 	 * Get the computed desired size for this layout, updating the internal cache as required
@@ -59,6 +59,13 @@ public:
 	 */
 	void DirtyLayout();
 
+	/**
+	 * Override the text style used and immediately update the text layout (if required).
+	 * This can be used to override the text style after calling ComputeDesiredSize (eg, if you can only compute your text style in OnPaint)
+	 * Please note that changing the size or font used by the text may causing clipping issues until the next call to ComputeDesiredSize
+	 */
+	void OverrideTextStyle(const FTextBlockStyle& InTextStyle);
+
 	/**
 	 * Get the child widgets of this layout
 	 */
@@ -73,6 +80,9 @@ private:
 	/** Updates the text layout to contain the given text */
 	void UpdateTextLayout(const FText& InText);
 
+	/** Updates the text layout to contain the given string */
+	void UpdateTextLayout(const FString& InText);
+
 	/** Update the text highlights */
 	void UpdateTextHighlights(const FText& InHighlightText);
 
@@ -82,13 +92,13 @@ private:
 	/** Calculate the wrapping width based on the given fixed wrap width, and whether we're auto-wrapping */
 	float CalculateWrappingWidth(const FWidgetArgs& InWidgetArgs) const;
 
-	FTextBlockLayout(FTextBlockStyle InDefaultTextStyle, TSharedRef<FBaseTextLayoutMarshaller> InMarshaller, TSharedPtr<IBreakIterator> InLineBreakPolicy);
+	FTextBlockLayout(FTextBlockStyle InDefaultTextStyle, TSharedRef<ITextLayoutMarshaller> InMarshaller, TSharedPtr<IBreakIterator> InLineBreakPolicy);
 
 	/** In control of the layout and wrapping of the text */
 	TSharedPtr<FSlateTextLayout> TextLayout;
 
 	/** The marshaller used to get/set the text to/from the text layout. */
-	TSharedPtr<FBaseTextLayoutMarshaller> Marshaller;
+	TSharedPtr<ITextLayoutMarshaller> Marshaller;
 
 	/** Used to render the current highlights in the text layout */
 	TSharedPtr<ISlateRunRenderer> TextHighlighter;
