Avoid bad boolean initialization in InputChord

diff --git a/Engine/Source/Runtime/Slate/Public/Framework/Commands/InputChord.h b/Engine/Source/Runtime/Slate/Public/Framework/Commands/InputChord.h
--- a/Engine/Source/Runtime/Slate/Public/Framework/Commands/InputChord.h
+++ b/Engine/Source/Runtime/Slate/Public/Framework/Commands/InputChord.h
@@ -82,10 +82,10 @@ struct SLATE_API FInputChord
 	 */
 	FInputChord( const EModifierKey::Type InModifierKeys, const FKey InKey )
 		: Key(InKey)
-		, bShift(InModifierKeys & EModifierKey::Shift)
-		, bCtrl(InModifierKeys & EModifierKey::Control)
-		, bAlt(InModifierKeys & EModifierKey::Alt)
-		, bCmd(InModifierKeys & EModifierKey::Command)
+		, bShift((InModifierKeys & EModifierKey::Shift) !=0)
+		, bCtrl((InModifierKeys & EModifierKey::Control) != 0)
+		, bAlt((InModifierKeys & EModifierKey::Alt) != 0)
+		, bCmd((InModifierKeys & EModifierKey::Command) != 0)
 	{
 	}
 
@@ -97,10 +97,10 @@ struct SLATE_API FInputChord
 	 */
 	FInputChord( const FKey InKey, const EModifierKey::Type InModifierKeys )
 		: Key(InKey)
-		, bShift(InModifierKeys & EModifierKey::Shift)
-		, bCtrl(InModifierKeys & EModifierKey::Control)
-		, bAlt(InModifierKeys & EModifierKey::Alt)
-		, bCmd(InModifierKeys & EModifierKey::Command)
+		, bShift((InModifierKeys & EModifierKey::Shift) !=0)
+		, bCtrl((InModifierKeys & EModifierKey::Control) != 0)
+		, bAlt((InModifierKeys & EModifierKey::Alt) != 0)
+		, bCmd((InModifierKeys & EModifierKey::Command) != 0)
 	{
 	}
 
