Added RMB+Drag scroll to SMultiLineEditableText widget.

diff --git a/Engine/Source/Runtime/Slate/Private/Widgets/Text/SMultiLineEditableText.cpp b/Engine/Source/Runtime/Slate/Private/Widgets/Text/SMultiLineEditableText.cpp
--- a/Engine/Source/Runtime/Slate/Private/Widgets/Text/SMultiLineEditableText.cpp
+++ b/Engine/Source/Runtime/Slate/Private/Widgets/Text/SMultiLineEditableText.cpp
@@ -204,6 +204,8 @@ SMultiLineEditableText::SMultiLineEditableText()
 	, IsReadOnly(false)
 	, UICommandList(new FUICommandList())
 	, bTextChangedByVirtualKeyboard(false)
+	, AmountScrolledWhileRightMouseDown(0.0f)
+	, bIsSoftwareCursor(false)
 {
 	
 }
@@ -550,6 +552,8 @@ FReply SMultiLineEditableText::OnFocusReceived( const FGeometry& MyGeometry, con
 
 void SMultiLineEditableText::OnFocusLost( const FFocusEvent& InFocusEvent )
 {
+	bIsSoftwareCursor = false;
+
 	// Skip the focus lost code if it's due to the context menu opening
 	if (!ActiveContextMenu.IsValid())
 	{
@@ -2314,6 +2318,19 @@ int32 SMultiLineEditableText::OnPaint( const FPaintArgs& Args, const FGeometry&
 
 	LayerId = TextLayout->OnPaint( Args, AllottedGeometry, MyClippingRect, OutDrawElements, LayerId, InWidgetStyle, ShouldBeEnabled( bParentEnabled ) );
 
+	if (bIsSoftwareCursor)
+	{
+		const FSlateBrush* Brush = FCoreStyle::Get().GetBrush(TEXT("SoftwareCursor_Grab"));
+
+		FSlateDrawElement::MakeBox(
+			OutDrawElements,
+			++LayerId,
+			AllottedGeometry.ToPaintGeometry(SoftwareCursorPosition - (Brush->ImageSize / 2), Brush->ImageSize),
+			Brush,
+			MyClippingRect
+			);
+	}
+
 	return LayerId;
 }
 
@@ -2434,19 +2451,74 @@ FReply SMultiLineEditableText::OnKeyUp( const FGeometry& MyGeometry, const FKeyE
 
 FReply SMultiLineEditableText::OnMouseButtonDown( const FGeometry& MyGeometry, const FPointerEvent& MouseEvent ) 
 {
-	FReply Reply = FTextEditHelper::OnMouseButtonDown( MyGeometry, MouseEvent, SharedThis( this ) );
+	if (MouseEvent.GetEffectingButton() == EKeys::RightMouseButton)
+	{
+		AmountScrolledWhileRightMouseDown = 0.0f;
+	}
+
+	FReply Reply = FTextEditHelper::OnMouseButtonDown(MyGeometry, MouseEvent, SharedThis(this));
+//	UpdateCursorHighlight();
 	return Reply;
 }
 
 FReply SMultiLineEditableText::OnMouseButtonUp( const FGeometry& MyGeometry, const FPointerEvent& MouseEvent ) 
 {
+	if (MouseEvent.GetEffectingButton() == EKeys::RightMouseButton)
+	{
+		bool bWasRightClickScrolling = IsRightClickScrolling();
+		AmountScrolledWhileRightMouseDown = 0.0f;
+
+		if (bWasRightClickScrolling)
+		{
+			bIsSoftwareCursor = false;
+			const FVector2D CursorPosition = MyGeometry.LocalToAbsolute(SoftwareCursorPosition);
+			const FIntPoint OriginalMousePos(CursorPosition.X, CursorPosition.Y);
+			return FReply::Handled().ReleaseMouseCapture().SetMousePos(OriginalMousePos);
+		}
+	}
+
 	FReply Reply = FTextEditHelper::OnMouseButtonUp( MyGeometry, MouseEvent, SharedThis( this ) );
 	return Reply;
 }
 
 FReply SMultiLineEditableText::OnMouseMove( const FGeometry& MyGeometry, const FPointerEvent& MouseEvent )
 {
-	FReply Reply = FTextEditHelper::OnMouseMove( MyGeometry, MouseEvent, SharedThis( this ) );
+	if (MouseEvent.IsMouseButtonDown(EKeys::RightMouseButton))
+	{
+		const float ScrollByAmount = MouseEvent.GetCursorDelta().Y / MyGeometry.Scale;
+
+		// If scrolling with the right mouse button, we need to remember how much we scrolled.
+		// If we did not scroll at all, we will bring up the context menu when the mouse is released.
+		AmountScrolledWhileRightMouseDown += FMath::Abs(ScrollByAmount);
+
+		if (IsRightClickScrolling())
+		{
+			const float PreviousScrollOffset = ScrollOffset.Y;
+			ScrollOffset.Y -= ScrollByAmount;
+
+			const float ContentSize = TextLayout->GetSize().Y;
+			const float ScrollMin = 0.0f;
+			const float ScrollMax = ContentSize - MyGeometry.Size.Y;
+			ScrollOffset.Y = FMath::Clamp(ScrollOffset.Y, ScrollMin, ScrollMax);
+
+			if (!bIsSoftwareCursor)
+			{
+				SoftwareCursorPosition = MyGeometry.AbsoluteToLocal(MouseEvent.GetScreenSpacePosition());
+				bIsSoftwareCursor = true;
+			}
+
+			if (PreviousScrollOffset != ScrollOffset.Y)
+			{
+				const float ScrollbarOffset = (ScrollMax != 0.0f) ? ScrollOffset.Y / ScrollMax : 0.0f;
+				OnVScrollBarUserScrolled.ExecuteIfBound(ScrollbarOffset);
+				SoftwareCursorPosition.Y += (PreviousScrollOffset - ScrollOffset.Y);
+			}
+
+			return FReply::Handled().UseHighPrecisionMouseMovement(AsShared());
+		}
+	}
+
+	FReply Reply = FTextEditHelper::OnMouseMove(MyGeometry, MouseEvent, SharedThis(this));
 	return Reply;
 }
 
@@ -2483,7 +2555,19 @@ FReply SMultiLineEditableText::OnMouseButtonDoubleClick(const FGeometry& MyGeome
 
 FCursorReply SMultiLineEditableText::OnCursorQuery( const FGeometry& MyGeometry, const FPointerEvent& CursorEvent ) const
 {
-	return FCursorReply::Cursor( EMouseCursor::TextEditBeam );
+	if (IsRightClickScrolling() && CursorEvent.IsMouseButtonDown(EKeys::RightMouseButton))
+	{
+		return FCursorReply::Cursor(EMouseCursor::None);
+	}
+	else
+	{
+		return FCursorReply::Cursor(EMouseCursor::TextEditBeam);
+	}
+}
+
+bool SMultiLineEditableText::IsRightClickScrolling() const
+{
+	return AmountScrolledWhileRightMouseDown >= FSlateApplication::Get().GetDragTriggerDistance() && VScrollBar.IsValid() && VScrollBar->IsNeeded();
 }
 
 /** Remember where the cursor was when we started selecting. */
diff --git a/Engine/Source/Runtime/Slate/Public/Widgets/Text/SMultiLineEditableText.h b/Engine/Source/Runtime/Slate/Public/Widgets/Text/SMultiLineEditableText.h
--- a/Engine/Source/Runtime/Slate/Public/Widgets/Text/SMultiLineEditableText.h
+++ b/Engine/Source/Runtime/Slate/Public/Widgets/Text/SMultiLineEditableText.h
@@ -559,6 +559,9 @@ private:
 	void OnHScrollBarMoved(const float InScrollOffsetFraction);
 	void OnVScrollBarMoved(const float InScrollOffsetFraction);
 
+	/** Return whether a RMB+Drag scroll operation is taking place */
+	bool IsRightClickScrolling() const;
+
 	/**
 	 * Ensure that we will get a Tick() soon (either due to having active focus, or something having changed progmatically and requiring an update)
 	 * Does nothing if the active tick timer is already enabled
@@ -694,6 +697,15 @@ private:
 
 	/** The timer that is actively driving this widget to Tick() even when Slate is idle */
 	TWeakPtr<FActiveTimerHandle> ActiveTickTimer;
+
+	/** How much we scrolled while RMB was being held */
+	float AmountScrolledWhileRightMouseDown;
+
+	/** Whether a software cursor is currently active */
+	bool bIsSoftwareCursor;
+
+	/**	The current position of the software cursor */
+	FVector2D SoftwareCursorPosition;
 };
 
 
