Slate: Prevent a frame-latent resize issue on SMultiLineEditableText when the scale changes

diff --git a/Engine/Source/Runtime/Slate/Private/Widgets/Text/SMultiLineEditableText.cpp b/Engine/Source/Runtime/Slate/Private/Widgets/Text/SMultiLineEditableText.cpp
--- a/Engine/Source/Runtime/Slate/Private/Widgets/Text/SMultiLineEditableText.cpp
+++ b/Engine/Source/Runtime/Slate/Private/Widgets/Text/SMultiLineEditableText.cpp
@@ -2192,8 +2192,6 @@ void SMultiLineEditableText::Tick( const FGeometry& AllottedGeometry, const doub
 		TextInputMethodChangeNotifier->NotifyLayoutChanged(ITextInputMethodChangeNotifier::ELayoutChangeType::Changed);
 	}
 
-	TextLayout->SetScale( AllottedGeometry.Scale );
-
 	const bool bShouldAppearFocused = HasKeyboardFocus() || ActiveContextMenu.IsValid();
 
 	if (!bShouldAppearFocused)
@@ -2348,6 +2346,7 @@ void SMultiLineEditableText::CacheDesiredSize(float LayoutScaleMultiplier)
 		WrappingWidth = (WrappingWidth >= 1.0f) ? FMath::Min(WrappingWidth, CachedSize.X) : CachedSize.X;
 	}
 
+	TextLayout->SetScale( LayoutScaleMultiplier );
 	TextLayout->SetWrappingWidth( WrappingWidth );
 	TextLayout->SetMargin( OurMargin );
 	TextLayout->SetLineHeightPercentage( LineHeightPercentage.Get() );
