[INTEGRATE] merge ue4-launcherDev->ue4 cl 2484701

diff --git a/Engine/Source/Runtime/Slate/Private/Framework/Application/SlateApplication.cpp b/Engine/Source/Runtime/Slate/Private/Framework/Application/SlateApplication.cpp
--- a/Engine/Source/Runtime/Slate/Private/Framework/Application/SlateApplication.cpp
+++ b/Engine/Source/Runtime/Slate/Private/Framework/Application/SlateApplication.cpp
@@ -3696,8 +3696,11 @@ bool FSlateApplication::ProcessKeyDownEvent( FKeyEvent& InKeyEvent )
 		// Ctrl+Shift+~ summons the Toolbox.
 		if (InKeyEvent.GetKey() == EKeys::Tilde && InKeyEvent.IsControlDown() && InKeyEvent.IsShiftDown())
 		{
-			IToolboxModule& ToolboxModule = FModuleManager::LoadModuleChecked<IToolboxModule>("Toolbox");
-			ToolboxModule.SummonToolbox();
+			IToolboxModule* ToolboxModule = FModuleManager::LoadModulePtr<IToolboxModule>("Toolbox");
+			if (ToolboxModule)
+			{
+				ToolboxModule->SummonToolbox();
+			}
 		}
 
 #endif //!(UE_BUILD_SHIPPING || UE_BUILD_TEST)
