Added EConsumeMouseWheel functionality to the SScrollBox.
Inspired by pull request #886, however we deemed it better to use the previously existing EConsumeMouseWheel enumeration to keep the feature inline with the ListView functionality. UE-11386

diff --git a/Engine/Source/Runtime/Slate/Private/Widgets/Layout/SScrollBox.cpp b/Engine/Source/Runtime/Slate/Private/Widgets/Layout/SScrollBox.cpp
--- a/Engine/Source/Runtime/Slate/Private/Widgets/Layout/SScrollBox.cpp
+++ b/Engine/Source/Runtime/Slate/Private/Widgets/Layout/SScrollBox.cpp
@@ -161,6 +161,7 @@ void SScrollBox::Construct( const FArguments& InArgs )
 	Orientation = InArgs._Orientation;
 	bScrollToEnd = false;
 	bIsActiveTimerRegistered = false;
+	ConsumeMouseWheel = InArgs._ConsumeMouseWheel;
 
 	if (InArgs._ExternalScrollbar.IsValid())
 	{
@@ -642,7 +643,7 @@ bool SScrollBox::ScrollBy( const FGeometry& AllottedGeometry, float ScrollAmount
 
 	OnUserScrolled.ExecuteIfBound(DesiredScrollOffset);
 
-	return DesiredScrollOffset != PreviousScrollOffset;
+	return ConsumeMouseWheel == EConsumeMouseWheel::Always || DesiredScrollOffset != PreviousScrollOffset;
 }
 
 FReply SScrollBox::OnDragDetected( const FGeometry& MyGeometry, const FPointerEvent& MouseEvent )
diff --git a/Engine/Source/Runtime/Slate/Public/Widgets/Layout/SScrollBox.h b/Engine/Source/Runtime/Slate/Public/Widgets/Layout/SScrollBox.h
--- a/Engine/Source/Runtime/Slate/Public/Widgets/Layout/SScrollBox.h
+++ b/Engine/Source/Runtime/Slate/Public/Widgets/Layout/SScrollBox.h
@@ -47,6 +47,7 @@ public:
 		, _ScrollBarAlwaysVisible(false)
 		, _ScrollBarThickness(FVector2D(5, 5))
 		, _OnUserScrolled()
+		, _ConsumeMouseWheel(EConsumeMouseWheel::WhenScrollingPossible)
 		{}
 		
 		SLATE_SUPPORTS_SLOT( FSlot )
@@ -72,6 +73,8 @@ public:
 		/** Called when the button is clicked */
 		SLATE_EVENT(FOnUserScrolled, OnUserScrolled)
 
+		SLATE_ARGUMENT(EConsumeMouseWheel, ConsumeMouseWheel);
+
 	SLATE_END_ARGS()
 
 
@@ -221,6 +224,9 @@ private:
 	/** Style resource for the scrollbar */
 	const FScrollBarStyle* ScrollBarStyle;
 
+	/** How we should handle scrolling with the mouse wheel */
+	EConsumeMouseWheel ConsumeMouseWheel;
+
 	/** Cached geometry for use with the active timer */
 	FGeometry CachedGeometry;
 
diff --git a/Engine/Source/Runtime/Slate/Public/Widgets/Views/STableViewBase.h b/Engine/Source/Runtime/Slate/Public/Widgets/Views/STableViewBase.h
--- a/Engine/Source/Runtime/Slate/Public/Widgets/Views/STableViewBase.h
+++ b/Engine/Source/Runtime/Slate/Public/Widgets/Views/STableViewBase.h
@@ -15,16 +15,6 @@ enum class EAllowOverscroll
 	No
 };
 
-/** Used to determine how we should handle mouse wheel input events when someone scrolls. */
-enum class EConsumeMouseWheel
-{
-	/** Only consume the mouse wheel event when we actually scroll some amount. */
-	WhenScrollingPossible,
-
-	/** Always consume mouse wheel event even if we don't scroll at all. */
-	Always, 
-};
-
 /** If the list panel is arranging items horizontally, this enum dictates how the items should be aligned (basically, where any extra space is placed) */
 enum class EListItemAlignment : uint8
 {
diff --git a/Engine/Source/Runtime/SlateCore/Public/Styling/SlateTypes.h b/Engine/Source/Runtime/SlateCore/Public/Styling/SlateTypes.h
--- a/Engine/Source/Runtime/SlateCore/Public/Styling/SlateTypes.h
+++ b/Engine/Source/Runtime/SlateCore/Public/Styling/SlateTypes.h
@@ -9,6 +9,18 @@
 
 #include "SlateTypes.generated.h"
 
+/** Used to determine how we should handle mouse wheel input events when someone scrolls. */
+UENUM()
+enum class EConsumeMouseWheel
+{
+	/** Only consume the mouse wheel event when we actually scroll some amount. */
+	WhenScrollingPossible,
+
+	/** Always consume mouse wheel event even if we don't scroll at all. */
+	Always,
+};
+
+
 /** Type of check box */
 UENUM()
 namespace ESlateCheckBoxType
diff --git a/Engine/Source/Runtime/UMG/Private/Components/ScrollBox.cpp b/Engine/Source/Runtime/UMG/Private/Components/ScrollBox.cpp
--- a/Engine/Source/Runtime/UMG/Private/Components/ScrollBox.cpp
+++ b/Engine/Source/Runtime/UMG/Private/Components/ScrollBox.cpp
@@ -23,6 +23,8 @@ UScrollBox::UScrollBox(const FObjectInitializer& ObjectInitializer)
 	AlwaysShowScrollbar = false;
 	ScrollbarThickness = FVector2D(5, 5);
 	ScrollBarVisibility = ESlateVisibility::Visible;
+
+	ConsumeMouseWheel = EConsumeMouseWheel::WhenScrollingPossible;
 }
 
 void UScrollBox::ReleaseSlateResources(bool bReleaseChildren)
@@ -64,7 +66,8 @@ TSharedRef<SWidget> UScrollBox::RebuildWidget()
 	MyScrollBox = SNew(SScrollBox)
 		.Style(&WidgetStyle)
 		.ScrollBarStyle(&WidgetBarStyle)
-		.Orientation(Orientation);
+		.Orientation(Orientation)
+		.ConsumeMouseWheel(ConsumeMouseWheel);
 
 	for ( UPanelSlot* Slot : Slots )
 	{
diff --git a/Engine/Source/Runtime/UMG/Public/Components/ScrollBox.h b/Engine/Source/Runtime/UMG/Public/Components/ScrollBox.h
--- a/Engine/Source/Runtime/UMG/Public/Components/ScrollBox.h
+++ b/Engine/Source/Runtime/UMG/Public/Components/ScrollBox.h
@@ -36,6 +36,10 @@ public:
 	UPROPERTY(EditDefaultsOnly, BlueprintReadOnly, Category="Scroll")
 	TEnumAsByte<ESlateVisibility> ScrollBarVisibility;
 
+	/**  Enable to always consume mouse wheel event, even when scrolling is not possible */
+	UPROPERTY(EditDefaultsOnly, BlueprintReadOnly, Category = "Scroll")
+	TEnumAsByte<EConsumeMouseWheel> ConsumeMouseWheel;
+
 	/**  */
 	UPROPERTY(EditDefaultsOnly, BlueprintReadOnly, Category="Scroll")
 	FVector2D ScrollbarThickness;
