Fixed list views not being scrollable when they contained a single item that was larger than the available view area

UE-11793 - Chat window stops scrolling with chat if a single chat bubble stretches from top to bottom

The list view was making some assumptions that the first visible item would always be entirely visible (barring some scrolling), however it's possible that the first visible item might be larger than the available view area of the list view, in which case it needs to be scrollable.

diff --git a/Engine/Source/Runtime/Slate/Private/Widgets/Views/STableViewBase.cpp b/Engine/Source/Runtime/Slate/Private/Widgets/Views/STableViewBase.cpp
--- a/Engine/Source/Runtime/Slate/Private/Widgets/Views/STableViewBase.cpp
+++ b/Engine/Source/Runtime/Slate/Private/Widgets/Views/STableViewBase.cpp
@@ -233,7 +233,9 @@ void STableViewBase::Tick( const FGeometry& AllottedGeometry, const double InCur
 			const int32 NumItemsBeingObserved = GetNumItemsBeingObserved();
 
 			const int32 NumItemsWide = GetNumItemsWide();
-			const bool bEnoughRoomForAllItems = ReGenerateResults.ExactNumRowsOnScreen >= (NumItemsBeingObserved / NumItemsWide);
+			const int32 NumItemRows = NumItemsBeingObserved / NumItemsWide;
+
+			const bool bEnoughRoomForAllItems = ReGenerateResults.ExactNumRowsOnScreen >= NumItemRows;
 			if (bEnoughRoomForAllItems)
 			{
 				// We can show all the items, so make sure there is no scrolling.
@@ -259,11 +261,22 @@ void STableViewBase::Tick( const FGeometry& AllottedGeometry, const double InCur
 			// Update scrollbar
 			if (NumItemsBeingObserved > 0)
 			{
-				// The thumb size is whatever fraction of the items we are currently seeing (including partially seen items).
-				// e.g. if we are seeing 0.5 of the first generated widget and 0.75 of the last widget, that's 1.25 widgets.
-				const double ThumbSizeFraction = ReGenerateResults.ExactNumRowsOnScreen / (NumItemsBeingObserved / NumItemsWide);
-				const double OffsetFraction = ScrollOffset / NumItemsBeingObserved;
-				ScrollBar->SetState( OffsetFraction, ThumbSizeFraction );
+				if (ReGenerateResults.ExactNumRowsOnScreen < 1.0f)
+				{
+					// We are be observing a single row which is larger than the available visible area, so we should calculate thumb size based on that
+					const double VisibleSizeFraction = AllottedGeometry.GetLocalSize().Y / ReGenerateResults.HeightOfGeneratedItems;
+					const double ThumbSizeFraction = FMath::Min(VisibleSizeFraction, 1.0);
+					const double OffsetFraction = ScrollOffset / NumItemsBeingObserved;
+					ScrollBar->SetState( OffsetFraction, ThumbSizeFraction );
+				}
+				else
+				{
+					// The thumb size is whatever fraction of the items we are currently seeing (including partially seen items).
+					// e.g. if we are seeing 0.5 of the first generated widget and 0.75 of the last widget, that's 1.25 widgets.
+					const double ThumbSizeFraction = ReGenerateResults.ExactNumRowsOnScreen / NumItemRows;
+					const double OffsetFraction = ScrollOffset / NumItemsBeingObserved;
+					ScrollBar->SetState( OffsetFraction, ThumbSizeFraction );
+				}
 			}
 			else
 			{
diff --git a/Engine/Source/Runtime/Slate/Public/Widgets/Views/SListView.h b/Engine/Source/Runtime/Slate/Public/Widgets/Views/SListView.h
--- a/Engine/Source/Runtime/Slate/Public/Widgets/Views/SListView.h
+++ b/Engine/Source/Runtime/Slate/Public/Widgets/Views/SListView.h
@@ -813,7 +813,13 @@ public:
 				if (bIsFirstItem)
 				{
 					// The first item may not be fully visible (but cannot exceed 1)
-					ItemsInView += 1.0f - FMath::Max(FMath::Fractional(ScrollOffset), 0.0f);
+					// FirstItemFractionScrolledIntoView is the fraction of the item that is visible after taking into account anything that may be scrolled off the top of the list view
+					// FirstItemHeightScrolledIntoView is the height of the item, ignoring anything that is scrolled off the top of the list view
+					// FirstItemVisibleFraction is either: The visible item height as a fraction of the available list view height (if the item size is larger than the available size, otherwise this will be >1), or just FirstItemFractionScrolledIntoView (which can never be >1)
+					const float FirstItemFractionScrolledIntoView = 1.0f - FMath::Max(FMath::Fractional(ScrollOffset), 0.0f);
+					const float FirstItemHeightScrolledIntoView = ItemHeight * FirstItemFractionScrolledIntoView;
+					const float FirstItemVisibleFraction = FMath::Min(MyGeometry.Size.Y / FirstItemHeightScrolledIntoView, FirstItemFractionScrolledIntoView);
+					ItemsInView += FirstItemVisibleFraction;
 				}
 				else if (ViewHeightUsedSoFar + ItemHeight > MyGeometry.Size.Y)
 				{
@@ -1378,7 +1384,7 @@ protected:
 	/** Delegate to invoke when selection changes. */
 	FOnSelectionChanged OnSelectionChanged;
 
-	/** Caled when the user clicks on an element int he list view with the left mouse button */
+	/** Called when the user clicks on an element int he list view with the left mouse button */
 	FOnMouseButtonClick OnClick;
 
 	/** Called when the user double-clicks on an element in the list view with the left mouse button */
