Implemented Page Up/Down support for multi-line editable text

UE-1496 - Slate: Implement PageUp PageDown support for multiline editable text

This will move the multi-line text by a "page" (based on the number of visible lines), and then scroll the document to try and keep the cursor at the same relative location in the view.

This also fixes some issues with scrolling when dealing with a scaled multi-line editable text, as the scroll offset is now always considered to be unscaled until it is given to the text layout. This prevents issues where we could over or under scroll when using a scale other than 1.

diff --git a/Engine/Source/Runtime/Slate/Private/Framework/Text/TextEditHelper.cpp b/Engine/Source/Runtime/Slate/Private/Framework/Text/TextEditHelper.cpp
--- a/Engine/Source/Runtime/Slate/Private/Framework/Text/TextEditHelper.cpp
+++ b/Engine/Source/Runtime/Slate/Private/Framework/Text/TextEditHelper.cpp
@@ -164,6 +164,24 @@ FReply FTextEditHelper::OnKeyDown( const FKeyEvent& InKeyEvent, const TSharedRef
 
 		return FReply::Handled();
 	}
+	else if( Key == EKeys::PageUp )
+	{
+		// Go to the previous page of the document document; select text if Shift is down.
+		TextEditor->JumpTo(
+			ETextLocation::PreviousPage,
+			(InKeyEvent.IsShiftDown()) ? ECursorAction::SelectText : ECursorAction::MoveCursor );
+
+		return FReply::Handled();
+	}
+	else if ( Key == EKeys::PageDown )
+	{
+		// Go to the next page of the document document; select text if Shift is down.
+		TextEditor->JumpTo(
+			ETextLocation::NextPage,
+			(InKeyEvent.IsShiftDown()) ? ECursorAction::SelectText : ECursorAction::MoveCursor );
+
+		return FReply::Handled();
+	}
 	else if( Key == EKeys::Enter && !TextEditor->GetIsReadOnly() )
 	{
 		FScopedTextTransaction TextTransaction(TextEditor);
diff --git a/Engine/Source/Runtime/Slate/Private/Widgets/Input/SEditableText.cpp b/Engine/Source/Runtime/Slate/Private/Widgets/Input/SEditableText.cpp
--- a/Engine/Source/Runtime/Slate/Private/Widgets/Input/SEditableText.cpp
+++ b/Engine/Source/Runtime/Slate/Private/Widgets/Input/SEditableText.cpp
@@ -780,6 +780,7 @@ void SEditableText::JumpTo(ETextLocation JumpLocation, ECursorAction Action)
 	{
 		case ETextLocation::BeginningOfLine:
 		case ETextLocation::BeginningOfDocument:
+		case ETextLocation::PreviousPage:
 		{
 			const int32 OldCaretPosition = CaretPosition;
 			SetCaretPosition( 0 );
@@ -797,6 +798,7 @@ void SEditableText::JumpTo(ETextLocation JumpLocation, ECursorAction Action)
 
 		case ETextLocation::EndOfLine:
 		case ETextLocation::EndOfDocument:
+		case ETextLocation::NextPage:
 		{
 			const int32 OldCaretPosition = CaretPosition;
 			SetCaretPosition( EditedText.ToString().Len() );
@@ -811,7 +813,7 @@ void SEditableText::JumpTo(ETextLocation JumpLocation, ECursorAction Action)
 			}
 		}
 		break;
-	};
+	}
 }
 
 
diff --git a/Engine/Source/Runtime/Slate/Private/Widgets/Text/SMultiLineEditableText.cpp b/Engine/Source/Runtime/Slate/Private/Widgets/Text/SMultiLineEditableText.cpp
--- a/Engine/Source/Runtime/Slate/Private/Widgets/Text/SMultiLineEditableText.cpp
+++ b/Engine/Source/Runtime/Slate/Private/Widgets/Text/SMultiLineEditableText.cpp
@@ -485,13 +485,13 @@ void SMultiLineEditableText::ForceRefreshTextLayout(const FText& CurrentText)
 
 void SMultiLineEditableText::OnHScrollBarMoved(const float InScrollOffsetFraction)
 {
-	ScrollOffset.X = FMath::Clamp<float>(InScrollOffsetFraction, 0.0, 1.0) * GetDesiredSize().X;
+	ScrollOffset.X = FMath::Clamp<float>(InScrollOffsetFraction, 0.0, 1.0) * TextLayout->GetSize().X;
 	OnHScrollBarUserScrolled.ExecuteIfBound(InScrollOffsetFraction);
 }
 
 void SMultiLineEditableText::OnVScrollBarMoved(const float InScrollOffsetFraction)
 {
-	ScrollOffset.Y = FMath::Clamp<float>(InScrollOffsetFraction, 0.0, 1.0) * GetDesiredSize().Y;
+	ScrollOffset.Y = FMath::Clamp<float>(InScrollOffsetFraction, 0.0, 1.0) * TextLayout->GetSize().Y;
 	OnVScrollBarUserScrolled.ExecuteIfBound(InScrollOffsetFraction);
 }
 
@@ -542,7 +542,7 @@ FReply SMultiLineEditableText::OnFocusReceived( const FGeometry& MyGeometry, con
 		// gain focus since it can cause the scroll position to jump unexpectedly
 		// If we gained focus via a mouse click that moved the cursor, then MoveCursor will already take care
 		// of making sure that gets scrolled into view
-		PositionToScrollIntoView = TOptional<FScrollInfo>();
+		PositionToScrollIntoView.Reset();
 
 		return SWidget::OnFocusReceived( MyGeometry, InFocusEvent );
 	}
@@ -600,7 +600,7 @@ void SMultiLineEditableText::OnFocusLost( const FFocusEvent& InFocusEvent )
 
 		// UpdateCursorHighlight always tries to scroll to the cursor, but we don't want that to happen when we 
 		// lose focus since it can cause the scroll position to jump unexpectedly
-		PositionToScrollIntoView = TOptional<FScrollInfo>();
+		PositionToScrollIntoView.Reset();
 	}
 }
 
@@ -1187,6 +1187,32 @@ void SMultiLineEditableText::UpdatePreferredCursorScreenOffsetInLine()
 
 void SMultiLineEditableText::JumpTo(ETextLocation JumpLocation, ECursorAction Action)
 {
+	// Utility function to count the number of fully visible lines (vertically)
+	// We consider this to be the number of lines on the current page
+	auto CountVisibleLines = [](const TArray<FTextLayout::FLineView>& LineViews, const float VisibleHeight) -> int32
+	{
+		int32 LinesInView = 0;
+		for (const auto& LineView : LineViews)
+		{
+			// The line view is scrolled such that lines above the top of the text area have negative offsets
+			if (LineView.Offset.Y >= 0.0f)
+			{
+				const float EndOffsetY = LineView.Offset.Y + LineView.Size.Y;
+				if (EndOffsetY <= VisibleHeight)
+				{
+					// Line is completely in view
+					++LinesInView;
+				}
+				else
+				{
+					// Line extends beyond the bottom of the text area - we've finished finding visible lines
+					break;
+				}
+			}
+		}
+		return LinesInView;
+	};
+
 	switch (JumpLocation)
 	{
 		case ETextLocation::BeginningOfLine:
@@ -1278,6 +1304,7 @@ void SMultiLineEditableText::JumpTo(ETextLocation JumpLocation, ECursorAction Ac
 			}
 		}
 		break;
+
 		case ETextLocation::EndOfDocument:
 		{
 			if (!TextLayout->IsEmpty())
@@ -1307,7 +1334,121 @@ void SMultiLineEditableText::JumpTo(ETextLocation JumpLocation, ECursorAction Ac
 			}
 		}
 		break;
-	};
+
+		case ETextLocation::PreviousPage:
+		{
+			const FTextLocation CursorInteractionPosition = CursorInfo.GetCursorInteractionLocation();
+			const TArray< FTextLayout::FLineView >& LineViews = TextLayout->GetLineViews();
+			const int32 CurrentLineViewIndex = TextLayout->GetLineViewIndexForTextLocation(LineViews, CursorInteractionPosition, CursorInfo.GetCursorAlignment() == ECursorAlignment::Right );
+
+			if (LineViews.IsValidIndex(CurrentLineViewIndex))
+			{
+				const FTextLayout::FLineView& CurrentLineView = LineViews[CurrentLineViewIndex];
+				
+				const FTextLocation OldCursorPosition = CursorInteractionPosition;
+
+				FTextLocation NewCursorPosition;
+				TOptional<ECursorAlignment> NewCursorAlignment;
+				const int32 NumLinesToMove = FMath::Max(1, CountVisibleLines(LineViews, CachedSize.Y));
+				TranslateLocationVertical(OldCursorPosition, -NumLinesToMove, TextLayout->GetScale(), NewCursorPosition, NewCursorAlignment);
+
+				if (Action == ECursorAction::SelectText)
+				{
+					if (!SelectionStart.IsSet())
+					{
+						this->SelectionStart = OldCursorPosition;
+					}
+				}
+				else
+				{
+					ClearSelection();
+				}
+
+				if (NewCursorAlignment.IsSet())
+				{
+					CursorInfo.SetCursorLocationAndAlignment(NewCursorPosition, NewCursorAlignment.GetValue());
+				}
+				else
+				{
+					CursorInfo.SetCursorLocationAndCalculateAlignment(TextLayout, NewCursorPosition);
+				}
+				OnCursorMoved.ExecuteIfBound(CursorInfo.GetCursorInteractionLocation());
+				UpdatePreferredCursorScreenOffsetInLine();
+				UpdateCursorHighlight();
+
+				// We need to scroll by the delta vertical offset value of the old line and the new line
+				// This will (try to) keep the cursor in the same relative location after the page jump
+				const int32 NewLineViewIndex = TextLayout->GetLineViewIndexForTextLocation(LineViews, CursorInfo.GetCursorInteractionLocation(), CursorInfo.GetCursorAlignment() == ECursorAlignment::Right );
+				if (LineViews.IsValidIndex(NewLineViewIndex))
+				{
+					const FTextLayout::FLineView& NewLineView = LineViews[NewLineViewIndex];
+					const float DeltaScrollY = (NewLineView.Offset.Y - CurrentLineView.Offset.Y) / TextLayout->GetScale();
+					ScrollOffset.Y = FMath::Max(0.0f, ScrollOffset.Y + DeltaScrollY);
+
+					// Disable the normal cursor scrolling that UpdateCursorHighlight triggers
+					PositionToScrollIntoView.Reset();
+				}
+			}
+		}
+		break;
+
+		case ETextLocation::NextPage:
+		{
+			const FTextLocation CursorInteractionPosition = CursorInfo.GetCursorInteractionLocation();
+			const TArray< FTextLayout::FLineView >& LineViews = TextLayout->GetLineViews();
+			const int32 CurrentLineViewIndex = TextLayout->GetLineViewIndexForTextLocation(LineViews, CursorInteractionPosition, CursorInfo.GetCursorAlignment() == ECursorAlignment::Right );
+
+			if (LineViews.IsValidIndex(CurrentLineViewIndex))
+			{
+				const FTextLayout::FLineView& CurrentLineView = LineViews[CurrentLineViewIndex];
+				
+				const FTextLocation OldCursorPosition = CursorInteractionPosition;
+
+				FTextLocation NewCursorPosition;
+				TOptional<ECursorAlignment> NewCursorAlignment;
+				const int32 NumLinesToMove = FMath::Max(1, CountVisibleLines(LineViews, CachedSize.Y));
+				TranslateLocationVertical(OldCursorPosition, NumLinesToMove, TextLayout->GetScale(), NewCursorPosition, NewCursorAlignment);
+
+				if (Action == ECursorAction::SelectText)
+				{
+					if (!SelectionStart.IsSet())
+					{
+						this->SelectionStart = OldCursorPosition;
+					}
+				}
+				else
+				{
+					ClearSelection();
+				}
+
+				if (NewCursorAlignment.IsSet())
+				{
+					CursorInfo.SetCursorLocationAndAlignment(NewCursorPosition, NewCursorAlignment.GetValue());
+				}
+				else
+				{
+					CursorInfo.SetCursorLocationAndCalculateAlignment(TextLayout, NewCursorPosition);
+				}
+				OnCursorMoved.ExecuteIfBound(CursorInfo.GetCursorInteractionLocation());
+				UpdatePreferredCursorScreenOffsetInLine();
+				UpdateCursorHighlight();
+
+				// We need to scroll by the delta vertical offset value of the old line and the new line
+				// This will (try to) keep the cursor in the same relative location after the page jump
+				const int32 NewLineViewIndex = TextLayout->GetLineViewIndexForTextLocation(LineViews, CursorInfo.GetCursorInteractionLocation(), CursorInfo.GetCursorAlignment() == ECursorAlignment::Right );
+				if (LineViews.IsValidIndex(NewLineViewIndex))
+				{
+					const FTextLayout::FLineView& NewLineView = LineViews[NewLineViewIndex];
+					const float DeltaScrollY = (NewLineView.Offset.Y - CurrentLineView.Offset.Y) / TextLayout->GetScale();
+					ScrollOffset.Y = FMath::Min(TextLayout->GetSize().Y - CachedSize.Y, ScrollOffset.Y + DeltaScrollY);
+
+					// Disable the normal cursor scrolling that UpdateCursorHighlight triggers
+					PositionToScrollIntoView.Reset();
+				}
+			}
+		}
+		break;
+	}
 }
 
 void SMultiLineEditableText::ClearSelection()
@@ -2240,7 +2381,7 @@ void SMultiLineEditableText::Tick( const FGeometry& AllottedGeometry, const doub
 			}
 		}
 
-		PositionToScrollIntoView = TOptional<FScrollInfo>();
+		PositionToScrollIntoView.Reset();
 	}
 
 	{
@@ -2290,7 +2431,7 @@ void SMultiLineEditableText::Tick( const FGeometry& AllottedGeometry, const doub
 		}
 	}
 
-	TextLayout->SetVisibleRegion(AllottedGeometry.Size, ScrollOffset);
+	TextLayout->SetVisibleRegion(AllottedGeometry.Size, ScrollOffset * TextLayout->GetScale());
 }
 
 int32 SMultiLineEditableText::OnPaint( const FPaintArgs& Args, const FGeometry& AllottedGeometry, const FSlateRect& MyClippingRect, FSlateWindowElementList& OutDrawElements, int32 LayerId, const FWidgetStyle& InWidgetStyle, bool bParentEnabled ) const
@@ -2351,7 +2492,7 @@ void SMultiLineEditableText::CacheDesiredSize(float LayoutScaleMultiplier)
 	TextLayout->SetMargin( OurMargin );
 	TextLayout->SetLineHeightPercentage( LineHeightPercentage.Get() );
 	TextLayout->SetJustification( Justification.Get() );
-	TextLayout->SetVisibleRegion( CachedSize, ScrollOffset );
+	TextLayout->SetVisibleRegion( CachedSize, ScrollOffset * TextLayout->GetScale() );
 	TextLayout->UpdateIfNeeded();
 
 	SWidget::CacheDesiredSize(LayoutScaleMultiplier);
@@ -2604,7 +2745,7 @@ FTextLocation SMultiLineEditableText::TranslatedLocation( const FTextLocation& L
 	}
 }
 
-void SMultiLineEditableText::TranslateLocationVertical( const FTextLocation& Location, int8 Direction, float GeometryScale, FTextLocation& OutCursorPosition, TOptional<ECursorAlignment>& OutCursorAlignment ) const
+void SMultiLineEditableText::TranslateLocationVertical( const FTextLocation& Location, int32 NumLinesToMove, float GeometryScale, FTextLocation& OutCursorPosition, TOptional<ECursorAlignment>& OutCursorAlignment ) const
 {
 	const TArray< FTextLayout::FLineView >& LineViews = TextLayout->GetLineViews();
 	const int32 NumberOfLineViews = LineViews.Num();
@@ -2613,7 +2754,7 @@ void SMultiLineEditableText::TranslateLocationVertical( const FTextLocation& Loc
 	ensure(CurrentLineViewIndex != INDEX_NONE);
 	const FTextLayout::FLineView& CurrentLineView = LineViews[ CurrentLineViewIndex ];
 
-	const int32 NewLineViewIndex = FMath::Clamp( CurrentLineViewIndex + Direction, 0, NumberOfLineViews - 1 );
+	const int32 NewLineViewIndex = FMath::Clamp( CurrentLineViewIndex + NumLinesToMove, 0, NumberOfLineViews - 1 );
 	const FTextLayout::FLineView& NewLineView = LineViews[ NewLineViewIndex ];
 	
 	// Our horizontal position is the clamped version of whatever the user explicitly set with horizontal movement.
diff --git a/Engine/Source/Runtime/Slate/Public/Widgets/Text/ITextEditorWidget.h b/Engine/Source/Runtime/Slate/Public/Widgets/Text/ITextEditorWidget.h
--- a/Engine/Source/Runtime/Slate/Public/Widgets/Text/ITextEditorWidget.h
+++ b/Engine/Source/Runtime/Slate/Public/Widgets/Text/ITextEditorWidget.h
@@ -37,7 +37,11 @@ enum class ETextLocation
 	/** Jump to the beginning of the line or beginning of text within a line (e.g. Home) */
 	BeginningOfLine,
 	/** Jump to the end of line (e.g. End) */
-	EndOfLine
+	EndOfLine,
+	/** Jump to the previous page in this document (e.g. PageUp) */
+	PreviousPage,
+	/** Jump to the next page in this document (e.g. PageDown) */
+	NextPage,
 };
 
 /**
diff --git a/Engine/Source/Runtime/Slate/Public/Widgets/Text/SMultiLineEditableText.h b/Engine/Source/Runtime/Slate/Public/Widgets/Text/SMultiLineEditableText.h
--- a/Engine/Source/Runtime/Slate/Public/Widgets/Text/SMultiLineEditableText.h
+++ b/Engine/Source/Runtime/Slate/Public/Widgets/Text/SMultiLineEditableText.h
@@ -492,12 +492,12 @@ private:
 	 * Given a location and a Direction to offset, return a new location.
 	 *
 	 * @param Location              Cursor location from which to offset
-	 * @param Direction             Positive means down, negative means up.
+	 * @param NumLinesToMove        Number of lines to move in a given direction. Positive means down, negative means up.
 	 * @param GeometryScale         Geometry DPI scale at which the widget is being rendered
 	 * @param OutCursorPosition     Fill with the updated cursor position.
 	 * @param OutCursorAlignment    Optionally fill with a new cursor alignment (will be auto-calculated if not set).
 	 */
-	void TranslateLocationVertical( const FTextLocation& Location, int8 Direction, float GeometryScale, FTextLocation& OutCursorPosition, TOptional<ECursorAlignment>& OutCursorAlignment ) const;
+	void TranslateLocationVertical( const FTextLocation& Location, int32 NumLinesToMove, float GeometryScale, FTextLocation& OutCursorPosition, TOptional<ECursorAlignment>& OutCursorAlignment ) const;
 
 	/** Find the closest word boundary */
 	FTextLocation ScanForWordBoundary( const FTextLocation& Location, int8 Direction ) const; 
@@ -603,7 +603,7 @@ private:
 	/** The last known size of the control from the previous OnPaint, used to recalculate wrapping. */
 	mutable FVector2D CachedSize;
 
-	/** The scroll offset (in Slate units) for this text */
+	/** The scroll offset (in unscaled Slate units) for this text */
 	FVector2D ScrollOffset;
 
 	TAttribute< FMargin > Margin;
