Slate: Fixed variable and function shadowing in SComboButton that broke opening the menu programmatically via SetIsOpen

diff --git a/Engine/Source/Runtime/Slate/Private/Widgets/Input/SComboButton.cpp b/Engine/Source/Runtime/Slate/Private/Widgets/Input/SComboButton.cpp
--- a/Engine/Source/Runtime/Slate/Private/Widgets/Input/SComboButton.cpp
+++ b/Engine/Source/Runtime/Slate/Private/Widgets/Input/SComboButton.cpp
@@ -13,7 +13,6 @@ void SComboButton::Construct( const FArguments& InArgs )
 	MenuBorderBrush = &InArgs._ComboButtonStyle->MenuBorderBrush;
 	MenuBorderPadding = InArgs._ComboButtonStyle->MenuBorderPadding;
 	
-	OnGetMenuContent = InArgs._OnGetMenuContent;
 	OnComboBoxOpened = InArgs._OnComboBoxOpened;
 	ContentWidgetPtr = InArgs._MenuContent.Widget;
 	bIsFocusable = InArgs._IsFocusable;
@@ -24,6 +23,7 @@ void SComboButton::Construct( const FArguments& InArgs )
 		.Placement(InArgs._MenuPlacement)
 		.Method(InArgs._Method)
 		.OnMenuOpenChanged(InArgs._OnMenuOpenChanged)
+		.OnGetMenuContent(InArgs._OnGetMenuContent)
 		[
 			SNew( SButton )
 			.ButtonStyle( OurButtonStyle )
@@ -72,13 +72,6 @@ void SComboButton::Construct( const FArguments& InArgs )
 
 FReply SComboButton::OnButtonClicked()
 {
-	TSharedPtr<SWidget> Content = nullptr;
-	if( OnGetMenuContent.IsBound() )
-	{
-		Content = OnGetMenuContent.Execute();
-		SetMenuContent( Content.ToSharedRef() );
-	}
-
 	// Button was clicked; show the popup.
 	// Do nothing if clicking on the button also dismissed the menu, because we will end up doing the same thing twice.
 	this->SetIsOpen( ShouldOpenDueToClick(), bIsFocusable );
@@ -98,7 +91,7 @@ FReply SComboButton::OnButtonClicked()
 		if (!WidgetToFocus.IsValid())
 		{
 			// no explicitly focused widget, try to focus the content
-			WidgetToFocus = Content;
+			WidgetToFocus = MenuContent;
 		}
 
 		if (!WidgetToFocus.IsValid())
@@ -116,13 +109,7 @@ FReply SComboButton::OnButtonClicked()
 	return ButtonClickedReply;
 }
 
-
-/**
- * Sets the content for this border
- *
- * @param	InContent	The widget to use as content for the border
- */
-void SComboButton::SetMenuContent( const TSharedRef< SWidget >& InContent )
+void SComboButton::SetMenuContent(TSharedRef<SWidget> InContent)
 {
 	MenuContent = 
 		SNew(SBorder)
diff --git a/Engine/Source/Runtime/Slate/Private/Widgets/Input/SMenuAnchor.cpp b/Engine/Source/Runtime/Slate/Private/Widgets/Input/SMenuAnchor.cpp
--- a/Engine/Source/Runtime/Slate/Private/Widgets/Input/SMenuAnchor.cpp
+++ b/Engine/Source/Runtime/Slate/Private/Widgets/Input/SMenuAnchor.cpp
@@ -309,9 +309,12 @@ void SMenuAnchor::SetIsOpen( bool InIsOpen, const bool bFocusMenu )
 	{
 		if ( InIsOpen )
 		{
-			TSharedPtr< SWidget > MenuContentPtr = OnGetMenuContent.IsBound() ? OnGetMenuContent.Execute() : MenuContent;
+			if ( OnGetMenuContent.IsBound() )
+			{
+				SetMenuContent(OnGetMenuContent.Execute());
+			}
 
-			if ( MenuContentPtr.IsValid() )
+			if ( MenuContent.IsValid() )
 			{
 				// OPEN POPUP
 				if ( OnMenuOpenChanged.IsBound() )
@@ -332,7 +335,7 @@ void SMenuAnchor::SetIsOpen( bool InIsOpen, const bool bFocusMenu )
 					SlatePrepass(LayoutScaleMultiplier);
 
 					// Figure out how big the content widget is so we can set the window's initial size properly
-					TSharedRef< SWidget > MenuContentRef = MenuContentPtr.ToSharedRef();
+					TSharedRef< SWidget > MenuContentRef = MenuContent.ToSharedRef();
 					MenuContentRef->SlatePrepass(LayoutScaleMultiplier);
 
 					// Combo-boxes never size down smaller than the widget that spawned them, but all
diff --git a/Engine/Source/Runtime/Slate/Public/Widgets/Input/SComboButton.h b/Engine/Source/Runtime/Slate/Public/Widgets/Input/SComboButton.h
--- a/Engine/Source/Runtime/Slate/Public/Widgets/Input/SComboButton.h
+++ b/Engine/Source/Runtime/Slate/Public/Widgets/Input/SComboButton.h
@@ -55,13 +55,10 @@ public:
 		/** Spawn a new window or reuse current window for this combo*/
 		SLATE_ARGUMENT( TOptional<EPopupMethod>, Method )
 	SLATE_END_ARGS()
-	
-	/**
-	 * Sets the content for this button
-	 *
-	 * @param	InContent	The widget to use as content for this button
-	 */
-	void SetMenuContent( const TSharedRef< SWidget >& InContent );
+
+	// SMenuAnchor interface
+	virtual void SetMenuContent(TSharedRef<SWidget> InContent) override;
+	// End of SMenuAnchor interface
 
 	/** See the OnGetMenuContent event */
 	void SetOnGetMenuContent( FOnGetContent InOnGetMenuContent );
@@ -88,9 +85,6 @@ protected:
 	/** Area where the button's content resides */
 	SHorizontalBox::FSlot* ButtonContentSlot;
 
-	/** Delegate to execute to get the menu content of this button */
-	FOnGetContent OnGetMenuContent;
-
 	/** Delegate to execute when the combo list is opened */
 	FOnComboBoxOpened OnComboBoxOpened;
 
diff --git a/Engine/Source/Runtime/Slate/Public/Widgets/Input/SMenuAnchor.h b/Engine/Source/Runtime/Slate/Public/Widgets/Input/SMenuAnchor.h
--- a/Engine/Source/Runtime/Slate/Public/Widgets/Input/SMenuAnchor.h
+++ b/Engine/Source/Runtime/Slate/Public/Widgets/Input/SMenuAnchor.h
@@ -67,7 +67,7 @@ public:
 	void SetContent(TSharedRef<SWidget> InContent);
 
 	/** See MenuContent attribute */
-	void SetMenuContent(TSharedRef<SWidget> InMenuContent);
+	virtual void SetMenuContent(TSharedRef<SWidget> InMenuContent);
 	
 	/**
 	 * Open or close the popup
