[UE-11357] - User would like to have a keyboard shortcut to close current tab
Added "Close Tab" command handled by [REDACTED_PERSON] widgets that closes the global 'active' tab if in the same window or closes the foreground tab of the focused stack.

#note So the usual expected behaviour is that this will close the last clicked or highlighted (yellow line on the tab's title) minor tab. It does this by checking if the the current global active tab is in the same window as the focussed tab stack. If these's no active minor tab in our window it tries to close the foreground tab of the focussed stack. Is this case the foreground tab will be a major tab as these are excluded from the global minor tab tracking. This is the most intuitive combination of closing behaviours I could find after some testing.

diff --git a/Engine/Source/Runtime/Slate/Private/Framework/Application/SlateApplication.cpp b/Engine/Source/Runtime/Slate/Private/Framework/Application/SlateApplication.cpp
--- a/Engine/Source/Runtime/Slate/Private/Framework/Application/SlateApplication.cpp
+++ b/Engine/Source/Runtime/Slate/Private/Framework/Application/SlateApplication.cpp
@@ -11,7 +11,7 @@
 #include "NotificationManager.h"
 #include "IInputProcessor.h"
 #include "ToolboxModule.h"
-
+#include "TabCommands.h"
 
 class FEventRouter
 {
@@ -759,6 +759,7 @@ FSlateApplication::FSlateApplication()
 	FInputCoreModule& InputCore = FModuleManager::LoadModuleChecked<FInputCoreModule>(TEXT("InputCore"));
 
 	FGenericCommands::Register();
+	FTabCommands::Register();
 
 	NormalExecutionGetter.BindRaw( this, &FSlateApplication::IsNormalExecution );
 	PointerIndexLastPositionMap.Add(CursorPointerIndex, FVector2D::ZeroVector);
@@ -766,6 +767,9 @@ FSlateApplication::FSlateApplication()
 
 FSlateApplication::~FSlateApplication()
 {
+	FTabCommands::Unregister();
+	FGenericCommands::Unregister();
+	
 	if (SlateTextField != nullptr)
 	{
 		delete SlateTextField;
diff --git a/Engine/Source/Runtime/Slate/Private/Framework/Docking/SDockingTabStack.cpp b/Engine/Source/Runtime/Slate/Private/Framework/Docking/SDockingTabStack.cpp
--- a/Engine/Source/Runtime/Slate/Private/Framework/Docking/SDockingTabStack.cpp
+++ b/Engine/Source/Runtime/Slate/Private/Framework/Docking/SDockingTabStack.cpp
@@ -2,7 +2,7 @@
 
 #include "SlatePrivatePCH.h"
 #include "DockingPrivate.h"
-
+#include "TabCommands.h"
 
 #define LOCTEXT_NAMESPACE "DockTabStack"
 
@@ -34,6 +34,8 @@ public:
 
 void SDockingTabStack::Construct( const FArguments& InArgs, const TSharedRef<FTabManager::FStack>& PersistentNode  )
 {
+	BindTabCommands();
+
 	Tabs = PersistentNode->Tabs;
 	this->SetSizeCoefficient(PersistentNode->GetSizeCoefficient());
 
@@ -414,6 +416,15 @@ FReply SDockingTabStack::OnMouseButtonDown( const FGeometry& MyGeometry, const F
 	}
 }
 
+FReply SDockingTabStack::OnKeyDown(const FGeometry& MyGeometry, const FKeyEvent& InKeyEvent)
+{
+	if (ActionList->ProcessCommandBindings(InKeyEvent))
+	{
+		return FReply::Handled();
+	}
+
+	return SDockingNode::OnKeyDown(MyGeometry, InKeyEvent);
+}
 
 FReply SDockingTabStack::OnUserAttemptingDock( SDockingNode::RelativeDirection Direction, const FDragDropEvent& DragDropEvent )
 {
@@ -1017,5 +1028,57 @@ FString SDockingTabStack::ShowPersistentTabs() const
 
 #endif
 
+void SDockingTabStack::BindTabCommands()
+{
+	check(!ActionList.IsValid());
+
+	ActionList = MakeShareable(new FUICommandList);
+
+	const FTabCommands& Commands = FTabCommands::Get();
+	ActionList->MapAction(Commands.CloseTab, FExecuteAction::CreateSP(this, &SDockingTabStack::ExecuteCloseTabCommand), FCanExecuteAction::CreateSP(this, &SDockingTabStack::CanExecuteCloseTabCommand));
+}
+
+void SDockingTabStack::ExecuteCloseTabCommand()
+{
+	auto DockArea = GetDockArea();
+	if (DockArea.IsValid())
+	{
+		TSharedPtr<FGlobalTabmanager> GlobalTabManager = FGlobalTabmanager::Get();
+		TSharedPtr<SDockTab> ActiveTab = GlobalTabManager->GetActiveTab();
+		if (ActiveTab.IsValid())
+		{
+			if (ActiveTab->GetParentWindow() == DockArea->GetParentWindow())
+			{
+				// Close the global active (minor) tab
+				ActiveTab->RequestCloseTab();
+				return;
+			}
+		}
+	}
+
+	// Close this stack's foreground tab
+	CloseForegroundTab();
+}
+
+bool SDockingTabStack::CanExecuteCloseTabCommand()
+{
+	auto DockArea = GetDockArea();
+	if (DockArea.IsValid())
+	{
+		TSharedPtr<FGlobalTabmanager> GlobalTabManager = FGlobalTabmanager::Get();
+		TSharedPtr<SDockTab> ActiveTab = GlobalTabManager->GetActiveTab();
+		if (ActiveTab.IsValid())
+		{
+			if (ActiveTab->GetParentWindow() == DockArea->GetParentWindow())
+			{
+				// Can close the global active (minor) tab because it's in the same window as this tab stack
+				return true;
+			}
+		}
+	}
+
+	// Not closing the global active tab - can we close this stack's foreground tab?
+	return CanCloseForegroundTab();
+}
 
 #undef LOCTEXT_NAMESPACE
diff --git a/Engine/Source/Runtime/Slate/Private/Framework/Docking/SDockingTabStack.h b/Engine/Source/Runtime/Slate/Private/Framework/Docking/SDockingTabStack.h
--- a/Engine/Source/Runtime/Slate/Private/Framework/Docking/SDockingTabStack.h
+++ b/Engine/Source/Runtime/Slate/Private/Framework/Docking/SDockingTabStack.h
@@ -99,6 +99,8 @@ public:
 
 	virtual FReply OnMouseButtonDown( const FGeometry& MyGeometry, const FPointerEvent& MouseEvent ) override;
 
+	virtual FReply OnKeyDown(const FGeometry& MyGeometry, const FKeyEvent& InKeyEvent) override;
+
 	virtual bool SupportsKeyboardFocus() const override { return true; }
 
 protected:
@@ -222,19 +224,26 @@ private:
 	/** @return Gets the visibility state for spacers that pad out the tab well to make room for title bar widgets */
 	EVisibility GetTitleAreaSpacerVisibility();
 
-	/** The tab in this dock stack that is active */
-	TSharedPtr<SDockTab> ActiveTab;
-
 	/** Visibility of TitleBar spacer based on maximize/restore status of the window.
 	 ** This gives us a little more space to grab the title bar when the window is not maximized
 	*/
 	EVisibility GetMaximizeSpacerVisibility() const;
 
+	/** Bind tab commands into the ActionList */
+	void BindTabCommands();
+
+	/** Attempts to close a tab when the CloseTab command is executed */
+	void ExecuteCloseTabCommand();
+
+	/** Attempts to find a tab that can be closed by the CloseTab command */
+	bool CanExecuteCloseTabCommand();
 
 #if DEBUG_TAB_MANAGEMENT
 	FString ShowPersistentTabs() const;
 #endif
 
+	/** Tab command list */
+	TSharedPtr<FUICommandList> ActionList;
 };
 
 
diff --git a/Engine/Source/Runtime/Slate/Private/Framework/Docking/TabCommands.cpp b/Engine/Source/Runtime/Slate/Private/Framework/Docking/TabCommands.cpp
new file mode 100644
--- /dev/null
+++ b/Engine/Source/Runtime/Slate/Private/Framework/Docking/TabCommands.cpp
@@ -0,0 +1,13 @@
+// Copyright 1998-2015 Epic Games, Inc. All Rights Reserved.
+
+#include "SlatePrivatePCH.h"
+#include "TabCommands.h"
+
+#define LOCTEXT_NAMESPACE ""
+
+void FTabCommands::RegisterCommands()
+{
+	UI_COMMAND(CloseTab, "Close Tab", "Closes the current window's active minor tab or the current focussed major tab if no minor tab is active.)", EUserInterfaceActionType::Button, FInputChord())
+}
+
+#undef LOCTEXT_NAMESPACE
\ No newline at end of file
diff --git a/Engine/Source/Runtime/Slate/Private/Widgets/Docking/SDockTab.cpp b/Engine/Source/Runtime/Slate/Private/Widgets/Docking/SDockTab.cpp
--- a/Engine/Source/Runtime/Slate/Private/Widgets/Docking/SDockTab.cpp
+++ b/Engine/Source/Runtime/Slate/Private/Widgets/Docking/SDockTab.cpp
@@ -673,7 +673,7 @@ FReply SDockTab::OnCloseButtonClicked()
 
 EVisibility SDockTab::HandleIsCloseButtonVisible() const
 {
-	return MyTabManager.Pin()->IsTabCloseable(SharedThis(this)) ? EVisibility::Visible : EVisibility::Hidden;
+	return CanCloseTab() ? EVisibility::Visible : EVisibility::Hidden;
 }
 
 bool SDockTab::CanCloseTab() const
diff --git a/Engine/Source/Runtime/Slate/Public/Framework/Docking/TabCommands.h b/Engine/Source/Runtime/Slate/Public/Framework/Docking/TabCommands.h
new file mode 100644
--- /dev/null
+++ b/Engine/Source/Runtime/Slate/Public/Framework/Docking/TabCommands.h
@@ -0,0 +1,21 @@
+// Copyright 1998-2015 Epic Games, Inc. All Rights Reserved.
+
+#pragma once
+
+class SLATE_API FTabCommands : public TCommands < FTabCommands >
+{
+public:
+
+	FTabCommands()
+		: TCommands<FTabCommands>(TEXT("TabCommands"), NSLOCTEXT("TabCommands", "DockingTabCommands", "Docking Tab Commands"), NAME_None, FCoreStyle::Get().GetStyleSetName())
+	{
+	}
+
+	virtual ~FTabCommands()
+	{
+	}
+
+	virtual void RegisterCommands() override;
+
+	TSharedPtr<FUICommandInfo> CloseTab;
+};
