Fixed crash in SDockTab::HandleIsCloseButtonVisible() caused by CL 2489109

diff --git a/Engine/Source/Runtime/Slate/Private/Widgets/Docking/SDockTab.cpp b/Engine/Source/Runtime/Slate/Private/Widgets/Docking/SDockTab.cpp
--- a/Engine/Source/Runtime/Slate/Private/Widgets/Docking/SDockTab.cpp
+++ b/Engine/Source/Runtime/Slate/Private/Widgets/Docking/SDockTab.cpp
@@ -673,7 +673,7 @@ FReply SDockTab::OnCloseButtonClicked()
 
 EVisibility SDockTab::HandleIsCloseButtonVisible() const
 {
-	return CanCloseTab() ? EVisibility::Visible : EVisibility::Hidden;
+	return MyTabManager.Pin()->IsTabCloseable(SharedThis(this)) ? EVisibility::Visible : EVisibility::Hidden;
 }
 
 bool SDockTab::CanCloseTab() const
