Remove unused drawn windows list.

diff --git a/Engine/Source/Runtime/Slate/Private/Framework/Application/SlateApplication.cpp b/Engine/Source/Runtime/Slate/Private/Framework/Application/SlateApplication.cpp
--- a/Engine/Source/Runtime/Slate/Private/Framework/Application/SlateApplication.cpp
+++ b/Engine/Source/Runtime/Slate/Private/Framework/Application/SlateApplication.cpp
@@ -939,7 +939,6 @@ struct FDrawWindowArgs
 		, WidgetsUnderCursor( InWidgetsUnderCursor )
 	{}
 
-	TArray<FGenericWindow*, TInlineAllocator<10> > OutDrawnWindows;
 	FSlateDrawBuffer& OutDrawBuffer;
 	const FWidgetPath& WidgetsUnderCursor;
 };
@@ -1037,11 +1036,6 @@ void FSlateApplication::DrawWindowAndChildren( const TSharedRef<SWindow>& Window
 			MaxLayerId = WidgetReflector->VisualizeCursorAndKeys( WindowElementList, MaxLayerId );
 		}
 
-		// Keep track of windows that we're actually going to be presenting, so we can mark
-		// them as 'drawn' afterwards.
-		FGenericWindow* NativeWindow = WindowToDraw->GetNativeWindow().Get();
-		DrawWindowArgs.OutDrawnWindows.Add( NativeWindow );
-
 		// Draw the child windows
 		const TArray< TSharedRef<SWindow> >& WindowChildren = WindowToDraw->GetChildWindows();
 		for (int32 ChildIndex=0; ChildIndex < WindowChildren.Num(); ++ChildIndex)
