Slate: Added missing override keywords

diff --git a/Engine/Source/Runtime/Slate/Public/Framework/Docking/TabManager.h b/Engine/Source/Runtime/Slate/Public/Framework/Docking/TabManager.h
--- a/Engine/Source/Runtime/Slate/Public/Framework/Docking/TabManager.h
+++ b/Engine/Source/Runtime/Slate/Public/Framework/Docking/TabManager.h
@@ -841,7 +841,7 @@ protected:
 	virtual void UpdateStats() override;
 
 public:
-	virtual void OnTabManagerClosing();
+	virtual void OnTabManagerClosing() override;
 
 
 private:
diff --git a/Engine/Source/Runtime/Slate/Public/Framework/Text/SlateWidgetRun.h b/Engine/Source/Runtime/Slate/Public/Framework/Text/SlateWidgetRun.h
--- a/Engine/Source/Runtime/Slate/Public/Framework/Text/SlateWidgetRun.h
+++ b/Engine/Source/Runtime/Slate/Public/Framework/Text/SlateWidgetRun.h
@@ -39,36 +39,23 @@ public:
 
 	virtual FTextRange GetTextRange() const override;
 	virtual void SetTextRange( const FTextRange& Value ) override;
-
 	virtual int16 GetBaseLine( float Scale ) const override;
 	virtual int16 GetMaxHeight( float Scale ) const override;
-
 	virtual FVector2D Measure( int32 StartIndex, int32 EndIndex, float Scale ) const override;
 	virtual int8 GetKerning(int32 CurrentIndex, float Scale) const override;
-
 	virtual TSharedRef< ILayoutBlock > CreateBlock( int32 StartIndex, int32 EndIndex, FVector2D Size, const TSharedPtr< IRunRenderer >& Renderer ) override;
-
 	virtual int32 OnPaint( const FPaintArgs& Args, const FTextLayout::FLineView& Line, const TSharedRef< ILayoutBlock >& Block, const FTextBlockStyle& DefaultStyle, const FGeometry& AllottedGeometry, const FSlateRect& MyClippingRect, FSlateWindowElementList& OutDrawElements, int32 LayerId, const FWidgetStyle& InWidgetStyle, bool bParentEnabled ) const override;
-
 	virtual const TArray< TSharedRef<SWidget> >& GetChildren() override;
-
-	virtual void ArrangeChildren( const TSharedRef< ILayoutBlock >& Block, const FGeometry& AllottedGeometry, FArrangedChildren& ArrangedChildren ) const;
-
+	virtual void ArrangeChildren( const TSharedRef< ILayoutBlock >& Block, const FGeometry& AllottedGeometry, FArrangedChildren& ArrangedChildren ) const override;
 	virtual int32 GetTextIndexAt( const TSharedRef< ILayoutBlock >& Block, const FVector2D& Location, float Scale, ETextHitPoint* const OutHitPoint = nullptr ) const override;
-
 	virtual FVector2D GetLocationAt( const TSharedRef< ILayoutBlock >& Block, int32 Offset, float Scale ) const override;
-
 	virtual void BeginLayout() override {}
 	virtual void EndLayout() override {}
-
 	virtual void Move(const TSharedRef<FString>& NewText, const FTextRange& NewRange) override;
 	virtual TSharedRef<IRun> Clone() const override;
-
 	virtual void AppendTextTo(FString& Text) const override;
 	virtual void AppendTextTo(FString& AppendToText, const FTextRange& PartialRange) const override;
-
 	virtual const FRunInfo& GetRunInfo() const override;
-
 	virtual ERunAttributes GetRunAttributes() const override;
 
 private:
diff --git a/Engine/Source/Runtime/Slate/Public/Widgets/Input/SMenuAnchor.h b/Engine/Source/Runtime/Slate/Public/Widgets/Input/SMenuAnchor.h
--- a/Engine/Source/Runtime/Slate/Public/Widgets/Input/SMenuAnchor.h
+++ b/Engine/Source/Runtime/Slate/Public/Widgets/Input/SMenuAnchor.h
@@ -96,7 +96,7 @@ protected:
 	virtual void OnArrangeChildren( const FGeometry& AllottedGeometry, FArrangedChildren& ArrangedChildren ) const override;
 	virtual FVector2D ComputeDesiredSize(float) const override;
 	virtual FChildren* GetChildren() override;
-	int32 OnPaint( const FPaintArgs& Args, const FGeometry& AllottedGeometry, const FSlateRect& MyClippingRect, FSlateWindowElementList& OutDrawElements, int32 LayerId, const FWidgetStyle& InWidgetStyle, bool bParentEnabled ) const;
+	virtual int32 OnPaint( const FPaintArgs& Args, const FGeometry& AllottedGeometry, const FSlateRect& MyClippingRect, FSlateWindowElementList& OutDrawElements, int32 LayerId, const FWidgetStyle& InWidgetStyle, bool bParentEnabled ) const override;
 	// End of SWidget interface
 
 	/** Invoked when the popup window is being closed by the application */
diff --git a/Engine/Source/Runtime/Slate/Public/Widgets/Input/SSlider.h b/Engine/Source/Runtime/Slate/Public/Widgets/Input/SSlider.h
--- a/Engine/Source/Runtime/Slate/Public/Widgets/Input/SSlider.h
+++ b/Engine/Source/Runtime/Slate/Public/Widgets/Input/SSlider.h
@@ -89,7 +89,7 @@ public:
 	// SWidget overrides
 
 	virtual int32 OnPaint( const FPaintArgs& Args, const FGeometry& AllottedGeometry, const FSlateRect& MyClippingRect, FSlateWindowElementList& OutDrawElements, int32 LayerId, const FWidgetStyle& InWidgetStyle, bool bParentEnabled ) const override;
-	virtual FVector2D ComputeDesiredSize(float) const;
+	virtual FVector2D ComputeDesiredSize(float) const override;
 	virtual FReply OnMouseButtonDown( const FGeometry& MyGeometry, const FPointerEvent& MouseEvent ) override;
 	virtual FReply OnMouseButtonUp( const FGeometry& MyGeometry, const FPointerEvent& MouseEvent ) override;
 	virtual FReply OnMouseMove( const FGeometry& MyGeometry, const FPointerEvent& MouseEvent ) override;
diff --git a/Engine/Source/Runtime/Slate/Public/Widgets/SViewport.h b/Engine/Source/Runtime/Slate/Public/Widgets/SViewport.h
--- a/Engine/Source/Runtime/Slate/Public/Widgets/SViewport.h
+++ b/Engine/Source/Runtime/Slate/Public/Widgets/SViewport.h
@@ -170,7 +170,7 @@ public:
 	virtual TOptional<EPopupMethod> OnQueryPopupMethod() const override;
 	virtual void OnFinishedPointerInput() override;
 	virtual void OnArrangeChildren( const FGeometry& AllottedGeometry, FArrangedChildren& ArrangedChildren ) const override;
-	virtual TSharedPtr<struct FVirtualPointerPosition> TranslateMouseCoordinateFor3DChild(const TSharedRef<SWidget>& ChildWidget, const FGeometry& MyGeometry, const FVector2D& ScreenSpaceMouseCoordinate, const FVector2D& LastScreenSpaceMouseCoordinate) const;
+	virtual TSharedPtr<struct FVirtualPointerPosition> TranslateMouseCoordinateFor3DChild(const TSharedRef<SWidget>& ChildWidget, const FGeometry& MyGeometry, const FVector2D& ScreenSpaceMouseCoordinate, const FVector2D& LastScreenSpaceMouseCoordinate) const override;
 	virtual FNavigationReply OnNavigation( const FGeometry& MyGeometry, const FNavigationEvent& InNavigationEvent ) override;
 
 private:
diff --git a/Engine/Source/Runtime/Slate/Public/Widgets/Text/SRichTextBlock.h b/Engine/Source/Runtime/Slate/Public/Widgets/Text/SRichTextBlock.h
--- a/Engine/Source/Runtime/Slate/Public/Widgets/Text/SRichTextBlock.h
+++ b/Engine/Source/Runtime/Slate/Public/Widgets/Text/SRichTextBlock.h
@@ -103,13 +103,9 @@ public:
 
 	// Begin SWidget Interface
 	void Construct( const FArguments& InArgs );
-
-	virtual int32 OnPaint( const FPaintArgs& Args, const FGeometry& AllottedGeometry, const FSlateRect& MyClippingRect, FSlateWindowElementList& OutDrawElements, int32 LayerId, const FWidgetStyle& InWidgetStyle, bool bParentEnabled ) const;
-
+	virtual int32 OnPaint( const FPaintArgs& Args, const FGeometry& AllottedGeometry, const FSlateRect& MyClippingRect, FSlateWindowElementList& OutDrawElements, int32 LayerId, const FWidgetStyle& InWidgetStyle, bool bParentEnabled ) const override;
 	virtual FVector2D ComputeDesiredSize(float) const override;
-
 	virtual FChildren* GetChildren() override;
-
 	virtual void OnArrangeChildren( const FGeometry& AllottedGeometry, FArrangedChildren& ArrangedChildren ) const override;
 	// End SWidget Interface
 
diff --git a/Engine/Source/Runtime/SlateCore/Public/Layout/Children.h b/Engine/Source/Runtime/SlateCore/Public/Layout/Children.h
--- a/Engine/Source/Runtime/SlateCore/Public/Layout/Children.h
+++ b/Engine/Source/Runtime/SlateCore/Public/Layout/Children.h
@@ -81,9 +81,9 @@ class TSupportsOneChildMixin : public FChildren, public TSlotBase<MixedIntoType>
 public:
 	TSupportsOneChildMixin():TSlotBase<MixedIntoType>(){}
 
-	virtual int32 Num() const { return 1; }
-	virtual TSharedRef<SWidget> GetChildAt( int32 ChildIndex ) { check(ChildIndex == 0); return FSlotBase::GetWidget(); }
-	virtual TSharedRef<const SWidget> GetChildAt( int32 ChildIndex ) const { check(ChildIndex == 0); return FSlotBase::GetWidget(); }
+	virtual int32 Num() const override { return 1; }
+	virtual TSharedRef<SWidget> GetChildAt( int32 ChildIndex ) override { check(ChildIndex == 0); return FSlotBase::GetWidget(); }
+	virtual TSharedRef<const SWidget> GetChildAt( int32 ChildIndex ) const override { check(ChildIndex == 0); return FSlotBase::GetWidget(); }
 private:
 	virtual const FSlotBase& GetSlotAt(int32 ChildIndex) const override { check(ChildIndex == 0); return *this; }
 };
@@ -102,9 +102,9 @@ public:
 	{
 	}
 
-	virtual int32 Num() const { return WidgetPtr.IsValid() ? 1 : 0 ; }
-	virtual TSharedRef<SWidget> GetChildAt( int32 ChildIndex ) { check(ChildIndex == 0); return WidgetPtr.Pin().ToSharedRef(); }
-	virtual TSharedRef<const SWidget> GetChildAt( int32 ChildIndex ) const { check(ChildIndex == 0); return WidgetPtr.Pin().ToSharedRef(); }
+	virtual int32 Num() const override { return WidgetPtr.IsValid() ? 1 : 0 ; }
+	virtual TSharedRef<SWidget> GetChildAt( int32 ChildIndex ) override { check(ChildIndex == 0); return WidgetPtr.Pin().ToSharedRef(); }
+	virtual TSharedRef<const SWidget> GetChildAt( int32 ChildIndex ) const override { check(ChildIndex == 0); return WidgetPtr.Pin().ToSharedRef(); }
 
 private:
 	virtual const FSlotBase& GetSlotAt(int32 ChildIndex) const override { static FSlotBase NullSlot; check(ChildIndex == 0); return NullSlot; }
diff --git a/Engine/Source/Runtime/SlateCore/Public/Widgets/SLeafWidget.h b/Engine/Source/Runtime/SlateCore/Public/Widgets/SLeafWidget.h
--- a/Engine/Source/Runtime/SlateCore/Public/Widgets/SLeafWidget.h
+++ b/Engine/Source/Runtime/SlateCore/Public/Widgets/SLeafWidget.h
@@ -29,7 +29,7 @@ private:
 	 *
 	 * Whenever possible, LeafWidgets should avoid dealing with layout properties. See TextBlock for an example.
 	 */
-	virtual int32 OnPaint( const FPaintArgs& Args, const FGeometry& AllottedGeometry, const FSlateRect& MyClippingRect, FSlateWindowElementList& OutDrawElements, int32 LayerId, const FWidgetStyle& InWidgetStyle, bool bParentEnabled ) const = 0;
+	virtual int32 OnPaint( const FPaintArgs& Args, const FGeometry& AllottedGeometry, const FSlateRect& MyClippingRect, FSlateWindowElementList& OutDrawElements, int32 LayerId, const FWidgetStyle& InWidgetStyle, bool bParentEnabled ) const override = 0;
 	
 	/**
 	 * Overwritten from SWidget.
@@ -38,16 +38,16 @@ private:
 	 * take child widgets into account as LeafWidgets have none by definition. For example, the TextBlock widget simply
 	 * measures the area necessary to display its text with the given font and font size.
 	 */
-	virtual FVector2D ComputeDesiredSize(float) const = 0;
+	virtual FVector2D ComputeDesiredSize(float) const override = 0;
 	
 	/**
 	 * Overwritten from SWidget.
 	 *
 	 * Leaf widgets never have children.
 	 */
-	virtual FChildren* GetChildren( );
+	virtual FChildren* GetChildren() override;
 
-	virtual void OnArrangeChildren( const FGeometry& AllottedGeometry, FArrangedChildren& ArrangedChildren ) const;
+	virtual void OnArrangeChildren( const FGeometry& AllottedGeometry, FArrangedChildren& ArrangedChildren ) const override;
 
 private:
 
diff --git a/Engine/Source/Runtime/SlateCore/Public/Widgets/SPanel.h b/Engine/Source/Runtime/SlateCore/Public/Widgets/SPanel.h
--- a/Engine/Source/Runtime/SlateCore/Public/Widgets/SPanel.h
+++ b/Engine/Source/Runtime/SlateCore/Public/Widgets/SPanel.h
@@ -22,7 +22,7 @@ public:
 	 * @param AllottedGeometry    The geometry allotted for this widget by its parent.
 	 * @param ArrangedChildren    The array to which to add the WidgetGeometries that represent the arranged children.
 	 */
-	virtual void OnArrangeChildren( const FGeometry& AllottedGeometry, FArrangedChildren& ArrangedChildren ) const = 0;
+	virtual void OnArrangeChildren( const FGeometry& AllottedGeometry, FArrangedChildren& ArrangedChildren ) const override = 0;
 
 	/**
 	 * A Panel's desired size in the space required to arrange of its children on the screen while respecting all of
@@ -30,7 +30,7 @@ public:
 	 *
 	 * @return The desired size.
 	 */
-	virtual FVector2D ComputeDesiredSize(float) const = 0;
+	virtual FVector2D ComputeDesiredSize(float) const override = 0;
 
 	/**
 	 * All widgets must provide a way to access their children in a layout-agnostic way.
@@ -39,7 +39,7 @@ public:
 	 * provides layout information about the child it stores. In that case
 	 * GetChildren should simply return the TPanelChildren<Slot>. See StackPanel for an example.
 	 */
-	virtual FChildren* GetChildren( ) = 0;
+	virtual FChildren* GetChildren() override = 0;
 
 public:
 
@@ -47,7 +47,7 @@ public:
 	 * Most panels do not create widgets as part of their implementation, so
 	 * they do not need to implement a Construct()
 	 */
-	void Construct( ) { }
+	void Construct() { }
 
 public:
 
diff --git a/Engine/Source/Runtime/SlateCore/Public/Widgets/SWindow.h b/Engine/Source/Runtime/SlateCore/Public/Widgets/SWindow.h
--- a/Engine/Source/Runtime/SlateCore/Public/Widgets/SWindow.h
+++ b/Engine/Source/Runtime/SlateCore/Public/Widgets/SWindow.h
@@ -583,7 +583,7 @@ public:
 	}
 
 	// Events
-	virtual FCursorReply OnCursorQuery( const FGeometry& MyGeometry, const FPointerEvent& CursorEvent ) const;
+	virtual FCursorReply OnCursorQuery( const FGeometry& MyGeometry, const FPointerEvent& CursorEvent ) const override;
 
 	/** The system will call this method to notify the window that it has been places in the foreground or background. */
 	virtual bool OnIsActiveChanged( const FWindowActivateEvent& ActivateEvent );
@@ -607,8 +607,7 @@ public:
 	virtual bool SupportsKeyboardFocus() const override;
 
 private:
-	virtual FReply OnFocusReceived( const FGeometry& MyGeometry, const FFocusEvent& InFocusEvent );
-
+	virtual FReply OnFocusReceived( const FGeometry& MyGeometry, const FFocusEvent& InFocusEvent ) override;
 	virtual FReply OnMouseButtonDown( const FGeometry& MyGeometry, const FPointerEvent& MouseEvent ) override;
 	virtual FReply OnMouseButtonUp( const FGeometry& MyGeometry, const FPointerEvent& MouseEvent ) override;
 	virtual FReply OnMouseMove( const FGeometry& MyGeometry, const FPointerEvent& MouseEvent ) override;
