Slate - No longer ensuring that the widget path must be valid when showing a menu anchor.  If no valid path is found, no popup occurs.

diff --git a/Engine/Source/Runtime/Slate/Private/Widgets/Input/SMenuAnchor.cpp b/Engine/Source/Runtime/Slate/Private/Widgets/Input/SMenuAnchor.cpp
--- a/Engine/Source/Runtime/Slate/Private/Widgets/Input/SMenuAnchor.cpp
+++ b/Engine/Source/Runtime/Slate/Private/Widgets/Input/SMenuAnchor.cpp
@@ -327,7 +327,7 @@ void SMenuAnchor::SetIsOpen( bool InIsOpen, const bool bFocusMenu )
 				// NOTE: Careful, GeneratePathToWidget can be reentrant in that it can call visibility delegates and such
 				FWidgetPath MyWidgetPath;
 				FSlateApplication::Get().GeneratePathToWidgetUnchecked(AsShared(), MyWidgetPath);
-				if (ensure(MyWidgetPath.IsValid()))
+				if (MyWidgetPath.IsValid())
 				{
 					const FGeometry& MyGeometry = MyWidgetPath.Widgets.Last().Geometry;
 					const float LayoutScaleMultiplier = MyGeometry.GetAccumulatedLayoutTransform().GetScale();
