Merging using UE4-To-UE4-Orion

#UE4

- Adding a "User Specified Scale" to Scale Box.

Submitting to Dev branch, then moving over to Main

diff --git a/Engine/Source/Runtime/Slate/Private/Widgets/Layout/SScaleBox.cpp b/Engine/Source/Runtime/Slate/Private/Widgets/Layout/SScaleBox.cpp
--- a/Engine/Source/Runtime/Slate/Private/Widgets/Layout/SScaleBox.cpp
+++ b/Engine/Source/Runtime/Slate/Private/Widgets/Layout/SScaleBox.cpp
@@ -13,6 +13,7 @@ void SScaleBox::Construct( const SScaleBox::FArguments& InArgs )
 {
 	Stretch = InArgs._Stretch;
 	StretchDirection = InArgs._StretchDirection;
+	UserSpecifiedScale = InArgs._UserSpecifiedScale;
 
 	ChildSlot
 	.HAlign(InArgs._HAlign)
@@ -52,6 +53,9 @@ void SScaleBox::OnArrangeChildren( const FGeometry& AllottedGeometry, FArrangedC
 			case EStretch::ScaleToFill:
 				FinalScale = FMath::Max(AllottedGeometry.Size.X / DesiredSize.X, AllottedGeometry.Size.Y / DesiredSize.Y);
 				break;
+			case EStretch::UserSpecified:
+				FinalScale = UserSpecifiedScale.Get(1.0f);
+				break;
 			}
 
 			switch ( CurrentStretchDirection )
@@ -118,3 +122,8 @@ void SScaleBox::SetStretch(EStretch::Type InStretch)
 {
 	Stretch = InStretch;
 }
+
+void SScaleBox::SetUserSpecifiedScale(float InUserSpecifiedScale)
+{
+	UserSpecifiedScale = InUserSpecifiedScale;
+}
diff --git a/Engine/Source/Runtime/Slate/Public/Widgets/Layout/SScaleBox.h b/Engine/Source/Runtime/Slate/Public/Widgets/Layout/SScaleBox.h
--- a/Engine/Source/Runtime/Slate/Public/Widgets/Layout/SScaleBox.h
+++ b/Engine/Source/Runtime/Slate/Public/Widgets/Layout/SScaleBox.h
@@ -30,7 +30,9 @@ namespace EStretch
 		/** Scales the content uniformly (preserving aspect ratio) until it can no longer scale the content without clipping it. */
 		ScaleToFit,
 		/** Scales the content uniformly (preserving aspect ratio), until all sides meet or exceed the size of the area.  Will result in clipping longer sides. */
-		ScaleToFill
+		ScaleToFill,
+		/** Scales the content by the scale specified by the user */
+		UserSpecified
 	};
 }
 
@@ -48,6 +50,7 @@ public:
 	, _VAlign(VAlign_Center)
 	, _StretchDirection(EStretchDirection::Both)
 	, _Stretch(EStretch::None)
+	, _UserSpecifiedScale(1.0f)
 	{}
 		/** Slot for this designers content (optional) */
 		SLATE_DEFAULT_SLOT(FArguments, Content)
@@ -64,6 +67,9 @@ public:
 		/** The stretching rule to apply when content is stretched */
 		SLATE_ATTRIBUTE(EStretch::Type, Stretch)
 
+		/** Optional scale that can be specified by the User */
+		SLATE_ATTRIBUTE(float, UserSpecifiedScale)
+
 	SLATE_END_ARGS()
 
 	void Construct(const FArguments& InArgs);
@@ -87,6 +93,9 @@ public:
 
 	/** See Stretch argument */
 	void SetStretch(EStretch::Type InStretch);
+
+	/** See UserSpecifiedScale argument */
+	void SetUserSpecifiedScale(float InUserSpecifiedScale);
 	
 private:
 	/** The allowed direction of stretching of the content */
@@ -94,4 +103,7 @@ private:
 
 	/** The method of scaling that is applied to the content. */
 	TAttribute<EStretch::Type> Stretch;
+
+	/** Optional scale that can be specified by the User */
+	TAttribute<float> UserSpecifiedScale;
 };
diff --git a/Engine/Source/Runtime/UMG/Private/Components/ScaleBox.cpp b/Engine/Source/Runtime/UMG/Private/Components/ScaleBox.cpp
--- a/Engine/Source/Runtime/UMG/Private/Components/ScaleBox.cpp
+++ b/Engine/Source/Runtime/UMG/Private/Components/ScaleBox.cpp
@@ -14,6 +14,7 @@ UScaleBox::UScaleBox(const FObjectInitializer& ObjectInitializer)
 
 	StretchDirection = EStretchDirection::Both;
 	Stretch = EStretch::ScaleToFit;
+	UserSpecifiedScale = 1.0f;
 }
 
 void UScaleBox::ReleaseSlateResources(bool bReleaseChildren)
@@ -41,6 +42,7 @@ void UScaleBox::SynchronizeProperties()
 
 	MyScaleBox->SetStretchDirection(StretchDirection);
 	MyScaleBox->SetStretch(Stretch);
+	MyScaleBox->SetUserSpecifiedScale(UserSpecifiedScale);
 }
 
 UClass* UScaleBox::GetSlotClass() const
diff --git a/Engine/Source/Runtime/UMG/Public/Components/ScaleBox.h b/Engine/Source/Runtime/UMG/Public/Components/ScaleBox.h
--- a/Engine/Source/Runtime/UMG/Public/Components/ScaleBox.h
+++ b/Engine/Source/Runtime/UMG/Public/Components/ScaleBox.h
@@ -27,11 +27,16 @@ public:
 	UPROPERTY(EditDefaultsOnly, BlueprintReadOnly, Category="Stretching")
 	TEnumAsByte<EStretch::Type> Stretch;
 
+	/** Optional scale that can be specified by the User. Isn't used if Stretch != EStretch::UserSpecified */
+	UPROPERTY(EditDefaultsOnly, BlueprintReadOnly, Category="Stretching")
+	float UserSpecifiedScale;
+
 public:
 
 	//TODO UMG Add Set Stretch
 	//TODO UMG Add Set Stretch DIrection
-
+	//TODO UMG Add Set UserSpecifiedScale
+	
 public:
 
 	// UWidget interface
