Split "Close Tab" commands into major and minor tab versions. Added tooltips to tab close buttons showing key bindings. Added default key binding for closing major tabs (Ctrl-Alt-F4)
[UE-11357] - User would like to have a keyboard shortcut to close current tab

diff --git a/Engine/Source/Runtime/Slate/Private/Framework/Docking/FDockingDragOperation.cpp b/Engine/Source/Runtime/Slate/Private/Framework/Docking/FDockingDragOperation.cpp
--- a/Engine/Source/Runtime/Slate/Private/Framework/Docking/FDockingDragOperation.cpp
+++ b/Engine/Source/Runtime/Slate/Private/Framework/Docking/FDockingDragOperation.cpp
@@ -272,7 +272,7 @@ FDockingDragOperation::FDockingDragOperation( const TSharedRef<SDockTab>& InTabT
 /** @return The offset into the tab where the user grabbed in Slate Units. */
 const FVector2D FDockingDragOperation::GetDecoratorOffsetFromCursor()
 {
-	const ETabRole RoleToUse = TabBeingDragged->IsNomadTabWithMajorTabStyle() ? ETabRole::MajorTab : TabBeingDragged->GetTabRole();
+	const ETabRole RoleToUse = TabBeingDragged->GetVisualTabRole();
 	const FVector2D TabDesiredSize = TabBeingDragged->GetDesiredSize();
 	const FVector2D MaxTabSize = FDockingConstants::GetMaxTabSizeFor(RoleToUse);
 
diff --git a/Engine/Source/Runtime/Slate/Private/Framework/Docking/SDockingTabStack.cpp b/Engine/Source/Runtime/Slate/Private/Framework/Docking/SDockingTabStack.cpp
--- a/Engine/Source/Runtime/Slate/Private/Framework/Docking/SDockingTabStack.cpp
+++ b/Engine/Source/Runtime/Slate/Private/Framework/Docking/SDockingTabStack.cpp
@@ -831,7 +831,7 @@ FReply SDockingTabStack::UnhideTabWell()
 
 bool SDockingTabStack::CanHideTabWell() const
 {
-	return GetNumTabs() == 1 && (GetTabs()[0]->GetTabRole() != ETabRole::MajorTab && !GetTabs()[0]->IsNomadTabWithMajorTabStyle());
+	return GetNumTabs() == 1 && FGlobalTabmanager::Get()->CanSetAsActiveTab(GetTabs()[0]);
 }
 
 bool SDockingTabStack::CanCloseForegroundTab() const
@@ -1035,32 +1035,36 @@ void SDockingTabStack::BindTabCommands()
 	ActionList = MakeShareable(new FUICommandList);
 
 	const FTabCommands& Commands = FTabCommands::Get();
-	ActionList->MapAction(Commands.CloseTab, FExecuteAction::CreateSP(this, &SDockingTabStack::ExecuteCloseTabCommand), FCanExecuteAction::CreateSP(this, &SDockingTabStack::CanExecuteCloseTabCommand));
+	ActionList->MapAction(Commands.CloseMajorTab, FExecuteAction::CreateSP(this, &SDockingTabStack::ExecuteCloseMajorTabCommand), FCanExecuteAction::CreateSP(this, &SDockingTabStack::CanExecuteCloseMajorTabCommand));
+	ActionList->MapAction(Commands.CloseMinorTab, FExecuteAction::CreateSP(this, &SDockingTabStack::ExecuteCloseMinorTabCommand), FCanExecuteAction::CreateSP(this, &SDockingTabStack::CanExecuteCloseMinorTabCommand));
 }
 
-void SDockingTabStack::ExecuteCloseTabCommand()
+void SDockingTabStack::ExecuteCloseMajorTabCommand()
 {
-	auto DockArea = GetDockArea();
-	if (DockArea.IsValid())
+	// Close this stack's foreground tab (if it's a major tab)
+	if (CanExecuteCloseMajorTabCommand())
 	{
-		TSharedPtr<FGlobalTabmanager> GlobalTabManager = FGlobalTabmanager::Get();
-		TSharedPtr<SDockTab> ActiveTab = GlobalTabManager->GetActiveTab();
-		if (ActiveTab.IsValid())
-		{
-			if (ActiveTab->GetParentWindow() == DockArea->GetParentWindow())
-			{
-				// Close the global active (minor) tab
-				ActiveTab->RequestCloseTab();
-				return;
-			}
-		}
+		TabWell->GetForegroundTab()->RequestCloseTab();
 	}
+}
+
+bool SDockingTabStack::CanExecuteCloseMajorTabCommand()
+{
+	// Can we close this stack's foreground tab (if it's a major tab)?
+	TSharedPtr<SDockTab> ForegroundTab = TabWell->GetForegroundTab();
+	return ForegroundTab.IsValid() && !FGlobalTabmanager::Get()->CanSetAsActiveTab(ForegroundTab);
+}
 
-	// Close this stack's foreground tab
-	CloseForegroundTab();
+void SDockingTabStack::ExecuteCloseMinorTabCommand()
+{
+	if (CanExecuteCloseMinorTabCommand())
+	{
+		// Close the global active (minor) tab
+		FGlobalTabmanager::Get()->GetActiveTab()->RequestCloseTab();
+	}
 }
 
-bool SDockingTabStack::CanExecuteCloseTabCommand()
+bool SDockingTabStack::CanExecuteCloseMinorTabCommand()
 {
 	auto DockArea = GetDockArea();
 	if (DockArea.IsValid())
@@ -1076,9 +1080,7 @@ bool SDockingTabStack::CanExecuteCloseTabCommand()
 			}
 		}
 	}
-
-	// Not closing the global active tab - can we close this stack's foreground tab?
-	return CanCloseForegroundTab();
+	return false;
 }
 
 #undef LOCTEXT_NAMESPACE
diff --git a/Engine/Source/Runtime/Slate/Private/Framework/Docking/SDockingTabStack.h b/Engine/Source/Runtime/Slate/Private/Framework/Docking/SDockingTabStack.h
--- a/Engine/Source/Runtime/Slate/Private/Framework/Docking/SDockingTabStack.h
+++ b/Engine/Source/Runtime/Slate/Private/Framework/Docking/SDockingTabStack.h
@@ -232,11 +232,17 @@ private:
 	/** Bind tab commands into the ActionList */
 	void BindTabCommands();
 
-	/** Attempts to close a tab when the CloseTab command is executed */
-	void ExecuteCloseTabCommand();
+	/** Attempts to close the foreground tab when the CloseMajorTab command is executed */
+	void ExecuteCloseMajorTabCommand();
 
-	/** Attempts to find a tab that can be closed by the CloseTab command */
-	bool CanExecuteCloseTabCommand();
+	/** Attempts to find the foreground tab that can be closed by the CloseMajorTab command */
+	bool CanExecuteCloseMajorTabCommand();
+
+	/** Attempts to close the active tab when the CloseMinorTab command is executed */
+	void ExecuteCloseMinorTabCommand();
+
+	/** Attempts to find the active tab that can be closed by the CloseMinorTab command */
+	bool CanExecuteCloseMinorTabCommand();
 
 #if DEBUG_TAB_MANAGEMENT
 	FString ShowPersistentTabs() const;
diff --git a/Engine/Source/Runtime/Slate/Private/Framework/Docking/SDockingTabWell.cpp b/Engine/Source/Runtime/Slate/Private/Framework/Docking/SDockingTabWell.cpp
--- a/Engine/Source/Runtime/Slate/Private/Framework/Docking/SDockingTabWell.cpp
+++ b/Engine/Source/Runtime/Slate/Private/Framework/Docking/SDockingTabWell.cpp
@@ -213,7 +213,7 @@ FVector2D SDockingTabWell::ComputeChildSize( const FGeometry& AllottedGeometry )
 	FVector2D MaxTabSize(0,0);
 	if ( FirstTab.IsValid() )
 	{
-		const ETabRole RoleToUse = FirstTab->IsNomadTabWithMajorTabStyle() ? ETabRole::MajorTab : FirstTab->GetTabRole();
+		const ETabRole RoleToUse = FirstTab->GetVisualTabRole();
 		MaxTabSize = FDockingConstants::GetMaxTabSizeFor(RoleToUse);
 	}
 	else
diff --git a/Engine/Source/Runtime/Slate/Private/Framework/Docking/TabCommands.cpp b/Engine/Source/Runtime/Slate/Private/Framework/Docking/TabCommands.cpp
--- a/Engine/Source/Runtime/Slate/Private/Framework/Docking/TabCommands.cpp
+++ b/Engine/Source/Runtime/Slate/Private/Framework/Docking/TabCommands.cpp
@@ -7,7 +7,8 @@
 
 void FTabCommands::RegisterCommands()
 {
-	UI_COMMAND(CloseTab, "Close Tab", "Closes the current window's active minor tab or the current focussed major tab if no minor tab is active.)", EUserInterfaceActionType::Button, FInputChord())
+	UI_COMMAND(CloseMajorTab, "Close Major Tab", "Closes the focused major tab", EUserInterfaceActionType::Button, FInputChord(EModifierKey::Control | EModifierKey::Alt, EKeys::F4))
+	UI_COMMAND(CloseMinorTab, "Close Minor Tab", "Closes the current window's active minor tab", EUserInterfaceActionType::Button, FInputChord())
 }
 
 #undef LOCTEXT_NAMESPACE
\ No newline at end of file
diff --git a/Engine/Source/Runtime/Slate/Private/Framework/Docking/TabManager.cpp b/Engine/Source/Runtime/Slate/Private/Framework/Docking/TabManager.cpp
--- a/Engine/Source/Runtime/Slate/Private/Framework/Docking/TabManager.cpp
+++ b/Engine/Source/Runtime/Slate/Private/Framework/Docking/TabManager.cpp
@@ -1588,13 +1588,16 @@ TSharedPtr<class SDockTab> FGlobalTabmanager::GetActiveTab() const
 	return ActiveTabPtr.Pin();
 }
 
+bool FGlobalTabmanager::CanSetAsActiveTab(const TSharedPtr<SDockTab>& Tab)
+{
+	// Setting NULL wipes out the active tab; always apply that change.
+	// Major tabs are ignored for the purposes of active-tab tracking. We do not care about their
+	return !Tab.IsValid() || Tab->GetVisualTabRole() != ETabRole::MajorTab;
+}
+
 void FGlobalTabmanager::SetActiveTab( const TSharedPtr<class SDockTab>& NewActiveTab )
 {
-	const bool bShouldApplyChange =
-		// Setting NULL wipes out the active tab; always apply that change.
-		!NewActiveTab.IsValid() ||
-		// Major tabs are ignored for the purposes of active-tab tracking. We do not care about their 
-		(NewActiveTab->GetTabRole() != ETabRole::MajorTab && !NewActiveTab->IsNomadTabWithMajorTabStyle());
+	const bool bShouldApplyChange = CanSetAsActiveTab(NewActiveTab);
 	
 	TSharedPtr<SDockTab> CurrentlyActiveTab = GetActiveTab();
 
diff --git a/Engine/Source/Runtime/Slate/Private/Widgets/Docking/SDockTab.cpp b/Engine/Source/Runtime/Slate/Private/Widgets/Docking/SDockTab.cpp
--- a/Engine/Source/Runtime/Slate/Private/Widgets/Docking/SDockTab.cpp
+++ b/Engine/Source/Runtime/Slate/Private/Widgets/Docking/SDockTab.cpp
@@ -3,7 +3,7 @@
 #include "SlatePrivatePCH.h"
 #include "SDockTab.h"
 #include "DockingPrivate.h"
-
+#include "TabCommands.h"
 
 namespace SDockTabDefs
 {
@@ -193,26 +193,34 @@ ETabRole SDockTab::GetTabRole() const
 	return TabRole;
 }
 
-bool SDockTab::IsNomadTabWithMajorTabStyle() const
+ETabRole SDockTab::GetVisualTabRole() const
 {
-	if ( this->TabRole == ETabRole::NomadTab )
+	// If the tab role is NomadTab but is being visualized as a major tab
+	if (this->TabRole == ETabRole::NomadTab)
 	{
-		if ( DraggedOverDockingArea.IsValid() )
+		bool bNomadMajorStyle = false;
+
+		if (DraggedOverDockingArea.IsValid())
 		{
-			return DraggedOverDockingArea->GetTabManager() == FGlobalTabmanager::Get();
+			bNomadMajorStyle = DraggedOverDockingArea->GetTabManager() == FGlobalTabmanager::Get();
 		}
-		else if ( GetParent().IsValid() && GetParent()->GetDockArea().IsValid() )
+		else if (GetParent().IsValid() && GetParent()->GetDockArea().IsValid())
 		{
-			return GetParent()->GetDockArea()->GetTabManager() == FGlobalTabmanager::Get();
+			bNomadMajorStyle = GetParent()->GetDockArea()->GetTabManager() == FGlobalTabmanager::Get();
 		}
 		else
 		{
 			// We are dragging or have no parent, but we are not dragging over anything, assume major
-			return true;
+			bNomadMajorStyle = true;
+		}
+
+		if (bNomadMajorStyle)
+		{
+			return ETabRole::MajorTab;
 		}
 	}
 
-	return false;
+	return GetTabRole();
 }
 
 const FSlateBrush* SDockTab::GetContentAreaBrush() const
@@ -522,6 +530,7 @@ void SDockTab::Construct( const FArguments& InArgs )
 					.ButtonStyle( CloseButtonStyle )
 					.OnClicked( this, &SDockTab::OnCloseButtonClicked )
 					.ContentPadding( 0 )
+					.ToolTipText(this, &SDockTab::GetCloseButtonToolTipText)
 					.Visibility(this, &SDockTab::HandleIsCloseButtonVisible)
 					[
 						SNew(SSpacer)
@@ -541,11 +550,7 @@ EActiveTimerReturnType SDockTab::TriggerActivateTab( double InCurrentTime, float
 
 const FDockTabStyle& SDockTab::GetCurrentStyle() const
 {
-	if ( this->TabRole == ETabRole::MajorTab )
-	{
-		return *MajorTabStyle;
-	}
-	else if ( IsNomadTabWithMajorTabStyle() )
+	if ( GetVisualTabRole() == ETabRole::MajorTab )
 	{
 		return *MajorTabStyle;
 	}
@@ -671,6 +676,27 @@ FReply SDockTab::OnCloseButtonClicked()
 	return FReply::Handled();
 }
 
+FText SDockTab::GetCloseButtonToolTipText() const
+{
+	TSharedPtr<FUICommandInfo> CloseCommand =
+		GetVisualTabRole() == ETabRole::MajorTab ? FTabCommands::Get().CloseMajorTab : FTabCommands::Get().CloseMinorTab;
+
+	FFormatNamedArguments Arguments;
+	Arguments.Add(TEXT("Label"), CloseCommand->GetLabel());
+
+	FText InputText = CloseCommand->GetInputText();
+	if (InputText.IsEmptyOrWhitespace())
+	{
+		Arguments.Add(TEXT("InputText"), InputText);
+	}
+	else
+	{
+		Arguments.Add(TEXT("InputText"), FText::Format(NSLOCTEXT("DockTab", "CloseButtonInputText", " ({0})"), InputText));
+	}
+
+	return FText::Format(NSLOCTEXT("DockTab", "CloseButtonToolTip", "{Label}{InputText}"), Arguments);
+}
+
 EVisibility SDockTab::HandleIsCloseButtonVisible() const
 {
 	return MyTabManager.Pin()->IsTabCloseable(SharedThis(this)) ? EVisibility::Visible : EVisibility::Hidden;
diff --git a/Engine/Source/Runtime/Slate/Public/Framework/Docking/TabCommands.h b/Engine/Source/Runtime/Slate/Public/Framework/Docking/TabCommands.h
--- a/Engine/Source/Runtime/Slate/Public/Framework/Docking/TabCommands.h
+++ b/Engine/Source/Runtime/Slate/Public/Framework/Docking/TabCommands.h
@@ -17,5 +17,6 @@ public:
 
 	virtual void RegisterCommands() override;
 
-	TSharedPtr<FUICommandInfo> CloseTab;
+	TSharedPtr<FUICommandInfo> CloseMajorTab;
+	TSharedPtr<FUICommandInfo> CloseMinorTab;
 };
diff --git a/Engine/Source/Runtime/Slate/Public/Framework/Docking/TabManager.h b/Engine/Source/Runtime/Slate/Public/Framework/Docking/TabManager.h
--- a/Engine/Source/Runtime/Slate/Public/Framework/Docking/TabManager.h
+++ b/Engine/Source/Runtime/Slate/Public/Framework/Docking/TabManager.h
@@ -777,6 +777,9 @@ public:
 	/** @return the currently active tab; NULL pointer if there is no active tab */
 	TSharedPtr<SDockTab> GetActiveTab() const;
 
+	/** Can the manager activate this Tab as the new active tab? */
+	bool CanSetAsActiveTab(const TSharedPtr<SDockTab>& Tab);
+
 	/** Activate the NewActiveTab. If NewActiveTab is NULL, the active tab is cleared. */
 	void SetActiveTab( const TSharedPtr<SDockTab>& NewActiveTab );
 
diff --git a/Engine/Source/Runtime/Slate/Public/Widgets/Docking/SDockTab.h b/Engine/Source/Runtime/Slate/Public/Widgets/Docking/SDockTab.h
--- a/Engine/Source/Runtime/Slate/Public/Widgets/Docking/SDockTab.h
+++ b/Engine/Source/Runtime/Slate/Public/Widgets/Docking/SDockTab.h
@@ -106,8 +106,12 @@ public:
 	/** Is this an MajorTab? A tool panel tab? */
 	ETabRole GetTabRole() const;
 
-	/** Returns true if the tab role is NomadTab, but is being visualized as a major tab. */
-	bool IsNomadTabWithMajorTabStyle() const;
+	/** Similar to GetTabRole() but returns the correct role for UI style and user input purposes */
+	ETabRole SDockTab::GetVisualTabRole() const;
+
+	/** Returns true if the tab role is NomadTab, but is being visualized as a major tab */
+	DEPRECATED(4.8, "Use GetVisualTabRole(), which will return ETabRole::MajorTab for nomads with major tab style.")
+	bool IsNomadTabWithMajorTabStyle() const { return GetTabRole() == ETabRole::NomadTab && GetVisualTabRole() == ETabRole::MajorTab; }
 
 	/**
 	 * What should the content area look like for this type of tab?
@@ -275,6 +279,9 @@ protected:
 	/** Called when the close button is clicked on the tab. */
 	FReply OnCloseButtonClicked();
 
+	/** The close button tooltip showing the appropriate close command shortcut */
+	FText GetCloseButtonToolTipText() const;
+
 	/** Specify the TabId that was used to spawn this tab. */
 	void SetLayoutIdentifier( const FTabId& TabId );
 
