Linux build fix

diff --git a/Engine/Source/Runtime/Slate/Public/Widgets/Docking/SDockTab.h b/Engine/Source/Runtime/Slate/Public/Widgets/Docking/SDockTab.h
--- a/Engine/Source/Runtime/Slate/Public/Widgets/Docking/SDockTab.h
+++ b/Engine/Source/Runtime/Slate/Public/Widgets/Docking/SDockTab.h
@@ -107,7 +107,7 @@ public:
 	ETabRole GetTabRole() const;
 
 	/** Similar to GetTabRole() but returns the correct role for UI style and user input purposes */
-	ETabRole SDockTab::GetVisualTabRole() const;
+	ETabRole GetVisualTabRole() const;
 
 	/** Returns true if the tab role is NomadTab, but is being visualized as a major tab */
 	DEPRECATED(4.8, "Use GetVisualTabRole(), which will return ETabRole::MajorTab for nomads with major tab style.")
