Changed Close Major tab command's default input chord to Ctrl+F4

diff --git a/Engine/Source/Runtime/Slate/Private/Framework/Docking/TabCommands.cpp b/Engine/Source/Runtime/Slate/Private/Framework/Docking/TabCommands.cpp
--- a/Engine/Source/Runtime/Slate/Private/Framework/Docking/TabCommands.cpp
+++ b/Engine/Source/Runtime/Slate/Private/Framework/Docking/TabCommands.cpp
@@ -7,7 +7,7 @@
 
 void FTabCommands::RegisterCommands()
 {
-	UI_COMMAND(CloseMajorTab, "Close Major Tab", "Closes the focused major tab", EUserInterfaceActionType::Button, FInputChord(EModifierKey::Control | EModifierKey::Alt, EKeys::F4))
+	UI_COMMAND(CloseMajorTab, "Close Major Tab", "Closes the focused major tab", EUserInterfaceActionType::Button, FInputChord(EModifierKey::Control, EKeys::F4))
 	UI_COMMAND(CloseMinorTab, "Close Minor Tab", "Closes the current window's active minor tab", EUserInterfaceActionType::Button, FInputChord())
 }
 
