Added support for absolutely sized slots in SSplitter

These are resizable, like the FractionOfParent slots, however unlike the FractionOfParent slots, they don't affect the size of their neighboring slots and are only able to resize themselves.

This means they behave like cells in a data grid (they resize themselves, and move everything else along).

diff --git a/Engine/Source/Runtime/Slate/Private/Widgets/Layout/SSplitter.cpp b/Engine/Source/Runtime/Slate/Private/Widgets/Layout/SSplitter.cpp
--- a/Engine/Source/Runtime/Slate/Private/Widgets/Layout/SSplitter.cpp
+++ b/Engine/Source/Runtime/Slate/Private/Widgets/Layout/SSplitter.cpp
@@ -88,8 +88,8 @@ TArray<FLayoutGeometry> SSplitter::ArrangeChildrenForLayout(const FGeometry& All
 	int32 NumNonCollapsedChildren = 0;
 	int32 NumResizeableChildren = 0;
 	float CoefficientTotal = 0;
-	// Some space is claimed by non-resizeable elements (auto-sized elements)
-	float NonResizeableSpace = 0;
+	// Some space is claimed by non-proportional elements (auto or fixed-sized elements)
+	float NonProportionalSpace = 0;
 	{
 		for (int32 ChildIndex=0; ChildIndex < Children.Num(); ++ChildIndex)	
 		{
@@ -99,7 +99,11 @@ TArray<FLayoutGeometry> SSplitter::ArrangeChildrenForLayout(const FGeometry& All
 
 				if ( Children[ChildIndex].SizingRule == SSplitter::SizeToContent )
 				{
-					NonResizeableSpace += Children[ChildIndex].GetWidget()->GetDesiredSize()[AxisIndex];
+					NonProportionalSpace += Children[ChildIndex].GetWidget()->GetDesiredSize()[AxisIndex];
+				}
+				else if ( Children[ChildIndex].SizingRule == SSplitter::AbsoluteSize )
+				{
+					NonProportionalSpace += Children[ChildIndex].SizeValue.Get();
 				}
 				else // SizingRule == SSplitter::FractionOfParent
 				{
@@ -113,7 +117,7 @@ TArray<FLayoutGeometry> SSplitter::ArrangeChildrenForLayout(const FGeometry& All
 
 	// The user-sizeable children must make room for the resize handles and for auto-sized children.
 	const float SpaceNeededForHandles = FMath::Max(0, NumNonCollapsedChildren-1) * PhysicalSplitterHandleSize;
-	const float ResizeableSpace = AllottedGeometry.Size.Component(AxisIndex) - SpaceNeededForHandles - NonResizeableSpace;
+	const float ProportionalSpace = AllottedGeometry.Size.Component(AxisIndex) - SpaceNeededForHandles - NonProportionalSpace;
 
 	// Arrange the children horizontally or vertically.
 	float XOffset = 0;
@@ -121,9 +125,19 @@ TArray<FLayoutGeometry> SSplitter::ArrangeChildrenForLayout(const FGeometry& All
 	{
 		const FSlot& CurSlot = Children[ChildIndex];
 
-		const float ChildSpace = ( CurSlot.SizingRule == SSplitter::SizeToContent )
-			? CurSlot.GetWidget()->GetDesiredSize()[AxisIndex]
-		: ResizeableSpace * CurSlot.SizeValue.Get() / CoefficientTotal;
+		float ChildSpace = 0.0f;
+		if ( CurSlot.SizingRule == SSplitter::SizeToContent )
+		{
+			ChildSpace = CurSlot.GetWidget()->GetDesiredSize()[AxisIndex];
+		}
+		else if ( CurSlot.SizingRule == SSplitter::AbsoluteSize )
+		{
+			ChildSpace = CurSlot.SizeValue.Get();
+		}
+		else // SizingRule == SSplitter::FractionOfParent
+		{
+			ChildSpace =  ProportionalSpace * CurSlot.SizeValue.Get() / CoefficientTotal;
+		}
 
 		const EVisibility ChildVisibility = CurSlot.GetWidget()->GetVisibility();
 
@@ -234,13 +248,13 @@ static FVector2D ComputeDesiredSizeForSplitter( const float PhysicalSplitterHand
 			FVector2D ChildDesiredSize( CurSlot.GetWidget()->GetDesiredSize() );
 			if ( Orientation == Orient_Horizontal )
 			{
-				MyDesiredSize.X += ChildDesiredSize.X;
+				MyDesiredSize.X += (CurSlot.SizingRule == SSplitter::AbsoluteSize) ? CurSlot.SizeValue.Get() : ChildDesiredSize.X;
 				MyDesiredSize.Y = FMath::Max(ChildDesiredSize.Y, MyDesiredSize.Y);
 			}
 			else
 			{
 				MyDesiredSize.X = FMath::Max(ChildDesiredSize.X, MyDesiredSize.X);				
-				MyDesiredSize.Y += ChildDesiredSize.Y;
+				MyDesiredSize.Y += (CurSlot.SizingRule == SSplitter::AbsoluteSize) ? CurSlot.SizeValue.Get() : ChildDesiredSize.Y;
 			}
 		}
 	}
@@ -502,114 +516,134 @@ void SSplitter::HandleResizing( const float PhysicalSplitterHandleSize, const ES
 	float Delta = LocalMousePos.Component(AxisIndex) - HandlePos;
 
 	const int32 SlotBeforeDragHandle = FindResizeableSlotBeforeHandle( DraggedHandle, Children );
-
-	TArray< int32 > SlotsAfterDragHandleIndicies;
-	if ( ResizeMode == ESplitterResizeMode::Fixed )
+	if ( SlotBeforeDragHandle == INDEX_NONE )
 	{
-		const int32 SlotAfterDragHandle = FindResizeableSlotAfterHandle( DraggedHandle, Children );
+		return;
+	}
+
+	FSlot& PrevChild = Children[SlotBeforeDragHandle];
+	const FLayoutGeometry& PrevChildGeom = ChildGeometries[SlotBeforeDragHandle];
+
+	// Compute the new sizes of the slot
+	const float PrevChildLength = PrevChildGeom.GetSizeInParentSpace().Component(AxisIndex);
+	float NewPrevChildLength = ClampChild( PrevChildLength + Delta );
 
-		if ( SlotAfterDragHandle < NumChildren )
+	// Changing the size of an absolutely sized slot doesn't affect the size of any subsequent slots, it just moves them
+	if ( PrevChild.SizingRule == SSplitter::AbsoluteSize )
+	{
+		if (PrevChild.OnSlotResized_Handler.IsBound())
 		{
-			SlotsAfterDragHandleIndicies.Add( SlotAfterDragHandle );
+			PrevChild.OnSlotResized_Handler.Execute( NewPrevChildLength );
+		}
+		else
+		{
+			PrevChild.SizeValue = NewPrevChildLength;
 		}
 	}
-	else if ( ResizeMode == ESplitterResizeMode::Fill )
-	{
-		FindAllResizeableSlotsAfterHandle( DraggedHandle, Children, /*OUT*/ SlotsAfterDragHandleIndicies );
-	}
-
-	if ( SlotBeforeDragHandle >= 0 && SlotsAfterDragHandleIndicies.Num() > 0 )
+	else
 	{
-		struct FSlotInfo 
+		// Proportionally sized slots need to redistribute their size accordingly
+		TArray< int32 > SlotsAfterDragHandleIndicies;
+		if ( ResizeMode == ESplitterResizeMode::Fixed )
 		{
-			FSlot* Slot;
-			const FLayoutGeometry* Geometry;
-			float NewSize;
-		};
+			const int32 SlotAfterDragHandle = FindResizeableSlotAfterHandle( DraggedHandle, Children );
 
-		TArray< FSlotInfo > SlotsAfterDragHandle;
-		for (int SlotIndex = 0; SlotIndex < SlotsAfterDragHandleIndicies.Num(); SlotIndex++)
+			if ( SlotAfterDragHandle < NumChildren )
+			{
+				SlotsAfterDragHandleIndicies.Add( SlotAfterDragHandle );
+			}
+		}
+		else if ( ResizeMode == ESplitterResizeMode::Fill )
 		{
-			FSlotInfo SlotInfo;
-
-			SlotInfo.Slot = &Children[ SlotsAfterDragHandleIndicies[ SlotIndex ] ];
-			SlotInfo.Geometry = &ChildGeometries[ SlotsAfterDragHandleIndicies[ SlotIndex ] ];
-			SlotInfo.NewSize = SlotInfo.Geometry->GetSizeInParentSpace().Component( AxisIndex );
-
-			SlotsAfterDragHandle.Add( SlotInfo );
+			FindAllResizeableSlotsAfterHandle( DraggedHandle, Children, /*OUT*/ SlotsAfterDragHandleIndicies );
 		}
 
-		// Get references the prev and next children and their layout settings so that we can modify them.
-		FSlot& PrevChild = Children[SlotBeforeDragHandle];
-		const FLayoutGeometry& PrevChildGeom = ChildGeometries[SlotBeforeDragHandle];
-
-		// Compute the new sizes of the children
-		const float PrevChildLength = PrevChildGeom.GetSizeInParentSpace().Component(AxisIndex);
-		float NewPrevChildLength = ClampChild( PrevChildLength + Delta );
-		Delta = NewPrevChildLength - PrevChildLength;
-
-		// Distribute the Delta across the affected slots after the drag handle
-		float UnusedDelta = Delta;
-		for (int DistributionCount = 0; DistributionCount < SlotsAfterDragHandle.Num() && UnusedDelta != 0; DistributionCount++)
+		if ( SlotsAfterDragHandleIndicies.Num() > 0)
 		{
-			float DividedDelta = UnusedDelta / SlotsAfterDragHandle.Num();
-			UnusedDelta = 0;
-			for (int SlotIndex = 0; SlotIndex < SlotsAfterDragHandle.Num(); SlotIndex++)
+			struct FSlotInfo 
 			{
-				FSlotInfo& SlotInfo = SlotsAfterDragHandle[ SlotIndex ];
+				FSlot* Slot;
+				const FLayoutGeometry* Geometry;
+				float NewSize;
+			};
 
-				float CurrentSize = ClampChild( SlotInfo.Geometry->GetSizeInParentSpace().Component(AxisIndex) );
-				SlotInfo.NewSize = ClampChild( CurrentSize - DividedDelta );
-
-				// If one of the slots couldn't be fully adjusted by the delta due to min/max constraints then
-				// the leftover delta needs to be evenly distributed to all of the other slots
-				UnusedDelta += SlotInfo.NewSize - ( CurrentSize - DividedDelta );
-			}
-		}
+			TArray< FSlotInfo > SlotsAfterDragHandle;
+			for (int SlotIndex = 0; SlotIndex < SlotsAfterDragHandleIndicies.Num(); SlotIndex++)
+			{
+				FSlotInfo SlotInfo;
 
-		Delta = Delta - UnusedDelta;
+				SlotInfo.Slot = &Children[ SlotsAfterDragHandleIndicies[ SlotIndex ] ];
+				SlotInfo.Geometry = &ChildGeometries[ SlotsAfterDragHandleIndicies[ SlotIndex ] ];
+				SlotInfo.NewSize = SlotInfo.Geometry->GetSizeInParentSpace().Component( AxisIndex );
 
-		// PrevChildLength needs to be updated: it's value has to take into account the next child's min/max restrictions
-		NewPrevChildLength = ClampChild( PrevChildLength + Delta );
+				SlotsAfterDragHandle.Add( SlotInfo );
+			}
 
-		// Cells being resized are both stretch values -> redistribute the stretch coefficients proportionately
-		// to match the new child sizes on the screen.
-		{
-			float TotalLength = NewPrevChildLength;
-			float TotalStretchCoefficients = PrevChild.SizeValue.Get();
+			Delta = NewPrevChildLength - PrevChildLength;
 
-			for (int SlotIndex = 0; SlotIndex < SlotsAfterDragHandle.Num(); SlotIndex++)
+			// Distribute the Delta across the affected slots after the drag handle
+			float UnusedDelta = Delta;
+			for (int DistributionCount = 0; DistributionCount < SlotsAfterDragHandle.Num() && UnusedDelta != 0; DistributionCount++)
 			{
-				FSlotInfo SlotInfo = SlotsAfterDragHandle[ SlotIndex ];
+				float DividedDelta = UnusedDelta / SlotsAfterDragHandle.Num();
+				UnusedDelta = 0;
+				for (int SlotIndex = 0; SlotIndex < SlotsAfterDragHandle.Num(); SlotIndex++)
+				{
+					FSlotInfo& SlotInfo = SlotsAfterDragHandle[ SlotIndex ];
+
+					float CurrentSize = ClampChild( SlotInfo.Geometry->GetSizeInParentSpace().Component(AxisIndex) );
+					SlotInfo.NewSize = ClampChild( CurrentSize - DividedDelta );
 
-				TotalLength += SlotInfo.NewSize;
-				TotalStretchCoefficients += SlotInfo.Slot->SizeValue.Get();
+					// If one of the slots couldn't be fully adjusted by the delta due to min/max constraints then
+					// the leftover delta needs to be evenly distributed to all of the other slots
+					UnusedDelta += SlotInfo.NewSize - ( CurrentSize - DividedDelta );
+				}
 			}
 
-			const float NewPrevChildSize = ( TotalStretchCoefficients * NewPrevChildLength / TotalLength );
+			Delta = Delta - UnusedDelta;
 
-			if (PrevChild.OnSlotResized_Handler.IsBound())
-			{
-				PrevChild.OnSlotResized_Handler.Execute( NewPrevChildSize );
-			}
-			else
-			{
-				PrevChild.SizeValue = NewPrevChildSize;
-			}
+			// PrevChildLength needs to be updated: it's value has to take into account the next child's min/max restrictions
+			NewPrevChildLength = ClampChild( PrevChildLength + Delta );
 
-			for (int SlotIndex = 0; SlotIndex < SlotsAfterDragHandle.Num(); SlotIndex++)
+			// Cells being resized are both stretch values -> redistribute the stretch coefficients proportionately
+			// to match the new child sizes on the screen.
 			{
-				FSlotInfo SlotInfo = SlotsAfterDragHandle[ SlotIndex ];
+				float TotalLength = NewPrevChildLength;
+				float TotalStretchCoefficients = PrevChild.SizeValue.Get();
+
+				for (int SlotIndex = 0; SlotIndex < SlotsAfterDragHandle.Num(); SlotIndex++)
+				{
+					FSlotInfo SlotInfo = SlotsAfterDragHandle[ SlotIndex ];
+
+					TotalLength += SlotInfo.NewSize;
+					TotalStretchCoefficients += SlotInfo.Slot->SizeValue.Get();
+				}
 
-				const float NewNextChildSize = ( TotalStretchCoefficients * SlotInfo.NewSize / TotalLength );
+				const float NewPrevChildSize = ( TotalStretchCoefficients * NewPrevChildLength / TotalLength );
 
-				if (SlotInfo.Slot->OnSlotResized_Handler.IsBound())
+				if (PrevChild.OnSlotResized_Handler.IsBound())
 				{
-					SlotInfo.Slot->OnSlotResized_Handler.Execute(NewNextChildSize);
+					PrevChild.OnSlotResized_Handler.Execute( NewPrevChildSize );
 				}
 				else
 				{
-					SlotInfo.Slot->SizeValue = NewNextChildSize;
+					PrevChild.SizeValue = NewPrevChildSize;
+				}
+
+				for (int SlotIndex = 0; SlotIndex < SlotsAfterDragHandle.Num(); SlotIndex++)
+				{
+					FSlotInfo SlotInfo = SlotsAfterDragHandle[ SlotIndex ];
+
+					const float NewNextChildSize = ( TotalStretchCoefficients * SlotInfo.NewSize / TotalLength );
+
+					if (SlotInfo.Slot->OnSlotResized_Handler.IsBound())
+					{
+						SlotInfo.Slot->OnSlotResized_Handler.Execute(NewNextChildSize);
+					}
+					else
+					{
+						SlotInfo.Slot->SizeValue = NewNextChildSize;
+					}
 				}
 			}
 		}
diff --git a/Engine/Source/Runtime/Slate/Public/Widgets/Layout/SSplitter.h b/Engine/Source/Runtime/Slate/Public/Widgets/Layout/SSplitter.h
--- a/Engine/Source/Runtime/Slate/Public/Widgets/Layout/SSplitter.h
+++ b/Engine/Source/Runtime/Slate/Public/Widgets/Layout/SSplitter.h
@@ -26,7 +26,9 @@ public:
 		/** Get the DesiredSize() of the content */
 		SizeToContent,
 		/** Use a fraction of the parent's size */
-		FractionOfParent
+		FractionOfParent,
+		/** Use the SizeValue of the slot as an absolute size */
+		AbsoluteSize,
 	};
 
 	DECLARE_DELEGATE_OneParam(
