Slightly modified version of Pull Request #976

*All parameters can be changed now to allow for different speed, acceleration, ... setups
*There is a second mode without acceleration as the analog sticks already allow different acceleration by user input
*Sub-pixel offsets are now preserved otherwise the cursor may stop movement in one or both directions if speed is very low

diff --git a/Engine/Source/Runtime/Slate/Private/Framework/Application/AnalogCursor.cpp b/Engine/Source/Runtime/Slate/Private/Framework/Application/AnalogCursor.cpp
--- a/Engine/Source/Runtime/Slate/Private/Framework/Application/AnalogCursor.cpp
+++ b/Engine/Source/Runtime/Slate/Private/Framework/Application/AnalogCursor.cpp
@@ -5,80 +5,110 @@
 
 #include "HittestGrid.h"
 
-const float FAnalogCursor::DefaultAcceleration = 150.0f;
-const float FAnalogCursor::DefaultMaxSpeed = 250.0f;
-
 FAnalogCursor::FAnalogCursor()
 : CurrentSpeed(FVector2D::ZeroVector)
+, CurrentOffset(FVector2D::ZeroVector)
+, Acceleration(1000.0f)
+, MaxSpeed(1500.0f)
+, StickySlowdown(0.5f)
+, DeadZone(0.1f)
+, Mode(AnalogCursorMode::Accelerated)
 , AnalogValues(FVector2D::ZeroVector)
 {
 }
 
 void FAnalogCursor::Tick(const float DeltaTime, FSlateApplication& SlateApp, TSharedRef<ICursor> Cursor)
 {
-	FVector2D OldPos = Cursor->GetPosition();
+	const FVector2D OldPosition = Cursor->GetPosition();
 
 	float SpeedMult = 1.0f; // Used to do a speed multiplication before adding the delta to the position to make widgets sticky
 	FVector2D AdjAnalogVals = AnalogValues; // A copy of the analog values so I can modify them based being over a widget, not currently doing this
 
-	//if the size of the analog vals vector is too small, then set it to zero
-	static const float DeadZoneSizeSq = (0.1f * 0.1f);
-	if (AdjAnalogVals.SizeSquared() < DeadZoneSizeSq)
+	// Adjust analog values according to dead zone
+	const float AnalogValsSize = AdjAnalogVals.Size();
+
+	if (AnalogValsSize > 0.0f)
 	{
-		AdjAnalogVals = FVector2D::ZeroVector;
+		const float TargetSize = FMath::Max(AnalogValsSize - DeadZone, 0.0f) / (1.0f - DeadZone);
+		AdjAnalogVals /= AnalogValsSize;
+		AdjAnalogVals *= TargetSize;
 	}
 
-	FWidgetPath WidgetPath = SlateApp.LocateWindowUnderMouse(OldPos, SlateApp.GetInteractiveTopLevelWindows());
+
+	// Check if there is a sticky widget beneath the cursor
+	FWidgetPath WidgetPath = SlateApp.LocateWindowUnderMouse(OldPosition, SlateApp.GetInteractiveTopLevelWindows());
 	if (WidgetPath.IsValid())
 	{
-		FArrangedWidget& ArrangedWidget = WidgetPath.Widgets.Last();
-		TSharedRef<SWidget> Widget = ArrangedWidget.Widget;
-		if (Widget->SupportsKeyboardFocus() && Widget->GetType() != FName("SViewport"))
+		const FArrangedChildren::FArrangedWidgetArray& AllArrangedWidgets = WidgetPath.Widgets.GetInternalArray();
+		for(const FArrangedWidget& ArrangedWidget : AllArrangedWidgets)
 		{
-			SpeedMult = 0.5f;
-			//FVector2D Adjustment = WidgetsAndCursors.Last().Geometry.Position - OldPos; // example of calculating distance from cursor to widget center
+			TSharedRef<SWidget> Widget = ArrangedWidget.Widget;
+			if (Widget->SupportsKeyboardFocus() && Widget->IsInteractable())
+			{
+				SpeedMult = StickySlowdown;
+				//FVector2D Adjustment = WidgetsAndCursors.Last().Geometry.Position - OldPosition; // example of calculating distance from cursor to widget center
+				break;
+			}
 		}
 	}
 
-	// Generate Min and Max for X to clamp the speed, this gives us instant direction change when crossing the axis
-	float CurrentMinSpeedX = 0.0f;
-	float CurrentMaxSpeedX = 0.0f;
-	if (AdjAnalogVals.X > 0.0f)
-	{ 
-		CurrentMaxSpeedX = AdjAnalogVals.X * DefaultMaxSpeed;
-	}
-	else
-	{
-		CurrentMinSpeedX = AdjAnalogVals.X * DefaultMaxSpeed;
-	}
-
-	// Generate Min and Max for Y to clamp the speed, this gives us instant direction change when crossing the axis
-	float CurrentMinSpeedY = 0.0f;
-	float CurrentMaxSpeedY = 0.0f;
-	if (AdjAnalogVals.Y > 0.0f)
+	switch(Mode)
 	{
-		CurrentMaxSpeedY = AdjAnalogVals.Y * DefaultMaxSpeed;
-	}
-	else
-	{
-		CurrentMinSpeedY = AdjAnalogVals.Y * DefaultMaxSpeed;
-	}
+		case AnalogCursorMode::Accelerated:
+		{
+			// Generate Min and Max for X to clamp the speed, this gives us instant direction change when crossing the axis
+			float CurrentMinSpeedX = 0.0f;
+			float CurrentMaxSpeedX = 0.0f;
+			if (AdjAnalogVals.X > 0.0f)
+			{ 
+				CurrentMaxSpeedX = AdjAnalogVals.X * MaxSpeed;
+			}
+			else
+			{
+				CurrentMinSpeedX = AdjAnalogVals.X * MaxSpeed;
+			}
+
+			// Generate Min and Max for Y to clamp the speed, this gives us instant direction change when crossing the axis
+			float CurrentMinSpeedY = 0.0f;
+			float CurrentMaxSpeedY = 0.0f;
+			if (AdjAnalogVals.Y > 0.0f)
+			{
+				CurrentMaxSpeedY = AdjAnalogVals.Y * MaxSpeed;
+			}
+			else
+			{
+				CurrentMinSpeedY = AdjAnalogVals.Y * MaxSpeed;
+			}
+
+			// Cubic acceleration curve
+			FVector2D ExpAcceleration = AdjAnalogVals * AdjAnalogVals * AdjAnalogVals * Acceleration;
+			// Preserve direction (if we use a squared equation above)
+			//ExpAcceleration.X *= FMath::Sign(AnalogValues.X);
+			//ExpAcceleration.Y *= FMath::Sign(AnalogValues.Y);
+
+			CurrentSpeed += ExpAcceleration * DeltaTime;
+
+			CurrentSpeed.X = FMath::Clamp(CurrentSpeed.X, CurrentMinSpeedX, CurrentMaxSpeedX);
+			CurrentSpeed.Y = FMath::Clamp(CurrentSpeed.Y, CurrentMinSpeedY, CurrentMaxSpeedY);
+
+			break;
+		}
 
-	// Cubic acceleration curve
-	FVector2D ExpAcceleration = AdjAnalogVals * AdjAnalogVals * AdjAnalogVals * DefaultAcceleration;
-	// Preserve direction (if we use a squared equation above)
-	//ExpAcceleration.X *= FMath::Sign(AnalogValues.X);
-	//ExpAcceleration.Y *= FMath::Sign(AnalogValues.Y);
+		case AnalogCursorMode::Direct:
 
-	CurrentSpeed += ExpAcceleration * DeltaTime;
+			CurrentSpeed = AdjAnalogVals * MaxSpeed;
 
-	CurrentSpeed.X = FMath::Clamp(CurrentSpeed.X, CurrentMinSpeedX, CurrentMaxSpeedX);
-	CurrentSpeed.Y = FMath::Clamp(CurrentSpeed.Y, CurrentMinSpeedY, CurrentMaxSpeedY);
+			break;
+	}
 
-	const FVector2D NewPosition = OldPos + (CurrentSpeed * SpeedMult);
+	CurrentOffset += CurrentSpeed * DeltaTime * SpeedMult;
+	const FVector2D NewPosition = OldPosition + CurrentOffset;
 
-	//update the cursor position
+	// update the cursor position
 	UpdateCursorPosition(SlateApp, Cursor, NewPosition);
+
+	// save the remaining sub-pixel offset after setting the cursor
+	CurrentOffset = NewPosition - Cursor->GetPosition();
 }
 
 bool FAnalogCursor::HandleKeyDownEvent(FSlateApplication& SlateApp, const FKeyEvent& InKeyEvent)
@@ -166,6 +196,33 @@ bool FAnalogCursor::HandleAnalogInputEvent(FSlateApplication& SlateApp, const FA
 	return true;
 }
 
+void FAnalogCursor::SetAcceleration(float NewAcceleration)
+{
+	Acceleration = NewAcceleration;
+}
+
+void FAnalogCursor::SetMaxSpeed(float NewMaxSpeed)
+{
+	MaxSpeed = NewMaxSpeed;
+}
+
+void FAnalogCursor::SetStickySlowdown(float NewStickySlowdown)
+{
+	StickySlowdown = NewStickySlowdown;
+}
+
+void FAnalogCursor::SetDeadZone(float NewDeadZone)
+{
+	DeadZone = NewDeadZone;
+}
+
+void FAnalogCursor::SetMode(AnalogCursorMode::Type NewMode)
+{
+	Mode = NewMode;
+
+	CurrentSpeed = FVector2D::ZeroVector;
+}
+
 void FAnalogCursor::UpdateCursorPosition(FSlateApplication& SlateApp, TSharedRef<ICursor> Cursor, const FVector2D& NewPosition)
 {
 	//grab the old position
@@ -193,3 +250,4 @@ void FAnalogCursor::UpdateCursorPosition(FSlateApplication& SlateApp, TSharedRef
 		SlateApp.ProcessMouseMoveEvent(MouseEvent);
 	}
 }
+
diff --git a/Engine/Source/Runtime/Slate/Public/Framework/Application/AnalogCursor.h b/Engine/Source/Runtime/Slate/Public/Framework/Application/AnalogCursor.h
--- a/Engine/Source/Runtime/Slate/Public/Framework/Application/AnalogCursor.h
+++ b/Engine/Source/Runtime/Slate/Public/Framework/Application/AnalogCursor.h
@@ -4,6 +4,15 @@
 
 #include "IInputProcessor.h"
 
+namespace AnalogCursorMode
+{
+	enum Type
+	{
+		Accelerated,
+		Direct,
+	};
+}
+
 /**
  * A class that simulates a cursor driven by an analog stick.
  */
@@ -22,6 +31,12 @@ public:
 	virtual bool HandleKeyUpEvent(FSlateApplication& SlateApp, const FKeyEvent& InKeyEvent) override;
 	virtual bool HandleAnalogInputEvent(FSlateApplication& SlateApp, const FAnalogInputEvent& InAnalogInputEvent) override;
 
+	void SetAcceleration(float NewAcceleration);
+	void SetMaxSpeed(float NewMaxSpeed);
+	void SetStickySlowdown(float NewStickySlowdown);
+	void SetDeadZone(float NewDeadZone);
+	void SetMode(AnalogCursorMode::Type NewMode);
+
 protected:
 
 	/** Getter */
@@ -36,9 +51,15 @@ protected:
 	/** Current speed of the cursor */
 	FVector2D CurrentSpeed;
 
-	/** Helpful statics */
-	static const float DefaultAcceleration;
-	static const float DefaultMaxSpeed;
+	/** Current sub-pixel offset */
+	FVector2D CurrentOffset;
+
+	/** Current settings */
+	float Acceleration;
+	float MaxSpeed;
+	float StickySlowdown;
+	float DeadZone;
+	AnalogCursorMode::Type Mode;
 
 private:
 
