Added some test cases for the ManualSize splitter slot type

Renamed the slot type from AbsoluteSize to ManualSize to avoid overloading the term "absolute".

Fixed some resizing issues when a ManualSize slot was used alongside the other slot types.

Improved the validation used when deciding whether a drag handle cursor should be shown, as this wasn't working correctly for ManualSize slots.

diff --git a/Engine/Source/Editor/DataTableEditor/Private/DataTableEditor.cpp b/Engine/Source/Editor/DataTableEditor/Private/DataTableEditor.cpp
--- a/Engine/Source/Editor/DataTableEditor/Private/DataTableEditor.cpp
+++ b/Engine/Source/Editor/DataTableEditor/Private/DataTableEditor.cpp
@@ -313,7 +313,7 @@ TSharedRef<SWidget> FDataTableEditor::CreateGridPanel()
 					RowSplitter->AddSlot()
 					.Value(TAttribute<float>::Create(TAttribute<float>::FGetter::CreateSP(this, &FDataTableEditor::GetColumnWidth, ColumnIndex)))
 					.OnSlotResized(SSplitter::FOnSlotResized::CreateSP(this, &FDataTableEditor::OnColumnResized, ColumnIndex))
-					.SizeRule(SSplitter::AbsoluteSize)
+					.SizeRule(SSplitter::ManualSize)
 					[
 						SAssignNew(ColumnEntryWidget, SBorder)
 						.Padding(2.0f)
@@ -342,7 +342,7 @@ TSharedRef<SWidget> FDataTableEditor::CreateGridPanel()
 				RowSplitter->AddSlot()
 				.Value(TAttribute<float>::Create(TAttribute<float>::FGetter::CreateSP(this, &FDataTableEditor::GetColumnWidth, static_cast<int32>(INDEX_NONE))))
 				.OnSlotResized(SSplitter::FOnSlotResized::CreateSP(this, &FDataTableEditor::OnColumnResized, static_cast<int32>(INDEX_NONE)))
-				.SizeRule(SSplitter::AbsoluteSize)
+				.SizeRule(SSplitter::ManualSize)
 				[
 					SNullWidget::NullWidget
 				];
diff --git a/Engine/Source/Runtime/AppFramework/Private/Widgets/Testing/STestSuite.cpp b/Engine/Source/Runtime/AppFramework/Private/Widgets/Testing/STestSuite.cpp
--- a/Engine/Source/Runtime/AppFramework/Private/Widgets/Testing/STestSuite.cpp
+++ b/Engine/Source/Runtime/AppFramework/Private/Widgets/Testing/STestSuite.cpp
@@ -933,14 +933,27 @@ public:
 						.Text( LOCTEXT("Col0Row2Visibility", "Col0Row2 Visibility") )
 						.OnClicked( Col0Row2Vis, &FVisibilityCycler::CycleVisibility )
 					]
-
+					+SVerticalBox::Slot()
+					.AutoHeight()
+					[
+						SNew(SButton)
+						.Text( LOCTEXT("Col0Row3Visibility", "Col0Row3 Visibility") )
+						.OnClicked( Col0Row3Vis, &FVisibilityCycler::CycleVisibility )
+					]
 				]
 				+SHorizontalBox::Slot()
 				.AutoWidth()
 				[
 					SNew(SButton)
-					.Text( LOCTEXT("CenterVisibility", "CenterVis Visibility") )
-					.OnClicked( CenterVis, &FVisibilityCycler::CycleVisibility )
+					.Text( LOCTEXT("Center0Visibility", "Center0 Visibility") )
+					.OnClicked( Center0Vis, &FVisibilityCycler::CycleVisibility )
+				]
+				+SHorizontalBox::Slot()
+				.AutoWidth()
+				[
+					SNew(SButton)
+					.Text( LOCTEXT("Center1Visibility", "Center1 Visibility") )
+					.OnClicked( Center1Vis, &FVisibilityCycler::CycleVisibility )
 				]
 				+SHorizontalBox::Slot()
 				.AutoWidth()
@@ -959,7 +972,21 @@ public:
 						SNew(SButton)
 						.Text( LOCTEXT("Col2Row1Visibility", "Col2Row1 Visibility") )
 						.OnClicked( Col2Row1Vis, &FVisibilityCycler::CycleVisibility )
-					]					
+					]
+					+SVerticalBox::Slot()
+					.AutoHeight()
+					[
+						SNew(SButton)
+						.Text( LOCTEXT("Col2Row2Visibility", "Col2Row2 Visibility") )
+						.OnClicked( Col2Row2Vis, &FVisibilityCycler::CycleVisibility )
+					]
+					+SVerticalBox::Slot()
+					.AutoHeight()
+					[
+						SNew(SButton)
+						.Text( LOCTEXT("Col2Row3Visibility", "Col2Row3 Visibility") )
+						.OnClicked( Col2Row3Vis, &FVisibilityCycler::CycleVisibility )
+					]
 				]
 			]
 			+ SVerticalBox::Slot()
@@ -971,15 +998,15 @@ public:
 				+ SSplitter::Slot()
 				[
 					SAssignNew(Nested0, SSplitter)
-						.Orientation(Orient_Vertical)
-						.ResizeMode( ESplitterResizeMode::Fill )
-						+SSplitter::Slot()
+					.Orientation(Orient_Vertical)
+					.ResizeMode( ESplitterResizeMode::Fill )
+					+SSplitter::Slot()
 					[
 						SNew(SBorder)
 							.Visibility(Col0Row0Vis, &FVisibilityCycler::GetVisibility)
 						[
 							SNew(STextBlock)
-								.Text( LOCTEXT("Col0Row0", "Col 0 Row 0") )
+								.Text( LOCTEXT("Col0Row0", "Col 0 Row 0 (Fraction Of Parent)") )
 						]
 					]
 					+SSplitter::Slot()
@@ -987,49 +1014,89 @@ public:
 						SNew(SBorder)
 							.Visibility(Col0Row1Vis, &FVisibilityCycler::GetVisibility)
 						[
-							SNew(STextBlock) .Text( LOCTEXT("Col0Row1", "Col 0 Row 1") )
+							SNew(STextBlock) .Text( LOCTEXT("Col0Row1", "Col 0 Row 1 (Fraction Of Parent)") )
 						]
 					]
 					+SSplitter::Slot()
+					.Value(100.0f)
+					.SizeRule( SSplitter::ManualSize)
 					[
 						SNew(SBorder)
 							.Visibility(Col0Row2Vis, &FVisibilityCycler::GetVisibility)
 						[
-							SNew(STextBlock) .Text( LOCTEXT("Col0Row2", "Col 0 Row 2") )
+							SNew(STextBlock) .Text( LOCTEXT("Col0Row2", "Col 0 Row 2 (Manual Size)") )
 						]
 					]
+					+SSplitter::Slot()
+					[
+						SNew(SBorder)
+							.Visibility(Col0Row3Vis, &FVisibilityCycler::GetVisibility)
+						[
+							SNew(STextBlock) .Text( LOCTEXT("Col0Row3", "Col 0 Row 3 (Fraction Of Parent)") )
+						]
+					]
+				]
+				+ SSplitter::Slot()
+				.Value(50.0f)
+				.SizeRule( SSplitter::ManualSize )
+				[
+					SNew(SBorder)
+						.Visibility(Center0Vis, &FVisibilityCycler::GetVisibility)
+					[
+						SNew(STextBlock) .Text( LOCTEXT("Center0", "Center 0 (Manual Size)") )
+					]
 				]
 				+ SSplitter::Slot()
 				. SizeRule( SSplitter::SizeToContent )
 				[
 					SNew(SBorder)
-						.Visibility(CenterVis, &FVisibilityCycler::GetVisibility)
+						.Visibility(Center1Vis, &FVisibilityCycler::GetVisibility)
 						.Padding( 5.0f )
 					[
 						SNew(SButton)
 							.OnClicked( this, &SSplitterTest::FlipTopLevelSplitter )
-							.Text( LOCTEXT("Re-orient", "Re-orient") )
+							.Text( LOCTEXT("Center1", "Center 1 (Size To Content)\nRe-orient Splitter") )
 					]
 				]
 				+ SSplitter::Slot()
 				[
 					SAssignNew(Nested1, SSplitter)
-						.Orientation(Orient_Vertical)
-						.ResizeMode( ESplitterResizeMode::Fill )
-						+SSplitter::Slot()
+					.Orientation(Orient_Vertical)
+					.ResizeMode( ESplitterResizeMode::Fill )
+					+SSplitter::Slot()
 					[
 						SNew(SBorder)
 							.Visibility(Col2Row0Vis, &FVisibilityCycler::GetVisibility)
 						[
-							SNew(STextBlock) .Text( LOCTEXT("Col2Row0", "Col 2 Row 0") )
+							SNew(STextBlock) .Text( LOCTEXT("Col2Row0", "Col 2 Row 0 (Fraction Of Parent)") )
 						]
 					]
 					+SSplitter::Slot()
+					.Value(100.0f)
+					.SizeRule( SSplitter::ManualSize)
 					[
 						SNew(SBorder)
 							.Visibility(Col2Row1Vis, &FVisibilityCycler::GetVisibility)
 						[
-							SNew(STextBlock) .Text( LOCTEXT("Col2Row1", "Col 1 Row 1") )
+							SNew(STextBlock) .Text( LOCTEXT("Col2Row1", "Col 2 Row 1 (Manual Size)") )
+						]
+					]
+					+SSplitter::Slot()
+					.Value(100.0f)
+					.SizeRule( SSplitter::ManualSize)
+					[
+						SNew(SBorder)
+							.Visibility(Col2Row2Vis, &FVisibilityCycler::GetVisibility)
+						[
+							SNew(STextBlock) .Text( LOCTEXT("Col2Row2", "Col 2 Row 2 (Manual Size)") )
+						]
+					]
+					+SSplitter::Slot()
+					[
+						SNew(SBorder)
+							.Visibility(Col2Row3Vis, &FVisibilityCycler::GetVisibility)
+						[
+							SNew(STextBlock) .Text( LOCTEXT("Col2Row3", "Col 2 Row 3 (Fraction Of Parent)") )
 						]
 					]
 				]
@@ -1041,9 +1108,13 @@ public:
 		: Col0Row0Vis( new FVisibilityCycler() )
 		, Col0Row1Vis( new FVisibilityCycler() )
 		, Col0Row2Vis( new FVisibilityCycler() )
-		, CenterVis( new FVisibilityCycler() )
+		, Col0Row3Vis( new FVisibilityCycler() )
+		, Center0Vis( new FVisibilityCycler() )
+		, Center1Vis( new FVisibilityCycler() )
 		, Col2Row0Vis( new FVisibilityCycler() )
 		, Col2Row1Vis( new FVisibilityCycler() )
+		, Col2Row2Vis( new FVisibilityCycler() )
+		, Col2Row3Vis( new FVisibilityCycler() )
 	{
 	}
 
@@ -1060,12 +1131,15 @@ protected:
 	TSharedRef<FVisibilityCycler> Col0Row0Vis;
 	TSharedRef<FVisibilityCycler> Col0Row1Vis;
 	TSharedRef<FVisibilityCycler> Col0Row2Vis;
+	TSharedRef<FVisibilityCycler> Col0Row3Vis;
 
-	TSharedRef<FVisibilityCycler> CenterVis;
+	TSharedRef<FVisibilityCycler> Center0Vis;
+	TSharedRef<FVisibilityCycler> Center1Vis;
 
 	TSharedRef<FVisibilityCycler> Col2Row0Vis;
 	TSharedRef<FVisibilityCycler> Col2Row1Vis;
-		
+	TSharedRef<FVisibilityCycler> Col2Row2Vis;
+	TSharedRef<FVisibilityCycler> Col2Row3Vis;
 
 	TSharedPtr<SSplitter> TopLeveSplitter;
 	TSharedPtr<SSplitter> Nested0;
diff --git a/Engine/Source/Runtime/Slate/Private/Widgets/Layout/SSplitter.cpp b/Engine/Source/Runtime/Slate/Private/Widgets/Layout/SSplitter.cpp
--- a/Engine/Source/Runtime/Slate/Private/Widgets/Layout/SSplitter.cpp
+++ b/Engine/Source/Runtime/Slate/Private/Widgets/Layout/SSplitter.cpp
@@ -97,11 +97,11 @@ TArray<FLayoutGeometry> SSplitter::ArrangeChildrenForLayout(const FGeometry& All
 			{
 				++NumNonCollapsedChildren;
 
-				if ( Children[ChildIndex].SizingRule == SSplitter::SizeToContent )
+				if ( Children[ChildIndex].SizingRule.Get() == SSplitter::SizeToContent )
 				{
 					NonProportionalSpace += Children[ChildIndex].GetWidget()->GetDesiredSize()[AxisIndex];
 				}
-				else if ( Children[ChildIndex].SizingRule == SSplitter::AbsoluteSize )
+				else if ( Children[ChildIndex].SizingRule.Get() == SSplitter::ManualSize )
 				{
 					NonProportionalSpace += Children[ChildIndex].SizeValue.Get();
 				}
@@ -126,11 +126,11 @@ TArray<FLayoutGeometry> SSplitter::ArrangeChildrenForLayout(const FGeometry& All
 		const FSlot& CurSlot = Children[ChildIndex];
 
 		float ChildSpace = 0.0f;
-		if ( CurSlot.SizingRule == SSplitter::SizeToContent )
+		if ( CurSlot.SizingRule.Get() == SSplitter::SizeToContent )
 		{
 			ChildSpace = CurSlot.GetWidget()->GetDesiredSize()[AxisIndex];
 		}
-		else if ( CurSlot.SizingRule == SSplitter::AbsoluteSize )
+		else if ( CurSlot.SizingRule.Get() == SSplitter::ManualSize )
 		{
 			ChildSpace = CurSlot.SizeValue.Get();
 		}
@@ -248,13 +248,13 @@ static FVector2D ComputeDesiredSizeForSplitter( const float PhysicalSplitterHand
 			FVector2D ChildDesiredSize( CurSlot.GetWidget()->GetDesiredSize() );
 			if ( Orientation == Orient_Horizontal )
 			{
-				MyDesiredSize.X += (CurSlot.SizingRule == SSplitter::AbsoluteSize) ? CurSlot.SizeValue.Get() : ChildDesiredSize.X;
+				MyDesiredSize.X += (CurSlot.SizingRule.Get() == SSplitter::ManualSize) ? CurSlot.SizeValue.Get() : ChildDesiredSize.X;
 				MyDesiredSize.Y = FMath::Max(ChildDesiredSize.Y, MyDesiredSize.Y);
 			}
 			else
 			{
 				MyDesiredSize.X = FMath::Max(ChildDesiredSize.X, MyDesiredSize.X);				
-				MyDesiredSize.Y += (CurSlot.SizingRule == SSplitter::AbsoluteSize) ? CurSlot.SizeValue.Get() : ChildDesiredSize.Y;
+				MyDesiredSize.Y += (CurSlot.SizingRule.Get() == SSplitter::ManualSize) ? CurSlot.SizeValue.Get() : ChildDesiredSize.Y;
 			}
 		}
 	}
@@ -369,19 +369,10 @@ FReply SSplitter::OnMouseMove( const FGeometry& MyGeometry, const FPointerEvent&
 		return FReply::Handled();
 	}
 	else
-	{	
-		// Hit test which handle we are hovering over.		
+	{
 		HoveredHandleIndex = (Orientation == Orient_Horizontal)
-			? GetHandleBeingResizedFromMousePosition<Orient_Horizontal>( PhysicalSplitterHandleSize, HitDetectionSplitterHandleSize, LocalMousePosition, LayoutChildren )
-			: GetHandleBeingResizedFromMousePosition<Orient_Vertical>( PhysicalSplitterHandleSize, HitDetectionSplitterHandleSize, LocalMousePosition, LayoutChildren );
-
-		if (HoveredHandleIndex != INDEX_NONE)
-		{
-			if ( FindResizeableSlotBeforeHandle(HoveredHandleIndex, Children) <= INDEX_NONE || FindResizeableSlotAfterHandle(HoveredHandleIndex, Children) >= Children.Num() )
-			{
-				HoveredHandleIndex = INDEX_NONE;
-			}
-		}
+			? GetAndValidateHandleBeingResizedFromMousePosition<Orient_Horizontal>( LocalMousePosition, LayoutChildren )
+			: GetAndValidateHandleBeingResizedFromMousePosition<Orient_Vertical>( LocalMousePosition, LayoutChildren );
 
 		return FReply::Unhandled();
 	}
@@ -409,8 +400,8 @@ FCursorReply SSplitter::OnCursorQuery( const FGeometry& MyGeometry, const FPoint
 
 	// Hit test which handle we are hovering over.		
 	const int32 CurrentHoveredHandleIndex = (Orientation == Orient_Horizontal)
-		? GetHandleBeingResizedFromMousePosition<Orient_Horizontal>( PhysicalSplitterHandleSize, HitDetectionSplitterHandleSize, LocalMousePosition, LayoutChildren )
-		: GetHandleBeingResizedFromMousePosition<Orient_Vertical>( PhysicalSplitterHandleSize, HitDetectionSplitterHandleSize, LocalMousePosition, LayoutChildren );
+		? GetAndValidateHandleBeingResizedFromMousePosition<Orient_Horizontal>( LocalMousePosition, LayoutChildren )
+		: GetAndValidateHandleBeingResizedFromMousePosition<Orient_Vertical>( LocalMousePosition, LayoutChildren );
 
 	if (CurrentHoveredHandleIndex != INDEX_NONE)
 	{
@@ -467,7 +458,25 @@ int32 SSplitter::FindResizeableSlotBeforeHandle( int32 DraggedHandle, const TPan
 }
 
 
-int32 SSplitter::FindResizeableSlotAfterHandle( int32 DraggedHandle, const TPanelChildren<FSlot>& Children )
+int32 SSplitter::FindNonCollapsedSlotAfterHandle( int32 DraggedHandle, const TPanelChildren<FSlot>& Children )
+{
+	const int32 NumChildren = Children.Num();
+
+	// But the slots list does contain collapsed children! Make sure that we are not resizing a collapsed slot.
+	// We also cannot resize auto-sizing slots.
+	int32 SlotAfterDragHandle = DraggedHandle+1;
+	{
+		while( SlotAfterDragHandle < NumChildren && Children[SlotAfterDragHandle].GetWidget()->GetVisibility() == EVisibility::Collapsed )
+		{
+			++SlotAfterDragHandle;
+		}
+	}
+
+	return SlotAfterDragHandle;
+}
+
+
+int32 SSplitter::FindProportionallyResizeableSlotAfterHandle( int32 DraggedHandle, const TPanelChildren<FSlot>& Children )
 {
 	const int32 NumChildren = Children.Num();
 
@@ -475,7 +484,7 @@ int32 SSplitter::FindResizeableSlotAfterHandle( int32 DraggedHandle, const TPane
 	// We also cannot resize auto-sizing slots.
 	int32 SlotAfterDragHandle = DraggedHandle+1;
 	{
-		while( SlotAfterDragHandle < NumChildren && (Children[SlotAfterDragHandle].GetWidget()->GetVisibility() == EVisibility::Collapsed || Children[SlotAfterDragHandle].SizingRule.Get() == SSplitter::SizeToContent) )
+		while( SlotAfterDragHandle < NumChildren && (Children[SlotAfterDragHandle].GetWidget()->GetVisibility() == EVisibility::Collapsed || Children[SlotAfterDragHandle].SizingRule.Get() != SSplitter::FractionOfParent) )
 		{
 			++SlotAfterDragHandle;
 		}
@@ -485,20 +494,20 @@ int32 SSplitter::FindResizeableSlotAfterHandle( int32 DraggedHandle, const TPane
 }
 
 
-void SSplitter::FindAllResizeableSlotsAfterHandle( int32 DraggedHandle, const TPanelChildren<FSlot>& Children, TArray< int32 >& OutSlotIndicies )
+void SSplitter::FindAllProportionallyResizeableSlotsAfterHandle( int32 DraggedHandle, const TPanelChildren<FSlot>& Children, TArray< int32 >& OutSlotIndicies )
 {
 	const int32 NumChildren = Children.Num();
 
 	for (int SlotIndex = DraggedHandle+1; SlotIndex < NumChildren; SlotIndex++)
 	{
-		if (Children[ SlotIndex ].GetWidget()->GetVisibility() == EVisibility::Collapsed || Children[ SlotIndex ].SizingRule.Get() == SSplitter::SizeToContent)
+		if (Children[ SlotIndex ].GetWidget()->GetVisibility() == EVisibility::Collapsed || Children[ SlotIndex ].SizingRule.Get() != SSplitter::FractionOfParent)
 		{
 			continue;
 		}
 
 		OutSlotIndicies.Add( SlotIndex );
 	}
-};
+}
 
 
 template<EOrientation SplitterOrientation>
@@ -528,8 +537,9 @@ void SSplitter::HandleResizing( const float PhysicalSplitterHandleSize, const ES
 	const float PrevChildLength = PrevChildGeom.GetSizeInParentSpace().Component(AxisIndex);
 	float NewPrevChildLength = ClampChild( PrevChildLength + Delta );
 
-	// Changing the size of an absolutely sized slot doesn't affect the size of any subsequent slots, it just moves them
-	if ( PrevChild.SizingRule == SSplitter::AbsoluteSize )
+	// Changing the size of an absolutely sized slot doesn't adjust the proportional size of any subsequent slots
+	// Our change in size will be handled by the layout logic and any proportional slots will update accordingly
+	if ( PrevChild.SizingRule.Get() == SSplitter::ManualSize )
 	{
 		if (PrevChild.OnSlotResized_Handler.IsBound())
 		{
@@ -543,22 +553,22 @@ void SSplitter::HandleResizing( const float PhysicalSplitterHandleSize, const ES
 	else
 	{
 		// Proportionally sized slots need to redistribute their size accordingly
-		TArray< int32 > SlotsAfterDragHandleIndicies;
+		TArray< int32 > ProportionalSlotsAfterDragHandleIndicies;
 		if ( ResizeMode == ESplitterResizeMode::Fixed )
 		{
-			const int32 SlotAfterDragHandle = FindResizeableSlotAfterHandle( DraggedHandle, Children );
+			const int32 SlotAfterDragHandle = FindProportionallyResizeableSlotAfterHandle( DraggedHandle, Children );
 
 			if ( SlotAfterDragHandle < NumChildren )
 			{
-				SlotsAfterDragHandleIndicies.Add( SlotAfterDragHandle );
+				ProportionalSlotsAfterDragHandleIndicies.Add( SlotAfterDragHandle );
 			}
 		}
 		else if ( ResizeMode == ESplitterResizeMode::Fill )
 		{
-			FindAllResizeableSlotsAfterHandle( DraggedHandle, Children, /*OUT*/ SlotsAfterDragHandleIndicies );
+			FindAllProportionallyResizeableSlotsAfterHandle( DraggedHandle, Children, /*OUT*/ ProportionalSlotsAfterDragHandleIndicies );
 		}
 
-		if ( SlotsAfterDragHandleIndicies.Num() > 0)
+		if ( ProportionalSlotsAfterDragHandleIndicies.Num() > 0)
 		{
 			struct FSlotInfo 
 			{
@@ -567,29 +577,29 @@ void SSplitter::HandleResizing( const float PhysicalSplitterHandleSize, const ES
 				float NewSize;
 			};
 
-			TArray< FSlotInfo > SlotsAfterDragHandle;
-			for (int SlotIndex = 0; SlotIndex < SlotsAfterDragHandleIndicies.Num(); SlotIndex++)
+			TArray< FSlotInfo > ProportionalSlotsAfterDragHandle;
+			for (int SlotIndex = 0; SlotIndex < ProportionalSlotsAfterDragHandleIndicies.Num(); SlotIndex++)
 			{
 				FSlotInfo SlotInfo;
 
-				SlotInfo.Slot = &Children[ SlotsAfterDragHandleIndicies[ SlotIndex ] ];
-				SlotInfo.Geometry = &ChildGeometries[ SlotsAfterDragHandleIndicies[ SlotIndex ] ];
+				SlotInfo.Slot = &Children[ ProportionalSlotsAfterDragHandleIndicies[ SlotIndex ] ];
+				SlotInfo.Geometry = &ChildGeometries[ ProportionalSlotsAfterDragHandleIndicies[ SlotIndex ] ];
 				SlotInfo.NewSize = SlotInfo.Geometry->GetSizeInParentSpace().Component( AxisIndex );
 
-				SlotsAfterDragHandle.Add( SlotInfo );
+				ProportionalSlotsAfterDragHandle.Add( SlotInfo );
 			}
 
 			Delta = NewPrevChildLength - PrevChildLength;
 
 			// Distribute the Delta across the affected slots after the drag handle
 			float UnusedDelta = Delta;
-			for (int DistributionCount = 0; DistributionCount < SlotsAfterDragHandle.Num() && UnusedDelta != 0; DistributionCount++)
+			for (int DistributionCount = 0; DistributionCount < ProportionalSlotsAfterDragHandle.Num() && UnusedDelta != 0; DistributionCount++)
 			{
-				float DividedDelta = UnusedDelta / SlotsAfterDragHandle.Num();
+				float DividedDelta = UnusedDelta / ProportionalSlotsAfterDragHandle.Num();
 				UnusedDelta = 0;
-				for (int SlotIndex = 0; SlotIndex < SlotsAfterDragHandle.Num(); SlotIndex++)
+				for (int SlotIndex = 0; SlotIndex < ProportionalSlotsAfterDragHandle.Num(); SlotIndex++)
 				{
-					FSlotInfo& SlotInfo = SlotsAfterDragHandle[ SlotIndex ];
+					FSlotInfo& SlotInfo = ProportionalSlotsAfterDragHandle[ SlotIndex ];
 
 					float CurrentSize = ClampChild( SlotInfo.Geometry->GetSizeInParentSpace().Component(AxisIndex) );
 					SlotInfo.NewSize = ClampChild( CurrentSize - DividedDelta );
@@ -611,9 +621,9 @@ void SSplitter::HandleResizing( const float PhysicalSplitterHandleSize, const ES
 				float TotalLength = NewPrevChildLength;
 				float TotalStretchCoefficients = PrevChild.SizeValue.Get();
 
-				for (int SlotIndex = 0; SlotIndex < SlotsAfterDragHandle.Num(); SlotIndex++)
+				for (int SlotIndex = 0; SlotIndex < ProportionalSlotsAfterDragHandle.Num(); SlotIndex++)
 				{
-					FSlotInfo SlotInfo = SlotsAfterDragHandle[ SlotIndex ];
+					FSlotInfo SlotInfo = ProportionalSlotsAfterDragHandle[ SlotIndex ];
 
 					TotalLength += SlotInfo.NewSize;
 					TotalStretchCoefficients += SlotInfo.Slot->SizeValue.Get();
@@ -630,9 +640,9 @@ void SSplitter::HandleResizing( const float PhysicalSplitterHandleSize, const ES
 					PrevChild.SizeValue = NewPrevChildSize;
 				}
 
-				for (int SlotIndex = 0; SlotIndex < SlotsAfterDragHandle.Num(); SlotIndex++)
+				for (int SlotIndex = 0; SlotIndex < ProportionalSlotsAfterDragHandle.Num(); SlotIndex++)
 				{
-					FSlotInfo SlotInfo = SlotsAfterDragHandle[ SlotIndex ];
+					FSlotInfo SlotInfo = ProportionalSlotsAfterDragHandle[ SlotIndex ];
 
 					const float NewNextChildSize = ( TotalStretchCoefficients * SlotInfo.NewSize / TotalLength );
 
@@ -685,6 +695,44 @@ int32 SSplitter::GetHandleBeingResizedFromMousePosition( float PhysicalSplitterH
 	return INDEX_NONE;
 }
 
+template<EOrientation SplitterOrientation>
+int32 SSplitter::GetAndValidateHandleBeingResizedFromMousePosition( FVector2D LocalMousePos, const TArray<FLayoutGeometry>& ChildGeometries ) const
+{
+	// Hit test which handle we are hovering over.		
+	int32 LocalHoveredHandleIndex = (Orientation == Orient_Horizontal)
+		? GetHandleBeingResizedFromMousePosition<Orient_Horizontal>( PhysicalSplitterHandleSize, HitDetectionSplitterHandleSize, LocalMousePos, ChildGeometries )
+		: GetHandleBeingResizedFromMousePosition<Orient_Vertical>( PhysicalSplitterHandleSize, HitDetectionSplitterHandleSize, LocalMousePos, ChildGeometries );
+
+	if (LocalHoveredHandleIndex != INDEX_NONE)
+	{
+		const FSlot& PrevChild = Children[LocalHoveredHandleIndex];
+		if ( PrevChild.SizingRule.Get() == SSplitter::ManualSize )
+		{
+			// Manual sized slots can be resized as long as they're not the last slot in the splitter
+			const int32 NonCollapsedSlotAfterDragHandle = FindNonCollapsedSlotAfterHandle( LocalHoveredHandleIndex, Children );
+
+			// Valid index?
+			if ( NonCollapsedSlotAfterDragHandle >= Children.Num() )
+			{
+				LocalHoveredHandleIndex = INDEX_NONE;
+			}
+		}
+		else
+		{
+			// Other kinds of slots can be resized as long as they have at least one proportionally sized slot after them in the splitter
+			const int32 ProportionalSlotAfterDragHandle = FindProportionallyResizeableSlotAfterHandle( LocalHoveredHandleIndex, Children );
+
+			// Valid index?
+			if ( ProportionalSlotAfterDragHandle >= Children.Num() )
+			{
+				LocalHoveredHandleIndex = INDEX_NONE;
+			}
+		}
+	}
+
+	return LocalHoveredHandleIndex;
+}
+
 /********************************************************************
 * SSplitter2x2														*
 * A splitter which has exactly 4 children and allows simultaneous	*		
diff --git a/Engine/Source/Runtime/Slate/Public/Widgets/Layout/SSplitter.h b/Engine/Source/Runtime/Slate/Public/Widgets/Layout/SSplitter.h
--- a/Engine/Source/Runtime/Slate/Public/Widgets/Layout/SSplitter.h
+++ b/Engine/Source/Runtime/Slate/Public/Widgets/Layout/SSplitter.h
@@ -13,7 +13,7 @@ namespace ESplitterResizeMode
 
 class FLayoutGeometry;
 /**
- * SSPlitter divides its allotted area into N segments, where N is the number of children it has.
+ * SSplitter divides its allotted area into N segments, where N is the number of children it has.
  * It allows the users to resize the children along the splitters axis: that is, horizontally or vertically.
  */
 class SLATE_API SSplitter : public SPanel
@@ -27,8 +27,8 @@ public:
 		SizeToContent,
 		/** Use a fraction of the parent's size */
 		FractionOfParent,
-		/** Use the SizeValue of the slot as an absolute size */
-		AbsoluteSize,
+		/** Use the SizeValue of the slot as the desired size */
+		ManualSize,
 	};
 
 	DECLARE_DELEGATE_OneParam(
@@ -213,20 +213,27 @@ private:
 protected:
 
 	/**
-	 * Given the index of the dragged handle and the children, find a child above/left_of of the dragged handle that can be resized.
+	 * Given the index of the dragged handle and the children, find a child above/left_of of the dragged handle that can be resized (either proportionally or absolutely).
 	 *
 	 * @return INDEX_NONE if no such child can be found.
 	 */
 	static int32 FindResizeableSlotBeforeHandle( int32 DraggedHandle, const TPanelChildren<FSlot>& Children );
 
 	/**
-	 * Given the index of the dragged handle and the children, find a child below/right_of the dragged handle that can be resized
+	 * Given the index of the dragged handle and the children, find a child below/right_of the dragged handle that isn't collapsed
 	 *
 	 * @return Children.Num() if no such child can be found.
 	 */
-	static int32 FindResizeableSlotAfterHandle( int32 DraggedHandle, const TPanelChildren<FSlot>& Children );
+	static int32 FindNonCollapsedSlotAfterHandle( int32 DraggedHandle, const TPanelChildren<FSlot>& Children );
 
-	static void FindAllResizeableSlotsAfterHandle( int32 DraggedHandle, const TPanelChildren<FSlot>& Children, TArray< int32 >& OutSlotIndicies );
+	/**
+	 * Given the index of the dragged handle and the children, find a child below/right_of the dragged handle that can be proportionally resized
+	 *
+	 * @return Children.Num() if no such child can be found.
+	 */
+	static int32 FindProportionallyResizeableSlotAfterHandle( int32 DraggedHandle, const TPanelChildren<FSlot>& Children );
+
+	static void FindAllProportionallyResizeableSlotsAfterHandle( int32 DraggedHandle, const TPanelChildren<FSlot>& Children, TArray< int32 >& OutSlotIndicies );
 
 	/**
 	 * Resizes the children based on user input. The template parameter Orientation corresponds to the splitter being horizontal or vertical.
@@ -258,6 +265,18 @@ protected:
 	template<EOrientation SplitterOrientation>
 	static int32 GetHandleBeingResizedFromMousePosition(  float PhysicalSplitterHandleSize, float HitDetectionSplitterHandleSize, FVector2D LocalMousePos, const TArray<FLayoutGeometry>& ChildGeometries );
 
+	/**
+	 * Given a mouse position within the splitter, figure out which resize handle we are hovering (if any), and
+	 * then validate that the handle can actually be used to perform a sizing operation
+	 *
+	 * @param LocalMousePos  The mouse position within this splitter.
+	 * @param ChildGeometris The arranged children and their geometries; we need to test the mouse against them.
+	 *
+	 * @return The index of the handle being hovered, or INDEX_NONE if we are not hovering a handle, or the handle cannot be used to size a slot.
+	 */
+	template<EOrientation SplitterOrientation>
+	int32 GetAndValidateHandleBeingResizedFromMousePosition( FVector2D LocalMousePos, const TArray<FLayoutGeometry>& ChildGeometries ) const;
+
 	TPanelChildren< FSlot > Children;
 
 	int32 HoveredHandleIndex;
