Slate : If the widget asking for mouse capture is completely gone (including the entire path to that widget) do not crash; ensure instead.
Thanks to [REDACTED_PROFILE_URL] for pointing this one out and providing a fix.

diff --git a/Engine/Source/Runtime/Slate/Private/Framework/Application/SlateApplication.cpp b/Engine/Source/Runtime/Slate/Private/Framework/Application/SlateApplication.cpp
--- a/Engine/Source/Runtime/Slate/Private/Framework/Application/SlateApplication.cpp
+++ b/Engine/Source/Runtime/Slate/Private/Framework/Application/SlateApplication.cpp
@@ -444,7 +444,7 @@ void FSlateApplication::MouseCaptorHelper::SetMouseCaptor(uint32 PointerIndex, c
 		{
 			MouseCaptorWeakPath = NewMouseCaptorPath;
 		}
-		else
+		else if (EventPath.Widgets.Num() > 0)
 		{
 			// If the target widget wasn't found on the event path then start the search from the root
 			NewMouseCaptorPath = EventPath.GetPathDownTo( EventPath.Widgets[0].Widget );
@@ -454,6 +454,10 @@ void FSlateApplication::MouseCaptorHelper::SetMouseCaptor(uint32 PointerIndex, c
 				? NewMouseCaptorPath
 				: FWeakWidgetPath();
 		}
+		else
+		{
+			ensureMsgf(EventPath.Widgets.Num() > 0, TEXT("An unknown widget is attempting to set capture to %s"), *Widget->ToString() );
+		}
 
 		if (MouseCaptorWeakPath.IsValid())
 		{
