UMG/Slate - UE-13439 - The SWrapBox now correctly ignores collapsed visibility widgets when performing layout calculations.

diff --git a/Engine/Source/Runtime/Slate/Private/Widgets/Layout/SWrapBox.cpp b/Engine/Source/Runtime/Slate/Private/Widgets/Layout/SWrapBox.cpp
--- a/Engine/Source/Runtime/Slate/Private/Widgets/Layout/SWrapBox.cpp
+++ b/Engine/Source/Runtime/Slate/Private/Widgets/Layout/SWrapBox.cpp
@@ -68,18 +68,24 @@ void SWrapBox::OnArrangeChildren( const FGeometry& AllottedGeometry, FArrangedCh
 		float LineHeightSoFar = 0;
 		for ( int32 ChildIndex=0; ChildIndex < Slots.Num(); ++ChildIndex )
 		{
-			const TSharedRef<SWidget>& ThisChild = Slots[ChildIndex].GetWidget();
+			const FSlot& CurSlot = Slots[ChildIndex];
+			const TSharedRef<SWidget>& ThisChild = CurSlot.GetWidget();
+
+			if ( ThisChild->GetVisibility() == EVisibility::Collapsed )
+			{
+				continue;
+			}
 
 			// If slot isn't the first one on the line, we need to add the inner slot padding.
 			const float ConditionPaddingX = ( WidthSoFar == 0 ) ? 0 : InnerSlotPadding.X;
 			
-			FVector2D ThisSlotSize = ThisChild->GetDesiredSize() + Slots[ChildIndex].SlotPadding.Get().GetDesiredSize();
+			FVector2D ThisSlotSize = ThisChild->GetDesiredSize() + CurSlot.SlotPadding.Get().GetDesiredSize();
 
 			// If the width is smaller than this slots fill limit, we need to account for it filling the line
 			// and forcing all controls that follow it to wrap.
-			if ( Slots[ChildIndex].SlotFillLineWhenWidthLessThan.IsSet() )
+			if ( CurSlot.SlotFillLineWhenWidthLessThan.IsSet() )
 			{
-				if ( MyPreferredWidth < Slots[ChildIndex].SlotFillLineWhenWidthLessThan.GetValue() )
+				if ( MyPreferredWidth < CurSlot.SlotFillLineWhenWidthLessThan.GetValue() )
 				{
 					ThisSlotSize.X = MyPreferredWidth;
 				}
@@ -108,6 +114,14 @@ void SWrapBox::OnArrangeChildren( const FGeometry& AllottedGeometry, FArrangedCh
 	float YOffsetSoFar = 0;
 	for ( int32 CurrentLineIndex=0, ChildIndex=0; ChildIndex < Slots.Num(); ++ChildIndex )
 	{
+		const FSlot& CurSlot = Slots[ChildIndex];
+		const TSharedRef<SWidget>& ThisChild = CurSlot.GetWidget();
+
+		if ( ThisChild->GetVisibility() == EVisibility::Collapsed )
+		{
+			continue;
+		}
+
 		// If slot isn't on the first line, we need to add the inner slot padding.
 		const float ConditionPaddingY = ( YOffsetSoFar == 0 ) ? 0 : InnerSlotPadding.Y;
 
@@ -131,31 +145,30 @@ void SWrapBox::OnArrangeChildren( const FGeometry& AllottedGeometry, FArrangedCh
 		bool bLastWidgetOnLine = ((ChildIndex + 1) >= Slots.Num()) || WidgetIndexStartsNewLine[CurrentLineIndex] == ( ChildIndex + 1 );
 
 		const float CurrentLineHeight = LineHeights[ CurrentLineIndex ];
-		const FMargin& SlotPadding( Slots[ChildIndex].SlotPadding.Get() );
+		const FMargin& SlotPadding(CurSlot.SlotPadding.Get());
 
-		const TSharedRef<SWidget>& ThisChild = Slots[ChildIndex].GetWidget();
 		float ThisSlotSizeX = ThisChild->GetDesiredSize().X + SlotPadding.GetTotalSpaceAlong<Orient_Horizontal>();
 
 		// If the width is smaller than this slots fill limit, we need to account for it filling the line
 		// and forcing all controls that follow it to wrap.
-		if ( Slots[ChildIndex].SlotFillLineWhenWidthLessThan.IsSet() )
+		if ( CurSlot.SlotFillLineWhenWidthLessThan.IsSet() )
 		{
-			if ( MyPreferredWidth < Slots[ChildIndex].SlotFillLineWhenWidthLessThan.GetValue() )
+			if ( MyPreferredWidth < CurSlot.SlotFillLineWhenWidthLessThan.GetValue() )
 			{
 				ThisSlotSizeX = MyPreferredWidth;
 			}
 		}
 
 		// If this slot fills empty space, we need to have it fill the remaining space on the line.
-		if ( bLastWidgetOnLine && Slots[ChildIndex].bSlotFillEmptySpace )
+		if ( bLastWidgetOnLine && CurSlot.bSlotFillEmptySpace )
 		{
 			ThisSlotSizeX = MyPreferredWidth - (XOffsetSoFar - ConditionPaddingX);
 		}
 
 		// Horizontal alignment will always be equivalent to HAlign_Fill due to the nature of the widget.
 		// But we still need to account for padding.		
-		AlignmentArrangeResult XAlignmentResult = AlignChild<Orient_Horizontal>( ThisSlotSizeX, Slots[ChildIndex], SlotPadding + ConditionPaddingX );
-		AlignmentArrangeResult YAlignmentResult = AlignChild<Orient_Vertical>( CurrentLineHeight, Slots[ChildIndex], SlotPadding + ConditionPaddingY );
+		AlignmentArrangeResult XAlignmentResult = AlignChild<Orient_Horizontal>(ThisSlotSizeX, CurSlot, SlotPadding + ConditionPaddingX);
+		AlignmentArrangeResult YAlignmentResult = AlignChild<Orient_Vertical>(CurrentLineHeight, CurSlot, SlotPadding + ConditionPaddingY);
 
 		ArrangedChildren.AddWidget( AllottedGeometry.MakeChild(
 			ThisChild,
@@ -182,18 +195,26 @@ FVector2D SWrapBox::ComputeDesiredSize( float ) const
 	
 	// Find the size of every row. Once the size of the row exceeds the preferred width,
 	// start a new row. Whenever we start a new row, update the desired size of the entire widget.
-	for ( int32 ChildIndex=0; ChildIndex<Slots.Num(); ++ChildIndex )
+	for ( int32 ChildIndex=0; ChildIndex < Slots.Num(); ++ChildIndex )
 	{
+		const FSlot& CurSlot = Slots[ChildIndex];
+		const TSharedRef<SWidget>& ThisChild = CurSlot.GetWidget();
+
+		if ( ThisChild->GetVisibility() == EVisibility::Collapsed )
+		{
+			continue;
+		}
+
 		// If slot isn't the first one on the line, we need to add the inner slot padding.
 		const float ConditionPaddingX = ( RowDesiredSize.X == 0 ) ? 0 : InnerSlotPadding.X;
 
-		FVector2D ThisChildDesiredSize = Slots[ChildIndex].GetWidget()->GetDesiredSize() + Slots[ChildIndex].SlotPadding.Get().GetDesiredSize() + ConditionPaddingX;
+		FVector2D ThisChildDesiredSize = ThisChild->GetDesiredSize() + CurSlot.SlotPadding.Get().GetDesiredSize() + ConditionPaddingX;
 
 		// If the width is smaller than this slots fill limit, we need to account for it filling the line
 		// and forcing all controls that follow it to wrap.
-		if ( Slots[ChildIndex].SlotFillLineWhenWidthLessThan.IsSet() )
+		if ( CurSlot.SlotFillLineWhenWidthLessThan.IsSet() )
 		{
-			if ( MyPreferredWidth < Slots[ChildIndex].SlotFillLineWhenWidthLessThan.GetValue() )
+			if ( MyPreferredWidth <CurSlot.SlotFillLineWhenWidthLessThan.GetValue() )
 			{
 				ThisChildDesiredSize.X = MyPreferredWidth;
 			}
