Slate: Improved the error handling message in FUICommandList::MapAction to include the name of the conflicted command

diff --git a/Engine/Source/Runtime/Slate/Private/Framework/Commands/UICommandList.cpp b/Engine/Source/Runtime/Slate/Private/Framework/Commands/UICommandList.cpp
--- a/Engine/Source/Runtime/Slate/Private/Framework/Commands/UICommandList.cpp
+++ b/Engine/Source/Runtime/Slate/Private/Framework/Commands/UICommandList.cpp
@@ -42,7 +42,9 @@ void FUICommandList::MapAction( const TSharedPtr< const FUICommandInfo > InUICom
 	check( InUICommandInfo.IsValid() );
 
 	// Check against already mapped actions
-	checkSlow( !UICommandBindingMap.Contains( InUICommandInfo ) );
+#if DO_GUARD_SLOW
+	checkf(!UICommandBindingMap.Contains(InUICommandInfo), TEXT("Command list already contains a command named '%s'"), *InUICommandInfo->GetCommandName().ToString());
+#endif
 
 	ContextsInList.Add( InUICommandInfo->GetBindingContext() );
 	UICommandBindingMap.Add( InUICommandInfo, InUIAction );
