Added another column sizing mode to STableRow which allows Excel like sizing

This will replace the ManualSize mode in SSplitter as it's too specific to use for general layout sizing. The data table editor will be converted to use a SListView so it can make use of this sizing mode.

diff --git a/Engine/Source/Runtime/Slate/Private/Widgets/Views/SHeaderRow.cpp b/Engine/Source/Runtime/Slate/Private/Widgets/Views/SHeaderRow.cpp
--- a/Engine/Source/Runtime/Slate/Private/Widgets/Views/SHeaderRow.cpp
+++ b/Engine/Source/Runtime/Slate/Private/Widgets/Views/SHeaderRow.cpp
@@ -562,32 +562,125 @@ void SHeaderRow::RegenerateWidgets()
 				SAssignNew(NewlyMadeHeader, STableColumnHeader, SomeColumn, DefaultPadding)
 				.Style( (SlotIndex + 1 == Columns.Num()) ? &Style->LastColumnStyle : &Style->ColumnStyle );
 
-			if ( SomeColumn.SizeRule == EColumnSizeMode::Fixed )
+			switch(SomeColumn.SizeRule)
 			{
-				// Add resizable cell
-				Splitter->AddSlot()
-				.SizeRule( SSplitter::SizeToContent )
-				[
-					SNew( SBox )
-					.WidthOverride( SomeColumn.Width.Get() )
+			case EColumnSizeMode::Fill:
+				{
+					TAttribute<float> WidthBinding;
+					WidthBinding.BindRaw( &SomeColumn, &FColumn::GetWidth );
+
+					// Add resizable cell
+					Splitter->AddSlot()
+					.Value( WidthBinding )
+					.SizeRule( SSplitter::FractionOfParent )
+					.OnSlotResized( SSplitter::FOnSlotResized::CreateRaw( &SomeColumn, &FColumn::SetWidth ) )
 					[
 						NewHeader
-					]
-				];
-			}
-			else 
-			{
-				TAttribute<float> WidthBinding;
-				WidthBinding.BindRaw( &SomeColumn, &FColumn::GetWidth );
-
-				// Add resizable cell
-				Splitter->AddSlot()
-				.Value( WidthBinding )
-				.SizeRule( SSplitter::FractionOfParent )
-				.OnSlotResized( SSplitter::FOnSlotResized::CreateRaw( &SomeColumn, &FColumn::SetWidth ) )
-				[
-					NewHeader
-				];
+					];
+				}
+				break;
+
+			case EColumnSizeMode::Fixed:
+				{
+					// Add fixed size cell
+					Splitter->AddSlot()
+					.SizeRule( SSplitter::SizeToContent )
+					[
+						SNew( SBox )
+						.WidthOverride( SomeColumn.GetWidth() )
+						[
+							NewHeader
+						]
+					];
+				}
+				break;
+
+			case EColumnSizeMode::Manual:
+				{
+					// Sizing grip to put at the end of the column - we can't use a SSplitter here as it doesn't have the resizing behavior we need
+					const float GripSize = 5.0f;
+					TSharedRef<SBorder> SizingGrip = SNew(SBorder)
+						.Padding(0.0f)
+						.BorderImage( FCoreStyle::Get().GetBrush("NoBorder") )
+						.Cursor( EMouseCursor::ResizeLeftRight )
+						.Content()
+						[
+							SNew(SSpacer)
+							.Size(FVector2D(GripSize, GripSize))
+						];
+
+					TWeakPtr<SBorder> WeakSizingGrip = SizingGrip;
+					auto SizingGrip_OnMouseButtonDown = [&SomeColumn, WeakSizingGrip](const FGeometry&, const FPointerEvent&) -> FReply
+					{
+						TSharedPtr<SBorder> SizingGrip = WeakSizingGrip.Pin();
+						if ( SizingGrip.IsValid() )
+						{
+							return FReply::Handled().CaptureMouse(SizingGrip.ToSharedRef());
+						}
+						return FReply::Unhandled();
+					};
+
+					auto SizingGrip_OnMouseButtonUp = [&SomeColumn, WeakSizingGrip](const FGeometry&, const FPointerEvent&) -> FReply
+					{
+						TSharedPtr<SBorder> SizingGrip = WeakSizingGrip.Pin();
+						if ( SizingGrip.IsValid() && SizingGrip->HasMouseCapture() )
+						{
+							return FReply::Handled().ReleaseMouseCapture();
+						}
+						return FReply::Unhandled();
+					};
+
+					auto SizingGrip_OnMouseMove = [&SomeColumn, WeakSizingGrip](const FGeometry&, const FPointerEvent& InPointerEvent) -> FReply
+					{
+						TSharedPtr<SBorder> SizingGrip = WeakSizingGrip.Pin();
+						if ( SizingGrip.IsValid() && SizingGrip->HasMouseCapture() )
+						{
+							// The sizing grip has been moved, so update our columns size from the movement delta
+							const float NewWidth = SomeColumn.GetWidth() + InPointerEvent.GetCursorDelta().X;
+							SomeColumn.SetWidth(FMath::Max(20.0f, NewWidth));
+							return FReply::Handled();
+						}
+						return FReply::Unhandled();
+					};
+
+					// Bind the events to handle the drag sizing
+					SizingGrip->SetOnMouseButtonDown( FPointerEventHandler::CreateLambda(SizingGrip_OnMouseButtonDown) );
+					SizingGrip->SetOnMouseButtonUp( FPointerEventHandler::CreateLambda(SizingGrip_OnMouseButtonUp) );
+					SizingGrip->SetOnMouseMove( FPointerEventHandler::CreateLambda(SizingGrip_OnMouseMove) );
+
+					auto GetColumnWidthAsOptionalSize = [&SomeColumn]() -> FOptionalSize
+					{
+						const float DesiredWidth = SomeColumn.GetWidth();
+						return FOptionalSize(DesiredWidth);
+					};
+
+					TAttribute<FOptionalSize> WidthBinding;
+					WidthBinding.Bind(TAttribute<FOptionalSize>::FGetter::CreateLambda(GetColumnWidthAsOptionalSize));
+
+					// Add resizable cell
+					Splitter->AddSlot()
+					.SizeRule( SSplitter::SizeToContent )
+					[
+						SNew(SBox)
+						.WidthOverride(WidthBinding)
+						[
+							SNew(SOverlay)
+							+SOverlay::Slot()
+							[
+								NewHeader
+							]
+							+SOverlay::Slot()
+							.HAlign(HAlign_Right)
+							[
+								SizingGrip
+							]
+						]
+					];
+				}
+				break;
+
+			default:
+				break;
 			}
 		}
 	}
diff --git a/Engine/Source/Runtime/Slate/Public/Widgets/Views/SHeaderRow.h b/Engine/Source/Runtime/Slate/Public/Widgets/Views/SHeaderRow.h
--- a/Engine/Source/Runtime/Slate/Public/Widgets/Views/SHeaderRow.h
+++ b/Engine/Source/Runtime/Slate/Public/Widgets/Views/SHeaderRow.h
@@ -39,11 +39,14 @@ namespace EColumnSizeMode
 {
 	enum Type
 	{
-		/** Column stretches to a fractional of the header row */
+		/** Column stretches to a fraction of the header row */
 		Fill = 0,
 		
 		/**	Column is fixed width and cannot be resized */
 		Fixed = 1,
+
+		/** Column is set to a width which can be user-sized */
+		Manual = 2,
 	};
 };
 
@@ -85,6 +88,7 @@ public:
 			, _DefaultTooltip()
 			, _FillWidth( 1.0f )
 			, _FixedWidth()
+			, _ManualWidth()
 			, _OnWidthChanged()
 			, _HeaderContent()
 			, _HAlignHeader( HAlign_Fill )
@@ -102,6 +106,7 @@ public:
 			SLATE_TEXT_ATTRIBUTE( DefaultTooltip )
 			SLATE_ATTRIBUTE( float, FillWidth )
 			SLATE_ARGUMENT( TOptional< float >, FixedWidth )
+			SLATE_ATTRIBUTE( float, ManualWidth )
 			SLATE_EVENT( FOnWidthChanged, OnWidthChanged )
 
 			SLATE_DEFAULT_SLOT( FArguments, HeaderContent )
@@ -144,6 +149,11 @@ public:
 				Width = InArgs._FixedWidth.GetValue();
 				SizeRule = EColumnSizeMode::Fixed;
 			}
+			else if ( InArgs._ManualWidth.IsSet() )
+			{
+				Width = InArgs._ManualWidth;
+				SizeRule = EColumnSizeMode::Manual;
+			}
 			else
 			{
 				Width = InArgs._FillWidth;
diff --git a/Engine/Source/Runtime/Slate/Public/Widgets/Views/STableRow.h b/Engine/Source/Runtime/Slate/Public/Widgets/Views/STableRow.h
--- a/Engine/Source/Runtime/Slate/Public/Widgets/Views/STableRow.h
+++ b/Engine/Source/Runtime/Slate/Public/Widgets/Views/STableRow.h
@@ -1020,33 +1020,68 @@ protected:
 				CellContents = GenerateWidgetForColumn( Column.ColumnId );
 			}
 
-			if ( Column.SizeRule == EColumnSizeMode::Fixed )
+			switch(Column.SizeRule)
 			{
-				Box->AddSlot()
-				.AutoWidth()
-				[
-					SNew( SBox )
-					.WidthOverride( Column.Width.Get() )
+			case EColumnSizeMode::Fill:
+				{
+					TAttribute<float> WidthBinding;
+					WidthBinding.BindRaw( &Column, &SHeaderRow::FColumn::GetWidth );
+
+					SHorizontalBox::FSlot& NewSlot = Box->AddSlot()
 					.HAlign(Column.CellHAlignment)
 					.VAlign(Column.CellVAlignment)
-					.Content()
+					.FillWidth( WidthBinding )
 					[
 						CellContents
-					]
-				];
-			}
-			else
-			{
-				TAttribute<float> WidthBinding;
-				WidthBinding.BindRaw( &Column, &SHeaderRow::FColumn::GetWidth );
+					];
+				}
+				break;
 
-				SHorizontalBox::FSlot& NewSlot = Box->AddSlot()
-				.HAlign(Column.CellHAlignment)
-				.VAlign(Column.CellVAlignment)
-				.FillWidth( WidthBinding )
-				[
-					CellContents
-				];
+			case EColumnSizeMode::Fixed:
+				{
+					Box->AddSlot()
+					.AutoWidth()
+					[
+						SNew( SBox )
+						.WidthOverride( Column.Width.Get() )
+						.HAlign(Column.CellHAlignment)
+						.VAlign(Column.CellVAlignment)
+						.Content()
+						[
+							CellContents
+						]
+					];
+				}
+				break;
+
+			case EColumnSizeMode::Manual:
+				{
+					auto GetColumnWidthAsOptionalSize = [&Column]() -> FOptionalSize
+					{
+						const float DesiredWidth = Column.GetWidth();
+						return FOptionalSize(DesiredWidth);
+					};
+
+					TAttribute<FOptionalSize> WidthBinding;
+					WidthBinding.Bind(TAttribute<FOptionalSize>::FGetter::CreateLambda(GetColumnWidthAsOptionalSize));
+
+					Box->AddSlot()
+					.AutoWidth()
+					[
+						SNew( SBox )
+						.WidthOverride(WidthBinding)
+						.HAlign(Column.CellHAlignment)
+						.VAlign(Column.CellVAlignment)
+						.Content()
+						[
+							CellContents
+						]
+					];
+				}
+				break;
+
+			default:
+				break;
 			}
 
 			NewColumnIdToSlotContents.Add( Column.ColumnId, CellContents );
