Added a way for STableViewBase (and its derived types) to notify you when they are scrolled.

This allows you, with the combination of SetScrollOffset, to sync two or more list views to have the same scroll offset.

STableViewBase no-longer associates its header with an external scrollbar. This prevents the header from trying to compensate for something that doesn't exist, and causing the header and the rows to become out of line.

diff --git a/Engine/Source/Developer/AutomationWindow/Private/SAutomationTestTreeView.h b/Engine/Source/Developer/AutomationWindow/Private/SAutomationTestTreeView.h
--- a/Engine/Source/Developer/AutomationWindow/Private/SAutomationTestTreeView.h
+++ b/Engine/Source/Developer/AutomationWindow/Private/SAutomationTestTreeView.h
@@ -39,6 +39,8 @@ public:
 
 			SLATE_EVENT( FOnGenerateRow, OnGenerateRow )
 
+			SLATE_EVENT( FOnTableViewScrolled, OnTreeViewScrolled )
+
 			SLATE_EVENT( FOnItemScrolledIntoView, OnItemScrolledIntoView )
 
 			SLATE_EVENT( FOnGetChildren, OnGetChildren )
@@ -122,7 +124,7 @@ public:
 			else
 			{
 				// Make the TableView
-				this->ConstructChildren( 0, InArgs._ItemHeight, EListItemAlignment::LeftAligned, InArgs._HeaderRow, InArgs._ExternalScrollbar );
+				this->ConstructChildren( 0, InArgs._ItemHeight, EListItemAlignment::LeftAligned, InArgs._HeaderRow, InArgs._ExternalScrollbar, InArgs._OnTreeViewScrolled );
 			}
 		}
 
diff --git a/Engine/Source/Runtime/Slate/Private/Widgets/Views/STableViewBase.cpp b/Engine/Source/Runtime/Slate/Private/Widgets/Views/STableViewBase.cpp
--- a/Engine/Source/Runtime/Slate/Private/Widgets/Views/STableViewBase.cpp
+++ b/Engine/Source/Runtime/Slate/Private/Widgets/Views/STableViewBase.cpp
@@ -9,12 +9,14 @@ namespace ListConstants
 	static const float OvershootBounceRate = 250.0f;
 }
 
-void STableViewBase::ConstructChildren( const TAttribute<float>& InItemWidth, const TAttribute<float>& InItemHeight, const TAttribute<EListItemAlignment>& InItemAlignment, const TSharedPtr<SHeaderRow>& InHeaderRow, const TSharedPtr<SScrollBar>& InScrollBar )
+void STableViewBase::ConstructChildren( const TAttribute<float>& InItemWidth, const TAttribute<float>& InItemHeight, const TAttribute<EListItemAlignment>& InItemAlignment, const TSharedPtr<SHeaderRow>& InHeaderRow, const TSharedPtr<SScrollBar>& InScrollBar, const FOnTableViewScrolled& InOnTableViewScrolled )
 {
 	bItemsNeedRefresh = true;
 	
 	HeaderRow = InHeaderRow;
 
+	OnTableViewScrolled = InOnTableViewScrolled;
+
 	// If the user provided a scrollbar, we do not need to make one of our own.
 	if (InScrollBar.IsValid())
 	{
@@ -22,7 +24,6 @@ void STableViewBase::ConstructChildren( const TAttribute<float>& InItemWidth, co
 		ScrollBar->SetOnUserScrolled( FOnUserScrolled::CreateSP(this, &STableViewBase::ScrollBar_OnUserScrolled) );
 	}
 	
-	
 	TSharedRef<SWidget> ListAndScrollbar = (!ScrollBar.IsValid())
 		// We need to make our own scrollbar
 		? StaticCastSharedRef<SWidget>
@@ -60,7 +61,12 @@ void STableViewBase::ConstructChildren( const TAttribute<float>& InItemWidth, co
 
 	if (InHeaderRow.IsValid())
 	{
-		InHeaderRow->SetAssociatedVerticalScrollBar( ScrollBar.ToSharedRef(), 16 );
+		// Only associate the scrollbar if we created it.
+		// If the scrollbar was passed in from outside then it won't appear under our header row so doesn't need compensating for.
+		if (!InScrollBar.IsValid())
+		{
+			InHeaderRow->SetAssociatedVerticalScrollBar( ScrollBar.ToSharedRef(), 16 );
+		}
 
 		this->ChildSlot
 		[
@@ -247,7 +253,7 @@ void STableViewBase::Tick( const FGeometry& AllottedGeometry, const double InCur
 			}
 			
 			
-			ScrollOffset = FMath::Max(0.0, ScrollOffset);
+			SetScrollOffset( FMath::Max(0.0, ScrollOffset) );
 			ItemsPanel->SmoothScrollOffset( FMath::Fractional(ScrollOffset / GetNumItemsWide()) );
 
 			if (AllowOverscroll == EAllowOverscroll::Yes)
@@ -645,10 +651,8 @@ float STableViewBase::ScrollTo( float InScrollOffset )
 {
 	const float NewScrollOffset = FMath::Clamp( InScrollOffset, -10.0f, GetNumItemsBeingObserved()+10.0f );
 	float AmountScrolled = FMath::Abs( ScrollOffset - NewScrollOffset );
-	ScrollOffset = NewScrollOffset;
-
-	RequestListRefresh();	
-
+	SetScrollOffset( NewScrollOffset );
+	
 	if ( bWasAtEndOfList && NewScrollOffset >= ScrollOffset )
 	{
 		AmountScrolled = 0;
@@ -657,6 +661,16 @@ float STableViewBase::ScrollTo( float InScrollOffset )
 	return AmountScrolled;
 }
 
+void STableViewBase::SetScrollOffset( const float InScrollOffset )
+{
+	if ( ScrollOffset != InScrollOffset )
+	{
+		ScrollOffset = InScrollOffset;
+		OnTableViewScrolled.ExecuteIfBound( ScrollOffset );
+		RequestListRefresh();
+	}
+}
+
 void STableViewBase::InsertWidget( const TSharedRef<ITableRow> & WidgetToInset )
 {
 	ItemsPanel->AddSlot(0)
@@ -755,13 +769,13 @@ float STableViewBase::GetScrollRateInItems() const
 
 void STableViewBase::ScrollToTop()
 {
-	ScrollOffset = 0;
+	SetScrollOffset( 0 );
 	RequestListRefresh();
 }
 
 void STableViewBase::ScrollToBottom()
 {
-	ScrollOffset = GetNumItemsBeingObserved();
+	SetScrollOffset( GetNumItemsBeingObserved() );
 	RequestListRefresh();
 }
 
diff --git a/Engine/Source/Runtime/Slate/Public/Widgets/Views/STableViewBase.h b/Engine/Source/Runtime/Slate/Public/Widgets/Views/STableViewBase.h
--- a/Engine/Source/Runtime/Slate/Public/Widgets/Views/STableViewBase.h
+++ b/Engine/Source/Runtime/Slate/Public/Widgets/Views/STableViewBase.h
@@ -35,6 +35,11 @@ enum class EListItemAlignment : uint8
 };
 
 
+DECLARE_DELEGATE_OneParam(
+	FOnTableViewScrolled,
+	double );	/** Scroll offset from the beginning of the list in items */
+
+
 /**
  * Contains ListView functionality that does not depend on the type of data being observed by the ListView.
  */
@@ -45,7 +50,7 @@ class SLATE_API STableViewBase
 public:
 
 	/** Create the child widgets that comprise the list */
-	void ConstructChildren( const TAttribute<float>& InItemWidth, const TAttribute<float>& InItemHeight, const TAttribute<EListItemAlignment>& InItemAlignment, const TSharedPtr<SHeaderRow>& InColumnHeaders, const TSharedPtr<SScrollBar>& InScrollBar  );
+	void ConstructChildren( const TAttribute<float>& InItemWidth, const TAttribute<float>& InItemHeight, const TAttribute<EListItemAlignment>& InItemAlignment, const TSharedPtr<SHeaderRow>& InColumnHeaders, const TSharedPtr<SScrollBar>& InScrollBar, const FOnTableViewScrolled& InOnTableViewScrolled );
 
 	/** Sets the item height */
 	void SetItemHeight(TAttribute<float> Height);
@@ -88,6 +93,9 @@ public:
 	/** Scrolls the view to the bottom */
 	void ScrollToBottom();
 
+	/** Set the scroll offset of this view (in items) */
+	void SetScrollOffset( const float InScrollOffset );
+
 public:
 
 	// SWidget interface
@@ -121,7 +129,7 @@ public:
 protected:
 
 	STableViewBase( ETableViewMode::Type InTableViewMode );
-	
+
 	/**
 	 * Scroll the list view by some number of screen units.
 	 *
@@ -235,9 +243,13 @@ protected:
 
 	/** The panel which holds the visible widgets in this list */
 	TSharedPtr< SListPanel > ItemsPanel;
+
 	/** The scroll bar widget */
 	TSharedPtr< SScrollBar > ScrollBar;
 
+	/** Delegate to call when the table view is scrolled */
+	FOnTableViewScrolled OnTableViewScrolled;
+
 	/** Scroll offset from the beginning of the list in items */
 	double ScrollOffset;
 
diff --git a/Engine/Source/Runtime/Slate/Public/Widgets/Views/STileView.h b/Engine/Source/Runtime/Slate/Public/Widgets/Views/STileView.h
--- a/Engine/Source/Runtime/Slate/Public/Widgets/Views/STileView.h
+++ b/Engine/Source/Runtime/Slate/Public/Widgets/Views/STileView.h
@@ -39,6 +39,8 @@ public:
 
 		SLATE_EVENT( FOnGenerateRow, OnGenerateTile )
 
+		SLATE_EVENT( FOnTableViewScrolled, OnTileViewScrolled )
+
 		SLATE_EVENT( FOnItemScrolledIntoView, OnItemScrolledIntoView )
 
 		SLATE_ARGUMENT( const TArray<ItemType>* , ListItemsSource )
@@ -118,7 +120,7 @@ public:
 		else
 		{
 			// Make the TableView
-			this->ConstructChildren(InArgs._ItemWidth, InArgs._ItemHeight, InArgs._ItemAlignment, TSharedPtr<SHeaderRow>(), InArgs._ExternalScrollbar);
+			this->ConstructChildren(InArgs._ItemWidth, InArgs._ItemHeight, InArgs._ItemAlignment, TSharedPtr<SHeaderRow>(), InArgs._ExternalScrollbar, InArgs._OnTileViewScrolled);
 			if (this->ScrollBar.IsValid())
 			{
 				this->ScrollBar->SetUserVisibility(InArgs._ScrollbarVisibility);
diff --git a/Engine/Source/Runtime/Slate/Public/Widgets/Views/STreeView.h b/Engine/Source/Runtime/Slate/Public/Widgets/Views/STreeView.h
--- a/Engine/Source/Runtime/Slate/Public/Widgets/Views/STreeView.h
+++ b/Engine/Source/Runtime/Slate/Public/Widgets/Views/STreeView.h
@@ -92,6 +92,8 @@ public:
 
 		SLATE_EVENT( FOnGenerateRow, OnGenerateRow )
 
+		SLATE_EVENT( FOnTableViewScrolled, OnTreeViewScrolled )
+
 		SLATE_EVENT( FOnItemScrolledIntoView, OnItemScrolledIntoView )
 
 		SLATE_EVENT( FOnGetChildren, OnGetChildren )
@@ -178,7 +180,7 @@ public:
 		else
 		{
 			// Make the TableView
-			this->ConstructChildren( 0, InArgs._ItemHeight, EListItemAlignment::LeftAligned, InArgs._HeaderRow, InArgs._ExternalScrollbar );
+			this->ConstructChildren( 0, InArgs._ItemHeight, EListItemAlignment::LeftAligned, InArgs._HeaderRow, InArgs._ExternalScrollbar, InArgs._OnTreeViewScrolled );
 		}
 	}
 
