Added a way for STableViewBase (and its derived types) to notify you when they are scrolled.

This allows you, with the combination of SetScrollOffset, to sync two or more list views to have the same scroll offset.

STableViewBase no-longer associates its header with an external scrollbar. This prevents the header from trying to compensate for something that doesn't exist, and causing the header and the rows to become out of line.

diff --git a/Engine/Source/Runtime/Slate/Public/Widgets/Views/SListView.h b/Engine/Source/Runtime/Slate/Public/Widgets/Views/SListView.h
--- a/Engine/Source/Runtime/Slate/Public/Widgets/Views/SListView.h
+++ b/Engine/Source/Runtime/Slate/Public/Widgets/Views/SListView.h
@@ -61,6 +61,8 @@ public:
 
 		SLATE_EVENT( FOnGenerateRow, OnGenerateRow )
 
+		SLATE_EVENT( FOnTableViewScrolled, OnListViewScrolled )
+
 		SLATE_EVENT( FOnItemScrolledIntoView, OnItemScrolledIntoView )
 
 		SLATE_ARGUMENT( const TArray<ItemType>* , ListItemsSource )
@@ -141,7 +143,7 @@ public:
 		else
 		{
 			// Make the TableView
-			ConstructChildren( 0, InArgs._ItemHeight, EListItemAlignment::LeftAligned, InArgs._HeaderRow, InArgs._ExternalScrollbar );
+			ConstructChildren( 0, InArgs._ItemHeight, EListItemAlignment::LeftAligned, InArgs._HeaderRow, InArgs._ExternalScrollbar, InArgs._OnListViewScrolled );
 			if(ScrollBar.IsValid())
 			{
 				ScrollBar->SetUserVisibility(InArgs._ScrollbarVisibility);
@@ -1122,13 +1124,15 @@ protected:
 				if (IndexOfItem < ScrollOffset || IndexPlusOne > (ScrollOffset + NumLiveWidgets))
 				{
 					// Scroll the top of the listview to the item in question
-					ScrollOffset = IndexOfItem;
+					double NewScrollOffset = IndexOfItem;
 					// Center the list view on the item in question.
-					ScrollOffset -= (NumLiveWidgets / 2);
+					NewScrollOffset -= (NumLiveWidgets / 2);
 					//we also don't want the widget being chopped off if it is at the end of the list
-					const double MoveBackBy = FMath::Clamp<double>(IndexPlusOne - (ScrollOffset + NumLiveWidgets), 0, FLT_MAX);
+					const double MoveBackBy = FMath::Clamp<double>(IndexPlusOne - (NewScrollOffset + NumLiveWidgets), 0, FLT_MAX);
 					//Move to the correct center spot
-					ScrollOffset += MoveBackBy;
+					NewScrollOffset += MoveBackBy;
+
+					SetScrollOffset( NewScrollOffset );
 				}
 
 				RequestListRefresh();
