Backing out changelist 2495723 as this is now handled by SHeaderRow rather than SSplitter

diff --git a/Engine/Source/Runtime/Slate/Private/Widgets/Layout/SSplitter.cpp b/Engine/Source/Runtime/Slate/Private/Widgets/Layout/SSplitter.cpp
--- a/Engine/Source/Runtime/Slate/Private/Widgets/Layout/SSplitter.cpp
+++ b/Engine/Source/Runtime/Slate/Private/Widgets/Layout/SSplitter.cpp
@@ -88,8 +88,8 @@ TArray<FLayoutGeometry> SSplitter::ArrangeChildrenForLayout(const FGeometry& All
 	int32 NumNonCollapsedChildren = 0;
 	int32 NumResizeableChildren = 0;
 	float CoefficientTotal = 0;
-	// Some space is claimed by non-proportional elements (auto or fixed-sized elements)
-	float NonProportionalSpace = 0;
+	// Some space is claimed by non-resizeable elements (auto-sized elements)
+	float NonResizeableSpace = 0;
 	{
 		for (int32 ChildIndex=0; ChildIndex < Children.Num(); ++ChildIndex)	
 		{
@@ -99,11 +99,7 @@ TArray<FLayoutGeometry> SSplitter::ArrangeChildrenForLayout(const FGeometry& All
 
 				if ( Children[ChildIndex].SizingRule == SSplitter::SizeToContent )
 				{
-					NonProportionalSpace += Children[ChildIndex].GetWidget()->GetDesiredSize()[AxisIndex];
-				}
-				else if ( Children[ChildIndex].SizingRule == SSplitter::AbsoluteSize )
-				{
-					NonProportionalSpace += Children[ChildIndex].SizeValue.Get();
+					NonResizeableSpace += Children[ChildIndex].GetWidget()->GetDesiredSize()[AxisIndex];
 				}
 				else // SizingRule == SSplitter::FractionOfParent
 				{
@@ -117,7 +113,7 @@ TArray<FLayoutGeometry> SSplitter::ArrangeChildrenForLayout(const FGeometry& All
 
 	// The user-sizeable children must make room for the resize handles and for auto-sized children.
 	const float SpaceNeededForHandles = FMath::Max(0, NumNonCollapsedChildren-1) * PhysicalSplitterHandleSize;
-	const float ProportionalSpace = AllottedGeometry.Size.Component(AxisIndex) - SpaceNeededForHandles - NonProportionalSpace;
+	const float ResizeableSpace = AllottedGeometry.Size.Component(AxisIndex) - SpaceNeededForHandles - NonResizeableSpace;
 
 	// Arrange the children horizontally or vertically.
 	float XOffset = 0;
@@ -125,19 +121,9 @@ TArray<FLayoutGeometry> SSplitter::ArrangeChildrenForLayout(const FGeometry& All
 	{
 		const FSlot& CurSlot = Children[ChildIndex];
 
-		float ChildSpace = 0.0f;
-		if ( CurSlot.SizingRule == SSplitter::SizeToContent )
-		{
-			ChildSpace = CurSlot.GetWidget()->GetDesiredSize()[AxisIndex];
-		}
-		else if ( CurSlot.SizingRule == SSplitter::AbsoluteSize )
-		{
-			ChildSpace = CurSlot.SizeValue.Get();
-		}
-		else // SizingRule == SSplitter::FractionOfParent
-		{
-			ChildSpace =  ProportionalSpace * CurSlot.SizeValue.Get() / CoefficientTotal;
-		}
+		const float ChildSpace = ( CurSlot.SizingRule == SSplitter::SizeToContent )
+			? CurSlot.GetWidget()->GetDesiredSize()[AxisIndex]
+		: ResizeableSpace * CurSlot.SizeValue.Get() / CoefficientTotal;
 
 		const EVisibility ChildVisibility = CurSlot.GetWidget()->GetVisibility();
 
@@ -248,13 +234,13 @@ static FVector2D ComputeDesiredSizeForSplitter( const float PhysicalSplitterHand
 			FVector2D ChildDesiredSize( CurSlot.GetWidget()->GetDesiredSize() );
 			if ( Orientation == Orient_Horizontal )
 			{
-				MyDesiredSize.X += (CurSlot.SizingRule == SSplitter::AbsoluteSize) ? CurSlot.SizeValue.Get() : ChildDesiredSize.X;
+				MyDesiredSize.X += ChildDesiredSize.X;
 				MyDesiredSize.Y = FMath::Max(ChildDesiredSize.Y, MyDesiredSize.Y);
 			}
 			else
 			{
 				MyDesiredSize.X = FMath::Max(ChildDesiredSize.X, MyDesiredSize.X);				
-				MyDesiredSize.Y += (CurSlot.SizingRule == SSplitter::AbsoluteSize) ? CurSlot.SizeValue.Get() : ChildDesiredSize.Y;
+				MyDesiredSize.Y += ChildDesiredSize.Y;
 			}
 		}
 	}
@@ -516,134 +502,114 @@ void SSplitter::HandleResizing( const float PhysicalSplitterHandleSize, const ES
 	float Delta = LocalMousePos.Component(AxisIndex) - HandlePos;
 
 	const int32 SlotBeforeDragHandle = FindResizeableSlotBeforeHandle( DraggedHandle, Children );
-	if ( SlotBeforeDragHandle == INDEX_NONE )
-	{
-		return;
-	}
-
-	FSlot& PrevChild = Children[SlotBeforeDragHandle];
-	const FLayoutGeometry& PrevChildGeom = ChildGeometries[SlotBeforeDragHandle];
-
-	// Compute the new sizes of the slot
-	const float PrevChildLength = PrevChildGeom.GetSizeInParentSpace().Component(AxisIndex);
-	float NewPrevChildLength = ClampChild( PrevChildLength + Delta );
 
-	// Changing the size of an absolutely sized slot doesn't affect the size of any subsequent slots, it just moves them
-	if ( PrevChild.SizingRule == SSplitter::AbsoluteSize )
+	TArray< int32 > SlotsAfterDragHandleIndicies;
+	if ( ResizeMode == ESplitterResizeMode::Fixed )
 	{
-		if (PrevChild.OnSlotResized_Handler.IsBound())
-		{
-			PrevChild.OnSlotResized_Handler.Execute( NewPrevChildLength );
-		}
-		else
+		const int32 SlotAfterDragHandle = FindResizeableSlotAfterHandle( DraggedHandle, Children );
+
+		if ( SlotAfterDragHandle < NumChildren )
 		{
-			PrevChild.SizeValue = NewPrevChildLength;
+			SlotsAfterDragHandleIndicies.Add( SlotAfterDragHandle );
 		}
 	}
-	else
+	else if ( ResizeMode == ESplitterResizeMode::Fill )
 	{
-		// Proportionally sized slots need to redistribute their size accordingly
-		TArray< int32 > SlotsAfterDragHandleIndicies;
-		if ( ResizeMode == ESplitterResizeMode::Fixed )
-		{
-			const int32 SlotAfterDragHandle = FindResizeableSlotAfterHandle( DraggedHandle, Children );
+		FindAllResizeableSlotsAfterHandle( DraggedHandle, Children, /*OUT*/ SlotsAfterDragHandleIndicies );
+	}
 
-			if ( SlotAfterDragHandle < NumChildren )
-			{
-				SlotsAfterDragHandleIndicies.Add( SlotAfterDragHandle );
-			}
-		}
-		else if ( ResizeMode == ESplitterResizeMode::Fill )
+	if ( SlotBeforeDragHandle >= 0 && SlotsAfterDragHandleIndicies.Num() > 0 )
+	{
+		struct FSlotInfo 
 		{
-			FindAllResizeableSlotsAfterHandle( DraggedHandle, Children, /*OUT*/ SlotsAfterDragHandleIndicies );
-		}
+			FSlot* Slot;
+			const FLayoutGeometry* Geometry;
+			float NewSize;
+		};
 
-		if ( SlotsAfterDragHandleIndicies.Num() > 0)
+		TArray< FSlotInfo > SlotsAfterDragHandle;
+		for (int SlotIndex = 0; SlotIndex < SlotsAfterDragHandleIndicies.Num(); SlotIndex++)
 		{
-			struct FSlotInfo 
-			{
-				FSlot* Slot;
-				const FLayoutGeometry* Geometry;
-				float NewSize;
-			};
+			FSlotInfo SlotInfo;
 
-			TArray< FSlotInfo > SlotsAfterDragHandle;
-			for (int SlotIndex = 0; SlotIndex < SlotsAfterDragHandleIndicies.Num(); SlotIndex++)
-			{
-				FSlotInfo SlotInfo;
+			SlotInfo.Slot = &Children[ SlotsAfterDragHandleIndicies[ SlotIndex ] ];
+			SlotInfo.Geometry = &ChildGeometries[ SlotsAfterDragHandleIndicies[ SlotIndex ] ];
+			SlotInfo.NewSize = SlotInfo.Geometry->GetSizeInParentSpace().Component( AxisIndex );
 
-				SlotInfo.Slot = &Children[ SlotsAfterDragHandleIndicies[ SlotIndex ] ];
-				SlotInfo.Geometry = &ChildGeometries[ SlotsAfterDragHandleIndicies[ SlotIndex ] ];
-				SlotInfo.NewSize = SlotInfo.Geometry->GetSizeInParentSpace().Component( AxisIndex );
+			SlotsAfterDragHandle.Add( SlotInfo );
+		}
 
-				SlotsAfterDragHandle.Add( SlotInfo );
-			}
+		// Get references the prev and next children and their layout settings so that we can modify them.
+		FSlot& PrevChild = Children[SlotBeforeDragHandle];
+		const FLayoutGeometry& PrevChildGeom = ChildGeometries[SlotBeforeDragHandle];
 
-			Delta = NewPrevChildLength - PrevChildLength;
+		// Compute the new sizes of the children
+		const float PrevChildLength = PrevChildGeom.GetSizeInParentSpace().Component(AxisIndex);
+		float NewPrevChildLength = ClampChild( PrevChildLength + Delta );
+		Delta = NewPrevChildLength - PrevChildLength;
 
-			// Distribute the Delta across the affected slots after the drag handle
-			float UnusedDelta = Delta;
-			for (int DistributionCount = 0; DistributionCount < SlotsAfterDragHandle.Num() && UnusedDelta != 0; DistributionCount++)
+		// Distribute the Delta across the affected slots after the drag handle
+		float UnusedDelta = Delta;
+		for (int DistributionCount = 0; DistributionCount < SlotsAfterDragHandle.Num() && UnusedDelta != 0; DistributionCount++)
+		{
+			float DividedDelta = UnusedDelta / SlotsAfterDragHandle.Num();
+			UnusedDelta = 0;
+			for (int SlotIndex = 0; SlotIndex < SlotsAfterDragHandle.Num(); SlotIndex++)
 			{
-				float DividedDelta = UnusedDelta / SlotsAfterDragHandle.Num();
-				UnusedDelta = 0;
-				for (int SlotIndex = 0; SlotIndex < SlotsAfterDragHandle.Num(); SlotIndex++)
-				{
-					FSlotInfo& SlotInfo = SlotsAfterDragHandle[ SlotIndex ];
+				FSlotInfo& SlotInfo = SlotsAfterDragHandle[ SlotIndex ];
 
-					float CurrentSize = ClampChild( SlotInfo.Geometry->GetSizeInParentSpace().Component(AxisIndex) );
-					SlotInfo.NewSize = ClampChild( CurrentSize - DividedDelta );
+				float CurrentSize = ClampChild( SlotInfo.Geometry->GetSizeInParentSpace().Component(AxisIndex) );
+				SlotInfo.NewSize = ClampChild( CurrentSize - DividedDelta );
 
-					// If one of the slots couldn't be fully adjusted by the delta due to min/max constraints then
-					// the leftover delta needs to be evenly distributed to all of the other slots
-					UnusedDelta += SlotInfo.NewSize - ( CurrentSize - DividedDelta );
-				}
+				// If one of the slots couldn't be fully adjusted by the delta due to min/max constraints then
+				// the leftover delta needs to be evenly distributed to all of the other slots
+				UnusedDelta += SlotInfo.NewSize - ( CurrentSize - DividedDelta );
 			}
+		}
+
+		Delta = Delta - UnusedDelta;
 
-			Delta = Delta - UnusedDelta;
+		// PrevChildLength needs to be updated: it's value has to take into account the next child's min/max restrictions
+		NewPrevChildLength = ClampChild( PrevChildLength + Delta );
 
-			// PrevChildLength needs to be updated: it's value has to take into account the next child's min/max restrictions
-			NewPrevChildLength = ClampChild( PrevChildLength + Delta );
+		// Cells being resized are both stretch values -> redistribute the stretch coefficients proportionately
+		// to match the new child sizes on the screen.
+		{
+			float TotalLength = NewPrevChildLength;
+			float TotalStretchCoefficients = PrevChild.SizeValue.Get();
 
-			// Cells being resized are both stretch values -> redistribute the stretch coefficients proportionately
-			// to match the new child sizes on the screen.
+			for (int SlotIndex = 0; SlotIndex < SlotsAfterDragHandle.Num(); SlotIndex++)
 			{
-				float TotalLength = NewPrevChildLength;
-				float TotalStretchCoefficients = PrevChild.SizeValue.Get();
+				FSlotInfo SlotInfo = SlotsAfterDragHandle[ SlotIndex ];
 
-				for (int SlotIndex = 0; SlotIndex < SlotsAfterDragHandle.Num(); SlotIndex++)
-				{
-					FSlotInfo SlotInfo = SlotsAfterDragHandle[ SlotIndex ];
+				TotalLength += SlotInfo.NewSize;
+				TotalStretchCoefficients += SlotInfo.Slot->SizeValue.Get();
+			}
 
-					TotalLength += SlotInfo.NewSize;
-					TotalStretchCoefficients += SlotInfo.Slot->SizeValue.Get();
-				}
+			const float NewPrevChildSize = ( TotalStretchCoefficients * NewPrevChildLength / TotalLength );
+
+			if (PrevChild.OnSlotResized_Handler.IsBound())
+			{
+				PrevChild.OnSlotResized_Handler.Execute( NewPrevChildSize );
+			}
+			else
+			{
+				PrevChild.SizeValue = NewPrevChildSize;
+			}
+
+			for (int SlotIndex = 0; SlotIndex < SlotsAfterDragHandle.Num(); SlotIndex++)
+			{
+				FSlotInfo SlotInfo = SlotsAfterDragHandle[ SlotIndex ];
 
-				const float NewPrevChildSize = ( TotalStretchCoefficients * NewPrevChildLength / TotalLength );
+				const float NewNextChildSize = ( TotalStretchCoefficients * SlotInfo.NewSize / TotalLength );
 
-				if (PrevChild.OnSlotResized_Handler.IsBound())
+				if (SlotInfo.Slot->OnSlotResized_Handler.IsBound())
 				{
-					PrevChild.OnSlotResized_Handler.Execute( NewPrevChildSize );
+					SlotInfo.Slot->OnSlotResized_Handler.Execute(NewNextChildSize);
 				}
 				else
 				{
-					PrevChild.SizeValue = NewPrevChildSize;
-				}
-
-				for (int SlotIndex = 0; SlotIndex < SlotsAfterDragHandle.Num(); SlotIndex++)
-				{
-					FSlotInfo SlotInfo = SlotsAfterDragHandle[ SlotIndex ];
-
-					const float NewNextChildSize = ( TotalStretchCoefficients * SlotInfo.NewSize / TotalLength );
-
-					if (SlotInfo.Slot->OnSlotResized_Handler.IsBound())
-					{
-						SlotInfo.Slot->OnSlotResized_Handler.Execute(NewNextChildSize);
-					}
-					else
-					{
-						SlotInfo.Slot->SizeValue = NewNextChildSize;
-					}
+					SlotInfo.Slot->SizeValue = NewNextChildSize;
 				}
 			}
 		}
diff --git a/Engine/Source/Runtime/Slate/Public/Widgets/Layout/SSplitter.h b/Engine/Source/Runtime/Slate/Public/Widgets/Layout/SSplitter.h
--- a/Engine/Source/Runtime/Slate/Public/Widgets/Layout/SSplitter.h
+++ b/Engine/Source/Runtime/Slate/Public/Widgets/Layout/SSplitter.h
@@ -13,7 +13,7 @@ namespace ESplitterResizeMode
 
 class FLayoutGeometry;
 /**
- * SSPlitter divides its allotted area into N segments, where N is the number of children it has.
+ * SSplitter divides its allotted area into N segments, where N is the number of children it has.
  * It allows the users to resize the children along the splitters axis: that is, horizontally or vertically.
  */
 class SLATE_API SSplitter : public SPanel
@@ -26,9 +26,7 @@ public:
 		/** Get the DesiredSize() of the content */
 		SizeToContent,
 		/** Use a fraction of the parent's size */
-		FractionOfParent,
-		/** Use the SizeValue of the slot as an absolute size */
-		AbsoluteSize,
+		FractionOfParent
 	};
 
 	DECLARE_DELEGATE_OneParam(
