Fixed SListView scrolling too far if you requested an item be scrolled into view before it had been Ticked

If you created a SListView and then immediately requested that it scroll item X into view, it would actually scroll item X+1 into view. This was because the logic in SListView::ScrollIntoView couldn't handle the number of live items being 0 (which it will be before the SListView has been fully ticked at least once) and would always overscroll by 1 item.

This change will defer such a request until the next Tick, at which point the SListView will be capable of handling the scroll correctly.

diff --git a/Engine/Source/Runtime/Slate/Private/Widgets/Views/STableViewBase.cpp b/Engine/Source/Runtime/Slate/Private/Widgets/Views/STableViewBase.cpp
--- a/Engine/Source/Runtime/Slate/Private/Widgets/Views/STableViewBase.cpp
+++ b/Engine/Source/Runtime/Slate/Private/Widgets/Views/STableViewBase.cpp
@@ -231,7 +231,7 @@ void STableViewBase::Tick( const FGeometry& AllottedGeometry, const double InCur
 			PanelGeometryLastTick = PanelGeometry;
 
 			// We never create the ItemsPanel if the user did not specify all the parameters required to successfully make a list.
-			ScrollIntoView( PanelGeometry );
+			const EScrollIntoViewResult ScrollIntoViewResult = ScrollIntoView( PanelGeometry );
 
 			const FReGenerateResults ReGenerateResults = ReGenerateItems( PanelGeometry );
 			LastGenerateResults = ReGenerateResults;
@@ -291,12 +291,20 @@ void STableViewBase::Tick( const FGeometry& AllottedGeometry, const double InCur
 				ScrollBar->SetState( OffsetFraction, ThumbSizeFraction );
 			}
 
-			NotifyItemScrolledIntoView();
-
-			bWasAtEndOfList = ScrollBar->DistanceFromBottom() < KINDA_SMALL_NUMBER ? true : false;
+			bWasAtEndOfList = (ScrollBar->DistanceFromBottom() < KINDA_SMALL_NUMBER);
 
 			bItemsNeedRefresh = false;
 			ItemsPanel->SetRefreshPending(false);
+
+			if (ScrollIntoViewResult == EScrollIntoViewResult::Deferred)
+			{
+				// We call this rather than just leave bItemsNeedRefresh as true to ensure that EnsureTickToRefresh is registered
+				RequestListRefresh();
+			}
+			else
+			{
+				NotifyItemScrolledIntoView();
+			}
 		}
 	}
 }
diff --git a/Engine/Source/Runtime/Slate/Public/Widgets/Views/SListView.h b/Engine/Source/Runtime/Slate/Public/Widgets/Views/SListView.h
--- a/Engine/Source/Runtime/Slate/Public/Widgets/Views/SListView.h
+++ b/Engine/Source/Runtime/Slate/Public/Widgets/Views/SListView.h
@@ -1105,7 +1105,7 @@ protected:
 	 * 
 	 * @param ListViewGeometry  The geometry of the listView; can be useful for centering the item.
 	 */
-	virtual void ScrollIntoView( const FGeometry& ListViewGeometry ) override
+	virtual EScrollIntoViewResult ScrollIntoView( const FGeometry& ListViewGeometry ) override
 	{
 		if ( TListTypeTraits<ItemType>::IsPtrValid(ItemToScrollIntoView) && ItemsSource != nullptr )
 		{
@@ -1117,8 +1117,15 @@ protected:
 				{
 					// Use the last number of widgets on screen to estimate if we actually need to scroll.
 					NumLiveWidgets = LastGenerateResults.ExactNumRowsOnScreen;
-				}
 
+					// If we still don't have any widgets, we're not in a situation where we can scroll an item into view
+					// (probably as nothing has been generated yet), so we'll defer this again until the next frame
+					if (NumLiveWidgets == 0)
+					{
+						return EScrollIntoViewResult::Deferred;
+					}
+				}
+				
 				// Only scroll the item into view if it's not already in the visible range
 				const double IndexPlusOne = IndexOfItem+1;
 				if (IndexOfItem < ScrollOffset || IndexPlusOne > (ScrollOffset + NumLiveWidgets))
@@ -1142,6 +1149,8 @@ protected:
 
 			TListTypeTraits<ItemType>::ResetPtr(ItemToScrollIntoView);
 		}
+
+		return EScrollIntoViewResult::Success;
 	}
 
 	virtual void NotifyItemScrolledIntoView() override
diff --git a/Engine/Source/Runtime/Slate/Public/Widgets/Views/STableViewBase.h b/Engine/Source/Runtime/Slate/Public/Widgets/Views/STableViewBase.h
--- a/Engine/Source/Runtime/Slate/Public/Widgets/Views/STableViewBase.h
+++ b/Engine/Source/Runtime/Slate/Public/Widgets/Views/STableViewBase.h
@@ -229,12 +229,20 @@ protected:
 	/** @return how many items there are in the TArray being observed */
 	virtual int32 GetNumItemsBeingObserved() const = 0;
 
+	enum class EScrollIntoViewResult
+	{
+		/** The function scrolled an item (if set) into view (or the item was already in view) */
+		Success,
+		/** The function did not have enough data to scroll the given item into view, so it should be deferred until the next Tick */
+		Deferred,
+	};
+
 	/**
 	 * If there is a pending request to scroll an item into view, do so.
 	 * 
 	 * @param ListViewGeometry  The geometry of the listView; can be useful for centering the item.
 	 */
-	virtual void ScrollIntoView(const FGeometry& ListViewGeometry) = 0;
+	virtual EScrollIntoViewResult ScrollIntoView(const FGeometry& ListViewGeometry) = 0;
 
 	/**
 	 * Called when an item has entered the visible geometry to check to see if the ItemScrolledIntoView delegate should be fired.
