Some minor adjustments to make the MultiLineEditable text conditionally behave more like the single line variant.
#UE-13296 - SMultiLineEditableTextBox does not implement revert on escape
#Branch UE4
#Proj Slate

diff --git a/Engine/Source/Runtime/Slate/Private/Widgets/Input/SMultiLineEditableTextBox.cpp b/Engine/Source/Runtime/Slate/Private/Widgets/Input/SMultiLineEditableTextBox.cpp
--- a/Engine/Source/Runtime/Slate/Private/Widgets/Input/SMultiLineEditableTextBox.cpp
+++ b/Engine/Source/Runtime/Slate/Private/Widgets/Input/SMultiLineEditableTextBox.cpp
@@ -37,16 +37,11 @@ void SMultiLineEditableTextBox::Construct( const FArguments& InArgs )
 	TAttribute<FMargin> Padding = InArgs._Padding.IsSet() ? InArgs._Padding : InArgs._Style->Padding;
 	TAttribute<FMargin> HScrollBarPadding = InArgs._HScrollBarPadding.IsSet() ? InArgs._HScrollBarPadding : InArgs._Style->HScrollBarPadding;
 	TAttribute<FMargin> VScrollBarPadding = InArgs._VScrollBarPadding.IsSet() ? InArgs._VScrollBarPadding : InArgs._Style->VScrollBarPadding;
+	TAttribute<FSlateFontInfo> Font = InArgs._Font.IsSet() ? InArgs._Font : InArgs._Style->Font;
 	TAttribute<FSlateColor> ForegroundColor = InArgs._ForegroundColor.IsSet() ? InArgs._ForegroundColor : InArgs._Style->ForegroundColor;
 	TAttribute<FSlateColor> BackgroundColor = InArgs._BackgroundColor.IsSet() ? InArgs._BackgroundColor : InArgs._Style->BackgroundColor;
 	ReadOnlyForegroundColor = InArgs._ReadOnlyForegroundColor.IsSet() ? InArgs._ReadOnlyForegroundColor : InArgs._Style->ReadOnlyForegroundColor;
 
-	TAttribute<FSlateFontInfo> Font;
-	if (InArgs._Font.IsSet())
-	{
-		Font = InArgs._Font;
-	}
-
 	TSharedPtr<SScrollBar> HScrollBar = InArgs._HScrollBar;
 	if (!HScrollBar.IsValid())
 	{
@@ -100,6 +95,9 @@ void SMultiLineEditableTextBox::Construct( const FArguments& InArgs )
 					.OnCursorMoved( InArgs._OnCursorMoved )
 					.ContextMenuExtender( InArgs._ContextMenuExtender )
 					.Justification(InArgs._Justification)
+					.RevertTextOnEscape(InArgs._RevertTextOnEscape)
+					.SelectAllTextWhenFocused(InArgs._SelectAllTextWhenFocused)
+					.ClearKeyboardFocusOnCommit(InArgs._ClearKeyboardFocusOnCommit)
 					.LineHeightPercentage(InArgs._LineHeightPercentage)
 					.Margin(InArgs._Margin)
 					.WrapTextAt(InArgs._WrapTextAt)
diff --git a/Engine/Source/Runtime/Slate/Private/Widgets/Text/SMultiLineEditableText.cpp b/Engine/Source/Runtime/Slate/Private/Widgets/Text/SMultiLineEditableText.cpp
--- a/Engine/Source/Runtime/Slate/Private/Widgets/Text/SMultiLineEditableText.cpp
+++ b/Engine/Source/Runtime/Slate/Private/Widgets/Text/SMultiLineEditableText.cpp
@@ -269,6 +269,9 @@ void SMultiLineEditableText::Construct( const FArguments& InArgs )
 	OnTextChanged = InArgs._OnTextChanged;
 	OnTextCommitted = InArgs._OnTextCommitted;
 	OnCursorMoved = InArgs._OnCursorMoved;
+	bSelectAllTextWhenFocused = InArgs._SelectAllTextWhenFocused;
+	bClearKeyboardFocusOnCommit = InArgs._ClearKeyboardFocusOnCommit;
+	bRevertTextOnEscape = InArgs._RevertTextOnEscape;
 	OnHScrollBarUserScrolled = InArgs._OnHScrollBarUserScrolled;
 	OnVScrollBarUserScrolled = InArgs._OnVScrollBarUserScrolled;
 
@@ -535,6 +538,14 @@ FReply SMultiLineEditableText::OnFocusReceived( const FGeometry& MyGeometry, con
 				TextInputMethodSystem->ActivateContext(TextInputMethodContext.ToSharedRef());
 			}
 		}
+		// Store undo state to use for escape key reverts
+		MakeUndoState(OriginalText);
+
+		// Select All Text
+		if(SelectAllTextWhenFocused())
+		{
+			SelectAllText();
+		}
 
 		UpdateCursorHighlight();
 
@@ -1758,6 +1769,28 @@ void SMultiLineEditableText::Refresh()
 	}
 }
 
+void SMultiLineEditableText::RestoreOriginalText()
+{
+	if(HasTextChangedFromOriginal())
+	{
+		SaveText(OriginalText.Text);
+
+		// Let outsiders know that the text content has been changed
+		OnTextChanged.ExecuteIfBound(OriginalText.Text);
+	}
+}
+
+bool SMultiLineEditableText::HasTextChangedFromOriginal() const
+{
+	bool bHasChanged = false;
+	if(!IsReadOnly.Get())
+	{
+		const FText EditedText = GetEditableText();
+		bHasChanged = !EditedText.EqualTo(OriginalText.Text);
+	}
+	return bHasChanged;
+}
+
 bool SMultiLineEditableText::CanExecuteSelectAll() const
 {
 	bool bCanExecute = true;
@@ -1908,7 +1941,31 @@ void SMultiLineEditableText::SetHasDragSelectedSinceFocused( bool Value )
 
 FReply SMultiLineEditableText::OnEscape()
 {
-	return FReply::Unhandled();
+	FReply MyReply = FReply::Unhandled();
+
+	if(AnyTextSelected())
+	{
+		// Clear selection
+		ClearSelection();
+		UpdateCursorHighlight();
+		MyReply = FReply::Handled();
+	}
+
+	if(!GetIsReadOnly())
+	{
+		// Restore the text if the revert flag is set
+		if(bRevertTextOnEscape.Get() && HasTextChangedFromOriginal())
+		{
+			RestoreOriginalText();
+			// Release input focus
+			if(bClearKeyboardFocusOnCommit.Get())
+			{
+				FSlateApplication::Get().ClearKeyboardFocus(EFocusCause::Cleared);
+			}
+			MyReply = FReply::Handled();
+		}
+	}
+	return MyReply;
 }
 
 void SMultiLineEditableText::OnEnter()
@@ -1930,6 +1987,12 @@ void SMultiLineEditableText::OnEnter()
 		// Reload underlying value now it is committed  (commit may alter the value) 
 		// so it can be re-displayed in the edit box if it retains focus
 		LoadText();
+		// Release input focus
+		if(bClearKeyboardFocusOnCommit.Get())
+		{
+			ClearSelection();
+			FSlateApplication::Get().ClearKeyboardFocus(EFocusCause::Cleared);
+		}
 	}
 }
 
diff --git a/Engine/Source/Runtime/Slate/Public/Widgets/Text/SMultiLineEditableText.h b/Engine/Source/Runtime/Slate/Public/Widgets/Text/SMultiLineEditableText.h
--- a/Engine/Source/Runtime/Slate/Public/Widgets/Text/SMultiLineEditableText.h
+++ b/Engine/Source/Runtime/Slate/Public/Widgets/Text/SMultiLineEditableText.h
@@ -33,6 +33,9 @@ public:
 		, _IsReadOnly(false)
 		, _OnTextChanged()
 		, _OnTextCommitted()
+		, _SelectAllTextWhenFocused(false)
+		, _RevertTextOnEscape(false)
+		, _ClearKeyboardFocusOnCommit(true)
 		, _OnCursorMoved()
 		, _ContextMenuExtender()
 		, _ModiferKeyForNewLine(EModifierKey::None)
@@ -84,6 +87,15 @@ public:
 		/** Called whenever the text is committed.  This happens when the user presses enter or the text box loses focus. */
 		SLATE_EVENT(FOnTextCommitted, OnTextCommitted)
 
+		/** Whether to select all text when the user clicks to give focus on the widget */
+		SLATE_ATTRIBUTE(bool, SelectAllTextWhenFocused)
+
+		/** Whether to allow the user to back out of changes when they press the escape key */
+		SLATE_ATTRIBUTE(bool, RevertTextOnEscape)
+
+		/** Whether to clear keyboard focus when pressing enter to commit changes */
+		SLATE_ATTRIBUTE(bool, ClearKeyboardFocusOnCommit)
+
 		/** Called whenever the horizontal scrollbar is moved by the user */
 		SLATE_EVENT(FOnUserScrolled, OnHScrollBarUserScrolled)
 
@@ -153,6 +165,12 @@ public:
 	/** Refresh this text box immediately, rather than wait for the usual caching mechanisms to take affect on the text Tick */
 	void Refresh();
 
+	/** Restores the text to the original state */
+	void RestoreOriginalText();
+
+	/** Returns whether the current text varies from the original */
+	bool HasTextChangedFromOriginal() const;
+
 private:
 	
 	/** 
@@ -626,6 +644,12 @@ private:
 	/** Whether to select all text when the user clicks to give focus on the widget */
 	TAttribute< bool > bSelectAllTextWhenFocused;
 
+	/** True if any changes should be reverted if we recieve an escape key */
+	TAttribute< bool > bRevertTextOnEscape;
+
+	/** True if we want the text control to lose focus on an text commit/revert events */
+	TAttribute< bool > bClearKeyboardFocusOnCommit;
+
 	/** True if we're currently selecting text by dragging the mouse cursor with the left button held down */
 	bool bIsDragSelecting;
 
@@ -647,6 +671,9 @@ private:
 	/** Undo state that will be pushed if text is actually changed between calls to StartChangingText() and FinishChangingText() */
 	FUndoState StateBeforeChangingText;
 
+	/** Original text undo state */
+	FUndoState OriginalText;
+
 	/** True if we're currently (potentially) changing the text string */
 	bool bIsChangingText;
 
