Another minor adjustment to ensure selection is cleared when the textbox loses focus.
#UE-13296 - SMultiLineEditableTextBox does not implement revert on escape
#Branch UE4
#Proj Slate

diff --git a/Engine/Source/Runtime/Slate/Private/Widgets/Text/SMultiLineEditableText.cpp b/Engine/Source/Runtime/Slate/Private/Widgets/Text/SMultiLineEditableText.cpp
--- a/Engine/Source/Runtime/Slate/Private/Widgets/Text/SMultiLineEditableText.cpp
+++ b/Engine/Source/Runtime/Slate/Private/Widgets/Text/SMultiLineEditableText.cpp
@@ -607,6 +607,12 @@ void SMultiLineEditableText::OnFocusLost( const FFocusEvent& InFocusEvent )
 		const FText EditedText = GetEditableText();
 
 		OnTextCommitted.ExecuteIfBound(EditedText, TextAction);
+
+		if(bClearKeyboardFocusOnCommit.Get())
+		{
+			ClearSelection();
+		}
+
 		UpdateCursorHighlight();
 
 		// UpdateCursorHighlight always tries to scroll to the cursor, but we don't want that to happen when we 
