Fixed invalid widget path in PushMenu() for 3D widgets
[UE-9940] - Trying to type of a text box that is applied to a 3D widget will result in a crash

#change Passing the current window into ITextEditorWidget::SummonContextMenu() so that the window can be passed to PushMenu(). This stops PushMenu() trying to generate a path to the 3D widget just to obtain the parent window.

diff --git a/Engine/Source/Runtime/Slate/Private/Framework/Text/TextEditHelper.cpp b/Engine/Source/Runtime/Slate/Private/Framework/Text/TextEditHelper.cpp
--- a/Engine/Source/Runtime/Slate/Private/Framework/Text/TextEditHelper.cpp
+++ b/Engine/Source/Runtime/Slate/Private/Framework/Text/TextEditHelper.cpp
@@ -421,7 +421,7 @@ FReply FTextEditHelper::OnMouseButtonUp( const FGeometry& MyGeometry, const FPoi
 			if ( MyGeometry.IsUnderLocation( InMouseEvent.GetScreenSpacePosition() ) )
 			{
 				// Right clicked, so summon a context menu if the cursor is within the widget
-				TextEditor->SummonContextMenu( InMouseEvent.GetScreenSpacePosition() );
+				TextEditor->SummonContextMenu(InMouseEvent.GetScreenSpacePosition(), InMouseEvent.GetWindow());
 			}
 
 			// Release mouse capture
diff --git a/Engine/Source/Runtime/Slate/Private/Widgets/Input/SEditableText.cpp b/Engine/Source/Runtime/Slate/Private/Widgets/Input/SEditableText.cpp
--- a/Engine/Source/Runtime/Slate/Private/Widgets/Input/SEditableText.cpp
+++ b/Engine/Source/Runtime/Slate/Private/Widgets/Input/SEditableText.cpp
@@ -1936,7 +1936,7 @@ void SEditableText::SaveText()
 }
 
 
-void SEditableText::SummonContextMenu( const FVector2D& InLocation )
+void SEditableText::SummonContextMenu(const FVector2D& InLocation, TSharedPtr<SWindow> ParentWindow)
 {
 	// Set the menu to automatically close when the user commits to a choice
 	const bool bShouldCloseWindowAfterMenuSelection = true;
@@ -1984,7 +1984,8 @@ void SEditableText::SummonContextMenu( const FVector2D& InLocation )
 	ActiveContextMenu.PrepareToSummon();
 
 	const bool bFocusImmediately = true;
-	TSharedPtr< SWindow > ContextMenuWindow = FSlateApplication::Get().PushMenu( SharedThis( this ), MenuBuilder.MakeWidget(), InLocation, FPopupTransitionEffect( FPopupTransitionEffect::ContextMenu ), bFocusImmediately );
+	TSharedRef<SWidget> MenuParent = ParentWindow.IsValid() ? StaticCastSharedRef<SWidget>(ParentWindow.ToSharedRef()) : SharedThis(this);
+	TSharedPtr< SWindow > ContextMenuWindow = FSlateApplication::Get().PushMenu(MenuParent, MenuBuilder.MakeWidget(), InLocation, FPopupTransitionEffect(FPopupTransitionEffect::ContextMenu), bFocusImmediately );
 	
 	// Make sure the window is valid.  It's possible for the parent to already be in the destroy queue, for example if the editable text was configured to dismiss it's window during OnTextCommitted.
 	if( ContextMenuWindow.IsValid() )
diff --git a/Engine/Source/Runtime/Slate/Private/Widgets/Text/SMultiLineEditableText.cpp b/Engine/Source/Runtime/Slate/Private/Widgets/Text/SMultiLineEditableText.cpp
--- a/Engine/Source/Runtime/Slate/Private/Widgets/Text/SMultiLineEditableText.cpp
+++ b/Engine/Source/Runtime/Slate/Private/Widgets/Text/SMultiLineEditableText.cpp
@@ -2305,7 +2305,7 @@ TSharedRef< SWidget > SMultiLineEditableText::GetWidget()
 	return SharedThis( this );
 }
 
-void SMultiLineEditableText::SummonContextMenu(const FVector2D& InLocation)
+void SMultiLineEditableText::SummonContextMenu(const FVector2D& InLocation, TSharedPtr<SWindow> ParentWindow)
 {
 	// Set the menu to automatically close when the user commits to a choice
 	const bool bShouldCloseWindowAfterMenuSelection = true;
@@ -2353,7 +2353,8 @@ void SMultiLineEditableText::SummonContextMenu(const FVector2D& InLocation)
 	ActiveContextMenu.PrepareToSummon();
 
 	const bool bFocusImmediately = true;
-	TSharedPtr< SWindow > ContextMenuWindow = FSlateApplication::Get().PushMenu(SharedThis(this), MenuBuilder.MakeWidget(), InLocation, FPopupTransitionEffect(FPopupTransitionEffect::ContextMenu), bFocusImmediately);
+	TSharedRef<SWidget> MenuParent = ParentWindow.IsValid() ? StaticCastSharedRef<SWidget>(ParentWindow.ToSharedRef()) : SharedThis(this);
+	TSharedPtr< SWindow > ContextMenuWindow = FSlateApplication::Get().PushMenu(MenuParent, MenuBuilder.MakeWidget(), InLocation, FPopupTransitionEffect(FPopupTransitionEffect::ContextMenu), bFocusImmediately);
 
 	// Make sure the window is valid.  It's possible for the parent to already be in the destroy queue, for example if the editable text was configured to dismiss it's window during OnTextCommitted.
 	if (ContextMenuWindow.IsValid())
diff --git a/Engine/Source/Runtime/Slate/Public/Widgets/Input/SEditableText.h b/Engine/Source/Runtime/Slate/Public/Widgets/Input/SEditableText.h
--- a/Engine/Source/Runtime/Slate/Public/Widgets/Input/SEditableText.h
+++ b/Engine/Source/Runtime/Slate/Public/Widgets/Input/SEditableText.h
@@ -312,7 +312,7 @@ public:
 	virtual void Undo() override;
 	virtual void Redo() override;
 	virtual TSharedRef< SWidget > GetWidget() override;
-	virtual void SummonContextMenu( const FVector2D& InLocation ) override;
+	virtual void SummonContextMenu(const FVector2D& InLocation, TSharedPtr<SWindow> ParentWindow = TSharedPtr<SWindow>()) override;
 	virtual void LoadText() override;
 	// END ITextEditorWidget interface
 
diff --git a/Engine/Source/Runtime/Slate/Public/Widgets/Text/ITextEditorWidget.h b/Engine/Source/Runtime/Slate/Public/Widgets/Text/ITextEditorWidget.h
--- a/Engine/Source/Runtime/Slate/Public/Widgets/Text/ITextEditorWidget.h
+++ b/Engine/Source/Runtime/Slate/Public/Widgets/Text/ITextEditorWidget.h
@@ -307,7 +307,7 @@ class SLATE_API ITextEditorWidget
 
 	virtual TSharedRef< SWidget > GetWidget() = 0;
 
-	virtual void SummonContextMenu( const FVector2D& InLocation ) = 0;
+	virtual void SummonContextMenu(const FVector2D& InLocation, TSharedPtr<SWindow> ParentWindow = TSharedPtr<SWindow>()) = 0;
 
 	virtual void LoadText() = 0;
 };
diff --git a/Engine/Source/Runtime/Slate/Public/Widgets/Text/SMultiLineEditableText.h b/Engine/Source/Runtime/Slate/Public/Widgets/Text/SMultiLineEditableText.h
--- a/Engine/Source/Runtime/Slate/Public/Widgets/Text/SMultiLineEditableText.h
+++ b/Engine/Source/Runtime/Slate/Public/Widgets/Text/SMultiLineEditableText.h
@@ -425,7 +425,7 @@ private:
 	virtual void Undo() override;
 	virtual void Redo() override;
 	virtual TSharedRef< SWidget > GetWidget() override;
-	virtual void SummonContextMenu( const FVector2D& InLocation ) override;
+	virtual void SummonContextMenu(const FVector2D& InLocation, TSharedPtr<SWindow> ParentWindow = TSharedPtr<SWindow>()) override;
 	virtual void LoadText() override;
 	// END ITextEditorWidget interface
 
