Fix Clang build error

diff --git a/Engine/Source/Runtime/Slate/Private/Widgets/Input/SEditableText.cpp b/Engine/Source/Runtime/Slate/Private/Widgets/Input/SEditableText.cpp
--- a/Engine/Source/Runtime/Slate/Private/Widgets/Input/SEditableText.cpp
+++ b/Engine/Source/Runtime/Slate/Private/Widgets/Input/SEditableText.cpp
@@ -1984,7 +1984,11 @@ void SEditableText::SummonContextMenu(const FVector2D& InLocation, TSharedPtr<SW
 	ActiveContextMenu.PrepareToSummon();
 
 	const bool bFocusImmediately = true;
-	TSharedRef<SWidget> MenuParent = ParentWindow.IsValid() ? StaticCastSharedRef<SWidget>(ParentWindow.ToSharedRef()) : SharedThis(this);
+	TSharedRef<SWidget> MenuParent = SharedThis(this);
+	if (ParentWindow.IsValid())
+	{
+		MenuParent = StaticCastSharedRef<SWidget>(ParentWindow.ToSharedRef());
+	}
 	TSharedPtr< SWindow > ContextMenuWindow = FSlateApplication::Get().PushMenu(MenuParent, MenuBuilder.MakeWidget(), InLocation, FPopupTransitionEffect(FPopupTransitionEffect::ContextMenu), bFocusImmediately );
 	
 	// Make sure the window is valid.  It's possible for the parent to already be in the destroy queue, for example if the editable text was configured to dismiss it's window during OnTextCommitted.
diff --git a/Engine/Source/Runtime/Slate/Private/Widgets/Text/SMultiLineEditableText.cpp b/Engine/Source/Runtime/Slate/Private/Widgets/Text/SMultiLineEditableText.cpp
--- a/Engine/Source/Runtime/Slate/Private/Widgets/Text/SMultiLineEditableText.cpp
+++ b/Engine/Source/Runtime/Slate/Private/Widgets/Text/SMultiLineEditableText.cpp
@@ -2353,7 +2353,11 @@ void SMultiLineEditableText::SummonContextMenu(const FVector2D& InLocation, TSha
 	ActiveContextMenu.PrepareToSummon();
 
 	const bool bFocusImmediately = true;
-	TSharedRef<SWidget> MenuParent = ParentWindow.IsValid() ? StaticCastSharedRef<SWidget>(ParentWindow.ToSharedRef()) : SharedThis(this);
+	TSharedRef<SWidget> MenuParent = SharedThis(this);
+	if (ParentWindow.IsValid())
+	{
+		MenuParent = StaticCastSharedRef<SWidget>(ParentWindow.ToSharedRef());
+	}
 	TSharedPtr< SWindow > ContextMenuWindow = FSlateApplication::Get().PushMenu(MenuParent, MenuBuilder.MakeWidget(), InLocation, FPopupTransitionEffect(FPopupTransitionEffect::ContextMenu), bFocusImmediately);
 
 	// Make sure the window is valid.  It's possible for the parent to already be in the destroy queue, for example if the editable text was configured to dismiss it's window during OnTextCommitted.
