Slate - Fixing STableRow to properly return the handled result on the new OnAcceptDrop handlers.

diff --git a/Engine/Source/Runtime/Slate/Public/Widgets/Views/STableRow.h b/Engine/Source/Runtime/Slate/Public/Widgets/Views/STableRow.h
--- a/Engine/Source/Runtime/Slate/Public/Widgets/Views/STableRow.h
+++ b/Engine/Source/Runtime/Slate/Public/Widgets/Views/STableRow.h
@@ -647,6 +647,8 @@ public:
 						// Expand the drop target just in case, so that what we dropped is visible.
 						OwnerWidget->Private_SetItemExpansion(*MyItem, true);
 					}
+
+					return DropReply;
 				}
 			}
 
