Exposed an option SlateApplication to close all existing windows.

Merging using UE4-To-UE4-LauncherDev

diff --git a/Engine/Source/Runtime/Slate/Private/Framework/Application/SlateApplication.cpp b/Engine/Source/Runtime/Slate/Private/Framework/Application/SlateApplication.cpp
--- a/Engine/Source/Runtime/Slate/Private/Framework/Application/SlateApplication.cpp
+++ b/Engine/Source/Runtime/Slate/Private/Framework/Application/SlateApplication.cpp
@@ -2227,25 +2227,29 @@ FModifierKeysState FSlateApplication::GetModifierKeys() const
 
 
 void FSlateApplication::OnShutdown()
+{
+	CloseAllWindowsImmediately();
+}
+
+void FSlateApplication::CloseAllWindowsImmediately()
 {
 	// Clean up our tooltip window
-	TSharedPtr< SWindow > PinnedToolTipWindow( ToolTipWindow.Pin() );
-	if( PinnedToolTipWindow.IsValid() )
+	TSharedPtr< SWindow > PinnedToolTipWindow(ToolTipWindow.Pin());
+	if (PinnedToolTipWindow.IsValid())
 	{
 		PinnedToolTipWindow->RequestDestroyWindow();
 		ToolTipWindow.Reset();
 	}
 
-	for( int32 WindowIndex = 0; WindowIndex < SlateWindows.Num(); ++WindowIndex )
+	for (int32 WindowIndex = 0; WindowIndex < SlateWindows.Num(); ++WindowIndex)
 	{
 		// Destroy all top level windows.  This will also request that all children of each window be destroyed
-		RequestDestroyWindow( SlateWindows[WindowIndex] );
+		RequestDestroyWindow(SlateWindows[WindowIndex]);
 	}
 
 	DestroyWindowsImmediately();
 }
 
-
 void FSlateApplication::DestroyWindowsImmediately()
 {
 	// Destroy any windows that were queued for deletion.
diff --git a/Engine/Source/Runtime/Slate/Public/Framework/Application/SlateApplication.h b/Engine/Source/Runtime/Slate/Public/Framework/Application/SlateApplication.h
--- a/Engine/Source/Runtime/Slate/Public/Framework/Application/SlateApplication.h
+++ b/Engine/Source/Runtime/Slate/Public/Framework/Application/SlateApplication.h
@@ -851,6 +851,9 @@ public:
 	/** Called when the slate application is being shut down. */
 	void OnShutdown();
 
+	/** Closes all active windows immediately */
+	void CloseAllWindowsImmediately();
+
 	/**
 	 * Destroy the native and slate windows in the array provided.
 	 *
