Fixed some issues with EListItemAlignment in STileView

Items were always being padded, even when not using EListItemAlignment::EvenlyDistributed. Also EListItemAlignment::CenterAligned and EListItemAlignment::RightAligned were applying too much line padding as bIsNewLine was never being cleared.

diff --git a/Engine/Source/Runtime/Slate/Private/Widgets/Views/SListPanel.cpp b/Engine/Source/Runtime/Slate/Private/Widgets/Views/SListPanel.cpp
--- a/Engine/Source/Runtime/Slate/Private/Widgets/Views/SListPanel.cpp
+++ b/Engine/Source/Runtime/Slate/Private/Widgets/Views/SListPanel.cpp
@@ -88,6 +88,8 @@ void SListPanel::OnArrangeChildren( const FGeometry& AllottedGeometry, FArranged
 						WidthSoFar += HalfLinePadding;
 					}
 				}
+
+				bIsNewLine = false;
 			}
 
 			ArrangedChildren.AddWidget(
@@ -234,12 +236,12 @@ float SListPanel::GetItemPadding(const FGeometry& AllottedGeometry, const EListI
 {
 	const float LocalItemWidth = GetDesiredItemWidth();
 	const int32 NumItemsWide = LocalItemWidth > 0 ? FMath::FloorToInt(AllottedGeometry.Size.X / LocalItemWidth) : 0;
-	const bool IsFillAlignment = ListItemAlignment == EListItemAlignment::Fill;
+	const bool IsEvenlyDistributedAlignment = ListItemAlignment == EListItemAlignment::EvenlyDistributed;
 
 	// Only add padding between items if we have more total items that we can fit on a single row.  Otherwise,
 	// the padding around items would continue to change proportionately with the (ample) free horizontal space
 	float Padding = 0.0f;
-	if (!IsFillAlignment && NumItemsWide > 0 && Children.Num() > NumItemsWide)
+	if (IsEvenlyDistributedAlignment && NumItemsWide > 0 && Children.Num() > NumItemsWide)
 	{
 		// Subtract a tiny amount from the available width to avoid floating point precision problems when arranging children
 		static const float FloatingPointPrecisionOffset = 0.001f;
