Added an option to create Slate window without a close button and used it for Source Control Login window, replacing a custom title bar it was using that didn't look well on Mac

diff --git a/Engine/Source/Developer/SourceControl/Private/SSourceControlLogin.cpp b/Engine/Source/Developer/SourceControl/Private/SSourceControlLogin.cpp
--- a/Engine/Source/Developer/SourceControl/Private/SSourceControlLogin.cpp
+++ b/Engine/Source/Developer/SourceControl/Private/SSourceControlLogin.cpp
@@ -14,47 +14,6 @@
 
 #define LOCTEXT_NAMESPACE "SSourceControlLogin"
 
-class SSourceControlTitleBar : public SBorder
-{
-public:
-	SLATE_BEGIN_ARGS(SSourceControlTitleBar) {}
-
-	/** A reference to the parent window */
-	SLATE_ARGUMENT(TSharedPtr<SWindow>, ParentWindow)
-
-	SLATE_END_ARGS()
-
-	void Construct(const FArguments& InArgs)
-	{
-		ParentWindowPtr = InArgs._ParentWindow;
-
-		SBorder::Construct(SBorder::FArguments()
-			.BorderImage(FEditorStyle::GetBrush("Window.Title.Active"))
-			[
-				SNew(SHorizontalBox)
-				.Visibility( EVisibility::HitTestInvisible )
-				+SHorizontalBox::Slot()
-				.HAlign(HAlign_Center)
-				[
-					SNew(STextBlock)
-					.Text( LOCTEXT("SourceControlLoginTitle", "Source Control Login") )
-					.TextStyle(FEditorStyle::Get(), "Window.TitleText")
-					.Visibility( EVisibility::HitTestInvisible )
-				]
-			]
-		);
-	}
-
-	virtual EWindowZone::Type GetWindowZoneOverride() const override
-	{
-		return EWindowZone::TitleBar;
-	}
-
-private:
-	/** The parent window of this widget */
-	TWeakPtr<SWindow> ParentWindowPtr;
-};
-
 const float SSourceControlLogin::RefreshFrequency = 0.5f;
 
 void SSourceControlLogin::Construct(const FArguments& InArgs)
@@ -80,18 +39,6 @@ void SSourceControlLogin::Construct(const FArguments& InArgs)
 		[
 			SNew(SVerticalBox)
 			+SVerticalBox::Slot()
-			.AutoHeight()
-			[
-				SNew(SHorizontalBox)
-				+SHorizontalBox::Slot()
-				.FillWidth(1.0f)
-				.Padding(0.0f, 0.0f, 0.0f, 10.0f)
-				[
-					SNew(SSourceControlTitleBar)
-					.ParentWindow(InArgs._ParentWindow)
-				]
-			]
-			+SVerticalBox::Slot()
 			.FillHeight(1.0f)
 			.Padding(0.0f, 6.0f, 0.0f, 0.0f)
 			[
diff --git a/Engine/Source/Developer/SourceControl/Private/SourceControlModule.cpp b/Engine/Source/Developer/SourceControl/Private/SourceControlModule.cpp
--- a/Engine/Source/Developer/SourceControl/Private/SourceControlModule.cpp
+++ b/Engine/Source/Developer/SourceControl/Private/SourceControlModule.cpp
@@ -117,9 +117,9 @@ void FSourceControlModule::ShowLoginDialog(const FSourceControlLoginClosed& InOn
 		// Create the window
 		SourceControlLoginWindowPtr = SNew(SWindow)
 			.Title( LOCTEXT("SourceControlLoginTitle", "Source Control Login") )
+			.HasCloseButton(false)
 			.SupportsMaximize(false) 
 			.SupportsMinimize(false)
-			.CreateTitleBar(false)
 			.SizingRule( ESizingRule::Autosized );
 
 		// Set the closed callback
diff --git a/Engine/Source/Runtime/Core/Private/Mac/MacWindow.cpp b/Engine/Source/Runtime/Core/Private/Mac/MacWindow.cpp
--- a/Engine/Source/Runtime/Core/Private/Mac/MacWindow.cpp
+++ b/Engine/Source/Runtime/Core/Private/Mac/MacWindow.cpp
@@ -58,7 +58,10 @@ void FMacWindow::Initialize( FMacApplication* const Application, const TSharedRe
 	uint32 WindowStyle = 0;
 	if( Definition->IsRegularWindow )
 	{
-		WindowStyle = NSClosableWindowMask;
+		if( Definition->HasCloseButton )
+		{
+			WindowStyle = NSClosableWindowMask;
+		}
 		
 		// In order to support rounded, shadowed windows set the window to be
 		// titled - we'll set the OpenGL view to cover the whole window
diff --git a/Engine/Source/Runtime/Core/Public/GenericPlatform/GenericWindowDefinition.h b/Engine/Source/Runtime/Core/Public/GenericPlatform/GenericWindowDefinition.h
--- a/Engine/Source/Runtime/Core/Public/GenericPlatform/GenericWindowDefinition.h
+++ b/Engine/Source/Runtime/Core/Public/GenericPlatform/GenericWindowDefinition.h
@@ -44,10 +44,11 @@ struct CORE_API FGenericWindowDefinition
 	bool AcceptsInput;
 	bool ActivateWhenFirstShown;
 
+	bool HasCloseButton;
 	bool SupportsMinimize;
 	bool SupportsMaximize;
 
-    bool IsModalWindow;
+	bool IsModalWindow;
 	bool IsRegularWindow;
 	bool HasSizingFrame;
 	bool SizeWillChangeOften;
diff --git a/Engine/Source/Runtime/Slate/Private/Framework/Application/SlateApplication.cpp b/Engine/Source/Runtime/Slate/Private/Framework/Application/SlateApplication.cpp
--- a/Engine/Source/Runtime/Slate/Private/Framework/Application/SlateApplication.cpp
+++ b/Engine/Source/Runtime/Slate/Private/Framework/Application/SlateApplication.cpp
@@ -1592,6 +1592,7 @@ TSharedRef< FGenericWindow > FSlateApplication::MakeWindow( TSharedRef<SWindow>
 	Definition->AcceptsInput = InSlateWindow->AcceptsInput();
 	Definition->ActivateWhenFirstShown = InSlateWindow->ActivateWhenFirstShown();
 
+	Definition->HasCloseButton = InSlateWindow->HasCloseBox();
 	Definition->SupportsMinimize = InSlateWindow->HasMinimizeBox();
 	Definition->SupportsMaximize = InSlateWindow->HasMaximizeBox();
 
diff --git a/Engine/Source/Runtime/Slate/Public/Framework/Application/SWindowTitleBar.h b/Engine/Source/Runtime/Slate/Public/Framework/Application/SWindowTitleBar.h
--- a/Engine/Source/Runtime/Slate/Public/Framework/Application/SWindowTitleBar.h
+++ b/Engine/Source/Runtime/Slate/Public/Framework/Application/SWindowTitleBar.h
@@ -165,46 +165,52 @@ protected:
 			return;
 		}
 
-		MinimizeButton = SNew(SButton)
-				.IsFocusable(false)
-				.IsEnabled(OwnerWindow->HasMinimizeBox())
-				.ContentPadding(0)
-				.OnClicked(this, &SWindowTitleBar::MinimizeButton_OnClicked)
-				.Cursor(EMouseCursor::Default)
-				.ButtonStyle(FCoreStyle::Get(), "NoBorder")
-				[
-					SNew(SImage)
-						.Image(this, &SWindowTitleBar::GetMinimizeImage)
-						.ColorAndOpacity(this, &SWindowTitleBar::GetWindowTitleContentColor)
-				]
-			;
-
-		MaximizeRestoreButton = SNew(SButton)
-				.IsFocusable(false)
-				.IsEnabled(OwnerWindow->HasMaximizeBox())
-				.ContentPadding(0.0f)
-				.OnClicked(this, &SWindowTitleBar::MaximizeRestoreButton_OnClicked)
-				.Cursor(EMouseCursor::Default)
-				.ButtonStyle(FCoreStyle::Get(), "NoBorder")
-				[
-					SNew(SImage)
-						.Image(this, &SWindowTitleBar::GetMaximizeRestoreImage)
-						.ColorAndOpacity(this, &SWindowTitleBar::GetWindowTitleContentColor)
-				]
-			;
-
-		CloseButton = SNew(SButton)
-				.IsFocusable(false)
-				.ContentPadding(0.0f)
-				.OnClicked(this, &SWindowTitleBar::CloseButton_OnClicked)
-				.Cursor(EMouseCursor::Default)
-				.ButtonStyle(FCoreStyle::Get(), "NoBorder")
-				[
-					SNew(SImage)
-						.Image(this, &SWindowTitleBar::GetCloseImage)
-						.ColorAndOpacity(this, &SWindowTitleBar::GetWindowTitleContentColor)
-				]
-			;
+		const bool bHasWindowButtons = OwnerWindow->HasCloseBox() || OwnerWindow->HasMinimizeBox() || OwnerWindow->HasMaximizeBox();
+
+		if (bHasWindowButtons)
+		{
+			MinimizeButton = SNew(SButton)
+					.IsFocusable(false)
+					.IsEnabled(OwnerWindow->HasMinimizeBox())
+					.ContentPadding(0)
+					.OnClicked(this, &SWindowTitleBar::MinimizeButton_OnClicked)
+					.Cursor(EMouseCursor::Default)
+					.ButtonStyle(FCoreStyle::Get(), "NoBorder")
+					[
+						SNew(SImage)
+							.Image(this, &SWindowTitleBar::GetMinimizeImage)
+							.ColorAndOpacity(this, &SWindowTitleBar::GetWindowTitleContentColor)
+					]
+				;
+
+			MaximizeRestoreButton = SNew(SButton)
+					.IsFocusable(false)
+					.IsEnabled(OwnerWindow->HasMaximizeBox())
+					.ContentPadding(0.0f)
+					.OnClicked(this, &SWindowTitleBar::MaximizeRestoreButton_OnClicked)
+					.Cursor(EMouseCursor::Default)
+					.ButtonStyle(FCoreStyle::Get(), "NoBorder")
+					[
+						SNew(SImage)
+							.Image(this, &SWindowTitleBar::GetMaximizeRestoreImage)
+							.ColorAndOpacity(this, &SWindowTitleBar::GetWindowTitleContentColor)
+					]
+				;
+
+			CloseButton = SNew(SButton)
+					.IsFocusable(false)
+					.IsEnabled(OwnerWindow->HasCloseBox())
+					.ContentPadding(0.0f)
+					.OnClicked(this, &SWindowTitleBar::CloseButton_OnClicked)
+					.Cursor(EMouseCursor::Default)
+					.ButtonStyle(FCoreStyle::Get(), "NoBorder")
+					[
+						SNew(SImage)
+							.Image(this, &SWindowTitleBar::GetCloseImage)
+							.ColorAndOpacity(this, &SWindowTitleBar::GetWindowTitleContentColor)
+					]
+				;
+		}
 
 #if PLATFORM_MAC
 
@@ -215,44 +221,51 @@ protected:
 #else // PLATFORM_MAC
 
 		// Windows UI layout
-		if (ShowAppIcon)
+		if (ShowAppIcon && bHasWindowButtons)
 		{
 			OutLeftContent = SNew(SAppIconWidget)
 				.IconColorAndOpacity(this, &SWindowTitleBar::GetWindowTitleContentColor);
 		}
 		else
 		{
-			OutLeftContent = SNullWidget::NullWidget;
+			OutLeftContent = SNew(SSpacer);
 		}
 
-		OutRightContent = SNew(SBox)
-			.Visibility(EVisibility::SelfHitTestInvisible)
-			.Padding(FMargin(2.0f, 0.0f, 0.0f, 0.0f))
-			[
-				// Minimize
-				SNew(SHorizontalBox)
-					.Visibility(EVisibility::SelfHitTestInvisible)
-
-				+ SHorizontalBox::Slot()
-					.AutoWidth()
-					[
-						MinimizeButton.ToSharedRef()
-					]
-
-				// Maximize/Restore
-				+ SHorizontalBox::Slot()
-					.AutoWidth()
-					[
-						MaximizeRestoreButton.ToSharedRef()
-					]
+		if (bHasWindowButtons)
+		{
+			OutRightContent = SNew(SBox)
+				.Visibility(EVisibility::SelfHitTestInvisible)
+				.Padding(FMargin(2.0f, 0.0f, 0.0f, 0.0f))
+				[
+					// Minimize
+					SNew(SHorizontalBox)
+						.Visibility(EVisibility::SelfHitTestInvisible)
 
-				// Close button
-				+ SHorizontalBox::Slot()
-					.AutoWidth()
-					[
-						CloseButton.ToSharedRef()
-					]
-			];
+					+ SHorizontalBox::Slot()
+						.AutoWidth()
+						[
+							MinimizeButton.ToSharedRef()
+						]
+
+					// Maximize/Restore
+					+ SHorizontalBox::Slot()
+						.AutoWidth()
+						[
+							MaximizeRestoreButton.ToSharedRef()
+						]
+
+					// Close button
+					+ SHorizontalBox::Slot()
+						.AutoWidth()
+						[
+							CloseButton.ToSharedRef()
+						]
+				];
+		}
+		else
+		{
+			OutRightContent = SNew(SSpacer);
+		}
 
 #endif // PLATFORM_MAC
 	}
@@ -289,6 +302,12 @@ protected:
 				];
 		}
 
+		// Adjust the center content alignment if needed. Windows without any title bar buttons look better if the title is centered.
+		if (LeftContent == SNullWidget::NullWidget && RightContent == SNullWidget::NullWidget && CenterContentAlignment == EHorizontalAlignment::HAlign_Left)
+		{
+			CenterContentAlignment = EHorizontalAlignment::HAlign_Center;
+		}
+
 		// calculate content dimensions
 		LeftContent->SlatePrepass();
 		RightContent->SlatePrepass();
diff --git a/Engine/Source/Runtime/SlateCore/Private/Widgets/SWindow.cpp b/Engine/Source/Runtime/SlateCore/Private/Widgets/SWindow.cpp
--- a/Engine/Source/Runtime/SlateCore/Private/Widgets/SWindow.cpp
+++ b/Engine/Source/Runtime/SlateCore/Private/Widgets/SWindow.cpp
@@ -193,6 +193,7 @@ void SWindow::Construct(const FArguments& InArgs)
 	this->bFocusWhenFirstShown = InArgs._FocusWhenFirstShown;
 	this->bActivateWhenFirstShown = InArgs._ActivateWhenFirstShown;
 	this->bHasOSWindowBorder = InArgs._UseOSWindowBorder;
+	this->bHasCloseButton = InArgs._HasCloseButton;
 	this->bHasMinimizeButton = InArgs._SupportsMinimize;
 	this->bHasMaximizeButton = InArgs._SupportsMaximize;
 	this->bHasSizingFrame = !InArgs._IsPopupWindow && InArgs._SizingRule == ESizingRule::UserSized;
@@ -1247,6 +1248,12 @@ bool SWindow::HasSizingFrame() const
 	return bHasSizingFrame;
 }
 
+/** @return true if this window has a close button/box on the titlebar area */
+bool SWindow::HasCloseBox() const
+{
+	return bHasCloseButton;
+}
+
 /** @return true if this window has a maximize button/box on the titlebar area */
 bool SWindow::HasMaximizeBox() const
 {
@@ -1630,6 +1637,7 @@ SWindow::SWindow()
 	, bFocusWhenFirstShown(true)
 	, bActivateWhenFirstShown(true)
 	, bHasOSWindowBorder( false )
+	, bHasCloseButton( false )
 	, bHasMinimizeButton( false )
 	, bHasMaximizeButton( false )
 	, bHasSizingFrame( false )
diff --git a/Engine/Source/Runtime/SlateCore/Public/Widgets/SWindow.h b/Engine/Source/Runtime/SlateCore/Public/Widgets/SWindow.h
--- a/Engine/Source/Runtime/SlateCore/Public/Widgets/SWindow.h
+++ b/Engine/Source/Runtime/SlateCore/Public/Widgets/SWindow.h
@@ -91,6 +91,7 @@ public:
 		, _FocusWhenFirstShown( true )
 		, _ActivateWhenFirstShown( true )
 		, _UseOSWindowBorder( false )
+		, _HasCloseButton( true )
 		, _SupportsMaximize( true )
 		, _SupportsMinimize( true )
 		, _CreateTitleBar( true )
@@ -142,6 +143,9 @@ public:
 		/** Use the default os look for the border of the window */
 		SLATE_ARGUMENT( bool, UseOSWindowBorder )
 
+		/** Does this window have a close button? */
+		SLATE_ARGUMENT( bool, HasCloseButton )
+
 		/** Can this window be maximized? */
 		SLATE_ARGUMENT( bool, SupportsMaximize )
 		
@@ -558,6 +562,9 @@ public:
 	/** @return true if this is a user-sized window with a thick edge */
 	bool HasSizingFrame() const;
 
+	/** @return true if this window has a close button/box on the titlebar area */
+	bool HasCloseBox() const;
+
 	/** @return true if this window has a maximize button/box on the titlebar area */
 	bool HasMaximizeBox() const;
 
@@ -783,7 +790,10 @@ protected:
 
 	/** True if this window displays the os window border instead of drawing one in slate */
 	bool bHasOSWindowBorder : 1;
-	
+
+	/** True if this window displays an enabled close button on the toolbar area */
+	bool bHasCloseButton : 1;
+
 	/** True if this window displays an enabled minimize button on the toolbar area */
 	bool bHasMinimizeButton : 1;
 
