Added Font and other appearance properties to the property sync methods of the TextBox and MultiLineTextBox UMG widgets. Updates these properly in the designer.
[UE-11833] - The font inside of the Text box does not update when changed until the user selects another tab

diff --git a/Engine/Source/Runtime/Slate/Private/Widgets/Input/SEditableText.cpp b/Engine/Source/Runtime/Slate/Private/Widgets/Input/SEditableText.cpp
--- a/Engine/Source/Runtime/Slate/Private/Widgets/Input/SEditableText.cpp
+++ b/Engine/Source/Runtime/Slate/Private/Widgets/Input/SEditableText.cpp
@@ -2151,3 +2151,33 @@ void SEditableText::SetColorAndOpacity(TAttribute<FSlateColor> Color)
 {
 	ColorAndOpacity = Color;
 }
+
+void SEditableText::SetMinDesiredWidth(const TAttribute<float>& InMinDesiredWidth)
+{
+	MinDesiredWidth = InMinDesiredWidth;
+}
+
+void SEditableText::SetIsCaretMovedWhenGainFocus(const TAttribute<bool>& InIsCaretMovedWhenGainFocus)
+{
+	IsCaretMovedWhenGainFocus = InIsCaretMovedWhenGainFocus;
+}
+
+void SEditableText::SetSelectAllTextWhenFocused(const TAttribute<bool>& InSelectAllTextWhenFocused)
+{
+	bSelectAllTextWhenFocused = InSelectAllTextWhenFocused;
+}
+
+void SEditableText::SetRevertTextOnEscape(const TAttribute<bool>& InRevertTextOnEscape)
+{
+	RevertTextOnEscape = InRevertTextOnEscape;
+}
+
+void SEditableText::SetClearKeyboardFocusOnCommit(const TAttribute<bool>& InClearKeyboardFocusOnCommit)
+{
+	ClearKeyboardFocusOnCommit = InClearKeyboardFocusOnCommit;
+}
+
+void SEditableText::SetSelectAllTextOnCommit(const TAttribute<bool>& InSelectAllTextOnCommit)
+{
+	SelectAllTextOnCommit = InSelectAllTextOnCommit;
+}
diff --git a/Engine/Source/Runtime/Slate/Private/Widgets/Input/SEditableTextBox.cpp b/Engine/Source/Runtime/Slate/Private/Widgets/Input/SEditableTextBox.cpp
--- a/Engine/Source/Runtime/Slate/Private/Widgets/Input/SEditableTextBox.cpp
+++ b/Engine/Source/Runtime/Slate/Private/Widgets/Input/SEditableTextBox.cpp
@@ -201,3 +201,68 @@ void SEditableTextBox::SetIsPassword(TAttribute< bool > InIsPassword)
 {
 	EditableText->SetIsPassword(InIsPassword);
 }
+
+
+void SEditableTextBox::SetFont(const TAttribute<FSlateFontInfo>& InFont)
+{
+	TAttribute<FSlateFontInfo> Font = InFont.IsSet() ? InFont : Style->Font;
+
+	EditableText->SetFont(Font);
+}
+
+
+void SEditableTextBox::SetTextBoxForegroundColor(const TAttribute<FSlateColor>& InForegroundColor)
+{
+	TAttribute<FSlateColor> ForegroundColor = InForegroundColor.IsSet() ? InForegroundColor : Style->ForegroundColor;
+
+	SetForegroundColor(ForegroundColor);
+}
+
+void SEditableTextBox::SetTextBoxBackgroundColor(const TAttribute<FSlateColor>& InBackgroundColor)
+{
+	TAttribute<FSlateColor> BackgroundColor = InBackgroundColor.IsSet() ? InBackgroundColor : Style->BackgroundColor;
+
+	SetBorderBackgroundColor(BackgroundColor);
+}
+
+
+void SEditableTextBox::SetReadOnlyForegroundColor(const TAttribute<FSlateColor>& InReadOnlyForegroundColor)
+{
+	ReadOnlyForegroundColor = InReadOnlyForegroundColor.IsSet() ? InReadOnlyForegroundColor : Style->ReadOnlyForegroundColor;
+}
+
+
+void SEditableTextBox::SetMinimumDesiredWidth(const TAttribute<float>& InMinimumDesiredWidth)
+{
+	EditableText->SetMinDesiredWidth(InMinimumDesiredWidth);
+}
+
+
+void SEditableTextBox::SetIsCaretMovedWhenGainFocus(const TAttribute<bool>& InIsCaretMovedWhenGainFocus)
+{
+	EditableText->SetIsCaretMovedWhenGainFocus(InIsCaretMovedWhenGainFocus);
+}
+
+
+void SEditableTextBox::SetSelectAllTextWhenFocused(const TAttribute<bool>& InSelectAllTextWhenFocused)
+{
+	EditableText->SetSelectAllTextWhenFocused(InSelectAllTextWhenFocused);
+}
+
+
+void SEditableTextBox::SetRevertTextOnEscape(const TAttribute<bool>& InRevertTextOnEscape)
+{
+	EditableText->SetRevertTextOnEscape(InRevertTextOnEscape);
+}
+
+
+void SEditableTextBox::SetClearKeyboardFocusOnCommit(const TAttribute<bool>& InClearKeyboardFocusOnCommit)
+{
+	EditableText->SetClearKeyboardFocusOnCommit(InClearKeyboardFocusOnCommit);
+}
+
+
+void SEditableTextBox::SetSelectAllTextOnCommit(const TAttribute<bool>& InSelectAllTextOnCommit)
+{
+	EditableText->SetSelectAllTextOnCommit(InSelectAllTextOnCommit);
+}
diff --git a/Engine/Source/Runtime/Slate/Private/Widgets/Input/SMultiLineEditableTextBox.cpp b/Engine/Source/Runtime/Slate/Private/Widgets/Input/SMultiLineEditableTextBox.cpp
--- a/Engine/Source/Runtime/Slate/Private/Widgets/Input/SMultiLineEditableTextBox.cpp
+++ b/Engine/Source/Runtime/Slate/Private/Widgets/Input/SMultiLineEditableTextBox.cpp
@@ -30,6 +30,8 @@ namespace
 void SMultiLineEditableTextBox::Construct( const FArguments& InArgs )
 {
 	check (InArgs._Style);
+	Style = InArgs._Style;
+
 	BorderImageNormal = &InArgs._Style->BackgroundImageNormal;
 	BorderImageHovered = &InArgs._Style->BackgroundImageHovered;
 	BorderImageFocused = &InArgs._Style->BackgroundImageFocused;
@@ -149,6 +151,25 @@ void SMultiLineEditableTextBox::SetHintText( const TAttribute< FText >& InHintTe
 	EditableText->SetHintText( InHintText );
 }
 
+void SMultiLineEditableTextBox::SetTextBoxForegroundColor(const TAttribute<FSlateColor>& InForegroundColor)
+{
+	TAttribute<FSlateColor> ForegroundColor = InForegroundColor.IsSet() ? InForegroundColor : Style->ForegroundColor;
+
+	SetForegroundColor(ForegroundColor);
+}
+
+void SMultiLineEditableTextBox::SetTextBoxBackgroundColor(const TAttribute<FSlateColor>& InBackgroundColor)
+{
+	TAttribute<FSlateColor> BackgroundColor = InBackgroundColor.IsSet() ? InBackgroundColor : Style->BackgroundColor;
+
+	SetBorderBackgroundColor(BackgroundColor);
+}
+
+void SMultiLineEditableTextBox::SetReadOnlyForegroundColor(const TAttribute<FSlateColor>& InReadOnlyForegroundColor)
+{
+	ReadOnlyForegroundColor = InReadOnlyForegroundColor.IsSet() ? InReadOnlyForegroundColor : Style->ReadOnlyForegroundColor;
+}
+
 void SMultiLineEditableTextBox::SetError( const FText& InError )
 {
 	SetError( InError.ToString() );
diff --git a/Engine/Source/Runtime/Slate/Public/Widgets/Input/SEditableText.h b/Engine/Source/Runtime/Slate/Public/Widgets/Input/SEditableText.h
--- a/Engine/Source/Runtime/Slate/Public/Widgets/Input/SEditableText.h
+++ b/Engine/Source/Runtime/Slate/Public/Widgets/Input/SEditableText.h
@@ -156,6 +156,48 @@ public:
 	 */
 	void SetFont( const TAttribute< FSlateFontInfo >& InNewFont );
 
+	/**
+	 * Sets the minimum width that a text block should be.
+	 *
+	 * @param  InMinDesiredWidth	The minimum width
+	 */
+	void SetMinDesiredWidth(const TAttribute<float>& InMinDesiredWidth);
+
+	/**
+	 * Workaround as we loose focus when the auto completion closes.
+	 *
+	 * @param  InIsCaretMovedWhenGainFocus	Workaround
+	 */
+	void SetIsCaretMovedWhenGainFocus(const TAttribute<bool>& InIsCaretMovedWhenGainFocus);
+
+	/**
+	 * Sets whether to select all text when the user clicks to give focus on the widget
+	 *
+	 * @param  InSelectAllTextWhenFocused	Select all text when the user clicks?
+	 */
+	void SetSelectAllTextWhenFocused(const TAttribute<bool>& InSelectAllTextWhenFocused);
+
+	/**
+	 * Sets whether to allow the user to back out of changes when they press the escape key
+	 *
+	 * @param  InRevertTextOnEscape			Allow the user to back out of changes?
+	 */
+	void SetRevertTextOnEscape(const TAttribute<bool>& InRevertTextOnEscape);
+
+	/**
+	 * Sets whether to clear keyboard focus when pressing enter to commit changes
+	 *
+	 * @param  InClearKeyboardFocusOnCommit		Clear keyboard focus when pressing enter?
+	 */
+	void SetClearKeyboardFocusOnCommit(const TAttribute<bool>& InClearKeyboardFocusOnCommit);
+
+	/**
+	 * Sets whether to select all text when pressing enter to commit changes
+	 *
+	 * @param  InSelectAllTextOnCommit		Select all text when pressing enter?
+	 */
+	void SetSelectAllTextOnCommit(const TAttribute<bool>& InSelectAllTextOnCommit);
+
 protected:
 
 	friend class FTextInputMethodContext;
diff --git a/Engine/Source/Runtime/Slate/Public/Widgets/Input/SEditableTextBox.h b/Engine/Source/Runtime/Slate/Public/Widgets/Input/SEditableTextBox.h
--- a/Engine/Source/Runtime/Slate/Public/Widgets/Input/SEditableTextBox.h
+++ b/Engine/Source/Runtime/Slate/Public/Widgets/Input/SEditableTextBox.h
@@ -125,6 +125,76 @@ public:
 	/** See the IsPassword attribute */
 	void SetIsPassword( TAttribute< bool > InIsPassword );
 
+	/**
+	 * Sets the font used to draw the text
+	 *
+	 * @param  InFont	The new font to use
+	 */
+	void SetFont(const TAttribute<FSlateFontInfo>& InFont);
+
+	/**
+	 * Sets the text color and opacity (overrides Style)
+	 *
+	 * @param  InForegroundColor 	The text color and opacity
+	 */
+	void SetTextBoxForegroundColor(const TAttribute<FSlateColor>& InForegroundColor);
+
+	/**
+	 * Sets the color of the background/border around the editable text (overrides Style) 
+	 *
+	 * @param  InBackgroundColor 	The background/border color
+	 */
+	void SetTextBoxBackgroundColor(const TAttribute<FSlateColor>& InBackgroundColor);
+
+	/**
+	 * Sets the text color and opacity when read-only (overrides Style) 
+	 *
+	 * @param  InReadOnlyForegroundColor 	The read-only text color and opacity
+	 */
+	void SetReadOnlyForegroundColor(const TAttribute<FSlateColor>& InReadOnlyForegroundColor);
+
+	/**
+	 * Sets the minimum width that a text box should be.
+	 *
+	 * @param  InMinimumDesiredWidth	The minimum width
+	 */
+	void SetMinimumDesiredWidth(const TAttribute<float>& InMinimumDesiredWidth);
+
+	/**
+	 * Workaround as we loose focus when the auto completion closes.
+	 *
+	 * @param  InIsCaretMovedWhenGainFocus	Workaround
+	 */
+	void SetIsCaretMovedWhenGainFocus(const TAttribute<bool>& InIsCaretMovedWhenGainFocus);
+
+	/**
+	 * Sets whether to select all text when the user clicks to give focus on the widget
+	 *
+	 * @param  InSelectAllTextWhenFocused	Select all text when the user clicks?
+	 */
+	void SetSelectAllTextWhenFocused(const TAttribute<bool>& InSelectAllTextWhenFocused);
+
+	/**
+	 * Sets whether to allow the user to back out of changes when they press the escape key
+	 *
+	 * @param  InRevertTextOnEscape			Allow the user to back out of changes?
+	 */
+	void SetRevertTextOnEscape(const TAttribute<bool>& InRevertTextOnEscape);
+
+	/**
+	 * Sets whether to clear keyboard focus when pressing enter to commit changes
+	 *
+	 * @param  InClearKeyboardFocusOnCommit		Clear keyboard focus when pressing enter?
+	 */
+	void SetClearKeyboardFocusOnCommit(const TAttribute<bool>& InClearKeyboardFocusOnCommit);
+
+	/**
+	 * Sets whether to select all text when pressing enter to commit changes
+	 *
+	 * @param  InSelectAllTextOnCommit		Select all text when pressing enter?
+	 */
+	void SetSelectAllTextOnCommit(const TAttribute<bool>& InSelectAllTextOnCommit);
+
 	/**
 	 * If InError is a non-empty string the TextBox will the ErrorReporting provided during construction
 	 * If no error reporting was provided, the TextBox will create a default error reporter.
diff --git a/Engine/Source/Runtime/Slate/Public/Widgets/Input/SMultiLineEditableTextBox.h b/Engine/Source/Runtime/Slate/Public/Widgets/Input/SMultiLineEditableTextBox.h
--- a/Engine/Source/Runtime/Slate/Public/Widgets/Input/SMultiLineEditableTextBox.h
+++ b/Engine/Source/Runtime/Slate/Public/Widgets/Input/SMultiLineEditableTextBox.h
@@ -183,6 +183,27 @@ public:
 	 */
 	void SetHintText( const TAttribute< FText >& InHintText );
 
+	/**
+	 * Sets the text color and opacity (overrides Style)
+	 *
+	 * @param  InForegroundColor 	The text color and opacity
+	 */
+	void SetTextBoxForegroundColor(const TAttribute<FSlateColor>& InForegroundColor);
+
+	/**
+	 * Sets the color of the background/border around the editable text (overrides Style) 
+	 *
+	 * @param  InBackgroundColor 	The background/border color
+	 */
+	void SetTextBoxBackgroundColor(const TAttribute<FSlateColor>& InBackgroundColor);
+
+	/**
+	 * Sets the text color and opacity when read-only (overrides Style) 
+	 *
+	 * @param  InReadOnlyForegroundColor 	The read-only text color and opacity
+	 */
+	void SetReadOnlyForegroundColor(const TAttribute<FSlateColor>& InReadOnlyForegroundColor);
+
 	/**
 	 * If InError is a non-empty string the TextBox will the ErrorReporting provided during construction
 	 * If no error reporting was provided, the TextBox will create a default error reporter.
@@ -245,6 +266,8 @@ protected:
 
 private:
 
+	const FEditableTextBoxStyle* Style;
+
 	/** Styling: border image to draw when not hovered or focused */
 	const FSlateBrush* BorderImageNormal;
 	/** Styling: border image to draw when hovered */
diff --git a/Engine/Source/Runtime/UMG/Private/Components/EditableTextBox.cpp b/Engine/Source/Runtime/UMG/Private/Components/EditableTextBox.cpp
--- a/Engine/Source/Runtime/UMG/Private/Components/EditableTextBox.cpp
+++ b/Engine/Source/Runtime/UMG/Private/Components/EditableTextBox.cpp
@@ -71,11 +71,19 @@ void UEditableTextBox::SynchronizeProperties()
 
 	MyEditableTextBlock->SetText(TextBinding);
 	MyEditableTextBlock->SetHintText(HintTextBinding);
+	MyEditableTextBlock->SetFont(Font);
+	MyEditableTextBlock->SetTextBoxForegroundColor(ForegroundColor);
+	MyEditableTextBlock->SetTextBoxBackgroundColor(BackgroundColor);
+	MyEditableTextBlock->SetReadOnlyForegroundColor(ReadOnlyForegroundColor);
 	MyEditableTextBlock->SetIsReadOnly(IsReadOnly);
 	MyEditableTextBlock->SetIsPassword(IsPassword);
-//	MyEditableTextBlock->SetColorAndOpacity(ColorAndOpacity);
-
-	// TODO UMG Complete making all properties settable on SEditableTextBox
+	MyEditableTextBlock->SetMinimumDesiredWidth(MinimumDesiredWidth);
+	MyEditableTextBlock->SetPadding(Padding);
+	MyEditableTextBlock->SetIsCaretMovedWhenGainFocus(IsCaretMovedWhenGainFocus);
+	MyEditableTextBlock->SetSelectAllTextWhenFocused(SelectAllTextWhenFocused);
+	MyEditableTextBlock->SetRevertTextOnEscape(RevertTextOnEscape);
+	MyEditableTextBlock->SetClearKeyboardFocusOnCommit(ClearKeyboardFocusOnCommit);
+	MyEditableTextBlock->SetSelectAllTextOnCommit(SelectAllTextOnCommit);
 }
 
 FText UEditableTextBox::GetText() const
diff --git a/Engine/Source/Runtime/UMG/Private/Components/MultiLineEditableTextBox.cpp b/Engine/Source/Runtime/UMG/Private/Components/MultiLineEditableTextBox.cpp
--- a/Engine/Source/Runtime/UMG/Private/Components/MultiLineEditableTextBox.cpp
+++ b/Engine/Source/Runtime/UMG/Private/Components/MultiLineEditableTextBox.cpp
@@ -68,6 +68,11 @@ void UMultiLineEditableTextBox::SynchronizeProperties()
 //	MyEditableTextBlock->SetIsPassword(IsPassword);
 //	MyEditableTextBlock->SetColorAndOpacity(ColorAndOpacity);
 
+//	MyEditableTextBlock->SetFont(Font);		@todo: can't set the font on an existing multiline text block
+	MyEditableTextBlock->SetTextBoxForegroundColor(ForegroundColor);
+	MyEditableTextBlock->SetTextBoxBackgroundColor(BackgroundColor);
+	MyEditableTextBlock->SetReadOnlyForegroundColor(ReadOnlyForegroundColor);
+
 	// TODO UMG Complete making all properties settable on SMultiLineEditableTextBox
 }
 
