Support PS4 SDK 2.5 and 2.0

diff --git a/Engine/Source/Runtime/Slate/Private/Widgets/Layout/SSplitter.cpp b/Engine/Source/Runtime/Slate/Private/Widgets/Layout/SSplitter.cpp
--- a/Engine/Source/Runtime/Slate/Private/Widgets/Layout/SSplitter.cpp
+++ b/Engine/Source/Runtime/Slate/Private/Widgets/Layout/SSplitter.cpp
@@ -861,7 +861,15 @@ void SSplitter2x2::ResizeChildren( const FGeometry& MyGeometry, const TArray<FLa
 		// Resize X and Y independently as they have different rules for X and Y
 		NewSizeTL.X = TopLeftSize.X + Delta.X;
 		NewSizeBL.X = BotLeftSize.X + Delta.X;
+
+		//workaround PS4 compiler crash.
+#if PLATFORM_PS4
+		volatile float workaround = 0.0f;
+		NewSizeTR.X = TopRightSize.X - Delta.X + workaround;
+#else
 		NewSizeTR.X = TopRightSize.X - Delta.X;
+#endif
+
 		NewSizeBR.X = BotRightSize.X - Delta.X;
 
 
