Fixed gamepad button X not being consumed by analog cursor code even when it was a handled event.

diff --git a/Engine/Source/Runtime/Slate/Private/Framework/Application/AnalogCursor.cpp b/Engine/Source/Runtime/Slate/Private/Framework/Application/AnalogCursor.cpp
--- a/Engine/Source/Runtime/Slate/Private/Framework/Application/AnalogCursor.cpp
+++ b/Engine/Source/Runtime/Slate/Private/Framework/Application/AnalogCursor.cpp
@@ -138,7 +138,8 @@ bool FAnalogCursor::HandleKeyDownEvent(FSlateApplication& SlateApp, const FKeyEv
 			);
 
 		TSharedPtr<FGenericWindow> GenWindow;
-		SlateApp.ProcessMouseButtonDownEvent(GenWindow, MouseEvent);
+		return SlateApp.ProcessMouseButtonDownEvent(GenWindow, MouseEvent);
+
 	}
 
 	return false;
