Added option to disable or add custom context menus to editable text widgets in Slate
[UE-14497] - Game client stops responding when the RMB is clicked in the Enter Your Hero Name textbox while in Full Screen mode

diff --git a/Engine/Source/Runtime/AppFramework/Private/Widgets/Testing/STestSuite.cpp b/Engine/Source/Runtime/AppFramework/Private/Widgets/Testing/STestSuite.cpp
--- a/Engine/Source/Runtime/AppFramework/Private/Widgets/Testing/STestSuite.cpp
+++ b/Engine/Source/Runtime/AppFramework/Private/Widgets/Testing/STestSuite.cpp
@@ -2060,6 +2060,22 @@ public:
 				[
 					SAssignNew(ErrorText, SErrorText)
 				]
+				+ SVerticalBox::Slot().AutoHeight() .HAlign(HAlign_Center) .VAlign(VAlign_Center) .Padding( 5 )
+				[
+					SNew( SEditableTextBox )
+					.Text( LOCTEXT( "CustomContextMenuInput", "This text box has a custom context menu" ) )
+					.RevertTextOnEscape( true )
+					.HintText(LOCTEXT("CustomContextMenuHint", "Custom context menu..."))
+					.OnContextMenuOpening( this, &STextEditTest::OnCustomContextMenuOpening )
+				]
+				+ SVerticalBox::Slot().AutoHeight() .HAlign(HAlign_Center) .VAlign(VAlign_Center) .Padding( 5 )
+				[
+					SNew( SEditableTextBox )
+					.Text( LOCTEXT( "DisabledContextMenuInput", "This text box has no context menu" ) )
+					.RevertTextOnEscape( true )
+					.HintText(LOCTEXT("DisabledContextMenuHint", "No context menu..."))
+					.OnContextMenuOpening(this, &STextEditTest::OnDisabledContextMenuOpening)
+				]
 			]
 		];
 	}
@@ -2089,6 +2105,22 @@ public:
 		NumericInput->SetError( Error );
 	}
 
+	TSharedPtr<SWidget> OnCustomContextMenuOpening()
+	{
+		return SNew(SBorder)
+			.Padding(5.0f)
+			.BorderImage(FCoreStyle::Get().GetBrush("BoxShadow"))
+			[
+				SNew(STextBlock)
+				.Text(LOCTEXT("CustomContextMenuContent", "This context menu content is just a text block"))
+			];
+	}
+
+	TSharedPtr<SWidget> OnDisabledContextMenuOpening()
+	{
+		return TSharedPtr<SWidget>();
+	}
+
 	void ClearSearchBox()
 	{
 		SearchBox->SetText( FText::GetEmpty() );
diff --git a/Engine/Source/Runtime/Slate/Private/Widgets/Input/SEditableText.cpp b/Engine/Source/Runtime/Slate/Private/Widgets/Input/SEditableText.cpp
--- a/Engine/Source/Runtime/Slate/Private/Widgets/Input/SEditableText.cpp
+++ b/Engine/Source/Runtime/Slate/Private/Widgets/Input/SEditableText.cpp
@@ -77,6 +77,7 @@ void SEditableText::Construct( const FArguments& InArgs )
 	bSelectAllTextWhenFocused = InArgs._SelectAllTextWhenFocused;
 	RevertTextOnEscape = InArgs._RevertTextOnEscape;
 	ClearKeyboardFocusOnCommit = InArgs._ClearKeyboardFocusOnCommit;
+	OnContextMenuOpening = InArgs._OnContextMenuOpening;
 	OnIsTypedCharValid = InArgs._OnIsTypedCharValid;
 	OnTextChanged = InArgs._OnTextChanged;
 	OnTextCommitted = InArgs._OnTextCommitted;
@@ -1938,68 +1939,82 @@ void SEditableText::SaveText()
 
 void SEditableText::SummonContextMenu(const FVector2D& InLocation, TSharedPtr<SWindow> ParentWindow)
 {
-	// Set the menu to automatically close when the user commits to a choice
-	const bool bShouldCloseWindowAfterMenuSelection = true;
+	TSharedPtr<SWidget> MenuContentWidget;
+
+	if (OnContextMenuOpening.IsBound())
+	{
+		MenuContentWidget = OnContextMenuOpening.Execute();
+	}
+	else
+	{
+		// Set the menu to automatically close when the user commits to a choice
+		const bool bShouldCloseWindowAfterMenuSelection = true;
 
 #define LOCTEXT_NAMESPACE "EditableTextContextMenu"
 
-	// This is a context menu which could be summoned from within another menu if this text block is in a menu
-	// it should not close the menu it is inside
-	bool bCloseSelfOnly = true;
-	FMenuBuilder MenuBuilder( bShouldCloseWindowAfterMenuSelection, UICommandList, MenuExtender, bCloseSelfOnly, &FCoreStyle::Get() );
-	{
-		MenuBuilder.BeginSection( "EditText", LOCTEXT( "Heading", "Modify Text" ) );
+		// This is a context menu which could be summoned from within another menu if this text block is in a menu
+		// it should not close the menu it is inside
+		bool bCloseSelfOnly = true;
+		FMenuBuilder MenuBuilder(bShouldCloseWindowAfterMenuSelection, UICommandList, MenuExtender, bCloseSelfOnly, &FCoreStyle::Get());
 		{
-			// Undo
-			MenuBuilder.AddMenuEntry( FGenericCommands::Get().Undo );
-		}
-		MenuBuilder.EndSection();
+			MenuBuilder.BeginSection("EditText", LOCTEXT("Heading", "Modify Text"));
+			{
+				// Undo
+				MenuBuilder.AddMenuEntry(FGenericCommands::Get().Undo);
+			}
+			MenuBuilder.EndSection();
 
-		MenuBuilder.BeginSection("EditableTextModify2");
-		{
-			// Cut
-			MenuBuilder.AddMenuEntry( FGenericCommands::Get().Cut );
+			MenuBuilder.BeginSection("EditableTextModify2");
+			{
+				// Cut
+				MenuBuilder.AddMenuEntry(FGenericCommands::Get().Cut);
 
-			// Copy
-			MenuBuilder.AddMenuEntry( FGenericCommands::Get().Copy );
+				// Copy
+				MenuBuilder.AddMenuEntry(FGenericCommands::Get().Copy);
 
-			// Paste
-			MenuBuilder.AddMenuEntry( FGenericCommands::Get().Paste );
+				// Paste
+				MenuBuilder.AddMenuEntry(FGenericCommands::Get().Paste);
 
-			// Delete
-			MenuBuilder.AddMenuEntry( FGenericCommands::Get().Delete );
-		}
-		MenuBuilder.EndSection();
+				// Delete
+				MenuBuilder.AddMenuEntry(FGenericCommands::Get().Delete);
+			}
+			MenuBuilder.EndSection();
 
-		MenuBuilder.BeginSection("EditableTextModify3");
-		{
-			// Select All
-			MenuBuilder.AddMenuEntry( FGenericCommands::Get().SelectAll );
+			MenuBuilder.BeginSection("EditableTextModify3");
+			{
+				// Select All
+				MenuBuilder.AddMenuEntry(FGenericCommands::Get().SelectAll);
+			}
+			MenuBuilder.EndSection();
 		}
-		MenuBuilder.EndSection();
-	}
 
-#undef LOCTEXT_NAMESPACE
+		MenuContentWidget = MenuBuilder.MakeWidget();
 
-	ActiveContextMenu.PrepareToSummon();
-
-	const bool bFocusImmediately = true;
-	TSharedRef<SWidget> MenuParent = SharedThis(this);
-	if (ParentWindow.IsValid())
-	{
-		MenuParent = StaticCastSharedRef<SWidget>(ParentWindow.ToSharedRef());
-	}
-	TSharedPtr< SWindow > ContextMenuWindow = FSlateApplication::Get().PushMenu(MenuParent, MenuBuilder.MakeWidget(), InLocation, FPopupTransitionEffect(FPopupTransitionEffect::ContextMenu), bFocusImmediately );
-	
-	// Make sure the window is valid.  It's possible for the parent to already be in the destroy queue, for example if the editable text was configured to dismiss it's window during OnTextCommitted.
-	if( ContextMenuWindow.IsValid() )
-	{
-		ContextMenuWindow->SetOnWindowClosed( FOnWindowClosed::CreateSP( this, &SEditableText::OnWindowClosed ) );
-		ActiveContextMenu.SummonSucceeded( ContextMenuWindow.ToSharedRef() );
+#undef LOCTEXT_NAMESPACE
 	}
-	else
+
+	if (MenuContentWidget.IsValid())
 	{
-		ActiveContextMenu.SummonFailed();
+		ActiveContextMenu.PrepareToSummon();
+
+		const bool bFocusImmediately = true;
+		TSharedRef<SWidget> MenuParent = SharedThis(this);
+		if (ParentWindow.IsValid())
+		{
+			MenuParent = StaticCastSharedRef<SWidget>(ParentWindow.ToSharedRef());
+		}
+		TSharedPtr< SWindow > ContextMenuWindow = FSlateApplication::Get().PushMenu(MenuParent, MenuContentWidget.ToSharedRef(), InLocation, FPopupTransitionEffect(FPopupTransitionEffect::ContextMenu), bFocusImmediately);
+
+		// Make sure the window is valid.  It's possible for the parent to already be in the destroy queue, for example if the editable text was configured to dismiss it's window during OnTextCommitted.
+		if (ContextMenuWindow.IsValid())
+		{
+			ContextMenuWindow->SetOnWindowClosed(FOnWindowClosed::CreateSP(this, &SEditableText::OnWindowClosed));
+			ActiveContextMenu.SummonSucceeded(ContextMenuWindow.ToSharedRef());
+		}
+		else
+		{
+			ActiveContextMenu.SummonFailed();
+		}
 	}
 }
 
diff --git a/Engine/Source/Runtime/Slate/Private/Widgets/Input/SEditableTextBox.cpp b/Engine/Source/Runtime/Slate/Private/Widgets/Input/SEditableTextBox.cpp
--- a/Engine/Source/Runtime/Slate/Private/Widgets/Input/SEditableTextBox.cpp
+++ b/Engine/Source/Runtime/Slate/Private/Widgets/Input/SEditableTextBox.cpp
@@ -45,6 +45,7 @@ void SEditableTextBox::Construct( const FArguments& InArgs )
 					.SelectAllTextWhenFocused( InArgs._SelectAllTextWhenFocused )
 					.RevertTextOnEscape( InArgs._RevertTextOnEscape )
 					.ClearKeyboardFocusOnCommit( InArgs._ClearKeyboardFocusOnCommit )
+					.OnContextMenuOpening( InArgs._OnContextMenuOpening )
 					.OnTextChanged( InArgs._OnTextChanged )
 					.OnTextCommitted( InArgs._OnTextCommitted )
 					.MinDesiredWidth( InArgs._MinDesiredWidth )
diff --git a/Engine/Source/Runtime/Slate/Private/Widgets/Input/SMultiLineEditableTextBox.cpp b/Engine/Source/Runtime/Slate/Private/Widgets/Input/SMultiLineEditableTextBox.cpp
--- a/Engine/Source/Runtime/Slate/Private/Widgets/Input/SMultiLineEditableTextBox.cpp
+++ b/Engine/Source/Runtime/Slate/Private/Widgets/Input/SMultiLineEditableTextBox.cpp
@@ -92,6 +92,7 @@ void SMultiLineEditableTextBox::Construct( const FArguments& InArgs )
 					.Marshaller( InArgs._Marshaller )
 					.Font( Font )
 					.IsReadOnly( InArgs._IsReadOnly )
+					.OnContextMenuOpening( InArgs._OnContextMenuOpening )
 					.OnTextChanged( InArgs._OnTextChanged )
 					.OnTextCommitted( InArgs._OnTextCommitted )
 					.OnCursorMoved( InArgs._OnCursorMoved )
diff --git a/Engine/Source/Runtime/Slate/Private/Widgets/Text/SMultiLineEditableText.cpp b/Engine/Source/Runtime/Slate/Private/Widgets/Text/SMultiLineEditableText.cpp
--- a/Engine/Source/Runtime/Slate/Private/Widgets/Text/SMultiLineEditableText.cpp
+++ b/Engine/Source/Runtime/Slate/Private/Widgets/Text/SMultiLineEditableText.cpp
@@ -271,6 +271,7 @@ void SMultiLineEditableText::Construct( const FArguments& InArgs )
 	OnCursorMoved = InArgs._OnCursorMoved;
 	bSelectAllTextWhenFocused = InArgs._SelectAllTextWhenFocused;
 	bClearKeyboardFocusOnCommit = InArgs._ClearKeyboardFocusOnCommit;
+	OnContextMenuOpening = InArgs._OnContextMenuOpening;
 	bRevertTextOnEscape = InArgs._RevertTextOnEscape;
 	OnHScrollBarUserScrolled = InArgs._OnHScrollBarUserScrolled;
 	OnVScrollBarUserScrolled = InArgs._OnVScrollBarUserScrolled;
@@ -2307,68 +2308,82 @@ TSharedRef< SWidget > SMultiLineEditableText::GetWidget()
 
 void SMultiLineEditableText::SummonContextMenu(const FVector2D& InLocation, TSharedPtr<SWindow> ParentWindow)
 {
-	// Set the menu to automatically close when the user commits to a choice
-	const bool bShouldCloseWindowAfterMenuSelection = true;
+	TSharedPtr<SWidget> MenuContentWidget;
+
+	if (OnContextMenuOpening.IsBound())
+	{
+		MenuContentWidget = OnContextMenuOpening.Execute();
+	}
+	else
+	{
+		// Set the menu to automatically close when the user commits to a choice
+		const bool bShouldCloseWindowAfterMenuSelection = true;
 
 #define LOCTEXT_NAMESPACE "EditableTextContextMenu"
 
-	// This is a context menu which could be summoned from within another menu if this text block is in a menu
-	// it should not close the menu it is inside
-	bool bCloseSelfOnly = true;
-	FMenuBuilder MenuBuilder(bShouldCloseWindowAfterMenuSelection, UICommandList, MenuExtender, bCloseSelfOnly, &FCoreStyle::Get());
-	{
-		MenuBuilder.BeginSection("EditText", LOCTEXT("Heading", "Modify Text"));
+		// This is a context menu which could be summoned from within another menu if this text block is in a menu
+		// it should not close the menu it is inside
+		bool bCloseSelfOnly = true;
+		FMenuBuilder MenuBuilder(bShouldCloseWindowAfterMenuSelection, UICommandList, MenuExtender, bCloseSelfOnly, &FCoreStyle::Get());
 		{
-			// Undo
-			MenuBuilder.AddMenuEntry(FGenericCommands::Get().Undo);
-		}
-		MenuBuilder.EndSection();
+			MenuBuilder.BeginSection("EditText", LOCTEXT("Heading", "Modify Text"));
+			{
+				// Undo
+				MenuBuilder.AddMenuEntry(FGenericCommands::Get().Undo);
+			}
+			MenuBuilder.EndSection();
 
-		MenuBuilder.BeginSection("EditableTextModify2");
-		{
-			// Cut
-			MenuBuilder.AddMenuEntry(FGenericCommands::Get().Cut);
+			MenuBuilder.BeginSection("EditableTextModify2");
+			{
+				// Cut
+				MenuBuilder.AddMenuEntry(FGenericCommands::Get().Cut);
 
-			// Copy
-			MenuBuilder.AddMenuEntry(FGenericCommands::Get().Copy);
+				// Copy
+				MenuBuilder.AddMenuEntry(FGenericCommands::Get().Copy);
 
-			// Paste
-			MenuBuilder.AddMenuEntry(FGenericCommands::Get().Paste);
+				// Paste
+				MenuBuilder.AddMenuEntry(FGenericCommands::Get().Paste);
 
-			// Delete
-			MenuBuilder.AddMenuEntry(FGenericCommands::Get().Delete);
-		}
-		MenuBuilder.EndSection();
+				// Delete
+				MenuBuilder.AddMenuEntry(FGenericCommands::Get().Delete);
+			}
+			MenuBuilder.EndSection();
 
-		MenuBuilder.BeginSection("EditableTextModify3");
-		{
-			// Select All
-			MenuBuilder.AddMenuEntry(FGenericCommands::Get().SelectAll);
+			MenuBuilder.BeginSection("EditableTextModify3");
+			{
+				// Select All
+				MenuBuilder.AddMenuEntry(FGenericCommands::Get().SelectAll);
+			}
+			MenuBuilder.EndSection();
 		}
-		MenuBuilder.EndSection();
-	}
 
-#undef LOCTEXT_NAMESPACE
-
-	ActiveContextMenu.PrepareToSummon();
+		MenuContentWidget = MenuBuilder.MakeWidget();
 
-	const bool bFocusImmediately = true;
-	TSharedRef<SWidget> MenuParent = SharedThis(this);
-	if (ParentWindow.IsValid())
-	{
-		MenuParent = StaticCastSharedRef<SWidget>(ParentWindow.ToSharedRef());
+#undef LOCTEXT_NAMESPACE
 	}
-	TSharedPtr< SWindow > ContextMenuWindow = FSlateApplication::Get().PushMenu(MenuParent, MenuBuilder.MakeWidget(), InLocation, FPopupTransitionEffect(FPopupTransitionEffect::ContextMenu), bFocusImmediately);
 
-	// Make sure the window is valid.  It's possible for the parent to already be in the destroy queue, for example if the editable text was configured to dismiss it's window during OnTextCommitted.
-	if (ContextMenuWindow.IsValid())
+	if (MenuContentWidget.IsValid())
 	{
-		ContextMenuWindow->SetOnWindowClosed(FOnWindowClosed::CreateSP(this, &SMultiLineEditableText::OnWindowClosed));
-		ActiveContextMenu.SummonSucceeded(ContextMenuWindow.ToSharedRef());
-	}
-	else
-	{
-		ActiveContextMenu.SummonFailed();
+		ActiveContextMenu.PrepareToSummon();
+
+		const bool bFocusImmediately = true;
+		TSharedRef<SWidget> MenuParent = SharedThis(this);
+		if (ParentWindow.IsValid())
+		{
+			MenuParent = StaticCastSharedRef<SWidget>(ParentWindow.ToSharedRef());
+		}
+		TSharedPtr< SWindow > ContextMenuWindow = FSlateApplication::Get().PushMenu(MenuParent, MenuContentWidget.ToSharedRef(), InLocation, FPopupTransitionEffect(FPopupTransitionEffect::ContextMenu), bFocusImmediately);
+
+		// Make sure the window is valid.  It's possible for the parent to already be in the destroy queue, for example if the editable text was configured to dismiss it's window during OnTextCommitted.
+		if (ContextMenuWindow.IsValid())
+		{
+			ContextMenuWindow->SetOnWindowClosed(FOnWindowClosed::CreateSP(this, &SMultiLineEditableText::OnWindowClosed));
+			ActiveContextMenu.SummonSucceeded(ContextMenuWindow.ToSharedRef());
+		}
+		else
+		{
+			ActiveContextMenu.SummonFailed();
+		}
 	}
 }
 
diff --git a/Engine/Source/Runtime/Slate/Public/Widgets/Input/SEditableText.h b/Engine/Source/Runtime/Slate/Public/Widgets/Input/SEditableText.h
--- a/Engine/Source/Runtime/Slate/Public/Widgets/Input/SEditableText.h
+++ b/Engine/Source/Runtime/Slate/Public/Widgets/Input/SEditableText.h
@@ -81,6 +81,9 @@ public:
 		/** Whether to clear keyboard focus when pressing enter to commit changes */
 		SLATE_ATTRIBUTE( bool, ClearKeyboardFocusOnCommit )
 
+		/** Delegate to call before a context menu is opened. User returns the menu content or null to the disable context menu */
+		SLATE_EVENT(FOnContextMenuOpening, OnContextMenuOpening)
+
 		/**
 		 * This is NOT for validating input!
 		 * 
@@ -571,6 +574,9 @@ private:
 	/** Whether to select all text when pressing enter to commit changes */
 	TAttribute< bool > SelectAllTextOnCommit;
 
+	/** Delegate to call before a context menu is opened */
+	FOnContextMenuOpening OnContextMenuOpening;
+
 	/** Called when a character is typed and we want to know if the text field supports typing this character. */
 	FOnIsTypedCharValid OnIsTypedCharValid;
 
diff --git a/Engine/Source/Runtime/Slate/Public/Widgets/Input/SEditableTextBox.h b/Engine/Source/Runtime/Slate/Public/Widgets/Input/SEditableTextBox.h
--- a/Engine/Source/Runtime/Slate/Public/Widgets/Input/SEditableTextBox.h
+++ b/Engine/Source/Runtime/Slate/Public/Widgets/Input/SEditableTextBox.h
@@ -66,6 +66,9 @@ public:
 		/** Whether to clear keyboard focus when pressing enter to commit changes */
 		SLATE_ATTRIBUTE( bool, ClearKeyboardFocusOnCommit )
 
+		/** Delegate to call before a context menu is opened. User returns the menu content or null to the disable context menu */
+		SLATE_EVENT(FOnContextMenuOpening, OnContextMenuOpening)
+
 		/** Called whenever the text is changed interactively by the user */
 		SLATE_EVENT( FOnTextChanged, OnTextChanged )
 
diff --git a/Engine/Source/Runtime/Slate/Public/Widgets/Input/SMultiLineEditableTextBox.h b/Engine/Source/Runtime/Slate/Public/Widgets/Input/SMultiLineEditableTextBox.h
--- a/Engine/Source/Runtime/Slate/Public/Widgets/Input/SMultiLineEditableTextBox.h
+++ b/Engine/Source/Runtime/Slate/Public/Widgets/Input/SMultiLineEditableTextBox.h
@@ -105,6 +105,9 @@ public:
 		/** Padding around the vertical scrollbar (overrides Style) */
 		SLATE_ATTRIBUTE( FMargin, VScrollBarPadding )
 
+		/** Delegate to call before a context menu is opened. User returns the menu content or null to the disable context menu */
+		SLATE_EVENT(FOnContextMenuOpening, OnContextMenuOpening)
+
 		/** Called whenever the text is changed interactively by the user */
 		SLATE_EVENT( FOnTextChanged, OnTextChanged )
 
diff --git a/Engine/Source/Runtime/Slate/Public/Widgets/Text/SMultiLineEditableText.h b/Engine/Source/Runtime/Slate/Public/Widgets/Text/SMultiLineEditableText.h
--- a/Engine/Source/Runtime/Slate/Public/Widgets/Text/SMultiLineEditableText.h
+++ b/Engine/Source/Runtime/Slate/Public/Widgets/Text/SMultiLineEditableText.h
@@ -96,6 +96,9 @@ public:
 		/** Whether to clear keyboard focus when pressing enter to commit changes */
 		SLATE_ATTRIBUTE(bool, ClearKeyboardFocusOnCommit)
 
+		/** Delegate to call before a context menu is opened. User returns the menu content or null to the disable context menu */
+		SLATE_EVENT(FOnContextMenuOpening, OnContextMenuOpening)
+
 		/** Called whenever the horizontal scrollbar is moved by the user */
 		SLATE_EVENT(FOnUserScrolled, OnHScrollBarUserScrolled)
 
@@ -683,6 +686,9 @@ private:
 	/** A list commands to execute if a user presses the corresponding keybinding in the text box */
 	TSharedRef< FUICommandList > UICommandList;
 
+	/** Delegate to call before a context menu is opened */
+	FOnContextMenuOpening OnContextMenuOpening;
+
 	/** Called whenever the text is changed interactively by the user */
 	FOnTextChanged OnTextChanged;
 
