Slate/UMG - UE-14099 - SScrollBox now has support for Touch input for scrolling and Overscroll.

diff --git a/Engine/Source/Runtime/Slate/Private/Framework/Layout/Overscroll.cpp b/Engine/Source/Runtime/Slate/Private/Framework/Layout/Overscroll.cpp
new file mode 100644
--- /dev/null
+++ b/Engine/Source/Runtime/Slate/Private/Framework/Layout/Overscroll.cpp
@@ -0,0 +1,71 @@
+// Copyright 1998-2015 Epic Games, Inc. All Rights Reserved.
+
+#include "SlatePrivatePCH.h"
+
+namespace OverscrollConstants
+{
+	/** The amount to scale the logarithm by to make it more loose */
+	static const float Looseness = 50.0f;
+	/** The bounce back rate when the overscroll stops. */
+	static const float OvershootBounceRate = 2000.0f;
+}
+
+FOverscroll::FOverscroll()
+	: OverscrollAmount(0.0f)
+{
+}
+
+float FOverscroll::ScrollBy(float Delta)
+{
+	const float ValueBeforeDeltaApplied = OverscrollAmount;
+	OverscrollAmount += Delta;
+
+	// Don't allow an interaction to change from positive <-> negative overscroll
+	const bool bCrossedOverscrollBoundary = FMath::Sign(ValueBeforeDeltaApplied) != FMath::Sign(OverscrollAmount);
+	if ( bCrossedOverscrollBoundary && ValueBeforeDeltaApplied != 0.f )
+	{
+		OverscrollAmount = 0.f;
+	}
+
+	return ValueBeforeDeltaApplied - OverscrollAmount;
+}
+
+float FOverscroll::GetOverscroll() const
+{
+	// We apply the resistance when observers attempt to read the overscroll.
+	// First we calculate the origin shift for the looseness factor, by default,
+	// straight up Log would be too tight, so we scale the result, the problem is that shifts
+	// the inflection point of the log function, so we subtract the shift amount to correct it.
+	const float OriginShift = OverscrollConstants::Looseness * FMath::Loge(OverscrollConstants::Looseness);
+
+	// Scale the log of the current overscroll and shift the origin so that the overscroll the resistance is continuous.
+	const float AbsElasticOverscroll = OverscrollConstants::Looseness * FMath::Loge(FMath::Abs(OverscrollAmount) + OverscrollConstants::Looseness) - OriginShift;
+
+	return FMath::Sign(OverscrollAmount) * AbsElasticOverscroll;
+}
+
+void FOverscroll::UpdateOverscroll(float InDeltaTime)
+{
+	const float PullForce = FMath::Abs(OverscrollAmount) + 1.0f;
+	const float EasedDelta = OverscrollConstants::OvershootBounceRate * InDeltaTime * PullForce;
+
+	if ( OverscrollAmount > 0 )
+	{
+		OverscrollAmount = FMath::Max(0.0f, OverscrollAmount - EasedDelta * InDeltaTime);
+	}
+	else
+	{
+		OverscrollAmount = FMath::Min(0.0f, OverscrollAmount + EasedDelta * InDeltaTime);
+	}
+}
+
+bool FOverscroll::ShouldApplyOverscroll(const bool bIsAtStartOfList, const bool bIsAtEndOfList, const float ScrollDelta) const
+{
+	const bool bShouldApplyOverscroll =
+		// We can scroll past the edge of the list only if we are at the edge
+		( bIsAtStartOfList && ScrollDelta < 0 ) || ( bIsAtEndOfList && ScrollDelta > 0 ) ||
+		// ... or if we are already past the edge and are scrolling in the opposite direction.
+		( OverscrollAmount > 0 && ScrollDelta < 0 ) || ( OverscrollAmount < 0 && ScrollDelta > 0 );
+
+	return bShouldApplyOverscroll;
+}
diff --git a/Engine/Source/Runtime/Slate/Private/Widgets/Layout/SScrollBox.cpp b/Engine/Source/Runtime/Slate/Private/Widgets/Layout/SScrollBox.cpp
--- a/Engine/Source/Runtime/Slate/Private/Widgets/Layout/SScrollBox.cpp
+++ b/Engine/Source/Runtime/Slate/Private/Widgets/Layout/SScrollBox.cpp
@@ -154,14 +154,17 @@ void SScrollBox::Construct( const FArguments& InArgs )
 	DesiredScrollOffset = 0;
 	bIsScrolling = false;
 	bAnimateScroll = false;
+	bStartedTouchInteraction = false;
 	AmountScrolledWhileRightMouseDown = 0;
 	bShowSoftwareCursor = false;
 	SoftwareCursorPosition = FVector2D::ZeroVector;
 	OnUserScrolled = InArgs._OnUserScrolled;
 	Orientation = InArgs._Orientation;
 	bScrollToEnd = false;
-	bIsActiveTimerRegistered = false;
+	bIsScrollingActiveTimerRegistered = false;
 	ConsumeMouseWheel = InArgs._ConsumeMouseWheel;
+	TickScrollDelta = 0;
+	AllowOverscroll = InArgs._AllowOverscroll;
 
 	if (InArgs._ExternalScrollbar.IsValid())
 	{
@@ -381,13 +384,13 @@ bool SScrollBox::ScrollDescendantIntoView(const FGeometry& MyGeometry, const TSh
 	if ( ensureMsg( WidgetGeometry, TEXT("Unable to scroll to descendant as it's not a child of the scrollbox") ) )
 	{
 		// @todo: This is a workaround because DesiredScrollOffset can exceed the ScrollMax when mouse dragging on the scroll bar and we need it clamped here or the offset is wrong
-		ScrollBy(MyGeometry, 0, false);
+		ScrollBy(MyGeometry, 0, EAllowOverscroll::No, false);
 
 		// Calculate how much we would need to scroll to bring this to the top/left of the scroll box
 		const float WidgetPosition = GetScrollComponentFromVector(WidgetGeometry->Geometry.AbsolutePosition);
 		const float MyPosition = GetScrollComponentFromVector(MyGeometry.AbsolutePosition);
 		const float ScrollOffset = WidgetPosition - MyPosition;
-		ScrollBy(MyGeometry, ScrollOffset, InAnimateScroll);
+		ScrollBy(MyGeometry, ScrollOffset, EAllowOverscroll::No, InAnimateScroll);
 		return true;
 	}
 	return false;
@@ -436,25 +439,66 @@ void SScrollBox::SetScrollBarThickness(FVector2D InThickness)
 
 EActiveTimerReturnType SScrollBox::UpdateInertialScroll(double InCurrentTime, float InDeltaTime)
 {
-	InertialScrollManager.UpdateScrollVelocity(InDeltaTime);
+	bool bKeepTicking = false;
 
-	if (bIsScrolling)
+	//if (bIsScrolling)
+	//{
+	//	const float ScrollVelocity = InertialScrollManager.GetScrollVelocity();
+	//	// Do not apply inertial scrolling while the user is actively scrolling via RMB.
+	//	if (ScrollVelocity != 0.f && !IsRightClickScrolling())
+	//	{
+	//		ScrollBy(CachedGeometry, ScrollVelocity * InDeltaTime, AllowOverscroll, true);
+	//	}
+
+	//	bKeepTicking = true;
+	//}
+
+	if ( IsRightClickScrolling() )
+	{
+		bKeepTicking = true;
+
+		// We sample for the inertial scroll on tick rather than on mouse/touch move so
+		// that we still get samples even if the mouse has not moved.
+		if ( CanUseInertialScroll(TickScrollDelta) )
+		{
+			InertialScrollManager.AddScrollSample(TickScrollDelta, InCurrentTime);
+		}
+	}
+	else
 	{
+		InertialScrollManager.UpdateScrollVelocity(InDeltaTime);
 		const float ScrollVelocity = InertialScrollManager.GetScrollVelocity();
-		// Do not apply inertial scrolling while the user is actively scrolling via RMB.
-		if (ScrollVelocity != 0.f && !IsRightClickScrolling())
+
+		if ( ScrollVelocity != 0.f )
 		{
-			ScrollBy(CachedGeometry, ScrollVelocity * InDeltaTime, true);
+			if ( CanUseInertialScroll(ScrollVelocity) )
+			{
+				bKeepTicking = true;
+				ScrollBy(CachedGeometry, ScrollVelocity * InDeltaTime, AllowOverscroll);
+			}
+			else
+			{
+				InertialScrollManager.ClearScrollVelocity();
+			}
 		}
 
-		return EActiveTimerReturnType::Continue;
+		if ( AllowOverscroll == EAllowOverscroll::Yes )
+		{
+			// If we are currently in overscroll, the list will need refreshing.
+			// Do this before UpdateOverscroll, as that could cause GetOverscroll() to be 0
+			if ( Overscroll.GetOverscroll() != 0.0f )
+			{
+				bKeepTicking = true;
+			}
+
+			Overscroll.UpdateOverscroll(InDeltaTime);
+		}
 	}
 
-	// Clear the scroll velocity so there isn't any delayed minuscule movement
-	InertialScrollManager.ClearScrollVelocity();
+	TickScrollDelta = 0.f;
 
-	bIsActiveTimerRegistered = false;
-	return EActiveTimerReturnType::Stop;
+	bIsScrollingActiveTimerRegistered = bKeepTicking;
+	return bKeepTicking ? EActiveTimerReturnType::Continue : EActiveTimerReturnType::Stop;
 }
 
 void SScrollBox::Tick( const FGeometry& AllottedGeometry, const double InCurrentTime, const float InDeltaTime )
@@ -475,7 +519,12 @@ void SScrollBox::Tick( const FGeometry& AllottedGeometry, const double InCurrent
 	const float ViewOffset = FMath::Clamp<float>( DesiredScrollOffset/ContentSize, 0.0, 1.0 - ViewFraction );
 	
 	// Update the scrollbar with the clamped version of the offset
-	const float TargetPhysicalOffset = GetScrollComponentFromVector(ViewOffset*ScrollPanel->GetDesiredSize());
+	float TargetPhysicalOffset = GetScrollComponentFromVector(ViewOffset*ScrollPanel->GetDesiredSize());
+	if ( AllowOverscroll == EAllowOverscroll::Yes )
+	{
+		TargetPhysicalOffset -= Overscroll.GetOverscroll();
+	}
+
 	bIsScrolling = !FMath::IsNearlyEqual(TargetPhysicalOffset, ScrollPanel->PhysicalOffset, 0.001f);
 	ScrollPanel->PhysicalOffset = (bAnimateScroll)
 		? FMath::FInterpTo(ScrollPanel->PhysicalOffset, TargetPhysicalOffset, InDeltaTime, 15.f)
@@ -487,11 +536,23 @@ void SScrollBox::Tick( const FGeometry& AllottedGeometry, const double InCurrent
 		// We cannot scroll, so ensure that there is no offset.
 		ScrollPanel->PhysicalOffset = 0.0f;
 	}
-	else if (bIsScrolling && !bIsActiveTimerRegistered)
+}
+
+FReply SScrollBox::OnPreviewMouseButtonDown(const FGeometry& MyGeometry, const FPointerEvent& MouseEvent)
+{
+	if ( MouseEvent.IsTouchEvent() )
+	{
+		// Clear any inertia 
+		this->InertialScrollManager.ClearScrollVelocity();
+		// We have started a new interaction; track how far the user has moved since they put their finger down.
+		AmountScrolledWhileRightMouseDown = 0;
+		// Someone put their finger down in this list, so they probably want to drag the list.
+		bStartedTouchInteraction = true;
+		return FReply::Unhandled();
+	}
+	else
 	{
-		// If scrolling and the scrollbar is needed, make sure the active timer is registered (possible b/c we may need to scroll when our geometry changes)
-		bIsActiveTimerRegistered = true;
-		RegisterActiveTimer(0.f, FWidgetActiveTimerDelegate::CreateSP(this, &SScrollBox::UpdateInertialScroll));
+		return FReply::Unhandled();
 	}
 }
 
@@ -519,12 +580,12 @@ FReply SScrollBox::OnMouseButtonUp( const FGeometry& MyGeometry, const FPointerE
 {
 	if ( MouseEvent.GetEffectingButton() == EKeys::RightMouseButton )
 	{
-		if (!bIsActiveTimerRegistered && IsRightClickScrolling())
+		if ( !bIsScrollingActiveTimerRegistered && IsRightClickScrolling() )
 		{
 			// Register the active timer to handle the inertial scrolling
 			CachedGeometry = MyGeometry;
 			bIsScrolling = true;
-			bIsActiveTimerRegistered = true;
+			bIsScrollingActiveTimerRegistered = true;
 			RegisterActiveTimer(0.f, FWidgetActiveTimerDelegate::CreateSP(this, &SScrollBox::UpdateInertialScroll));
 		}
 
@@ -566,7 +627,7 @@ FReply SScrollBox::OnMouseMove( const FGeometry& MyGeometry, const FPointerEvent
 		if( IsRightClickScrolling() )
 		{
 			this->InertialScrollManager.AddScrollSample(-ScrollByAmount, FPlatformTime::Seconds());
-			const bool bDidScroll = this->ScrollBy( MyGeometry, -ScrollByAmount );
+			const bool bDidScroll = this->ScrollBy( MyGeometry, -ScrollByAmount, AllowOverscroll );
 
 			FReply Reply = FReply::Handled();
 
@@ -584,6 +645,8 @@ FReply SScrollBox::OnMouseMove( const FGeometry& MyGeometry, const FPointerEvent
 				SetScrollComponentOnVector(SoftwareCursorPosition, GetScrollComponentFromVector(SoftwareCursorPosition) + GetScrollComponentFromVector(MouseEvent.GetCursorDelta()));
 			}
 
+			TickScrollDelta -= ScrollByAmount;
+
 			return Reply;
 		}
 	}
@@ -593,10 +656,11 @@ FReply SScrollBox::OnMouseMove( const FGeometry& MyGeometry, const FPointerEvent
 
 void SScrollBox::OnMouseLeave( const FPointerEvent& MouseEvent )
 {
-	if(this->HasMouseCapture() == false)
+	if (HasMouseCapture() == false)
 	{
 		// No longer scrolling (unless we have mouse capture)
 		AmountScrolledWhileRightMouseDown = 0;
+		bStartedTouchInteraction = false;
 	}
 }
 
@@ -607,14 +671,14 @@ FReply SScrollBox::OnMouseWheel( const FGeometry& MyGeometry, const FPointerEven
 		// Make sure scroll velocity is cleared so it doesn't fight with the mouse wheel input
 		InertialScrollManager.ClearScrollVelocity();
 
-		const bool bScrollWasHandled = this->ScrollBy( MyGeometry, -MouseEvent.GetWheelDelta()*WheelScrollAmount );
+		const bool bScrollWasHandled = this->ScrollBy(MyGeometry, -MouseEvent.GetWheelDelta()*WheelScrollAmount, EAllowOverscroll::No);
 
-		if (bScrollWasHandled && !bIsActiveTimerRegistered)
+		if ( bScrollWasHandled && !bIsScrollingActiveTimerRegistered )
 		{
 			// Register the active timer to handle the inertial scrolling
 			CachedGeometry = MyGeometry;
 			bIsScrolling = true;
-			bIsActiveTimerRegistered = true;
+			bIsScrollingActiveTimerRegistered = true;
 			RegisterActiveTimer(0.f, FWidgetActiveTimerDelegate::CreateSP(this, &SScrollBox::UpdateInertialScroll));
 		}
 
@@ -626,7 +690,7 @@ FReply SScrollBox::OnMouseWheel( const FGeometry& MyGeometry, const FPointerEven
 	}
 }
 
-bool SScrollBox::ScrollBy( const FGeometry& AllottedGeometry, float ScrollAmount, bool InAnimateScroll )
+bool SScrollBox::ScrollBy(const FGeometry& AllottedGeometry, float ScrollAmount, EAllowOverscroll Overscrolling, bool InAnimateScroll)
 {
 	bAnimateScroll = InAnimateScroll;
 
@@ -635,11 +699,17 @@ bool SScrollBox::ScrollBy( const FGeometry& AllottedGeometry, float ScrollAmount
 
 	float PreviousScrollOffset = DesiredScrollOffset;
 
-	DesiredScrollOffset += ScrollAmount;
+	float UnclampedDesiredScrollOffset = DesiredScrollOffset + ScrollAmount;
 
 	const float ScrollMin = 0.0f;
 	const float ScrollMax = ContentSize - GetScrollComponentFromVector(ScrollPanelGeometry.Size);
-	DesiredScrollOffset = FMath::Clamp(DesiredScrollOffset, ScrollMin, ScrollMax);
+	DesiredScrollOffset = FMath::Clamp(UnclampedDesiredScrollOffset, ScrollMin, ScrollMax);
+
+	float RemainingScrollDelta = UnclampedDesiredScrollOffset - DesiredScrollOffset;
+	if ( Overscrolling == EAllowOverscroll::Yes && Overscroll.ShouldApplyOverscroll(DesiredScrollOffset == 0, DesiredScrollOffset == ScrollMax, RemainingScrollDelta) )
+	{
+		Overscroll.ScrollBy(-RemainingScrollDelta);
+	}
 
 	OnUserScrolled.ExecuteIfBound(DesiredScrollOffset);
 
@@ -664,6 +734,59 @@ FCursorReply SScrollBox::OnCursorQuery( const FGeometry& MyGeometry, const FPoin
 	}
 }
 
+FReply SScrollBox::OnTouchStarted(const FGeometry& MyGeometry, const FPointerEvent& InTouchEvent)
+{
+	// See OnPreviewMouseButtonDown()
+	//     if (MouseEvent.IsTouchEvent())
+
+	return FReply::Unhandled();
+}
+
+FReply SScrollBox::OnTouchMoved(const FGeometry& MyGeometry, const FPointerEvent& InTouchEvent)
+{
+	if ( bStartedTouchInteraction )
+	{
+		const float ScrollByAmount = GetScrollComponentFromVector(InTouchEvent.GetCursorDelta()) / MyGeometry.Scale;
+		AmountScrolledWhileRightMouseDown += FMath::Abs(ScrollByAmount);
+		TickScrollDelta -= ScrollByAmount;
+
+		if ( AmountScrolledWhileRightMouseDown > FSlateApplication::Get().GetDragTriggerDistance() )
+		{
+			const float AmountScrolled = this->ScrollBy(MyGeometry, -ScrollByAmount, EAllowOverscroll::Yes, false);
+
+			// The user has moved the list some amount; they are probably
+			// trying to scroll. From now on, the list assumes the user is scrolling
+			// until they lift their finger.
+			return FReply::Handled().CaptureMouse(AsShared());
+		}
+		return FReply::Handled();
+	}
+	else
+	{
+		return FReply::Handled();
+	}
+}
+
+FReply SScrollBox::OnTouchEnded(const FGeometry& MyGeometry, const FPointerEvent& InTouchEvent)
+{
+	CachedGeometry = MyGeometry;
+
+	AmountScrolledWhileRightMouseDown = 0;
+	bStartedTouchInteraction = false;
+
+	bIsScrollingActiveTimerRegistered = true;
+	RegisterActiveTimer(0.f, FWidgetActiveTimerDelegate::CreateSP(this, &SScrollBox::UpdateInertialScroll));
+
+	if ( HasMouseCapture() )
+	{
+		return FReply::Handled().ReleaseMouseCapture();
+	}
+	else
+	{
+		return FReply::Handled();
+	}
+}
+
 int32 SScrollBox::OnPaint( const FPaintArgs& Args, const FGeometry& AllottedGeometry, const FSlateRect& MyClippingRect, FSlateWindowElementList& OutDrawElements, int32 LayerId, const FWidgetStyle& InWidgetStyle, bool bParentEnabled ) const
 {
 	int32 NewLayerId = SCompoundWidget::OnPaint( Args, AllottedGeometry, MyClippingRect, OutDrawElements, LayerId, InWidgetStyle, bParentEnabled );
@@ -709,3 +832,12 @@ FSlateColor SScrollBox::GetEndShadowOpacity() const
 	
 	return FLinearColor(1.0f, 1.0f, 1.0f, ShadowOpacity);
 }
+
+bool SScrollBox::CanUseInertialScroll(float ScrollAmount) const
+{
+	const auto CurrentOverscroll = Overscroll.GetOverscroll();
+
+	// We allow sampling for the inertial scroll if we are not in the overscroll region,
+	// Or if we are scrolling outwards of the overscroll region
+	return CurrentOverscroll == 0.f || FMath::Sign(CurrentOverscroll) != FMath::Sign(ScrollAmount);
+}
diff --git a/Engine/Source/Runtime/Slate/Private/Widgets/Views/STableViewBase.cpp b/Engine/Source/Runtime/Slate/Private/Widgets/Views/STableViewBase.cpp
--- a/Engine/Source/Runtime/Slate/Private/Widgets/Views/STableViewBase.cpp
+++ b/Engine/Source/Runtime/Slate/Private/Widgets/Views/STableViewBase.cpp
@@ -811,12 +811,12 @@ bool STableViewBase::CanUseInertialScroll( float ScrollAmount ) const
 	return CurrentOverscroll == 0.f || FMath::Sign(CurrentOverscroll) != FMath::Sign(ScrollAmount);
 }
 
-STableViewBase::FOverscroll::FOverscroll()
+STableViewBase::FListOverscroll::FListOverscroll()
 : OverscrollAmount( 0.0f )
 {
 }
 
-float STableViewBase::FOverscroll::ScrollBy( float Delta )
+float STableViewBase::FListOverscroll::ScrollBy(float Delta)
 {
 	const float ValueBeforeDeltaApplied = OverscrollAmount;
 	const float EasedDelta = Delta / (FMath::Abs(OverscrollAmount/ListConstants::OvershootMax)+1.0f);
@@ -832,12 +832,12 @@ float STableViewBase::FOverscroll::ScrollBy( float Delta )
 	return ValueBeforeDeltaApplied - OverscrollAmount;
 }
 
-float STableViewBase::FOverscroll::GetOverscroll() const
+float STableViewBase::FListOverscroll::GetOverscroll() const
 {
 	return OverscrollAmount;
 }
 
-void STableViewBase::FOverscroll::UpdateOverscroll( float InDeltaTime )
+void STableViewBase::FListOverscroll::UpdateOverscroll(float InDeltaTime)
 {
 	const float PullForce = FMath::Abs(OverscrollAmount) + 1.0f;
 	const float EasedDelta = ListConstants::OvershootBounceRate * InDeltaTime * PullForce;
@@ -852,7 +852,7 @@ void STableViewBase::FOverscroll::UpdateOverscroll( float InDeltaTime )
 	}
 }
 
-bool STableViewBase::FOverscroll::ShouldApplyOverscroll( const bool bIsAtStartOfList, const bool bIsAtEndOfList, const float ScrollDelta ) const
+bool STableViewBase::FListOverscroll::ShouldApplyOverscroll(const bool bIsAtStartOfList, const bool bIsAtEndOfList, const float ScrollDelta) const
 {
 	const bool bShouldApplyOverscroll =
 		// We can scroll past the edge of the list only if we are at the edge
diff --git a/Engine/Source/Runtime/Slate/Public/Framework/Layout/Overscroll.h b/Engine/Source/Runtime/Slate/Public/Framework/Layout/Overscroll.h
new file mode 100644
--- /dev/null
+++ b/Engine/Source/Runtime/Slate/Public/Framework/Layout/Overscroll.h
@@ -0,0 +1,40 @@
+// Copyright 1998-2015 Epic Games, Inc. All Rights Reserved.
+
+#pragma once
+
+enum class EAllowOverscroll : uint8
+{
+	Yes,
+	No
+};
+
+struct SLATE_API FOverscroll
+{
+public:
+
+	FOverscroll();
+
+	/** @return The Amount actually scrolled */
+	float ScrollBy(float Delta);
+
+	/** How far the user scrolled above/below the beginning/end of the list. */
+	float GetOverscroll() const;
+
+	/** Ticks the overscroll manager so it can animate. */
+	void UpdateOverscroll(float InDeltaTime);
+
+	/**
+	 * Should ScrollDelta be applied to overscroll or to regular item scrolling.
+	 *
+	 * @param bIsAtStartOfList  Are we at the very beginning of the list (i.e. showing the first item at the top of the view)?
+	 * @param bIsAtEndOfList    Are we showing the last item on the screen completely?
+	 * @param ScrollDelta       How much the user is trying to scroll in Slate Units.
+	 *
+	 * @return true if the user's scrolling should be applied toward overscroll.
+	 */
+	bool ShouldApplyOverscroll(const bool bIsAtStartOfList, const bool bIsAtEndOfList, const float ScrollDelta) const;
+
+private:
+	/** How much we've over-scrolled above/below the beginning/end of the list, stored in log form */
+	float OverscrollAmount;
+};
diff --git a/Engine/Source/Runtime/Slate/Public/Widgets/Layout/SScrollBox.h b/Engine/Source/Runtime/Slate/Public/Widgets/Layout/SScrollBox.h
--- a/Engine/Source/Runtime/Slate/Public/Widgets/Layout/SScrollBox.h
+++ b/Engine/Source/Runtime/Slate/Public/Widgets/Layout/SScrollBox.h
@@ -3,6 +3,7 @@
 #pragma once
 
 #include "InertialScrollManager.h"
+#include "Overscroll.h"
 
 /** SScrollBox can scroll through an arbitrary number of widgets. */
 class SLATE_API SScrollBox : public SCompoundWidget
@@ -46,6 +47,7 @@ public:
 		, _ScrollBarVisibility(EVisibility::Visible)
 		, _ScrollBarAlwaysVisible(false)
 		, _ScrollBarThickness(FVector2D(5, 5))
+		, _AllowOverscroll(EAllowOverscroll::Yes)
 		, _OnUserScrolled()
 		, _ConsumeMouseWheel(EConsumeMouseWheel::WhenScrollingPossible)
 		{}
@@ -70,6 +72,8 @@ public:
 
 		SLATE_ARGUMENT( FVector2D, ScrollBarThickness )
 
+		SLATE_ARGUMENT(EAllowOverscroll, AllowOverscroll);
+
 		/** Called when the button is clicked */
 		SLATE_EVENT(FOnUserScrolled, OnUserScrolled)
 
@@ -130,6 +134,7 @@ public:
 
 	// SWidget interface
 	virtual void Tick( const FGeometry& AllottedGeometry, const double InCurrentTime, const float InDeltaTime ) override;
+	virtual FReply OnPreviewMouseButtonDown(const FGeometry& MyGeometry, const FPointerEvent& MouseEvent) override;
 	virtual FReply OnMouseButtonDown( const FGeometry& MyGeometry, const FPointerEvent& MouseEvent ) override;
 	virtual FReply OnMouseButtonUp( const FGeometry& MyGeometry, const FPointerEvent& MouseEvent ) override;
 	virtual FReply OnMouseMove( const FGeometry& MyGeometry, const FPointerEvent& MouseEvent ) override;
@@ -138,6 +143,9 @@ public:
 	virtual FReply OnDragDetected( const FGeometry& MyGeometry, const FPointerEvent& MouseEvent ) override;
 	virtual FCursorReply OnCursorQuery( const FGeometry& MyGeometry, const FPointerEvent& CursorEvent ) const override;
 	virtual int32 OnPaint( const FPaintArgs& Args, const FGeometry& AllottedGeometry, const FSlateRect& MyClippingRect, FSlateWindowElementList& OutDrawElements, int32 LayerId, const FWidgetStyle& InWidgetStyle, bool bParentEnabled ) const override;
+	virtual FReply OnTouchStarted(const FGeometry& MyGeometry, const FPointerEvent& InTouchEvent) override;
+	virtual FReply OnTouchMoved(const FGeometry& MyGeometry, const FPointerEvent& InTouchEvent) override;
+	virtual FReply OnTouchEnded(const FGeometry& MyGeometry, const FPointerEvent& InTouchEvent) override;
 	// End of SWidget interface
 
 private:
@@ -181,7 +189,7 @@ private:
 	 * @param InAnimateScroll	Whether or not to animate the scroll
 	 * @return Whether or not the scroll was fully handled
 	 */
-	bool ScrollBy( const FGeometry& AllottedGeometry, float ScrollAmount, bool InAnimateScroll = true );
+	bool ScrollBy(const FGeometry& AllottedGeometry, float ScrollAmount, EAllowOverscroll Overscroll, bool InAnimateScroll = true);
 
 	/** Invoked when the user scroll via the scrollbar */
 	void ScrollBar_OnUserScrolled( float InScrollOffsetFraction );
@@ -195,6 +203,9 @@ private:
 	/** Active timer to update inertial scrolling as needed */
 	EActiveTimerReturnType UpdateInertialScroll(double InCurrentTime, float InDeltaTime);
 
+	/** Check whether the current state of the table warrants inertial scroll by the specified amount */
+	bool CanUseInertialScroll(float ScrollAmount) const;
+
 private:
 
 	/** The panel which stacks the child slots */
@@ -203,12 +214,24 @@ private:
 	/** The scrollbar which controls scrolling for the scrollbox. */
 	TSharedPtr<SScrollBar> ScrollBar;
 
+	/** The amount we have scrolled this tick cycle */
+	float TickScrollDelta;
+
+	/** Did the user start an interaction in this list? */
+	bool bStartedTouchInteraction;
+
 	/** How much we scrolled while the rmb has been held */
 	float AmountScrolledWhileRightMouseDown;
 
 	/** Helper object to manage inertial scrolling */
 	FInertialScrollManager InertialScrollManager;
 
+	/** The overscroll state management structure. */
+	FOverscroll Overscroll;
+
+	/** Whether to permit overscroll on this scroll box */
+	EAllowOverscroll AllowOverscroll;
+
 	/**	The current position of the software cursor */
 	FVector2D SoftwareCursorPosition;
 
@@ -246,7 +269,5 @@ private:
 	bool bScrollToEnd : 1;
 
 	/** Whether the active timer to update the inertial scroll is registered */
-	bool bIsActiveTimerRegistered : 1;
+	bool bIsScrollingActiveTimerRegistered : 1;
 };
-
-
diff --git a/Engine/Source/Runtime/Slate/Public/Widgets/Views/STableViewBase.h b/Engine/Source/Runtime/Slate/Public/Widgets/Views/STableViewBase.h
--- a/Engine/Source/Runtime/Slate/Public/Widgets/Views/STableViewBase.h
+++ b/Engine/Source/Runtime/Slate/Public/Widgets/Views/STableViewBase.h
@@ -4,17 +4,12 @@
 
 #include "IScrollableWidget.h"
 #include "InertialScrollManager.h"
+#include "Overscroll.h"
 
 class SListPanel;
 class SHeaderRow;
 
 
-enum class EAllowOverscroll
-{
-	Yes,
-	No
-};
-
 /** If the list panel is arranging items horizontally, this enum dictates how the items should be aligned (basically, where any extra space is placed) */
 enum class EListItemAlignment : uint8
 {
@@ -316,11 +311,11 @@ private:
 
 protected:
 
-	struct SLATE_API FOverscroll
+	struct SLATE_API FListOverscroll
 	{
 	public:
 
-		FOverscroll();
+		FListOverscroll();
 
 		/** @return The Amount actually scrolled */
 		float ScrollBy( float Delta );
