We no longer clear focus on child widgets of the game viewport when the window loses focus by default.

diff --git a/Engine/Source/Runtime/Slate/Private/Framework/Application/SlateApplication.cpp b/Engine/Source/Runtime/Slate/Private/Framework/Application/SlateApplication.cpp
--- a/Engine/Source/Runtime/Slate/Private/Framework/Application/SlateApplication.cpp
+++ b/Engine/Source/Runtime/Slate/Private/Framework/Application/SlateApplication.cpp
@@ -331,6 +331,9 @@ namespace SlateDefs
 
 	// How far tool tips should be pushed out from a force field border, in pixels
 	static const FVector2D ToolTipOffsetFromForceField( 4.0f, 3.0f );
+
+	// When true we will clear keyboard focus from widgets within the game view port when the window loses focus
+	static bool ClearGameViewportFocusWhenWindowLosesFocus = false;
 }
 
 
@@ -2210,7 +2213,7 @@ bool FSlateApplication::SetUserFocus(const uint32 InUserIndex, const FWidgetPath
 		
 		// Set the windows restore state.
 		//@Todo Slate: We need to store the full focus state 
-		//@Todo Slate: Why are we chacking FocusedWindow != NewFocusedWidget?
+		//@Todo Slate: Why are we checking FocusedWindow != NewFocusedWidget?
 		if (FocusedWindow.IsValid() && FocusedWindow != NewFocusedWidget)
 		{
 			FocusedWindow->SetWidgetToFocusOnActivate(NewFocusedWidget);
@@ -5096,7 +5099,7 @@ bool FSlateApplication::ProcessWindowActivatedEvent( const FWindowActivateEvent&
 		ActivateEvent.GetAffectedWindow()->OnIsActiveChanged( ActivateEvent );
 
 		// A window was deactivated; mouse capture should be cleared
-		ResetToDefaultInputSettings();
+		ResetToDefaultPointerInputSettings();
 	}
 
 	return true;
@@ -5118,7 +5121,7 @@ void FSlateApplication::ProcessApplicationActivationEvent( bool InAppActivated )
 	// If the user switched to a different application then we should dismiss our pop-ups.  In the case
 	// where a user clicked on a different Slate window, OnWindowActivatedMessage() will be call MenuStack.OnWindowActivated()
 	// to destroy any windows in our stack that are no longer appropriate to be displayed.
-	if( UserSwitchedAway )
+	if (UserSwitchedAway)
 	{
 		// Close pop-up menus
 		DismissAllMenus();
@@ -5129,8 +5132,18 @@ void FSlateApplication::ProcessApplicationActivationEvent( bool InAppActivated )
 		// No slate window is active when our entire app becomes inactive
 		bSlateWindowActive = false;
 
-		// Clear keyboard focus when the app is deactivated
-		ClearKeyboardFocus(EFocusCause::OtherWidgetLostFocus);
+		// Clear keyboard focus when the app is deactivated, if the widget isn't a child of the game viewport.
+		bool bClearKeyboardFocus = true;
+
+		if (!SlateDefs::ClearGameViewportFocusWhenWindowLosesFocus && GameViewportWidget.IsValid())
+		{
+			FUserFocusEntry& UserFocusEntry = UserFocusEntries[GetUserIndexForKeyboard()];
+			bClearKeyboardFocus = !(UserFocusEntry.WidgetPath.ContainsWidget(GameViewportWidget.Pin().ToSharedRef()));
+		}
+		if (bClearKeyboardFocus)
+		{
+			ClearKeyboardFocus(EFocusCause::OtherWidgetLostFocus);
+		}
 
 		// If we have a slate-only drag-drop occurring, stop the drag drop.
 		if ( IsDragDropping() && !DragDropContent->IsExternalOperation() )
