Slate - Tweaking the overscroll rebounding math to be more consistent without framerate affecting the speed.

diff --git a/Engine/Source/Runtime/Slate/Private/Framework/Layout/Overscroll.cpp b/Engine/Source/Runtime/Slate/Private/Framework/Layout/Overscroll.cpp
--- a/Engine/Source/Runtime/Slate/Private/Framework/Layout/Overscroll.cpp
+++ b/Engine/Source/Runtime/Slate/Private/Framework/Layout/Overscroll.cpp
@@ -6,8 +6,10 @@ namespace OverscrollConstants
 {
 	/** The amount to scale the logarithm by to make it more loose */
 	static const float Looseness = 50.0f;
+	/** The "max" used to perform the interpolation snap back, and make it faster the further away it is. */
+	static const float OvershootLooseMax = 100.0f;
 	/** The bounce back rate when the overscroll stops. */
-	static const float OvershootBounceRate = 2000.0f;
+	static const float OvershootBounceRate = 1500.0f;
 }
 
 FOverscroll::FOverscroll()
@@ -47,15 +49,15 @@ float FOverscroll::GetOverscroll() const
 void FOverscroll::UpdateOverscroll(float InDeltaTime)
 {
 	const float PullForce = FMath::Abs(OverscrollAmount) + 1.0f;
-	const float EasedDelta = OverscrollConstants::OvershootBounceRate * InDeltaTime * PullForce;
+	const float EasedDelta = OverscrollConstants::OvershootBounceRate * InDeltaTime * FMath::Max(1.0f, PullForce / OverscrollConstants::OvershootLooseMax);
 
 	if ( OverscrollAmount > 0 )
 	{
-		OverscrollAmount = FMath::Max(0.0f, OverscrollAmount - EasedDelta * InDeltaTime);
+		OverscrollAmount = FMath::Max(0.0f, OverscrollAmount - EasedDelta);
 	}
 	else
 	{
-		OverscrollAmount = FMath::Min(0.0f, OverscrollAmount + EasedDelta * InDeltaTime);
+		OverscrollAmount = FMath::Min(0.0f, OverscrollAmount + EasedDelta);
 	}
 }
 
diff --git a/Engine/Source/Runtime/Slate/Private/Widgets/Layout/SScrollBox.cpp b/Engine/Source/Runtime/Slate/Private/Widgets/Layout/SScrollBox.cpp
--- a/Engine/Source/Runtime/Slate/Private/Widgets/Layout/SScrollBox.cpp
+++ b/Engine/Source/Runtime/Slate/Private/Widgets/Layout/SScrollBox.cpp
@@ -522,7 +522,7 @@ void SScrollBox::Tick( const FGeometry& AllottedGeometry, const double InCurrent
 	float TargetPhysicalOffset = GetScrollComponentFromVector(ViewOffset*ScrollPanel->GetDesiredSize());
 	if ( AllowOverscroll == EAllowOverscroll::Yes )
 	{
-		TargetPhysicalOffset -= Overscroll.GetOverscroll();
+		TargetPhysicalOffset += Overscroll.GetOverscroll();
 	}
 
 	bIsScrolling = !FMath::IsNearlyEqual(TargetPhysicalOffset, ScrollPanel->PhysicalOffset, 0.001f);
@@ -697,18 +697,21 @@ bool SScrollBox::ScrollBy(const FGeometry& AllottedGeometry, float ScrollAmount,
 	const float ContentSize = GetScrollComponentFromVector(ScrollPanel->GetDesiredSize());
 	const FGeometry ScrollPanelGeometry = FindChildGeometry( AllottedGeometry, ScrollPanel.ToSharedRef() );
 
-	float PreviousScrollOffset = DesiredScrollOffset;
-
-	float UnclampedDesiredScrollOffset = DesiredScrollOffset + ScrollAmount;
+	const float PreviousScrollOffset = DesiredScrollOffset;
 
 	const float ScrollMin = 0.0f;
 	const float ScrollMax = ContentSize - GetScrollComponentFromVector(ScrollPanelGeometry.Size);
-	DesiredScrollOffset = FMath::Clamp(UnclampedDesiredScrollOffset, ScrollMin, ScrollMax);
 
-	float RemainingScrollDelta = UnclampedDesiredScrollOffset - DesiredScrollOffset;
-	if ( Overscrolling == EAllowOverscroll::Yes && Overscroll.ShouldApplyOverscroll(DesiredScrollOffset == 0, DesiredScrollOffset == ScrollMax, RemainingScrollDelta) )
+	if ( ScrollAmount != 0 )
 	{
-		Overscroll.ScrollBy(-RemainingScrollDelta);
+		if ( Overscrolling == EAllowOverscroll::Yes && Overscroll.ShouldApplyOverscroll(DesiredScrollOffset == 0, DesiredScrollOffset == ScrollMax, ScrollAmount) )
+		{
+			Overscroll.ScrollBy(ScrollAmount);
+		}
+		else
+		{
+			DesiredScrollOffset = FMath::Clamp(DesiredScrollOffset + ScrollAmount, ScrollMin, ScrollMax);
+		}
 	}
 
 	OnUserScrolled.ExecuteIfBound(DesiredScrollOffset);
diff --git a/Engine/Source/Runtime/Slate/Public/Framework/Layout/Overscroll.h b/Engine/Source/Runtime/Slate/Public/Framework/Layout/Overscroll.h
--- a/Engine/Source/Runtime/Slate/Public/Framework/Layout/Overscroll.h
+++ b/Engine/Source/Runtime/Slate/Public/Framework/Layout/Overscroll.h
@@ -8,6 +8,9 @@ enum class EAllowOverscroll : uint8
 	No
 };
 
+/**
+ * Handles overscroll management.
+ */
 struct SLATE_API FOverscroll
 {
 public:
