Fix for previous changes which still had problems when you hit escape a number of times.
#UE-13296 - SMultiLineEditableTextBox does not implement revert on escape
#Branch UE4
#Proj Slate

#info The wrong interface was being called to update the text which resulted in de-synchronisation between the text and the text consumer. Copy paste error :(

diff --git a/Engine/Source/Runtime/Slate/Private/Widgets/Text/SMultiLineEditableText.cpp b/Engine/Source/Runtime/Slate/Private/Widgets/Text/SMultiLineEditableText.cpp
--- a/Engine/Source/Runtime/Slate/Private/Widgets/Text/SMultiLineEditableText.cpp
+++ b/Engine/Source/Runtime/Slate/Private/Widgets/Text/SMultiLineEditableText.cpp
@@ -1780,10 +1780,11 @@ void SMultiLineEditableText::RestoreOriginalText()
 {
 	if(HasTextChangedFromOriginal())
 	{
-		SaveText(OriginalText.Text);
+		SetText(OriginalText.Text);
+		TextLayout->UpdateIfNeeded();
 
 		// Let outsiders know that the text content has been changed
-		OnTextChanged.ExecuteIfBound(OriginalText.Text);
+		OnTextCommitted.ExecuteIfBound(OriginalText.Text, ETextCommit::OnCleared);
 	}
 }
 
