[INTEGRATE] merge UE4-LauncherDev->UE4 cl 2532197

diff --git a/Engine/Source/Runtime/Slate/Private/Widgets/Text/SMultiLineEditableText.cpp b/Engine/Source/Runtime/Slate/Private/Widgets/Text/SMultiLineEditableText.cpp
--- a/Engine/Source/Runtime/Slate/Private/Widgets/Text/SMultiLineEditableText.cpp
+++ b/Engine/Source/Runtime/Slate/Private/Widgets/Text/SMultiLineEditableText.cpp
@@ -378,6 +378,22 @@ void SMultiLineEditableText::SetText(const TAttribute< FText >& InText)
 	}
 }
 
+const FText SMultiLineEditableText::GetPlainText() const
+{
+	const TArray< FTextLayout::FLineModel >& Lines = TextLayout->GetLineModels();
+	const int32 NumberOfLines = Lines.Num();
+	if( NumberOfLines > 0)
+	{
+		FString SelectedText;
+		const FTextLocation StartLocation(0, 0);
+		const FTextLocation EndLocation(NumberOfLines -1, Lines[NumberOfLines -1].Text->Len());
+		FTextSelection Selection(StartLocation, EndLocation);
+		TextLayout->GetSelectionAsText(SelectedText, Selection);
+		return FText::FromString(SelectedText);
+	}
+	return FText::GetEmpty();
+}
+
 void SMultiLineEditableText::SetHintText(const TAttribute< FText >& InHintText)
 {
 	HintText = InHintText;
@@ -1549,6 +1565,11 @@ void SMultiLineEditableText::GoTo(const FTextLocation& NewLocation)
 	}
 }
 
+void SMultiLineEditableText::GoTo(ETextLocation GoToLocation)
+{
+	JumpTo(GoToLocation, ECursorAction::MoveCursor);
+}
+
 void SMultiLineEditableText::ScrollTo(const FTextLocation& NewLocation)
 {
 	const TArray< FTextLayout::FLineModel >& Lines = TextLayout->GetLineModels();
diff --git a/Engine/Source/Runtime/Slate/Public/Widgets/Input/SMultiLineEditableTextBox.h b/Engine/Source/Runtime/Slate/Public/Widgets/Input/SMultiLineEditableTextBox.h
--- a/Engine/Source/Runtime/Slate/Public/Widgets/Input/SMultiLineEditableTextBox.h
+++ b/Engine/Source/Runtime/Slate/Public/Widgets/Input/SMultiLineEditableTextBox.h
@@ -172,6 +172,16 @@ public:
 		return EditableText->GetText();
 	}
 
+	/**
+	 * Returns the plain text string without richtext formatting
+	 *
+	 * @return  Text string
+	 */
+	const FText GetPlainText() const
+	{
+		return EditableText->GetPlainText();
+	}
+
 	/**
 	 * Sets the text string currently being edited 
 	 *
@@ -232,6 +242,12 @@ public:
 	/** Move the cursor to the given location in the document */
 	void GoTo(const FTextLocation& NewLocation);
 
+	/** Move the cursor to the specified location */
+	void GoTo(ETextLocation NewLocation)
+	{
+		EditableText->GoTo(NewLocation);
+	}
+
 	/** Scroll to the given location in the document (without moving the cursor) */
 	void ScrollTo(const FTextLocation& NewLocation);
 
diff --git a/Engine/Source/Runtime/Slate/Public/Widgets/Text/SMultiLineEditableText.h b/Engine/Source/Runtime/Slate/Public/Widgets/Text/SMultiLineEditableText.h
--- a/Engine/Source/Runtime/Slate/Public/Widgets/Text/SMultiLineEditableText.h
+++ b/Engine/Source/Runtime/Slate/Public/Widgets/Text/SMultiLineEditableText.h
@@ -126,6 +126,12 @@ public:
 	 */
 	void SetText(const TAttribute< FText >& InText);
 
+	/**
+	 * Returns the plain text string without richtext formatting
+	 * @return  Text string
+	 */
+	const FText GetPlainText() const;
+
 	/**
 	 * Sets the text that appears when there is no text in the text box
 	 */
@@ -147,6 +153,9 @@ public:
 	/** Move the cursor to the given location in the document (will also scroll to this point) */
 	void GoTo(const FTextLocation& NewLocation);
 
+	/** Move the cursor specified location */
+	void GoTo(ETextLocation NewLocation);
+
 	/** Scroll to the given location in the document (without moving the cursor) */
 	void ScrollTo(const FTextLocation& NewLocation);
 
