Fixed graphical corruption while in an intra-frame debugging loop due to hitting a breakpoint in a widget blueprint tick method.

https://jira.ol.epicgames.net/browse/UE-14586

diff --git a/Engine/Source/Runtime/Slate/Private/Framework/Application/SlateApplication.cpp b/Engine/Source/Runtime/Slate/Private/Framework/Application/SlateApplication.cpp
--- a/Engine/Source/Runtime/Slate/Private/Framework/Application/SlateApplication.cpp
+++ b/Engine/Source/Runtime/Slate/Private/Framework/Application/SlateApplication.cpp
@@ -13,6 +13,8 @@
 #include "ToolboxModule.h"
 #include "TabCommands.h"
 
+extern SLATECORE_API TOptional<FShortRect> GSlateScissorRect;
+
 class FEventRouter
 {
 
@@ -3078,6 +3080,11 @@ void FSlateApplication::EnterDebuggingMode()
 	GFirstFrameIntraFrameDebugging = true;
 #endif	//WITH_EDITORONLY_DATA
 
+
+	// The scissor rect stack must be reset when re-entering the tick loop to avoid graphical artifacts with existing clip rects applied new widgets
+	TOptional<FShortRect> PreviousScissorRect = GSlateScissorRect;
+	GSlateScissorRect.Reset();
+
 	// Tick slate from here in the event that we should not return until the modal window is closed.
 	while (!bRequestLeaveDebugMode)
 	{
@@ -3115,6 +3122,9 @@ void FSlateApplication::EnterDebuggingMode()
 		}
 	}
 
+	// Reset the scissor rect back to what it was before we started debugging
+	GSlateScissorRect = PreviousScissorRect;
+
 	bLeaveDebugForSingleStep = false;
 }
 
