Fixed updating of Mac menu bar when a new tab is added

diff --git a/Engine/Source/Runtime/Slate/Private/Framework/Docking/TabManager.cpp b/Engine/Source/Runtime/Slate/Private/Framework/Docking/TabManager.cpp
--- a/Engine/Source/Runtime/Slate/Private/Framework/Docking/TabManager.cpp
+++ b/Engine/Source/Runtime/Slate/Private/Framework/Docking/TabManager.cpp
@@ -860,6 +860,8 @@ void FTabManager::DrawAttention( const TSharedRef<SDockTab>& TabToHighlight )
 	}
 	TabToHighlight->GetParentDockTabStack()->BringToFront(TabToHighlight);
 	TabToHighlight->FlashTab();
+
+	FGlobalTabmanager::Get()->UpdateMainMenu(TabToHighlight, true);
 }
 
 void FTabManager::InsertNewDocumentTab( FName PlaceholderId, ESearchPreference::Type SearchPreference, const TSharedRef<SDockTab>& UnmanagedTab )
@@ -880,6 +882,9 @@ TSharedRef<SDockTab> FTabManager::InvokeTab( const FTabId& TabId )
 	{
 		ParentWindowPtr->SetTitle( NewTab->GetTabLabel() );
 	}
+#if PLATFORM_MAC
+	FPlatformMisc::bChachedMacMenuStateNeedsUpdate = true;
+#endif
 	return NewTab;
 }
 
