Fix crashes caused by setting the image asset to none for the touch interface control

diff --git a/Engine/Source/Runtime/Engine/Private/GameFramework/TouchInterface.cpp b/Engine/Source/Runtime/Engine/Private/GameFramework/TouchInterface.cpp
--- a/Engine/Source/Runtime/Engine/Private/GameFramework/TouchInterface.cpp
+++ b/Engine/Source/Runtime/Engine/Private/GameFramework/TouchInterface.cpp
@@ -33,8 +33,8 @@ void UTouchInterface::Activate(TSharedPtr<SVirtualJoystick> VirtualJoystick)
 			FTouchInputControl Control = Controls[ControlIndex];
 			SVirtualJoystick::FControlInfo* SlateControl = new(SlateControls)SVirtualJoystick::FControlInfo;
 
-			SlateControl->Image1 = FCoreStyle::GetDynamicImageBrush("Engine.Joystick.Image1", Control.Image1, Control.Image1->GetFName());
-			SlateControl->Image2 = FCoreStyle::GetDynamicImageBrush("Engine.Joystick.Image2", Control.Image2, Control.Image2->GetFName());
+			SlateControl->Image1 = (Control.Image1 ? FCoreStyle::GetDynamicImageBrush("Engine.Joystick.Image1", Control.Image1, Control.Image1->GetFName()) : nullptr);
+			SlateControl->Image2 = (Control.Image2 ? FCoreStyle::GetDynamicImageBrush("Engine.Joystick.Image2", Control.Image2, Control.Image2->GetFName()) : nullptr);
 			SlateControl->Center = Control.Center;
 			SlateControl->VisualSize = Control.VisualSize;
 			SlateControl->ThumbSize = Control.ThumbSize;
diff --git a/Engine/Source/Runtime/Slate/Private/Widgets/Input/SVirtualJoystick.cpp b/Engine/Source/Runtime/Slate/Private/Widgets/Input/SVirtualJoystick.cpp
--- a/Engine/Source/Runtime/Slate/Private/Widgets/Input/SVirtualJoystick.cpp
+++ b/Engine/Source/Runtime/Slate/Private/Widgets/Input/SVirtualJoystick.cpp
@@ -115,29 +115,35 @@ int32 SVirtualJoystick::OnPaint( const FPaintArgs& Args, const FGeometry& Allott
 		{
 			const FControlInfo& Control = Controls[ControlIndex];
 
-			FSlateDrawElement::MakeBox(
-				OutDrawElements,
-				RetLayerId++,
-				AllottedGeometry.ToPaintGeometry(
-				Control.VisualCenter - FVector2D(Control.VisualSize.X * 0.5f, Control.VisualSize.Y * 0.5f),
-				Control.VisualSize),
-				Control.Image2.Get(),
-				MyClippingRect,
-				ESlateDrawEffect::None,
-				ColorAndOpacitySRGB
-				);
-
-			FSlateDrawElement::MakeBox(
-				OutDrawElements,
-				RetLayerId++,
-				AllottedGeometry.ToPaintGeometry(
-				Control.VisualCenter + Control.ThumbPosition - FVector2D(Control.ThumbSize.X * 0.5f, Control.ThumbSize.Y * 0.5f),
-				Control.ThumbSize),
-				Control.Image1.Get(),
-				MyClippingRect,
-				ESlateDrawEffect::None,
-				ColorAndOpacitySRGB
-				);
+			if (Control.Image2.IsValid())
+			{
+				FSlateDrawElement::MakeBox(
+					OutDrawElements,
+					RetLayerId++,
+					AllottedGeometry.ToPaintGeometry(
+					Control.VisualCenter - FVector2D(Control.VisualSize.X * 0.5f, Control.VisualSize.Y * 0.5f),
+					Control.VisualSize),
+					Control.Image2.Get(),
+					MyClippingRect,
+					ESlateDrawEffect::None,
+					ColorAndOpacitySRGB
+					);
+			}
+
+			if (Control.Image1.IsValid())
+			{
+				FSlateDrawElement::MakeBox(
+					OutDrawElements,
+					RetLayerId++,
+					AllottedGeometry.ToPaintGeometry(
+					Control.VisualCenter + Control.ThumbPosition - FVector2D(Control.ThumbSize.X * 0.5f, Control.ThumbSize.Y * 0.5f),
+					Control.ThumbSize),
+					Control.Image1.Get(),
+					MyClippingRect,
+					ESlateDrawEffect::None,
+					ColorAndOpacitySRGB
+					);
+			}
 		}
 	}
 	
