Add bool FSlateApplicationBase::IsSlateAsleep() to allow querying whether Slate is currently sleeping.

Merging CL2533947 from LauncherDev back to main (using UE4-To-UE4-LauncherDev)

diff --git a/Engine/Source/Runtime/Slate/Private/Framework/Application/SlateApplication.cpp b/Engine/Source/Runtime/Slate/Private/Framework/Application/SlateApplication.cpp
--- a/Engine/Source/Runtime/Slate/Private/Framework/Application/SlateApplication.cpp
+++ b/Engine/Source/Runtime/Slate/Private/Framework/Application/SlateApplication.cpp
@@ -1390,8 +1390,10 @@ void FSlateApplication::Tick()
 	// This effectively means the slate application is totally idle and we don't need to update the UI.
 	// This relies on Widgets properly registering for Active timer when they need something to happen even
 	// when the user is not providing any input (ie, animations, viewport rendering, async polling, etc).
+	bIsSlateAsleep = true;
 	if (!AllowSlateToSleep.GetValueOnGameThread() || bAnyActiveTimersPending || !bIsUserIdle || bNeedsSyntheticMouseMouse)
 	{
+		bIsSlateAsleep = false; // if we get here, then Slate is not sleeping
 		if (!bFoldTick)
 		{
 			SCOPE_CYCLE_COUNTER(STAT_SlateTickWindowAndChildren);
diff --git a/Engine/Source/Runtime/SlateCore/Private/Application/SlateApplicationBase.cpp b/Engine/Source/Runtime/SlateCore/Private/Application/SlateApplicationBase.cpp
--- a/Engine/Source/Runtime/SlateCore/Private/Application/SlateApplicationBase.cpp
+++ b/Engine/Source/Runtime/SlateCore/Private/Application/SlateApplicationBase.cpp
@@ -22,6 +22,7 @@ FWidgetPath FHitTesting::LocateWidgetInWindow(FVector2D ScreenspaceMouseCoordina
 FSlateApplicationBase::FSlateApplicationBase()
 : Renderer()
 , HitTesting(this)
+, bIsSlateAsleep(false)
 {
 
 }
@@ -75,3 +76,8 @@ bool FSlateApplicationBase::AnyActiveTimersArePending()
 
 	return bAnyTickReady;
 }
+
+bool FSlateApplicationBase::IsSlateAsleep()
+{
+	return bIsSlateAsleep;
+}
diff --git a/Engine/Source/Runtime/SlateCore/Public/Application/SlateApplicationBase.h b/Engine/Source/Runtime/SlateCore/Public/Application/SlateApplicationBase.h
--- a/Engine/Source/Runtime/SlateCore/Public/Application/SlateApplicationBase.h
+++ b/Engine/Source/Runtime/SlateCore/Public/Application/SlateApplicationBase.h
@@ -421,5 +421,20 @@ protected:
 
 	// Holds a pointer to the platform application.
 	static TSharedPtr<class GenericApplication> PlatformApplication;
+
+public:
+
+	/**
+	 * Is Slate currently sleeping or not.
+	 *
+	 * @return True if Slate is sleeping.
+	 */
+	bool IsSlateAsleep();
+
+protected:
+
+	// Gets set when Slate goes to sleep and cleared when active.
+	bool bIsSlateAsleep;
+
 };
 
