Slate: Fixed an issue where STableRow::GetContent would always return SNullWidget when the content was set via SetWidget

diff --git a/Engine/Source/Runtime/Slate/Public/Widgets/Views/STableRow.h b/Engine/Source/Runtime/Slate/Public/Widgets/Views/STableRow.h
--- a/Engine/Source/Runtime/Slate/Public/Widgets/Views/STableRow.h
+++ b/Engine/Source/Runtime/Slate/Public/Widgets/Views/STableRow.h
@@ -706,6 +706,12 @@ public:
 		return SharedThis(this);
 	}
 
+	virtual void SetContent(TSharedRef< SWidget > InContent) override
+	{
+		this->Content = InContent;
+		SBorder::SetContent(InContent);
+	}
+
 	virtual TSharedPtr<SWidget> GetContent() override
 	{
 		if ( this->Content.IsValid() )
