UMG - Working on several general fixes to the navigation system.  Widgets other than the focused widget can respond to navigation requests with explicit or delegate based handling.  Adding a new function to allow querying the if descendants of a widget are focused by a specific user.  Implementing controller support for the SScrollBox; it has custom logic for handling navigation inside of it scrolling the next element into view for the focuser and moving focus to the next logical child.

diff --git a/Engine/Source/Runtime/Slate/Private/Framework/Application/SlateApplication.cpp b/Engine/Source/Runtime/Slate/Private/Framework/Application/SlateApplication.cpp
--- a/Engine/Source/Runtime/Slate/Private/Framework/Application/SlateApplication.cpp
+++ b/Engine/Source/Runtime/Slate/Private/Framework/Application/SlateApplication.cpp
@@ -3552,6 +3552,18 @@ bool FSlateApplication::ShowUserFocus(const TSharedPtr<const SWidget> Widget) co
 	return false;
 }
 
+bool FSlateApplication::HasUserFocusedDescendants(const TSharedRef< const SWidget >& Widget, int32 UserIndex) const
+{
+	const FUserFocusEntry& UserFocusEntry = UserFocusEntries[UserIndex];
+	const FWeakWidgetPath & FocusedWidgetPath = UserFocusEntry.WidgetPath;
+	if ( FocusedWidgetPath.IsValid() && FocusedWidgetPath.GetLastWidget().Pin() != Widget && FocusedWidgetPath.ContainsWidget(Widget) )
+	{
+		return true;
+	}
+
+	return false;
+}
+
 bool FSlateApplication::HasFocusedDescendants( const TSharedRef< const SWidget >& Widget ) const
 {
 	for (int32 UserIndex = 0; UserIndex < SlateApplicationDefs::MaxUsers; ++UserIndex)
@@ -4695,23 +4707,17 @@ bool FSlateApplication::AttemptNavigation(const FNavigationEvent& NavigationEven
 	{
 		EUINavigation NavigationType = NavigationEvent.GetNavigationType();
 
-		// If the FocusedWidget is the Boundary widget we don't need to do any complex work. 
-		if (FocusedWidget == NavigationReply.GetHandler())
+		if ( NavigationReply.GetBoundaryRule() == EUINavigationRule::Explicit )
 		{
-			if (NavigationReply.GetBoundaryRule() == EUINavigationRule::Explicit)
-			{
-				NewFocusedWidget = NavigationReply.GetFocusRecipient();
-			}
-			else if (NavigationReply.GetBoundaryRule() == EUINavigationRule::Custom)
+			NewFocusedWidget = NavigationReply.GetFocusRecipient();
+		}
+		else if ( NavigationReply.GetBoundaryRule() == EUINavigationRule::Custom )
+		{
+			const FNavigationDelegate& FocusDelegate = NavigationReply.GetFocusDelegate();
+			if ( FocusDelegate.IsBound() )
 			{
-				const FNavigationDelegate& FocusDelegate = NavigationReply.GetFocusDelegate();
-				if (FocusDelegate.IsBound())
-				{
-					NewFocusedWidget = FocusDelegate.Execute(NavigationType);
-				}
+				NewFocusedWidget = FocusDelegate.Execute(NavigationType);
 			}
-			// If it's not explicit or custom we do noting.
-			// This is an optimization because wrapping and stopping on the widget that is currently focused can't go anywhere.
 		}
 		else
 		{
diff --git a/Engine/Source/Runtime/Slate/Private/Widgets/Layout/SScrollBox.cpp b/Engine/Source/Runtime/Slate/Private/Widgets/Layout/SScrollBox.cpp
--- a/Engine/Source/Runtime/Slate/Private/Widgets/Layout/SScrollBox.cpp
+++ b/Engine/Source/Runtime/Slate/Private/Widgets/Layout/SScrollBox.cpp
@@ -386,11 +386,32 @@ bool SScrollBox::ScrollDescendantIntoView(const FGeometry& MyGeometry, const TSh
 		// @todo: This is a workaround because DesiredScrollOffset can exceed the ScrollMax when mouse dragging on the scroll bar and we need it clamped here or the offset is wrong
 		ScrollBy(MyGeometry, 0, EAllowOverscroll::No, false);
 
-		// Calculate how much we would need to scroll to bring this to the top/left of the scroll box
-		const float WidgetPosition = GetScrollComponentFromVector(WidgetGeometry->Geometry.AbsolutePosition);
-		const float MyPosition = GetScrollComponentFromVector(MyGeometry.AbsolutePosition);
-		const float ScrollOffset = WidgetPosition - MyPosition;
-		ScrollBy(MyGeometry, ScrollOffset, EAllowOverscroll::No, InAnimateScroll);
+		// If the scale is 0, we can't really calculate things properly.
+		if ( MyGeometry.Scale != 0 )
+		{
+			// Calculate how much we would need to scroll to bring this to the top/left of the scroll box
+			// NOTE: The scrollbox does all offset calculation using desired sizes which doesn't involve scale,
+			//       so we need to remove the scale from positions so that they are in the same unscaled space.
+			const float WidgetStartPosition = GetScrollComponentFromVector(WidgetGeometry->Geometry.AbsolutePosition / MyGeometry.Scale);
+			const float WidgetEndPosition = GetScrollComponentFromVector(WidgetGeometry->Geometry.AbsolutePosition / MyGeometry.Scale + WidgetGeometry->Geometry.GetLocalSize());
+			const float ViewStartPosition = GetScrollComponentFromVector(MyGeometry.AbsolutePosition / MyGeometry.Scale);
+			const float ViewEndPosition = GetScrollComponentFromVector(MyGeometry.AbsolutePosition / MyGeometry.Scale + MyGeometry.GetLocalSize());
+
+			const float ViewDelta = ( ViewEndPosition - ViewStartPosition );
+			const float WidgetDelta = ( WidgetEndPosition - WidgetStartPosition );
+
+			if ( WidgetStartPosition < ViewStartPosition )
+			{
+				const float ScrollOffset = WidgetStartPosition - ViewStartPosition;
+				ScrollBy(MyGeometry, ScrollOffset, EAllowOverscroll::No, InAnimateScroll);
+			}
+			else if ( WidgetEndPosition > ViewEndPosition )
+			{
+				const float ScrollOffset = ( WidgetEndPosition - ViewDelta ) - ViewStartPosition;
+				ScrollBy(MyGeometry, ScrollOffset, EAllowOverscroll::No, InAnimateScroll);
+			}
+		}
+
 		return true;
 	}
 	return false;
@@ -790,6 +811,73 @@ FReply SScrollBox::OnTouchEnded(const FGeometry& MyGeometry, const FPointerEvent
 	}
 }
 
+FNavigationReply SScrollBox::OnNavigation(const FGeometry& MyGeometry, const FNavigationEvent& InNavigationEvent)
+{
+	TSharedPtr<SWidget> FocusedChild;
+	int32 FocusedChildIndex = -1;
+
+	// Find the child with focus currently so that we can find the next logical child we're going to move to.
+	TPanelChildren<SScrollBox::FSlot>& Children = ScrollPanel->Children;
+	for ( int32 SlotIndex=0; SlotIndex < Children.Num(); ++SlotIndex )
+	{
+		if ( Children[SlotIndex].GetWidget()->HasUserFocus(InNavigationEvent.GetUserIndex()).IsSet() ||
+			 Children[SlotIndex].GetWidget()->HasUserFocusedDescendants(InNavigationEvent.GetUserIndex()) )
+		{
+			FocusedChild = Children[SlotIndex].GetWidget();
+			FocusedChildIndex = SlotIndex;
+			break;
+		}
+	}
+
+	if ( FocusedChild.IsValid() )
+	{
+		if ( Orientation == Orient_Vertical )
+		{
+			switch ( InNavigationEvent.GetNavigationType() )
+			{
+			case EUINavigation::Up:
+				FocusedChildIndex--;
+				break;
+			case EUINavigation::Down:
+				FocusedChildIndex++;
+				break;
+			default:
+				// If we don't handle this direction in our current orientation we can 
+				// just allow the behavior of the boundary rule take over.
+				return SCompoundWidget::OnNavigation(MyGeometry, InNavigationEvent);
+			}
+		}
+		else // Orient_Horizontal
+		{
+			switch ( InNavigationEvent.GetNavigationType() )
+			{
+			case EUINavigation::Left:
+				FocusedChildIndex--;
+				break;
+			case EUINavigation::Right:
+				FocusedChildIndex++;
+				break;
+			default:
+				// If we don't handle this direction in our current orientation we can 
+				// just allow the behavior of the boundary rule take over.
+				return SCompoundWidget::OnNavigation(MyGeometry, InNavigationEvent);
+			}
+		}
+
+		// If the focused child index is in a valid range we know we can successfully focus
+		// the new child we're moving focus to.
+		if ( FocusedChildIndex >= 0 && FocusedChildIndex < Children.Num() )
+		{
+			TSharedRef<SWidget> NextFocusedChild = Children[FocusedChildIndex].GetWidget();
+			ScrollDescendantIntoView(MyGeometry, NextFocusedChild, false);
+
+			return FNavigationReply::Explicit(NextFocusedChild);
+		}
+	}
+
+	return SCompoundWidget::OnNavigation(MyGeometry, InNavigationEvent);
+}
+
 int32 SScrollBox::OnPaint( const FPaintArgs& Args, const FGeometry& AllottedGeometry, const FSlateRect& MyClippingRect, FSlateWindowElementList& OutDrawElements, int32 LayerId, const FWidgetStyle& InWidgetStyle, bool bParentEnabled ) const
 {
 	int32 NewLayerId = SCompoundWidget::OnPaint( Args, AllottedGeometry, MyClippingRect, OutDrawElements, LayerId, InWidgetStyle, bParentEnabled );
diff --git a/Engine/Source/Runtime/Slate/Public/Framework/Application/SlateApplication.h b/Engine/Source/Runtime/Slate/Public/Framework/Application/SlateApplication.h
--- a/Engine/Source/Runtime/Slate/Public/Framework/Application/SlateApplication.h
+++ b/Engine/Source/Runtime/Slate/Public/Framework/Application/SlateApplication.h
@@ -997,6 +997,7 @@ public:
 	virtual bool HasAnyMouseCaptor() const override;
 	virtual FSlateRect GetPreferredWorkArea() const override;
 	virtual bool HasFocusedDescendants( const TSharedRef<const SWidget>& Widget ) const override;
+	virtual bool HasUserFocusedDescendants(const TSharedRef< const SWidget >& Widget, int32 UserIndex) const override;
 	virtual bool IsExternalUIOpened() override;
 	virtual FWidgetPath LocateWindowUnderMouse( FVector2D ScreenspaceMouseCoordinate, const TArray<TSharedRef<SWindow>>& Windows, bool bIgnoreEnabledStatus = false ) override;
 	virtual bool IsWindowHousingInteractiveTooltip(const TSharedRef<const SWindow>& WindowToTest) const override;
diff --git a/Engine/Source/Runtime/Slate/Public/Widgets/Layout/SScrollBox.h b/Engine/Source/Runtime/Slate/Public/Widgets/Layout/SScrollBox.h
--- a/Engine/Source/Runtime/Slate/Public/Widgets/Layout/SScrollBox.h
+++ b/Engine/Source/Runtime/Slate/Public/Widgets/Layout/SScrollBox.h
@@ -146,6 +146,7 @@ public:
 	virtual FReply OnTouchStarted(const FGeometry& MyGeometry, const FPointerEvent& InTouchEvent) override;
 	virtual FReply OnTouchMoved(const FGeometry& MyGeometry, const FPointerEvent& InTouchEvent) override;
 	virtual FReply OnTouchEnded(const FGeometry& MyGeometry, const FPointerEvent& InTouchEvent) override;
+	virtual FNavigationReply OnNavigation(const FGeometry& MyGeometry, const FNavigationEvent& InNavigationEvent);
 	// End of SWidget interface
 
 private:
diff --git a/Engine/Source/Runtime/SlateCore/Private/Widgets/SWidget.cpp b/Engine/Source/Runtime/SlateCore/Private/Widgets/SWidget.cpp
--- a/Engine/Source/Runtime/SlateCore/Private/Widgets/SWidget.cpp
+++ b/Engine/Source/Runtime/SlateCore/Private/Widgets/SWidget.cpp
@@ -428,6 +428,11 @@ TOptional<EFocusCause> SWidget::HasAnyUserFocus() const
 	return FSlateApplicationBase::Get().HasAnyUserFocus(SharedThis(this));
 }
 
+bool SWidget::HasUserFocusedDescendants(int32 UserIndex) const
+{
+	return FSlateApplicationBase::Get().HasUserFocusedDescendants(SharedThis(this), UserIndex);
+}
+
 bool SWidget::HasFocusedDescendants() const
 {
 	return FSlateApplicationBase::Get().HasFocusedDescendants(SharedThis(this));
diff --git a/Engine/Source/Runtime/SlateCore/Public/Application/SlateApplicationBase.h b/Engine/Source/Runtime/SlateCore/Public/Application/SlateApplicationBase.h
--- a/Engine/Source/Runtime/SlateCore/Public/Application/SlateApplicationBase.h
+++ b/Engine/Source/Runtime/SlateCore/Public/Application/SlateApplicationBase.h
@@ -224,6 +224,15 @@ public:
 	 */
 	virtual FSlateRect GetPreferredWorkArea( ) const = 0;
 
+	/**
+	 * Checks whether the specified widget has any descendants which are currently focused for the specified user user.
+	 *
+	 * @param Widget The widget to check.
+	 * @param InUserIndex Index of the user that we want to check for.
+	 * @return true if any descendants are focused, false otherwise.
+	 */
+	virtual bool HasUserFocusedDescendants(const TSharedRef< const SWidget >& Widget, int32 UserIndex) const = 0;
+
 	/**
 	 * Checks whether the specified widget has any descendants which are currently focused.
 	 *
diff --git a/Engine/Source/Runtime/SlateCore/Public/Widgets/SWidget.h b/Engine/Source/Runtime/SlateCore/Public/Widgets/SWidget.h
--- a/Engine/Source/Runtime/SlateCore/Public/Widgets/SWidget.h
+++ b/Engine/Source/Runtime/SlateCore/Public/Widgets/SWidget.h
@@ -594,6 +594,13 @@ public:
 	 */
 	TOptional<EFocusCause> HasAnyUserFocus() const;
 
+	/**
+	 * Gets whether or not the specified users has this widget or any descendant focused.
+	 *
+	 * @return The optional will be set with the focus cause, if unset this widget doesn't have focus.
+	 */
+	bool HasUserFocusedDescendants(int32 UserIndex) const;
+
 	/**
 	 * @return Whether this widget has any descendants with keyboard focus
 	 */
diff --git a/Engine/Source/Runtime/UMG/Private/Components/Widget.cpp b/Engine/Source/Runtime/UMG/Private/Components/Widget.cpp
--- a/Engine/Source/Runtime/UMG/Private/Components/Widget.cpp
+++ b/Engine/Source/Runtime/UMG/Private/Components/Widget.cpp
@@ -280,17 +280,6 @@ bool UWidget::HasKeyboardFocus() const
 	return false;
 }
 
-bool UWidget::HasFocusedDescendants() const
-{
-	TSharedPtr<SWidget> SafeWidget = GetCachedWidget();
-	if (SafeWidget.IsValid())
-	{
-		return SafeWidget->HasFocusedDescendants();
-	}
-
-	return false;
-}
-
 bool UWidget::HasMouseCapture() const
 {
 	TSharedPtr<SWidget> SafeWidget = GetCachedWidget();
@@ -348,6 +337,41 @@ bool UWidget::HasAnyUserFocus() const
 	return false;
 }
 
+bool UWidget::HasFocusedDescendants() const
+{
+	TSharedPtr<SWidget> SafeWidget = GetCachedWidget();
+	if ( SafeWidget.IsValid() )
+	{
+		return SafeWidget->HasFocusedDescendants();
+	}
+
+	return false;
+}
+
+bool UWidget::HasUserFocusedDescendants(APlayerController* PlayerController) const
+{
+	if ( PlayerController == nullptr || !PlayerController->IsLocalPlayerController() )
+	{
+		return false;
+	}
+
+	TSharedPtr<SWidget> SafeWidget = GetCachedWidget();
+	if ( SafeWidget.IsValid() )
+	{
+		FLocalPlayerContext Context(PlayerController);
+
+		if ( ULocalPlayer* LocalPlayer = Context.GetLocalPlayer() )
+		{
+			// HACK: We use the controller Id as the local player index for focusing widgets in Slate.
+			int32 UserIndex = LocalPlayer->GetControllerId();
+
+			return SafeWidget->HasUserFocusedDescendants(UserIndex);
+		}
+	}
+
+	return false;
+}
+
 void UWidget::SetUserFocus(APlayerController* PlayerController)
 {
 	if ( PlayerController == nullptr || !PlayerController->IsLocalPlayerController() )
diff --git a/Engine/Source/Runtime/UMG/Public/Components/Widget.h b/Engine/Source/Runtime/UMG/Public/Components/Widget.h
--- a/Engine/Source/Runtime/UMG/Public/Components/Widget.h
+++ b/Engine/Source/Runtime/UMG/Public/Components/Widget.h
@@ -278,12 +278,6 @@ public:
 	UFUNCTION(BlueprintCallable, Category="Widget")
 	bool HasKeyboardFocus() const;
 
-	/**
-	 * @return Whether this widget has any descendants with keyboard focus
-	 */
-	UFUNCTION(BlueprintCallable, Category="Widget")
-	bool HasFocusedDescendants() const;
-
 	/**
 	 * Checks to see if this widget is the current mouse captor
 	 * @return  True if this widget has captured the mouse
@@ -295,13 +289,21 @@ public:
 	UFUNCTION(BlueprintCallable, Category="Widget")
 	void SetKeyboardFocus();
 
-	/** Gets the focus to this widget. */
+	/** @return true if this widget is focused by a specific user. */
 	UFUNCTION(BlueprintCallable, Category="Widget")
 	bool HasUserFocus(APlayerController* PlayerController) const;
 
-	/** Gets the focus to this widget. */
+	/** @return true if this widget is focused by any user. */
 	UFUNCTION(BlueprintCallable, Category="Widget")
 	bool HasAnyUserFocus() const;
+
+	/** @return true if any descendant widget is focused by any user. */
+	UFUNCTION(BlueprintCallable, Category="Widget", meta=(DisplayName="HasAnyUserFocusedDescendants"))
+	bool HasFocusedDescendants() const;
+
+	/** @return true if any descendant widget is focused by a specific user. */
+	UFUNCTION(BlueprintCallable, Category="Widget")
+	bool HasUserFocusedDescendants(APlayerController* PlayerController) const;
 	
 	/** Sets the focus to this widget for a specific user */
 	UFUNCTION(BlueprintCallable, Category="Widget")
