Missed file in merge due to conflict with another CL.

diff --git a/Engine/Source/Runtime/Slate/Private/Framework/Application/SlateApplication.cpp b/Engine/Source/Runtime/Slate/Private/Framework/Application/SlateApplication.cpp
--- a/Engine/Source/Runtime/Slate/Private/Framework/Application/SlateApplication.cpp
+++ b/Engine/Source/Runtime/Slate/Private/Framework/Application/SlateApplication.cpp
@@ -3538,6 +3538,11 @@ TOptional<EFocusCause> FSlateApplication::HasAnyUserFocus(const TSharedPtr<const
 	return TOptional<EFocusCause>();
 }
 
+bool FSlateApplication::IsWidgetDirectlyHovered(const TSharedPtr<const SWidget> Widget) const
+{
+	return WidgetsUnderCursorLastEvent.IsValid() && Widget == WidgetsUnderCursorLastEvent.GetLastWidget().Pin();
+}
+
 bool FSlateApplication::ShowUserFocus(const TSharedPtr<const SWidget> Widget) const
 {
 	for (int32 UserIndex = 0; UserIndex < SlateApplicationDefs::MaxUsers; ++UserIndex)
