[INTEGRATE] merge UE4-LauncherDev > UE4 CL 2552378
Editable Text Updates - allow additional UI commands so users can intercept key down events

diff --git a/Engine/Source/Runtime/Slate/Private/Widgets/Input/SEditableText.cpp b/Engine/Source/Runtime/Slate/Private/Widgets/Input/SEditableText.cpp
--- a/Engine/Source/Runtime/Slate/Private/Widgets/Input/SEditableText.cpp
+++ b/Engine/Source/Runtime/Slate/Private/Widgets/Input/SEditableText.cpp
@@ -110,6 +110,13 @@ void SEditableText::Construct( const FArguments& InArgs )
 		FExecuteAction::CreateSP( this, &SEditableText::SelectAllText ),
 		FCanExecuteAction::CreateSP( this, &SEditableText::CanExecuteSelectAll ) );
 
+	// Append any additional commands that a consumer of EditableText wants us to be aware of.
+	const TSharedPtr<FUICommandList>& AdditionalCommands = InArgs._AdditionalCommands;
+	if ( AdditionalCommands.IsValid() )
+	{
+		UICommandList->Append( AdditionalCommands.ToSharedRef() );
+	}
+
 	// build context menu extender
 	MenuExtender = MakeShareable(new FExtender);
 	MenuExtender->AddMenuExtension("EditText", EExtensionHook::Before, TSharedPtr<FUICommandList>(), InArgs._ContextMenuExtender);
diff --git a/Engine/Source/Runtime/Slate/Private/Widgets/Input/SEditableTextBox.cpp b/Engine/Source/Runtime/Slate/Private/Widgets/Input/SEditableTextBox.cpp
--- a/Engine/Source/Runtime/Slate/Private/Widgets/Input/SEditableTextBox.cpp
+++ b/Engine/Source/Runtime/Slate/Private/Widgets/Input/SEditableTextBox.cpp
@@ -50,6 +50,7 @@ void SEditableTextBox::Construct( const FArguments& InArgs )
 					.OnTextCommitted( InArgs._OnTextCommitted )
 					.MinDesiredWidth( InArgs._MinDesiredWidth )
 					.SelectAllTextOnCommit( InArgs._SelectAllTextOnCommit )
+					.AdditionalCommands( InArgs._AdditionalCommands )
 				]
 			]
 		]
diff --git a/Engine/Source/Runtime/Slate/Private/Widgets/Input/SMultiLineEditableTextBox.cpp b/Engine/Source/Runtime/Slate/Private/Widgets/Input/SMultiLineEditableTextBox.cpp
--- a/Engine/Source/Runtime/Slate/Private/Widgets/Input/SMultiLineEditableTextBox.cpp
+++ b/Engine/Source/Runtime/Slate/Private/Widgets/Input/SMultiLineEditableTextBox.cpp
@@ -109,6 +109,7 @@ void SMultiLineEditableTextBox::Construct( const FArguments& InArgs )
 					.VScrollBar(VScrollBar)
 					.OnHScrollBarUserScrolled(InArgs._OnHScrollBarUserScrolled)
 					.OnVScrollBarUserScrolled(InArgs._OnVScrollBarUserScrolled)
+					.AdditionalCommands( InArgs._AdditionalCommands)
 					.ModiferKeyForNewLine(InArgs._ModiferKeyForNewLine)
 				]
 
diff --git a/Engine/Source/Runtime/Slate/Private/Widgets/Text/SMultiLineEditableText.cpp b/Engine/Source/Runtime/Slate/Private/Widgets/Text/SMultiLineEditableText.cpp
--- a/Engine/Source/Runtime/Slate/Private/Widgets/Text/SMultiLineEditableText.cpp
+++ b/Engine/Source/Runtime/Slate/Private/Widgets/Text/SMultiLineEditableText.cpp
@@ -340,6 +340,13 @@ void SMultiLineEditableText::Construct( const FArguments& InArgs )
 		FExecuteAction::CreateSP(this, &SMultiLineEditableText::SelectAllText),
 		FCanExecuteAction::CreateSP(this, &SMultiLineEditableText::CanExecuteSelectAll));
 
+	// Append any additional commands that a consumer of MultiLineEditableText wants us to be aware of.
+	const TSharedPtr<FUICommandList>& AdditionalCommands = InArgs._AdditionalCommands;
+	if ( AdditionalCommands.IsValid() )
+	{
+		UICommandList->Append( AdditionalCommands.ToSharedRef() );
+	}
+
 	// build context menu extender
 	MenuExtender = MakeShareable(new FExtender);
 	MenuExtender->AddMenuExtension("EditText", EExtensionHook::Before, TSharedPtr<FUICommandList>(), InArgs._ContextMenuExtender);
diff --git a/Engine/Source/Runtime/Slate/Public/Widgets/Input/SEditableText.h b/Engine/Source/Runtime/Slate/Public/Widgets/Input/SEditableText.h
--- a/Engine/Source/Runtime/Slate/Public/Widgets/Input/SEditableText.h
+++ b/Engine/Source/Runtime/Slate/Public/Widgets/Input/SEditableText.h
@@ -33,6 +33,7 @@ public:
 		, _ClearKeyboardFocusOnCommit(true)
 		, _MinDesiredWidth(0.0f)
 		, _SelectAllTextOnCommit( false )
+		, _AdditionalCommands()
 		, _VirtualKeyboardType(EKeyboardType::Keyboard_Default)
 		{}
 
@@ -104,6 +105,9 @@ public:
 		/** Whether to select all text when pressing enter to commit changes */
 		SLATE_ATTRIBUTE( bool, SelectAllTextOnCommit )
 
+		/** Additional UI commands. */
+		SLATE_ARGUMENT(TSharedPtr<FUICommandList>, AdditionalCommands)
+
 		/** Menu extender for the right-click context menu */
 		SLATE_EVENT( FMenuExtensionDelegate, ContextMenuExtender )
 
diff --git a/Engine/Source/Runtime/Slate/Public/Widgets/Input/SEditableTextBox.h b/Engine/Source/Runtime/Slate/Public/Widgets/Input/SEditableTextBox.h
--- a/Engine/Source/Runtime/Slate/Public/Widgets/Input/SEditableTextBox.h
+++ b/Engine/Source/Runtime/Slate/Public/Widgets/Input/SEditableTextBox.h
@@ -25,6 +25,7 @@ public:
 		, _ClearKeyboardFocusOnCommit( true )
 		, _MinDesiredWidth( 0.0f )
 		, _SelectAllTextOnCommit( false )
+		, _AdditionalCommands()
 		, _BackgroundColor()		
 		, _Padding()
 		, _ErrorReporting()
@@ -81,6 +82,9 @@ public:
 		/** Whether to select all text when pressing enter to commit changes */
 		SLATE_ATTRIBUTE( bool, SelectAllTextOnCommit )
 
+		/** Additional UI commands. */
+		SLATE_ARGUMENT(TSharedPtr<FUICommandList>, AdditionalCommands)
+
 		/** The color of the background/border around the editable text (overrides Style) */
 		SLATE_ATTRIBUTE( FSlateColor, BackgroundColor )
 
diff --git a/Engine/Source/Runtime/Slate/Public/Widgets/Input/SMultiLineEditableTextBox.h b/Engine/Source/Runtime/Slate/Public/Widgets/Input/SMultiLineEditableTextBox.h
--- a/Engine/Source/Runtime/Slate/Public/Widgets/Input/SMultiLineEditableTextBox.h
+++ b/Engine/Source/Runtime/Slate/Public/Widgets/Input/SMultiLineEditableTextBox.h
@@ -35,6 +35,7 @@ public:
 		, _WrapTextAt(0.0f)
 		, _AutoWrapText(false)
 		, _SelectAllTextOnCommit( false )
+		, _AdditionalCommands()
 		, _BackgroundColor()		
 		, _Padding()
 		, _Margin()
@@ -123,6 +124,9 @@ public:
 		/** Called when the cursor is moved within the text area */
 		SLATE_EVENT( SMultiLineEditableText::FOnCursorMoved, OnCursorMoved )
 
+		/** Additional UI commands. */
+		SLATE_ARGUMENT(TSharedPtr<FUICommandList>, AdditionalCommands)
+
 		/** Menu extender for the right-click context menu */
 		SLATE_EVENT( FMenuExtensionDelegate, ContextMenuExtender )
 
diff --git a/Engine/Source/Runtime/Slate/Public/Widgets/Text/SMultiLineEditableText.h b/Engine/Source/Runtime/Slate/Public/Widgets/Text/SMultiLineEditableText.h
--- a/Engine/Source/Runtime/Slate/Public/Widgets/Text/SMultiLineEditableText.h
+++ b/Engine/Source/Runtime/Slate/Public/Widgets/Text/SMultiLineEditableText.h
@@ -37,6 +37,7 @@ public:
 		, _RevertTextOnEscape(false)
 		, _ClearKeyboardFocusOnCommit(true)
 		, _OnCursorMoved()
+		, _AdditionalCommands()
 		, _ContextMenuExtender()
 		, _ModiferKeyForNewLine(EModifierKey::None)
 	{}
@@ -108,6 +109,9 @@ public:
 		/** Called when the cursor is moved within the text area */
 		SLATE_EVENT(FOnCursorMoved, OnCursorMoved)
 
+		/** Additional UI commands. */
+		SLATE_ARGUMENT(TSharedPtr<FUICommandList>, AdditionalCommands)
+
 		/** Menu extender for the right-click context menu */
 		SLATE_EVENT(FMenuExtensionDelegate, ContextMenuExtender)
 
