Activating a window now sets focus for all users

UE-14833 - Cannot get input from input device without Controller already being created

FSceneViewport used to require you to click on it after alt-tabbing away from the game so that you got focus back for all users (as FSceneViewport::OnMouseButtonDown claimed focus for all users). This was due to the window re-activation only claiming focus for the keyboard user.

diff --git a/Engine/Source/Runtime/Slate/Private/Framework/Application/SlateApplication.cpp b/Engine/Source/Runtime/Slate/Private/Framework/Application/SlateApplication.cpp
--- a/Engine/Source/Runtime/Slate/Private/Framework/Application/SlateApplication.cpp
+++ b/Engine/Source/Runtime/Slate/Private/Framework/Application/SlateApplication.cpp
@@ -2013,6 +2013,25 @@ void FSlateApplication::SetUserFocus(uint32 UserIndex, const TSharedPtr<SWidget>
 	}
 }
 
+void FSlateApplication::SetAllUserFocus(const TSharedPtr<SWidget>& WidgetToFocus, EFocusCause ReasonFocusIsChanging /*= EFocusCause::SetDirectly*/)
+{
+	const bool bValidWidget = WidgetToFocus.IsValid();
+	ensureMsgf(bValidWidget, TEXT("Attempting to focus an invalid widget. If your intent is to clear focus use ClearAllUserFocus()"));
+	if (bValidWidget)
+	{
+		FWidgetPath PathToWidget;
+		const bool bFound = FSlateWindowHelper::FindPathToWidget(SlateWindows, WidgetToFocus.ToSharedRef(), /*OUT*/ PathToWidget);
+		if (bFound)
+		{
+			SetAllUserFocus(PathToWidget, ReasonFocusIsChanging);
+		}
+		else
+		{
+			//ensureMsgf(bFound, TEXT("Attempting to focus a widget that isn't in the tree and visible: %s. If your intent is to clear focus use ClearAllUserFocus()"), WidgetToFocus->ToString());
+		}
+	}
+}
+
 TSharedPtr<SWidget> FSlateApplication::GetUserFocusedWidget(uint32 UserIndex) const
 {
 	const FUserFocusEntry& UserFocusEntry = UserFocusEntries[UserIndex];
@@ -2029,6 +2048,11 @@ void FSlateApplication::ClearUserFocus(uint32 UserIndex, EFocusCause ReasonFocus
 	SetUserFocus(UserIndex, FWidgetPath(), ReasonFocusIsChanging);
 }
 
+void FSlateApplication::ClearAllUserFocus(EFocusCause ReasonFocusIsChanging /*= EFocusCause::SetDirectly*/)
+{
+	SetAllUserFocus(FWidgetPath(), ReasonFocusIsChanging);
+}
+
 void FSlateApplication::ReleaseJoystickCapture(uint32 UserIndex)
 {
 	ClearUserFocus(UserIndex);
@@ -2229,6 +2253,15 @@ bool FSlateApplication::SetUserFocus(const uint32 InUserIndex, const FWidgetPath
 }
 
 
+void FSlateApplication::SetAllUserFocus(const FWidgetPath& InFocusPath, const EFocusCause InCause)
+{
+	for (int32 SlateUserIndex = 0; SlateUserIndex < SlateApplicationDefs::MaxUsers; ++SlateUserIndex)
+	{
+		SetUserFocus(SlateUserIndex, InFocusPath, InCause);
+	}
+}
+
+
 FModifierKeysState FSlateApplication::GetModifierKeys() const
 {
 	return PlatformApplication->GetModifierKeys();
diff --git a/Engine/Source/Runtime/Slate/Public/Framework/Application/SlateApplication.h b/Engine/Source/Runtime/Slate/Public/Framework/Application/SlateApplication.h
--- a/Engine/Source/Runtime/Slate/Public/Framework/Application/SlateApplication.h
+++ b/Engine/Source/Runtime/Slate/Public/Framework/Application/SlateApplication.h
@@ -29,16 +29,6 @@ DECLARE_DELEGATE(FModalWindowStackEnded)
 extern SLATE_API const FName NAME_UnrealOS;
 
 
-/**
- * Design constraints for Slate applications
- */
-namespace SlateApplicationDefs
-{
-	/** How many users can we support at once? */
-	static const int32 MaxUsers = 8;
-}
-
-
 /** Allow widgets to find out when someone clicked outside them. Currently needed by MenuAnchros. */
 class SLATE_API FPopupSupport
 {
@@ -379,6 +369,14 @@ public:
 	 */
 	void SetUserFocus(uint32 UserIndex, const TSharedPtr<SWidget>& WidgetToFocus, EFocusCause ReasonFocusIsChanging = EFocusCause::SetDirectly);
 
+	/**
+	 * Sets focus for all users to the SWidget passed in.
+	 *
+	 * @param WidgetToFocus the widget to set focus to
+	 * @param ReasonFocusIsChanging the contextual reason for the focus change
+	 */
+	void SetAllUserFocus(const TSharedPtr<SWidget>& WidgetToFocus, EFocusCause ReasonFocusIsChanging = EFocusCause::SetDirectly);
+
 	/** @return a pointer to the Widget that currently has the users focus; Empty pointer when the user has no focus. */
 	TSharedPtr< SWidget > GetUserFocusedWidget(uint32 UserIndex) const;
 
@@ -388,6 +386,9 @@ public:
 	/** Releases the users focus from whatever it currently is on. */
 	void ClearUserFocus(uint32 UserIndex, EFocusCause ReasonFocusIsChanging = EFocusCause::SetDirectly);
 
+	/** Releases the focus for all users from whatever it currently is on. */
+	void ClearAllUserFocus(EFocusCause ReasonFocusIsChanging = EFocusCause::SetDirectly);
+
 	DEPRECATED(4.6, "FSlateApplication::ReleaseJoystickCapture() is deprecated, use FSlateApplication::ClearUserFocus() instead.")
 	void ReleaseJoystickCapture(uint32 UserIndex);
 
@@ -1011,6 +1012,7 @@ public:
 	virtual void RequestDestroyWindow( TSharedRef<SWindow> WindowToDestroy ) override;
 	virtual bool SetKeyboardFocus( const FWidgetPath& InFocusPath, const EFocusCause InCause ) override;
 	virtual bool SetUserFocus(const uint32 InUserIndex, const FWidgetPath& InFocusPath, const EFocusCause InCause) override;
+	virtual void SetAllUserFocus(const FWidgetPath& InFocusPath, const EFocusCause InCause) override;
 
 public:
 
diff --git a/Engine/Source/Runtime/SlateCore/Private/Widgets/SWindow.cpp b/Engine/Source/Runtime/SlateCore/Private/Widgets/SWindow.cpp
--- a/Engine/Source/Runtime/SlateCore/Private/Widgets/SWindow.cpp
+++ b/Engine/Source/Runtime/SlateCore/Private/Widgets/SWindow.cpp
@@ -1337,7 +1337,7 @@ bool SWindow::OnIsActiveChanged( const FWindowActivateEvent& ActivateEvent )
 				FWidgetPath WidgetToFocusPath;
 				if( FSlateWindowHelper::FindPathToWidget( JustThisWindow, PinnedWidgetToFocus.ToSharedRef(), WidgetToFocusPath ) )
 				{
-					FSlateApplicationBase::Get().SetKeyboardFocus( WidgetToFocusPath, EFocusCause::SetDirectly );
+					FSlateApplicationBase::Get().SetAllUserFocus( WidgetToFocusPath, EFocusCause::SetDirectly );
 				}
 			}
 			else
@@ -1346,7 +1346,7 @@ bool SWindow::OnIsActiveChanged( const FWindowActivateEvent& ActivateEvent )
 				if( FSlateWindowHelper::FindPathToWidget( JustThisWindow, AsShared(), WindowWidgetPath ) )
 				{
 					FWeakWidgetPath WeakWindowPath(WindowWidgetPath);
-					FSlateApplicationBase::Get().SetKeyboardFocus( WeakWindowPath.ToNextFocusedPath(EUINavigation::Next), EFocusCause::SetDirectly );
+					FSlateApplicationBase::Get().SetAllUserFocus( WeakWindowPath.ToNextFocusedPath(EUINavigation::Next), EFocusCause::SetDirectly );
 				}
 			}
 		}
diff --git a/Engine/Source/Runtime/SlateCore/Public/Application/SlateApplicationBase.h b/Engine/Source/Runtime/SlateCore/Public/Application/SlateApplicationBase.h
--- a/Engine/Source/Runtime/SlateCore/Public/Application/SlateApplicationBase.h
+++ b/Engine/Source/Runtime/SlateCore/Public/Application/SlateApplicationBase.h
@@ -39,6 +39,15 @@ private:
 	FWidgetPath LocateWidgetInWindow(FVector2D ScreenspaceMouseCoordinate, const TSharedRef<SWindow>& Window, bool bIgnoreEnabledStatus) const;
 };
 
+/**
+ * Design constraints for Slate applications
+ */
+namespace SlateApplicationDefs
+{
+	/** How many users can we support at once? */
+	static const int32 MaxUsers = 8;
+}
+
 /**
  * Base class for Slate applications.
  *
@@ -336,6 +345,14 @@ public:
 	 */
 	virtual bool SetUserFocus(const uint32 InUserIndex, const FWidgetPath& InFocusPath, const EFocusCause InCause) = 0;
 
+	/**
+	 * Sets the focus for all users to the specified widget.  The widget must be allowed to receive focus.
+	 *
+	 * @param InWidget WidgetPath to the Widget to being focused.
+	 * @param InCause The reason that focus is changing.
+	 */
+	virtual void SetAllUserFocus(const FWidgetPath& InFocusPath, const EFocusCause InCause) = 0;
+
 private:
 	/**
 	 * Implementation for active timer registration. See SWidget::RegisterActiveTimer.
