Converted SEditableText::ScanForWordBoundary to use a word-break iterator

UE-15231 - Editor crashes on when modifying text in a SEditableText text box

This will make its movement behavior match that of the multi-line editable text, and will work much better for non-English languages.

diff --git a/Engine/Source/Runtime/Slate/Private/Widgets/Input/SEditableText.cpp b/Engine/Source/Runtime/Slate/Private/Widgets/Input/SEditableText.cpp
--- a/Engine/Source/Runtime/Slate/Private/Widgets/Input/SEditableText.cpp
+++ b/Engine/Source/Runtime/Slate/Private/Widgets/Input/SEditableText.cpp
@@ -1,6 +1,7 @@
 // Copyright 1998-2015 Epic Games, Inc. All Rights Reserved.
 
 #include "SlatePrivatePCH.h"
+#include "BreakIterator.h"
 #include "TextEditHelper.h"
 #include "GenericCommands.h"
 
@@ -1807,43 +1808,54 @@ int32 SEditableText::FindClickedCharacterIndex( const FVector2D& InLocalCursorPo
 
 int32 SEditableText::ScanForWordBoundary( const int32 Location, int8 Direction ) const
 {
+	// We delay creating the iterator until we actually need it
+	if (!WordBreakIterator.IsValid())
+	{
+		WordBreakIterator = FBreakIterator::CreateWordBreakIterator();
+	}
+
 	const FString& EditedTextString = EditedText.ToString();
-	const int32 StringLength = EditedTextString.Len();
+	WordBreakIterator->SetString(EditedTextString);
 
+	int32 NewLocation = Location;
 	if (Direction > 0)
 	{
-		// Scan right for text
-		int32 CurCharIndex = Location;
-		while( CurCharIndex < StringLength && !FText::IsWhitespace( EditedTextString[ CurCharIndex ] ) )
+		// First move right to the next break candidate
+		NewLocation = WordBreakIterator->MoveToCandidateAfter(NewLocation);
+
+		// Then step over any whitespace
+		while (EditedTextString.IsValidIndex(NewLocation) && FText::IsWhitespace(EditedTextString[NewLocation]))
 		{
-			++CurCharIndex;
+			++NewLocation;
 		}
 
-		// Scan right for whitespace
-		while( CurCharIndex < StringLength && FText::IsWhitespace( EditedTextString[ CurCharIndex ] ) )
+		// No subsequent break candidate, just set to the upper-range
+		if (NewLocation == INDEX_NONE)
 		{
-			++CurCharIndex;
+			NewLocation = EditedTextString.Len();
 		}
-
-		return CurCharIndex;
 	}
 	else
 	{
-		// Scan left for whitespace
-		int32 CurCharIndex = Location;
-		while( CurCharIndex > 0 && FText::IsWhitespace( EditedTextString[ CurCharIndex - 1 ] ) )
+		// First step over any whitespace
+		while (EditedTextString.IsValidIndex(NewLocation-1) && FText::IsWhitespace(EditedTextString[NewLocation-1]))
 		{
-			--CurCharIndex;
+			--NewLocation;
 		}
 
-		// Scan left for text
-		while( CurCharIndex > 0 && !FText::IsWhitespace( EditedTextString[ CurCharIndex - 1 ] ) )
+		// Then move left to the next previous candidate
+		NewLocation = WordBreakIterator->MoveToCandidateBefore(NewLocation);
+
+		// No previous break candidate, just set to the lower-range
+		if (NewLocation == INDEX_NONE)
 		{
-			--CurCharIndex;
+			NewLocation = 0;
 		}
-
-		return CurCharIndex;
 	}
+
+	WordBreakIterator->ClearString();
+
+	return NewLocation;
 }
 
 
diff --git a/Engine/Source/Runtime/Slate/Public/Widgets/Input/SEditableText.h b/Engine/Source/Runtime/Slate/Public/Widgets/Input/SEditableText.h
--- a/Engine/Source/Runtime/Slate/Public/Widgets/Input/SEditableText.h
+++ b/Engine/Source/Runtime/Slate/Public/Widgets/Input/SEditableText.h
@@ -2,6 +2,7 @@
 
 #pragma once
 
+#include "IBreakIterator.h"
 #include "ITextEditorWidget.h"
 #include "ITextInputMethodSystem.h"
 #include "IVirtualKeyboardEntry.h"
@@ -671,4 +672,6 @@ private:
 	/** True if a spring animation is currently in progress */
 	bool bIsSpringing;
 
+	/** The iterator to use to detect word boundaries */
+	mutable TSharedPtr<IBreakIterator> WordBreakIterator;
 };
