Removed the animation from SEditableText so that Slate can sleep

UE-11632 - Make SEditableText sleep-friendly.

SEditableText used to have spring animations for moving the caret and changing the selection, however these animations required the editable text to keep Slate active while they had focus in order to smoothly update the animation.

These animations have been removed, and the only ticking the editable text now needs to do is to blink the cursor (which is a much coarser animation). This also brings the look-and-feel of SEditableText and SMultiLineEditableText closer together, as the latter has never used these spring animations.

diff --git a/Engine/Source/Runtime/Slate/Private/Widgets/Input/SEditableText.cpp b/Engine/Source/Runtime/Slate/Private/Widgets/Input/SEditableText.cpp
--- a/Engine/Source/Runtime/Slate/Private/Widgets/Input/SEditableText.cpp
+++ b/Engine/Source/Runtime/Slate/Private/Widgets/Input/SEditableText.cpp
@@ -13,9 +13,7 @@ SEditableText* SEditableText::EditableTextUnderCursor = nullptr;
 SEditableText::SEditableText()
 	: ScrollHelper()
 	, CaretPosition( 0 )
-	, CaretVisualPositionSpring()
 	, LastCaretInteractionTime( -1000.0 )
-	, LastSelectionInteractionTime( -1000.0 )
 	, bIsDragSelecting( false )
 	, bWasFocusedByLastMouseDown( false )
 	, bHasDragSelectedSinceFocused( false )
@@ -24,16 +22,7 @@ SEditableText::SEditableText()
 	, bIsChangingText( false )
 	, UICommandList( new FUICommandList() )
 	, bTextChangedByVirtualKeyboard( false )
-	, bIsSpringing( false )
 {
-	// Setup springs
-	FFloatSpring1D::FSpringConfig SpringConfig;
-	SpringConfig.SpringConstant = EditableTextDefs::CaretSpringConstant;
-	CaretVisualPositionSpring.SetConfig( SpringConfig );
-
-	SpringConfig.SpringConstant = EditableTextDefs::SelectionTargetSpringConstant;
-	SelectionTargetLeftSpring.SetConfig( SpringConfig );
-	SelectionTargetRightSpring.SetConfig( SpringConfig );
 }
 
 /** Destructor */
@@ -68,7 +57,6 @@ void SEditableText::Construct( const FArguments& InArgs )
 	Font = InArgs._Font.IsSet() ? InArgs._Font : InArgs._Style->Font;
 	ColorAndOpacity = InArgs._ColorAndOpacity.IsSet() ? InArgs._ColorAndOpacity : InArgs._Style->ColorAndOpacity;
 	BackgroundImageSelected = InArgs._BackgroundImageSelected.IsSet() ? InArgs._BackgroundImageSelected : &InArgs._Style->BackgroundImageSelected;
-	BackgroundImageSelectionTarget = InArgs._BackgroundImageSelectionTarget.IsSet() ? InArgs._BackgroundImageSelectionTarget : &InArgs._Style->BackgroundImageSelectionTarget;
 	BackgroundImageComposing = InArgs._BackgroundImageComposing.IsSet() ? InArgs._BackgroundImageComposing : &InArgs._Style->BackgroundImageComposing;	
 	CaretImage = InArgs._CaretImage.IsSet() ? InArgs._CaretImage : &InArgs._Style->CaretImage;
 
@@ -235,37 +223,25 @@ void SEditableText::Tick( const FGeometry& AllottedGeometry, const double InCurr
 	const FString VisibleText = GetStringToRender();
 	const FSlateFontInfo& FontInfo = Font.Get();
 	const TSharedRef< FSlateFontMeasure > FontMeasureService = FSlateApplication::Get().GetRenderer()->GetFontMeasureService();
-	const float FontMaxCharHeight = FontMeasureService->Measure( FString( ), FontInfo, AllottedGeometry.Scale ).Y;
+	const float FontMaxCharHeight = FontMeasureService->GetMaxCharacterHeight( FontInfo );
 
 	const bool bShouldAppearFocused = ShouldAppearFocused();
 	
-	// Update caret position if focused
+	float CaretVisualPosition = 0.0f;
 	if( bShouldAppearFocused )
 	{
-		FString TextBeforeCaret( VisibleText.Left( CaretPosition ) );
-		const FVector2D TextBeforeCaretSize( FontMeasureService->Measure( TextBeforeCaret, FontInfo, AllottedGeometry.Scale ) );
-
-		if (FSlateApplication::Get().IsRunningAtTargetFrameRate())
-		{
-			// Spring toward the caret's position!
-			CaretVisualPositionSpring.SetTarget(TextBeforeCaretSize.X);
-		}
-		else
-		{
-			// Just set the caret position directly
-			CaretVisualPositionSpring.SetPosition(TextBeforeCaretSize.X);
-		}
+		const FVector2D TextBeforeCaretSize = FontMeasureService->Measure( VisibleText, 0, CaretPosition, FontInfo, false, AllottedGeometry.Scale );
+		CaretVisualPosition = TextBeforeCaretSize.X;
 	}
 
-
 	// Make sure the caret is scrolled into view.  Note that even if we don't have keyboard focus (and the caret
 	// is hidden), we still do this to make sure the beginning of the string is in view.
 	{
 		const float CaretWidth = FTextEditHelper::CalculateCaretWidth(FontMaxCharHeight);
 
 		// Grab the caret position in the scrolled area space
-		const float CaretLeft = CaretVisualPositionSpring.GetPosition();
-		const float CaretRight = CaretVisualPositionSpring.GetPosition() + CaretWidth;
+		const float CaretLeft = CaretVisualPosition;
+		const float CaretRight = CaretVisualPosition + CaretWidth;
 
 		// Figure out where the caret is in widget space
 		const float WidgetSpaceCaretLeft = ScrollHelper.FromScrollerSpace( FVector2D( CaretLeft, ScrollHelper.GetPosition().Y ) ).X;
@@ -284,7 +260,6 @@ void SEditableText::Tick( const FGeometry& AllottedGeometry, const double InCurr
 			ScrollHelper.SetPosition( FVector2D( CaretRight - AllotedGeometryDrawX, ScrollHelper.GetPosition( ).Y ) );
 		}
 
-
 		// Make sure text is never scrolled out of view when it doesn't need to be!
 		{
 			// Measure the total widget of text string, plus the caret's width!
@@ -299,42 +274,6 @@ void SEditableText::Tick( const FGeometry& AllottedGeometry, const double InCurr
 			}
 		}
 	}
-
-
-	// Update selection 'target' effect
-	{
-		bool bSelectionTargetChanged = false;
-		if( AnyTextSelected() )
-		{
-			const float SelectionLeftX = EditableTextDefs::SelectionRectLeftOffset + FontMeasureService->Measure( VisibleText.Left( Selection.GetMinIndex() ), FontInfo, AllottedGeometry.Scale ).X;
-			const float SelectionRightX = EditableTextDefs::SelectionRectRightOffset + FontMeasureService->Measure( VisibleText.Left( Selection.GetMaxIndex( ) ), FontInfo, AllottedGeometry.Scale ).X;
-
-			SelectionTargetLeftSpring.SetTarget( SelectionLeftX );
-			SelectionTargetRightSpring.SetTarget( SelectionRightX );
-		}
-	}
-}
-
-EActiveTimerReturnType SEditableText::AnimateSpringsWhileFocused(double InCurrentTime, float InDeltaTime)
-{
-	// Keep ticking as long as we should appear focused or the caret is in transit
-	const bool bShouldAppearFocused = ShouldAppearFocused();
-	if (bShouldAppearFocused || !CaretVisualPositionSpring.IsAtRest())
-	{
-		// Update all the springs
-		CaretVisualPositionSpring.Tick(InDeltaTime);
-		
-		const float TimeSinceSelectionInteraction = (float)( InCurrentTime - LastSelectionInteractionTime );
-		if (TimeSinceSelectionInteraction <= EditableTextDefs::SelectionTargetEffectDuration || ( bShouldAppearFocused && AnyTextSelected() ))
-		{
-			SelectionTargetLeftSpring.Tick(InDeltaTime);
-			SelectionTargetRightSpring.Tick(InDeltaTime);
-		}
-
-		return EActiveTimerReturnType::Continue;
-	}
-
-	return EActiveTimerReturnType::Stop;
 }
 
 bool SEditableText::FTextInputMethodContext::IsReadOnly()
@@ -467,8 +406,8 @@ bool SEditableText::FTextInputMethodContext::GetTextBounds(const uint32 BeginInd
 		const FString VisibleText = OwningWidgetPtr->GetStringToRender();
 		const TSharedRef< FSlateFontMeasure > FontMeasureService = FSlateApplication::Get().GetRenderer()->GetFontMeasureService();
 
-		const float TextLeft = FontMeasureService->Measure( VisibleText.Left( BeginIndex ), FontInfo, CachedGeometry.Scale ).X;
-		const float TextRight = FontMeasureService->Measure( VisibleText.Left( BeginIndex + Length ), FontInfo, CachedGeometry.Scale ).X;
+		const float TextLeft = FontMeasureService->Measure( VisibleText, 0, BeginIndex, FontInfo, false, CachedGeometry.Scale ).X;
+		const float TextRight = FontMeasureService->Measure( VisibleText, 0, BeginIndex + Length, FontInfo, false, CachedGeometry.Scale ).X;
 
 		const float ScrollAreaLeft = OwningWidgetPtr->ScrollHelper.ToScrollerSpace( FVector2D::ZeroVector ).X;
 		const float ScrollAreaRight = OwningWidgetPtr->ScrollHelper.ToScrollerSpace( CachedGeometry.Size ).X;
@@ -828,12 +767,6 @@ void SEditableText::JumpTo(ETextLocation JumpLocation, ECursorAction Action)
 
 void SEditableText::ClearSelection()
 {
-	// Only animate clearing of the selection if there was actually a character selected
-	if( AnyTextSelected() )
-	{
-		RestartSelectionTargetAnimation();
-	}
-
 	// Explicitly invalidate the selection range indices.  The user no longer has a selection
 	// start and finish point
 	Selection.Clear();
@@ -847,13 +780,6 @@ void SEditableText::SelectAllText()
 	
 	Selection.StartIndex = EditedText.ToString().Len();
 	Selection.FinishIndex = 0;
-	RestartSelectionTargetAnimation();
-
-	// Make sure the animated selection effect starts roughly where the cursor is
-	const TSharedRef< FSlateFontMeasure > FontMeasureService = FSlateApplication::Get().GetRenderer()->GetFontMeasureService();
-	const float SelectionTargetX = EditableTextDefs::SelectionRectLeftOffset + FontMeasureService->Measure( EditedText.ToString( ).Left( CaretPosition ), Font.Get( ), 1.0f/*CachedGeometry.Scale*/ ).X;
-	SelectionTargetLeftSpring.SetPosition( SelectionTargetX );
-	SelectionTargetRightSpring.SetPosition( SelectionTargetX );
 }
 
 bool SEditableText::SelectAllTextWhenFocused()
@@ -894,7 +820,6 @@ void SEditableText::SelectWordAt( const FGeometry& MyGeometry, const FVector2D&
 		// Select the word!
 		Selection.StartIndex = FirstWordCharIndex;
 		Selection.FinishIndex = LastWordCharIndex;
-		RestartSelectionTargetAnimation();
 	}
 }
 
@@ -1285,8 +1210,8 @@ int32 SEditableText::OnPaint( const FPaintArgs& Args, const FGeometry& AllottedG
 		// Figure out a universally visible selection color.
 		const FColor SelectionBackgroundColorAndOpacity = ( (FLinearColor::White - ThisColorAndOpacity)*0.5f + FLinearColor(-0.2f, -0.05f, 0.15f)) * InWidgetStyle.GetColorAndOpacityTint();
 
-		const float SelectionLeftX = EditableTextDefs::SelectionRectLeftOffset + FontMeasureService->Measure( VisibleText.Left( Selection.GetMinIndex() ), FontInfo, AllottedGeometry.Scale ).X;
-		const float SelectionRightX = EditableTextDefs::SelectionRectRightOffset + FontMeasureService->Measure( VisibleText.Left( Selection.GetMaxIndex( ) ), FontInfo, AllottedGeometry.Scale ).X;
+		const float SelectionLeftX = EditableTextDefs::SelectionRectLeftOffset + FontMeasureService->Measure( VisibleText, 0, Selection.GetMinIndex(), FontInfo, false, AllottedGeometry.Scale ).X;
+		const float SelectionRightX = EditableTextDefs::SelectionRectRightOffset + FontMeasureService->Measure( VisibleText, 0, Selection.GetMaxIndex(), FontInfo, false, AllottedGeometry.Scale ).X;
 		const float SelectionTopY = 0.0f;
 		const float SelectionBottomY = DrawGeometrySize.Y;
 
@@ -1308,10 +1233,10 @@ int32 SEditableText::OnPaint( const FPaintArgs& Args, const FGeometry& AllottedG
 	const FTextRange CompositionRange(TextInputMethodContext->CompositionBeginIndex, TextInputMethodContext->CompositionBeginIndex + TextInputMethodContext->CompositionLength);
 	if( TextInputMethodContext->IsComposing && CompositionRange.InclusiveContains(CaretPosition) )
 	{
-		const float CompositionLeft = FontMeasureService->Measure( VisibleText.Left( CompositionRange.BeginIndex ), FontInfo, AllottedGeometry.Scale ).X;
-		const float CompositionRight = FontMeasureService->Measure( VisibleText.Left( CompositionRange.EndIndex ), FontInfo, AllottedGeometry.Scale ).X;
+		const float CompositionLeft = FontMeasureService->Measure( VisibleText, 0, CompositionRange.BeginIndex, FontInfo, false, AllottedGeometry.Scale ).X;
+		const float CompositionRight = FontMeasureService->Measure( VisibleText, 0, CompositionRange.EndIndex, FontInfo, false, AllottedGeometry.Scale ).X;
 		const float CompositionTop = 0.0f;
-		const float CompositionBottom = FontMeasureService->Measure( VisibleText.Mid( TextInputMethodContext->CompositionBeginIndex, TextInputMethodContext->CompositionLength ), FontInfo, AllottedGeometry.Scale ).Y;
+		const float CompositionBottom = FontMeasureService->Measure( VisibleText, TextInputMethodContext->CompositionBeginIndex, TextInputMethodContext->CompositionLength, FontInfo, false, AllottedGeometry.Scale ).Y;
 
 		const FVector2D DrawPosition = ScrollHelper.FromScrollerSpace( FVector2D( CompositionLeft, CompositionTop ) );
 		const FVector2D DrawSize = ScrollHelper.SizeFromScrollerSpace( FVector2D( CompositionRight - CompositionLeft, CompositionBottom - CompositionTop ) );
@@ -1353,7 +1278,7 @@ int32 SEditableText::OnPaint( const FPaintArgs& Args, const FGeometry& AllottedG
 		const int32 FirstVisibleCharacter = FontMeasureService->FindCharacterIndexAtOffset( VisibleText, FontInfo, ScrollAreaLeft, AllottedGeometry.Scale );
 		const int32 LastVisibleCharacter = FontMeasureService->FindCharacterIndexAtOffset( VisibleText, FontInfo, ScrollAreaRight, AllottedGeometry.Scale );
 		const FString PotentiallyVisibleText( VisibleText.Mid( FirstVisibleCharacter, ( LastVisibleCharacter - FirstVisibleCharacter ) + 1 ) );
-		const float FirstVisibleCharacterPosition = FontMeasureService->Measure( VisibleText.Left( FirstVisibleCharacter ), FontInfo, AllottedGeometry.Scale ).X;
+		const float FirstVisibleCharacterPosition = FontMeasureService->Measure( VisibleText, 0, FirstVisibleCharacter, FontInfo, false, AllottedGeometry.Scale ).X;
 
 		const float TextVertOffset = EditableTextDefs::TextVertOffsetPercent * FontMaxCharHeight;
 		const FVector2D DrawPosition = ScrollHelper.FromScrollerSpace( FVector2D( FirstVisibleCharacterPosition, DrawPositionY + TextVertOffset ) );
@@ -1372,84 +1297,23 @@ int32 SEditableText::OnPaint( const FPaintArgs& Args, const FGeometry& AllottedG
 	}
 
 
-	// Draw selection targeting effect
-	const float TimeSinceSelectionInteraction = (float)( CurrentTime - LastSelectionInteractionTime );
-	if( TimeSinceSelectionInteraction <= EditableTextDefs::SelectionTargetEffectDuration )
-	{
-		// Don't draw selection effect when drag-selecting text using the mouse
-		if( !bIsDragSelecting || !bHasDragSelectedSinceMouseDown )
-		{
-			const bool bIsClearingSelection = !AnyTextSelected() || ( !bIsReadonly && !bShouldAppearFocused );
-
-			// Compute animation progress
-			float EffectAlpha = FMath::Clamp( TimeSinceSelectionInteraction / EditableTextDefs::SelectionTargetEffectDuration, 0.0f, 1.0f );
-			EffectAlpha = 1.0f - EffectAlpha * EffectAlpha;  // Inverse square falloff (looks nicer!)
-
-			// Apply extra opacity falloff when deselecting text
-			float EffectOpacity = EffectAlpha;
-			if( bIsClearingSelection )
-			{
-				EffectOpacity *= EffectOpacity;
-			}
-
-			const FSlateBrush* SelectionTargetBrush = BackgroundImageSelectionTarget.Get();
-
-			// Set final opacity of the effect
-			FLinearColor SelectionTargetColorAndOpacity = SelectionTargetBrush->GetTint( InWidgetStyle );
-			SelectionTargetColorAndOpacity.A *= EditableTextDefs::SelectionTargetOpacity * EffectOpacity;
-			const FColor SelectionTargetColorAndOpacitySRGB = SelectionTargetColorAndOpacity;
-
-			// Compute the bounds offset of the selection target from where the selection target spring
-			// extents currently lie.  This is used to "grow" or "shrink" the selection as needed.
-			const float SelectingAnimOffset = EditableTextDefs::SelectingAnimOffsetPercent * FontMaxCharHeight;
-			const float DeselectingAnimOffset = EditableTextDefs::DeselectingAnimOffsetPercent * FontMaxCharHeight;
-
-			// Choose an offset amount depending on whether we're selecting text, or clearing selection
-			const float EffectOffset = bIsClearingSelection ? ( 1.0f - EffectAlpha ) * DeselectingAnimOffset : EffectAlpha * SelectingAnimOffset;
-
-			const float SelectionLeftX = SelectionTargetLeftSpring.GetPosition() - EffectOffset;
-			const float SelectionRightX = SelectionTargetRightSpring.GetPosition() + EffectOffset;
-			const float SelectionTopY = 0.0f - EffectOffset;
-			const float SelectionBottomY = AllottedGeometry.Size.Y + EffectOffset;
-
-			const FVector2D DrawPosition = ScrollHelper.FromScrollerSpace( FVector2D( SelectionLeftX, SelectionTopY ) );
-			const FVector2D DrawSize = ScrollHelper.SizeFromScrollerSpace( FVector2D( SelectionRightX - SelectionLeftX, SelectionBottomY - SelectionTopY ) );
-
-			// NOTE: We rely on scissor clipping for the selection rectangle
-			FSlateDrawElement::MakeBox(
-				OutDrawElements,
-				LayerId + TextLayer,
-				AllottedGeometry.ToPaintGeometry( DrawPosition, DrawSize, ScaleInverse ),	// Position, Size, Scale
-				SelectionTargetBrush,										// Image
-				MyClippingRect,												// Clipping rect
-				DrawEffects,												// Effects to use
-				SelectionTargetColorAndOpacitySRGB );						// Color
-		}
-	}
-
-
 	// Draw the caret!
 	if( bShouldAppearFocused )
 	{
-		// Handle caret blinking
+		// The caret is always visible (i.e. not blinking) when we're interacting with it; otherwise it might get lost.
+		const bool bForceCaretVisible = (CurrentTime - LastCaretInteractionTime) < EditableTextDefs::CaretBlinkPauseTime;
+		float CaretOpacity = (bForceCaretVisible)
+			? 1.0f
+			: FMath::RoundToFloat( FMath::MakePulsatingValue( CurrentTime, EditableTextDefs::BlinksPerSecond ));
+		CaretOpacity *= CaretOpacity;	// Squared falloff, because it looks more interesting
+
 		FLinearColor CursorColorAndOpacity = ThisColorAndOpacity;
-		{
-			// Only blink when the user isn't actively typing
-			const double BlinkPauseEndTime = LastCaretInteractionTime + EditableTextDefs::CaretBlinkPauseTime;
-			if( CurrentTime > BlinkPauseEndTime )
-			{
-				// Caret pulsate time is relative to the last time that we stopped interacting with
-				// the cursor.  This just makes sure that the caret starts out fully opaque.
-				const double PulsateTime = CurrentTime - BlinkPauseEndTime;
-
-				// Blink the caret!
-				float CaretOpacity = FMath::MakePulsatingValue( PulsateTime, EditableTextDefs::BlinksPerSecond );
-				CaretOpacity *= CaretOpacity;	// Squared falloff, because it looks more interesting
-				CursorColorAndOpacity.A = CaretOpacity;
-			}
-		}
+		CursorColorAndOpacity.A = CaretOpacity;
+
+		const FVector2D TextBeforeCaretSize = FontMeasureService->Measure( VisibleText, 0, CaretPosition, FontInfo, false, AllottedGeometry.Scale );
+		const float CaretVisualPosition = TextBeforeCaretSize.X;
 
-		const float CaretX = CaretVisualPositionSpring.GetPosition();
+		const float CaretX = CaretVisualPosition;
 		const float CaretY = DrawPositionY + EditableTextDefs::CaretVertOffset;
 		const float CaretWidth = FTextEditHelper::CalculateCaretWidth(FontMaxCharHeight);
 		const float CaretHeight = EditableTextDefs::CaretHeightPercent * FontMaxCharHeight;
@@ -1511,20 +1375,14 @@ FReply SEditableText::OnFocusReceived( const FGeometry& MyGeometry, const FFocus
 	// Skip the rest of the focus received code if it's due to the context menu closing
 	if ( !ActiveContextMenu.IsValid() )
 	{
-		RegisterActiveTimer(0.f, FWidgetActiveTimerDelegate::CreateSP(this, &SEditableText::AnimateSpringsWhileFocused));
+		EnsureActiveTick();
 
 		// Don't unselect text when focus is set because another widget lost focus, as that may be in response to a
 		// dismissed context menu where the user clicked an option to select text.
 		if( InFocusEvent.GetCause() != EFocusCause::OtherWidgetLostFocus )
 		{
-			// Deselect text, but override the last selection interaction time, so that the user doesn't
-			// see a selection transition animation when keyboard focus is received.
 			ClearSelection();
-			LastSelectionInteractionTime = -1000.0;
 		}
-
-		// Reset the cursor position spring
-		CaretVisualPositionSpring.SetPosition( CaretVisualPositionSpring.GetPosition() );
 	 
 		bHasDragSelectedSinceFocused = false;
 
@@ -1578,15 +1436,7 @@ void SEditableText::OnFocusLost( const FFocusEvent& InFocusEvent )
 	if ( !ActiveContextMenu.IsValid() )
 	{
 		// Place the caret at the front of the text
-		if (FSlateApplication::Get().IsRunningAtTargetFrameRate())
-		{
-			// Spring back when the frame rate is high enough
-			CaretVisualPositionSpring.SetTarget(0.f);
-		}
-		else
-		{
-			CaretVisualPositionSpring.SetPosition(0.f);
-		}
+		SetCaretPosition(0);
 
 		FSlateApplication& SlateApplication = FSlateApplication::Get();
 		if (FPlatformMisc::GetRequiresVirtualKeyboard())
@@ -1691,25 +1541,6 @@ FReply SEditableText::OnMouseButtonDown( const FGeometry& InMyGeometry, const FP
 
 FReply SEditableText::OnMouseButtonUp( const FGeometry& InMyGeometry, const FPointerEvent& InMouseEvent )
 {
-	// The mouse must have been captured by either left or right button down before we'll process mouse ups
-	if( HasMouseCapture() )
-	{
-		if( InMouseEvent.GetEffectingButton() == EKeys::LeftMouseButton && bIsDragSelecting )
-		{
-			// If we selected any text while dragging, then kick off the selection target animation
-			if( bHasDragSelectedSinceMouseDown && Selection.StartIndex != Selection.FinishIndex )
-			{
-				// Start the selection animation wherever the user started the click/drag
-				const TSharedRef< FSlateFontMeasure > FontMeasureService = FSlateApplication::Get().GetRenderer()->GetFontMeasureService();
-				const float SelectionTargetX = EditableTextDefs::SelectionRectLeftOffset + FontMeasureService->Measure( EditedText.ToString( ).Left( Selection.StartIndex ), Font.Get( ), InMyGeometry.Scale ).X;
-				SelectionTargetLeftSpring.SetPosition( SelectionTargetX );
-				SelectionTargetRightSpring.SetPosition( SelectionTargetX );
-
-				RestartSelectionTargetAnimation();
-			}
-		}
-	}
-
 	return FTextEditHelper::OnMouseButtonUp( InMyGeometry, InMouseEvent, SharedThis( this ) );
 }
 
@@ -1745,8 +1576,6 @@ void SEditableText::SelectText( const int32 InOldCaretPosition )
 		{
 			// Start selecting
 			Selection.StartIndex = Selection.FinishIndex = InOldCaretPosition;
-
-			RestartSelectionTargetAnimation();
 		}
 		else
 		{
@@ -1770,12 +1599,6 @@ void SEditableText::DeleteSelectedText()
 		const FString NewText = OldText.Left( Selection.GetMinIndex() ) + OldText.Mid( Selection.GetMaxIndex() );
 		EditedText = FText::FromString(NewText);
 
-		// Visually collapse the selection target right before we clear selection
-		const TSharedRef< FSlateFontMeasure > FontMeasureService = FSlateApplication::Get().GetRenderer()->GetFontMeasureService();
-		const float SelectionTargetX = EditableTextDefs::SelectionRectLeftOffset + FontMeasureService->Measure( NewText.Left( CaretPosition ), Font.Get( ), 1.0f/*CachedGeometry.Scale*/ ).X;
-		SelectionTargetLeftSpring.SetTarget( SelectionTargetX );
-		SelectionTargetRightSpring.SetTarget( SelectionTargetX );
-
 		// Clear selection
 		ClearSelection();
 	}
@@ -2155,16 +1978,6 @@ FString SEditableText::GetStringToRender() const
 	return VisibleText;
 }
 
-void SEditableText::RestartSelectionTargetAnimation()
-{
-	// Don't bother with this unless we're running at a decent frame rate
-	if( FSlateApplication::Get().IsRunningAtTargetFrameRate() )
-	{
-		LastSelectionInteractionTime = FSlateApplication::Get().GetCurrentTime();
-	}
-}
-
-
 void SEditableText::OnWindowClosed(const TSharedRef<SWindow>&)
 {
 	// Note: We don't reset the ActiveContextMenu here, as Slate hasn't yet finished processing window focus events, and we need 
@@ -2223,3 +2036,21 @@ void SEditableText::SetSelectAllTextOnCommit(const TAttribute<bool>& InSelectAll
 {
 	SelectAllTextOnCommit = InSelectAllTextOnCommit;
 }
+
+void SEditableText::EnsureActiveTick()
+{
+	TSharedPtr<FActiveTimerHandle> ActiveTickTimerPin = ActiveTickTimer.Pin();
+	if(ActiveTickTimerPin.IsValid())
+	{
+		return;
+	}
+
+	auto DoActiveTick = [this](double InCurrentTime, float InDeltaTime) -> EActiveTimerReturnType
+	{
+		// Continue if we still have focus, otherwise treat as a fire-and-forget Tick() request
+		const bool bShouldAppearFocused = HasKeyboardFocus() || ActiveContextMenu.IsValid();
+		return (bShouldAppearFocused) ? EActiveTimerReturnType::Continue : EActiveTimerReturnType::Stop;
+	};
+
+	ActiveTickTimer = RegisterActiveTimer(0.5f, FWidgetActiveTimerDelegate::CreateLambda(DoActiveTick));
+}
diff --git a/Engine/Source/Runtime/Slate/Public/Widgets/Input/SEditableText.h b/Engine/Source/Runtime/Slate/Public/Widgets/Input/SEditableText.h
--- a/Engine/Source/Runtime/Slate/Public/Widgets/Input/SEditableText.h
+++ b/Engine/Source/Runtime/Slate/Public/Widgets/Input/SEditableText.h
@@ -23,7 +23,6 @@ public:
 		, _Font()
 		, _ColorAndOpacity()
 		, _BackgroundImageSelected()
-		, _BackgroundImageSelectionTarget()
 		, _BackgroundImageComposing()
 		, _CaretImage()
 		, _IsReadOnly( false )
@@ -56,9 +55,6 @@ public:
 		/** Background image for the selected text (overrides Style) */
 		SLATE_ATTRIBUTE( const FSlateBrush*, BackgroundImageSelected )
 
-		/** Background image for the selection targeting effect (overrides Style) */
-		SLATE_ATTRIBUTE( const FSlateBrush*, BackgroundImageSelectionTarget )
-
 		/** Background image for the composing text (overrides Style) */
 		SLATE_ATTRIBUTE( const FSlateBrush*, BackgroundImageComposing )
 
@@ -502,20 +498,12 @@ protected:
 	 */
 	void SetCaretPosition( int32 Position );
 
-	/**
-	 * Called to restart selection target animation after most types of text selection changes
-	 */
-	void RestartSelectionTargetAnimation();
-
 	/**
 	 * Called when the content menu window is closed
 	 */
 	void OnWindowClosed(const TSharedRef<SWindow>&);
 
 private:
-	/** Animates the caret and highlight selection springs */
-	EActiveTimerReturnType AnimateSpringsWhileFocused(double InCurrentTime, float InDeltaTime);
-	
 	/** @return Whether the editable text should appear focused */
 	bool ShouldAppearFocused() const;
 
@@ -529,6 +517,12 @@ private:
 	 */
 	FString GetStringToRender() const;
 
+	/**
+	 * Ensure that we will get a Tick() soon (either due to having active focus, or something having changed progmatically and requiring an update)
+	 * Does nothing if the active tick timer is already enabled
+	 */
+	void EnsureActiveTick();
+
 private:
 
 	/** A pointer to the editable text currently under the mouse cursor. nullptr when there isn't one. */
@@ -549,9 +543,6 @@ private:
 	/** Background image for the selected text */
 	TAttribute< const FSlateBrush* > BackgroundImageSelected;
 
-	/** Background image for the selection targeting effect */
-	TAttribute< const FSlateBrush* > BackgroundImageSelectionTarget;
-
 	/** Background image for the composing text */
 	TAttribute< const FSlateBrush* > BackgroundImageComposing;
 
@@ -603,24 +594,12 @@ private:
 	/** Caret position as an index into the editable text string's character array */
 	int32 CaretPosition;
 
-	/** Horizontal visual position of caret along the scrolling canvas */
-	FFloatSpring1D CaretVisualPositionSpring;
-
 	/** Stores the last time that the caret was moved by the user.  Used to prevent the caret from blinking. */
 	double LastCaretInteractionTime;
 
 	/** Current selection state */
 	FTextSelection Selection;
 
-	/** Selection "targeting" visual effect left position */
-	FFloatSpring1D SelectionTargetLeftSpring;
-
-	/** Selection "targeting" visual effect right position */
-	FFloatSpring1D SelectionTargetRightSpring;
-
-	/** Last time that the user had a major interaction with the selection */
-	double LastSelectionInteractionTime;
-
 	/** True if we're currently selecting text by dragging the mouse cursor with the left button held down */
 	bool bIsDragSelecting;
 
@@ -669,8 +648,8 @@ private:
 	/** True if the text has been changed by a virtual keyboard */
 	bool bTextChangedByVirtualKeyboard;
 
-	/** True if a spring animation is currently in progress */
-	bool bIsSpringing;
+	/** The timer that is actively driving this widget to Tick() even when Slate is idle */
+	TWeakPtr<FActiveTimerHandle> ActiveTickTimer;
 
 	/** The iterator to use to detect word boundaries */
 	mutable TSharedPtr<IBreakIterator> WordBreakIterator;
diff --git a/Engine/Source/Runtime/SlateCore/Private/Styling/CoreStyle.cpp b/Engine/Source/Runtime/SlateCore/Private/Styling/CoreStyle.cpp
--- a/Engine/Source/Runtime/SlateCore/Private/Styling/CoreStyle.cpp
+++ b/Engine/Source/Runtime/SlateCore/Private/Styling/CoreStyle.cpp
@@ -413,7 +413,6 @@ TSharedRef<ISlateStyle> FCoreStyle::Create( const FName& InStyleSetName )
 
 		const FEditableTextStyle NormalEditableTextStyle = FEditableTextStyle()
 			.SetBackgroundImageSelected( *SelectionBackground )
-			.SetBackgroundImageSelectionTarget( *SelectionTarget )
 			.SetBackgroundImageComposing( *CompositionBackground )
 			.SetCaretImage( *GenericWhiteBox );
 		Style->Set( "NormalEditableText", NormalEditableTextStyle );
diff --git a/Engine/Source/Runtime/SlateCore/Private/Styling/SlateTypes.cpp b/Engine/Source/Runtime/SlateCore/Private/Styling/SlateTypes.cpp
--- a/Engine/Source/Runtime/SlateCore/Private/Styling/SlateTypes.cpp
+++ b/Engine/Source/Runtime/SlateCore/Private/Styling/SlateTypes.cpp
@@ -191,7 +191,6 @@ FEditableTextStyle::FEditableTextStyle()
 	: Font(FStyleDefaults::GetFontInfo(9))
 	, ColorAndOpacity(FSlateColor::UseForeground())
 	, BackgroundImageSelected()
-	, BackgroundImageSelectionTarget()
 	, CaretImage()
 {
 }
@@ -199,7 +198,6 @@ FEditableTextStyle::FEditableTextStyle()
 void FEditableTextStyle::GetResources( TArray< const FSlateBrush* >& OutBrushes ) const
 {
 	OutBrushes.Add( &BackgroundImageSelected );
-	OutBrushes.Add( &BackgroundImageSelectionTarget );
 	OutBrushes.Add( &CaretImage );
 }
 
diff --git a/Engine/Source/Runtime/SlateCore/Public/Styling/SlateTypes.h b/Engine/Source/Runtime/SlateCore/Public/Styling/SlateTypes.h
--- a/Engine/Source/Runtime/SlateCore/Public/Styling/SlateTypes.h
+++ b/Engine/Source/Runtime/SlateCore/Public/Styling/SlateTypes.h
@@ -497,11 +497,6 @@ struct SLATECORE_API FEditableTextStyle : public FSlateWidgetStyle
 	FSlateBrush BackgroundImageSelected;
 	FEditableTextStyle& SetBackgroundImageSelected( const FSlateBrush& InBackgroundImageSelected ){ BackgroundImageSelected = InBackgroundImageSelected; return *this; }
 
-	/** Background image for the selection targeting effect */
-	UPROPERTY(EditAnywhere, BlueprintReadWrite, Category=Appearance)
-	FSlateBrush BackgroundImageSelectionTarget;
-	FEditableTextStyle& SetBackgroundImageSelectionTarget( const FSlateBrush& InBackgroundImageSelectionTarget ){ BackgroundImageSelectionTarget = InBackgroundImageSelectionTarget; return *this; }
-
 	/** Background image for the selected text */
 	UPROPERTY(EditAnywhere, BlueprintReadWrite, Category=Appearance)
 	FSlateBrush BackgroundImageComposing;
diff --git a/Engine/Source/Runtime/UMG/Private/Components/EditableText.cpp b/Engine/Source/Runtime/UMG/Private/Components/EditableText.cpp
--- a/Engine/Source/Runtime/UMG/Private/Components/EditableText.cpp
+++ b/Engine/Source/Runtime/UMG/Private/Components/EditableText.cpp
@@ -150,12 +150,6 @@ void UEditableText::PostLoad()
 			BackgroundImageSelected_DEPRECATED = nullptr;
 		}
 
-		if ( BackgroundImageSelectionTarget_DEPRECATED != nullptr )
-		{
-			WidgetStyle.BackgroundImageSelectionTarget = BackgroundImageSelectionTarget_DEPRECATED->Brush;
-			BackgroundImageSelectionTarget_DEPRECATED = nullptr;
-		}
-
 		if ( BackgroundImageComposing_DEPRECATED != nullptr )
 		{
 			WidgetStyle.BackgroundImageComposing = BackgroundImageComposing_DEPRECATED->Brush;
diff --git a/Engine/Source/Runtime/UMG/Public/Components/EditableText.h b/Engine/Source/Runtime/UMG/Public/Components/EditableText.h
--- a/Engine/Source/Runtime/UMG/Public/Components/EditableText.h
+++ b/Engine/Source/Runtime/UMG/Public/Components/EditableText.h
@@ -50,10 +50,6 @@ public:
 	UPROPERTY()
 	USlateBrushAsset* BackgroundImageSelected_DEPRECATED;
 
-	/** Background image for the selection targeting effect (overrides Style) */
-	UPROPERTY()
-	USlateBrushAsset* BackgroundImageSelectionTarget_DEPRECATED;
-
 	/** Background image for the composing text (overrides Style) */
 	UPROPERTY()
 	USlateBrushAsset* BackgroundImageComposing_DEPRECATED;
