Fixed SMultiLineEditableText Slate widget's abillity to select and copy text from.

diff --git a/Engine/Source/Runtime/Slate/Private/Widgets/Text/SMultiLineEditableText.cpp b/Engine/Source/Runtime/Slate/Private/Widgets/Text/SMultiLineEditableText.cpp
--- a/Engine/Source/Runtime/Slate/Private/Widgets/Text/SMultiLineEditableText.cpp
+++ b/Engine/Source/Runtime/Slate/Private/Widgets/Text/SMultiLineEditableText.cpp
@@ -1943,16 +1943,25 @@ bool SMultiLineEditableText::IsTextSelectedAt( const FGeometry& MyGeometry, cons
 	int32 SelectionEndLineIndex = Selection.GetEnd().GetLineIndex();
 	int32 SelectionEndLineOffset = Selection.GetEnd().GetOffset();
 
+	if (SelectionBeginningLineIndex == SelectionEndLineIndex)
+	{
+		return ClickedPosition.GetLineIndex() == SelectionBeginningLineIndex
+			&& SelectionBeginningLineOffset <= ClickedPosition.GetOffset()
+			&& SelectionEndLineOffset >= ClickedPosition.GetOffset();
+	}
+
 	if (SelectionBeginningLineIndex == ClickedPosition.GetLineIndex())
 	{
-		return SelectionBeginningLineOffset >= ClickedPosition.GetOffset();
+		return SelectionBeginningLineOffset <= ClickedPosition.GetOffset();
 	}
-	else if (SelectionEndLineIndex == ClickedPosition.GetLineIndex())
+
+	if (SelectionEndLineIndex == ClickedPosition.GetLineIndex())
 	{
-		return SelectionEndLineOffset < ClickedPosition.GetOffset();
+		return SelectionEndLineOffset >= ClickedPosition.GetOffset();
 	}
 	
-	return SelectionBeginningLineIndex < ClickedPosition.GetLineIndex() && SelectionEndLineIndex > ClickedPosition.GetLineIndex();
+	return SelectionBeginningLineIndex < ClickedPosition.GetLineIndex()
+		&& SelectionEndLineIndex > ClickedPosition.GetLineIndex();
 }
 
 void SMultiLineEditableText::SetWasFocusedByLastMouseDown( bool Value )
