Added a way for UI actions to be repeated if the chord used to perform them is being held down

The default is to disallow repeating to preserve the existing behavior, but the Undo/Redo and Paste actions in SEditableText and SMultiLineEditableText have been updated to allow repeating. This makes them behave more like how a text editor is expected to behave.

diff --git a/Engine/Source/Runtime/Slate/Private/Framework/Commands/UICommandList.cpp b/Engine/Source/Runtime/Slate/Private/Framework/Commands/UICommandList.cpp
--- a/Engine/Source/Runtime/Slate/Private/Framework/Commands/UICommandList.cpp
+++ b/Engine/Source/Runtime/Slate/Private/Framework/Commands/UICommandList.cpp
@@ -6,34 +6,34 @@
 /* FUICommandList interface
  *****************************************************************************/
 
-void FUICommandList::MapAction( const TSharedPtr< const FUICommandInfo > InUICommandInfo, FExecuteAction ExecuteAction )
+void FUICommandList::MapAction( const TSharedPtr< const FUICommandInfo > InUICommandInfo, FExecuteAction ExecuteAction, EUIActionRepeatMode RepeatMode )
 {
-	MapAction( InUICommandInfo, ExecuteAction, FCanExecuteAction(), FIsActionChecked(), FIsActionButtonVisible() );
+	MapAction( InUICommandInfo, ExecuteAction, FCanExecuteAction(), FIsActionChecked(), FIsActionButtonVisible(), RepeatMode );
 }
 
 
-void FUICommandList::MapAction( const TSharedPtr< const FUICommandInfo > InUICommandInfo, FExecuteAction ExecuteAction, FCanExecuteAction CanExecuteAction )
+void FUICommandList::MapAction( const TSharedPtr< const FUICommandInfo > InUICommandInfo, FExecuteAction ExecuteAction, FCanExecuteAction CanExecuteAction, EUIActionRepeatMode RepeatMode )
 {
-	MapAction( InUICommandInfo, ExecuteAction, CanExecuteAction, FIsActionChecked(), FIsActionButtonVisible() );
+	MapAction( InUICommandInfo, ExecuteAction, CanExecuteAction, FIsActionChecked(), FIsActionButtonVisible(), RepeatMode );
 }
 
 
-void FUICommandList::MapAction( const TSharedPtr< const FUICommandInfo > InUICommandInfo, FExecuteAction ExecuteAction, FCanExecuteAction CanExecuteAction, FIsActionChecked IsCheckedDelegate )
+void FUICommandList::MapAction( const TSharedPtr< const FUICommandInfo > InUICommandInfo, FExecuteAction ExecuteAction, FCanExecuteAction CanExecuteAction, FIsActionChecked IsCheckedDelegate, EUIActionRepeatMode RepeatMode )
 {
-	MapAction( InUICommandInfo, ExecuteAction, CanExecuteAction, IsCheckedDelegate, FIsActionButtonVisible() );
+	MapAction( InUICommandInfo, ExecuteAction, CanExecuteAction, IsCheckedDelegate, FIsActionButtonVisible(), RepeatMode );
 }
 
 
-void FUICommandList::MapAction( const TSharedPtr< const FUICommandInfo > InUICommandInfo, FExecuteAction ExecuteAction, FCanExecuteAction CanExecuteAction, FIsActionChecked IsCheckedDelegate, FIsActionButtonVisible IsVisibleDelegate )
+void FUICommandList::MapAction( const TSharedPtr< const FUICommandInfo > InUICommandInfo, FExecuteAction ExecuteAction, FCanExecuteAction CanExecuteAction, FIsActionChecked IsCheckedDelegate, FIsActionButtonVisible IsVisibleDelegate, EUIActionRepeatMode RepeatMode )
 {
 	FUIAction Action;
 	Action.ExecuteAction = ExecuteAction;
 	Action.CanExecuteAction = CanExecuteAction;
 	Action.IsCheckedDelegate = IsCheckedDelegate;
 	Action.IsActionVisibleDelegate = IsVisibleDelegate;
+	Action.RepeatMode = RepeatMode;
 
 	MapAction( InUICommandInfo, Action );
-
 }
 
 
@@ -161,7 +161,7 @@ bool FUICommandList::ProcessCommandBindings( const FKey Key, const FModifierKeys
 
 bool FUICommandList::ConditionalProcessCommandBindings( const FKey Key, bool bCtrl, bool bAlt, bool bShift, bool bCmd, bool bRepeat ) const
 {
-	if ( !bRepeat && !FSlateApplication::Get().IsDragDropping() )
+	if ( !FSlateApplication::Get().IsDragDropping() )
 	{
 		FInputChord CheckChord( Key, EModifierKey::FromBools(bCtrl, bAlt, bShift, bCmd) );
 
@@ -188,7 +188,7 @@ bool FUICommandList::ConditionalProcessCommandBindings( const FKey Key, bool bCt
 					// If there is no Action mapped to this command list, continue to the next context
 					if( Action )
 					{
-						if(Action->CanExecute())
+						if(Action->CanExecute() && (!bRepeat || Action->CanRepeat()))
 						{
 							// If the action was found and can be executed, do so now
 							Action->Execute();
diff --git a/Engine/Source/Runtime/Slate/Private/Widgets/Input/SEditableText.cpp b/Engine/Source/Runtime/Slate/Private/Widgets/Input/SEditableText.cpp
--- a/Engine/Source/Runtime/Slate/Private/Widgets/Input/SEditableText.cpp
+++ b/Engine/Source/Runtime/Slate/Private/Widgets/Input/SEditableText.cpp
@@ -77,7 +77,8 @@ void SEditableText::Construct( const FArguments& InArgs )
 	// Map UI commands to delegates which are called when the command should be executed
 	UICommandList->MapAction( FGenericCommands::Get().Undo,
 		FExecuteAction::CreateSP( this, &SEditableText::Undo ),
-		FCanExecuteAction::CreateSP( this, &SEditableText::CanExecuteUndo ) );
+		FCanExecuteAction::CreateSP( this, &SEditableText::CanExecuteUndo ),
+		EUIActionRepeatMode::RepeatEnabled );
 
 	UICommandList->MapAction( FGenericCommands::Get().Cut,
 		FExecuteAction::CreateSP( this, &SEditableText::CutSelectedTextToClipboard ),
@@ -85,7 +86,8 @@ void SEditableText::Construct( const FArguments& InArgs )
 
 	UICommandList->MapAction( FGenericCommands::Get().Paste,
 		FExecuteAction::CreateSP( this, &SEditableText::PasteTextFromClipboard ),
-		FCanExecuteAction::CreateSP( this, &SEditableText::CanExecutePaste ) );
+		FCanExecuteAction::CreateSP( this, &SEditableText::CanExecutePaste ),
+		EUIActionRepeatMode::RepeatEnabled );
 
 	UICommandList->MapAction( FGenericCommands::Get().Copy,
 		FExecuteAction::CreateSP( this, &SEditableText::CopySelectedTextToClipboard ),
diff --git a/Engine/Source/Runtime/Slate/Private/Widgets/Text/SMultiLineEditableText.cpp b/Engine/Source/Runtime/Slate/Private/Widgets/Text/SMultiLineEditableText.cpp
--- a/Engine/Source/Runtime/Slate/Private/Widgets/Text/SMultiLineEditableText.cpp
+++ b/Engine/Source/Runtime/Slate/Private/Widgets/Text/SMultiLineEditableText.cpp
@@ -318,7 +318,8 @@ void SMultiLineEditableText::Construct( const FArguments& InArgs )
 	// Map UI commands to delegates which are called when the command should be executed
 	UICommandList->MapAction(FGenericCommands::Get().Undo,
 		FExecuteAction::CreateSP(this, &SMultiLineEditableText::Undo),
-		FCanExecuteAction::CreateSP(this, &SMultiLineEditableText::CanExecuteUndo));
+		FCanExecuteAction::CreateSP(this, &SMultiLineEditableText::CanExecuteUndo),
+		EUIActionRepeatMode::RepeatEnabled);
 
 	UICommandList->MapAction(FGenericCommands::Get().Cut,
 		FExecuteAction::CreateSP(this, &SMultiLineEditableText::CutSelectedTextToClipboard),
@@ -326,7 +327,8 @@ void SMultiLineEditableText::Construct( const FArguments& InArgs )
 
 	UICommandList->MapAction(FGenericCommands::Get().Paste,
 		FExecuteAction::CreateSP(this, &SMultiLineEditableText::PasteTextFromClipboard),
-		FCanExecuteAction::CreateSP(this, &SMultiLineEditableText::CanExecutePaste));
+		FCanExecuteAction::CreateSP(this, &SMultiLineEditableText::CanExecutePaste),
+		EUIActionRepeatMode::RepeatEnabled);
 
 	UICommandList->MapAction(FGenericCommands::Get().Copy,
 		FExecuteAction::CreateSP(this, &SMultiLineEditableText::CopySelectedTextToClipboard),
diff --git a/Engine/Source/Runtime/Slate/Public/Framework/Commands/UIAction.h b/Engine/Source/Runtime/Slate/Public/Framework/Commands/UIAction.h
--- a/Engine/Source/Runtime/Slate/Public/Framework/Commands/UIAction.h
+++ b/Engine/Source/Runtime/Slate/Public/Framework/Commands/UIAction.h
@@ -16,6 +16,14 @@ DECLARE_DELEGATE_RetVal(bool, FIsActionChecked);
 DECLARE_DELEGATE_RetVal(bool, FIsActionButtonVisible);
 
 
+/** Enum controlling whether a given UI action can be repeated if the chord used to call it is held down */
+enum class EUIActionRepeatMode
+{
+	RepeatDisabled,
+	RepeatEnabled,
+};
+
+
 /**
  * Implements an UI action.
  */
@@ -33,21 +41,32 @@ struct FUIAction
 	/** Holds a delegate that is executed when determining whether this action is visible. */
 	FIsActionButtonVisible IsActionVisibleDelegate;
 
+	/** Can this action can be repeated if the chord used to call it is held down? */
+	EUIActionRepeatMode RepeatMode;
+
 public:
 
 	/** Default constructor. */
-	FUIAction( ) { }
+	FUIAction( )
+		: ExecuteAction()
+		, CanExecuteAction()
+		, IsCheckedDelegate()
+		, IsActionVisibleDelegate()
+		, RepeatMode(EUIActionRepeatMode::RepeatDisabled)
+	{ }
 
 	/**
 	 * Constructor that takes delegates to initialize the action with
 	 *
 	 * @param	ExecuteAction		The delegate to call when the action should be executed
+	 * @param	RepeatMode			Can this action can be repeated if the chord used to call it is held down?
 	 */
-	FUIAction( FExecuteAction InitExecuteAction )
+	FUIAction( FExecuteAction InitExecuteAction, EUIActionRepeatMode InitRepeatMode = EUIActionRepeatMode::RepeatDisabled )
 		: ExecuteAction(InitExecuteAction)
 		, CanExecuteAction(FCanExecuteAction::CreateRaw(&FSlateApplication::Get(), &FSlateApplication::IsNormalExecution))
 		, IsCheckedDelegate()
 		, IsActionVisibleDelegate()
+		, RepeatMode(InitRepeatMode)
 	{ }
 
 	/**
@@ -55,12 +74,14 @@ public:
 	 *
 	 * @param	ExecuteAction		The delegate to call when the action should be executed
 	 * @param	CanExecuteAction	The delegate to call to see if the action can be executed
+	 * @param	RepeatMode			Can this action can be repeated if the chord used to call it is held down?
 	 */
-	FUIAction( FExecuteAction InitExecuteAction, FCanExecuteAction InitCanExecuteAction )
+	FUIAction( FExecuteAction InitExecuteAction, FCanExecuteAction InitCanExecuteAction, EUIActionRepeatMode InitRepeatMode = EUIActionRepeatMode::RepeatDisabled )
 		: ExecuteAction(InitExecuteAction )
 		, CanExecuteAction(InitCanExecuteAction)
 		, IsCheckedDelegate()
 		, IsActionVisibleDelegate()
+		, RepeatMode(InitRepeatMode)
 	{ }
 
 	/**
@@ -69,12 +90,14 @@ public:
 	 * @param	ExecuteAction		The delegate to call when the action should be executed
 	 * @param	CanExecuteAction	The delegate to call to see if the action can be executed
 	 * @param	IsCheckedDelegate	The delegate to call to see if the action should appear checked when visualized
+	 * @param	RepeatMode			Can this action can be repeated if the chord used to call it is held down?
 	 */
-	FUIAction( FExecuteAction InitExecuteAction, FCanExecuteAction InitCanExecuteAction, FIsActionChecked InitIsCheckedDelegate )
+	FUIAction( FExecuteAction InitExecuteAction, FCanExecuteAction InitCanExecuteAction, FIsActionChecked InitIsCheckedDelegate, EUIActionRepeatMode InitRepeatMode = EUIActionRepeatMode::RepeatDisabled )
 		: ExecuteAction(InitExecuteAction)
 		, CanExecuteAction(InitCanExecuteAction)
 		, IsCheckedDelegate(InitIsCheckedDelegate)
 		, IsActionVisibleDelegate()
+		, RepeatMode(InitRepeatMode)
 	{ }
 
 	/**
@@ -84,12 +107,14 @@ public:
 	 * @param	CanExecuteAction		The delegate to call to see if the action can be executed
 	 * @param	IsCheckedDelegate		The delegate to call to see if the action should appear checked when visualized
 	 * @param	IsActionVisibleDelegate	The delegate to call to see if the action should be visible
+	 * @param	RepeatMode				Can this action can be repeated if the chord used to call it is held down?
 	 */
-	FUIAction( FExecuteAction InitExecuteAction, FCanExecuteAction InitCanExecuteAction, FIsActionChecked InitIsCheckedDelegate, FIsActionButtonVisible InitIsActionVisibleDelegate )
+	FUIAction( FExecuteAction InitExecuteAction, FCanExecuteAction InitCanExecuteAction, FIsActionChecked InitIsCheckedDelegate, FIsActionButtonVisible InitIsActionVisibleDelegate, EUIActionRepeatMode InitRepeatMode = EUIActionRepeatMode::RepeatDisabled )
 		: ExecuteAction(InitExecuteAction)
 		, CanExecuteAction(InitCanExecuteAction)
 		, IsCheckedDelegate(InitIsCheckedDelegate)
 		, IsActionVisibleDelegate(InitIsActionVisibleDelegate)
+		, RepeatMode(InitRepeatMode)
 	{ }
 
 public:
@@ -165,4 +190,14 @@ public:
 
 		return EVisibility::Visible;
 	}
+
+	/**
+	 * Checks whether this action can be repeated if the chord used to call it is held down.
+	 *
+	 * @return true if the delegate can be repeated, false otherwise.
+	 */
+	bool CanRepeat( ) const
+	{
+		return RepeatMode == EUIActionRepeatMode::RepeatEnabled;
+	}
 };
diff --git a/Engine/Source/Runtime/Slate/Public/Framework/Commands/UICommandList.h b/Engine/Source/Runtime/Slate/Public/Framework/Commands/UICommandList.h
--- a/Engine/Source/Runtime/Slate/Public/Framework/Commands/UICommandList.h
+++ b/Engine/Source/Runtime/Slate/Public/Framework/Commands/UICommandList.h
@@ -16,8 +16,9 @@ public:
 	 *
 	 * @param InUICommandInfo	The command info to map
 	 * @param ExecuteAction		The delegate to call when the command should be executed
+	 * @param RepeatMode		Can this action can be repeated if the chord used to call it is held down?
 	 */
-	void MapAction( const TSharedPtr< const FUICommandInfo > InUICommandInfo, FExecuteAction ExecuteAction );
+	void MapAction( const TSharedPtr< const FUICommandInfo > InUICommandInfo, FExecuteAction ExecuteAction, EUIActionRepeatMode RepeatMode = EUIActionRepeatMode::RepeatDisabled );
 
 	/**
 	 * Maps a command info to a series of delegates that are executed by a multibox or mouse/keyboard input
@@ -25,8 +26,9 @@ public:
 	 * @param InUICommandInfo	The command info to map
 	 * @param ExecuteAction		The delegate to call when the command should be executed
 	 * @param CanExecuteAction	The delegate to call to see if the command can be executed
+	 * @param RepeatMode		Can this action can be repeated if the chord used to call it is held down?
 	 */
-	void MapAction( const TSharedPtr< const FUICommandInfo > InUICommandInfo, FExecuteAction ExecuteAction, FCanExecuteAction CanExecuteAction );
+	void MapAction( const TSharedPtr< const FUICommandInfo > InUICommandInfo, FExecuteAction ExecuteAction, FCanExecuteAction CanExecuteAction, EUIActionRepeatMode RepeatMode = EUIActionRepeatMode::RepeatDisabled );
 
 	/**
 	 * Maps a command info to a series of delegates that are executed by a multibox or mouse/keyboard input
@@ -35,8 +37,9 @@ public:
 	 * @param ExecuteAction		The delegate to call when the command should be executed
 	 * @param CanExecuteAction	The delegate to call to see if the command can be executed
 	 * @param IsCheckedDelegate	The delegate to call to see if the command should appear checked when visualized in a multibox
+	 * @param RepeatMode		Can this action can be repeated if the chord used to call it is held down?
 	 */
-	void MapAction( const TSharedPtr< const FUICommandInfo > InUICommandInfo, FExecuteAction ExecuteAction, FCanExecuteAction CanExecuteAction, FIsActionChecked IsCheckedDelegate );
+	void MapAction( const TSharedPtr< const FUICommandInfo > InUICommandInfo, FExecuteAction ExecuteAction, FCanExecuteAction CanExecuteAction, FIsActionChecked IsCheckedDelegate, EUIActionRepeatMode RepeatMode = EUIActionRepeatMode::RepeatDisabled );
 
 	/**
 	 * Maps a command info to a series of delegates that are executed by a multibox or mouse/keyboard input
@@ -46,8 +49,9 @@ public:
 	 * @param CanExecuteAction	The delegate to call to see if the command can be executed
 	 * @param IsCheckedDelegate	The delegate to call to see if the command should appear checked when visualized in a multibox
 	 * @param IsVisibleDelegate	The delegate to call to see if the command should appear or be hidden when visualized in a multibox
+	 * @param RepeatMode		Can this action can be repeated if the chord used to call it is held down?
 	 */
-	void MapAction( const TSharedPtr< const FUICommandInfo > InUICommandInfo, FExecuteAction ExecuteAction, FCanExecuteAction CanExecuteAction, FIsActionChecked IsCheckedDelegate, FIsActionButtonVisible IsVisibleDelegate );
+	void MapAction( const TSharedPtr< const FUICommandInfo > InUICommandInfo, FExecuteAction ExecuteAction, FCanExecuteAction CanExecuteAction, FIsActionChecked IsCheckedDelegate, FIsActionButtonVisible IsVisibleDelegate, EUIActionRepeatMode RepeatMode = EUIActionRepeatMode::RepeatDisabled );
 
 	/**
 	 * Maps a command info to a series of delegates that are executed by a multibox or mouse/keyboard input
