Fixed "Delete" on an empty line in ML text being handled like a "Backspace"

The ML text logic for using delete on an empty line was trying to ensure that the first line wasn't removed, and would then move the cursor back to the end of the previous line. This is incorrect, as that is the behavior of the backspace key - delete should make sure that the last line can't be removed, and shouldn't move the cursor at all.

diff --git a/Engine/Source/Runtime/Slate/Private/Widgets/Text/SMultiLineEditableText.cpp b/Engine/Source/Runtime/Slate/Private/Widgets/Text/SMultiLineEditableText.cpp
--- a/Engine/Source/Runtime/Slate/Private/Widgets/Text/SMultiLineEditableText.cpp
+++ b/Engine/Source/Runtime/Slate/Private/Widgets/Text/SMultiLineEditableText.cpp
@@ -819,15 +819,10 @@ void SMultiLineEditableText::DeleteChar()
 			//If we are at the very beginning of the line...
 			if (Line.Text->Len() == 0)
 			{
-				//And if the current line isn't the very first line then...
-				if (CursorInteractionPosition.GetLineIndex() > 0)
+				//And if the current line isn't the very last line then...
+				if (Lines.IsValidIndex(CursorInteractionPosition.GetLineIndex() + 1))
 				{
-					if (TextLayout->RemoveLine(CursorInteractionPosition.GetLineIndex()))
-					{
-						//Update the cursor so it appears at the end of the previous line,
-						//as we're going to delete the imaginary \n separating them
-						FinalCursorLocation = FTextLocation(CursorInteractionPosition.GetLineIndex() - 1, Lines[CursorInteractionPosition.GetLineIndex() - 1].Text->Len());
-					}
+					TextLayout->RemoveLine(CursorInteractionPosition.GetLineIndex());
 				}
 				//else do nothing as the FinalCursorLocation is already correct
 			}
