Empty lines in ML text now draw a little highlight so that you can tell that the line itself is highlighted

diff --git a/Engine/Source/Runtime/Slate/Private/Widgets/Text/SMultiLineEditableText.cpp b/Engine/Source/Runtime/Slate/Private/Widgets/Text/SMultiLineEditableText.cpp
--- a/Engine/Source/Runtime/Slate/Private/Widgets/Text/SMultiLineEditableText.cpp
+++ b/Engine/Source/Runtime/Slate/Private/Widgets/Text/SMultiLineEditableText.cpp
@@ -162,8 +162,10 @@ int32 SMultiLineEditableText::FTextSelectionRunRenderer::OnPaint( const FPaintAr
 	// The block size and offset values are pre-scaled, so we need to account for that when converting the block offsets into paint geometry
 	const float InverseScale = Inverse(AllottedGeometry.Scale);
 
-	const float HighlightWidth = Block->GetSize().X;
-	if (HighlightWidth)
+	// We still want to show a small selection outline on empty lines to make it clear that the line itself is selected despite being empty
+	const float MinHighlightWidth = (Line.Range.IsEmpty()) ? 4.0f * AllottedGeometry.Scale : 0.0f;
+	const float HighlightWidth = FMath::Max(Block->GetSize().X, MinHighlightWidth);
+	if (HighlightWidth > 0.0f)
 	{
 		// Draw the actual highlight rectangle
 		FSlateDrawElement::MakeBox(
@@ -1150,10 +1152,7 @@ void SMultiLineEditableText::UpdateCursorHighlight()
 		if ( SelectionBeginningLineIndex == SelectionEndLineIndex )
 		{
 			const FTextRange Range(SelectionBeginningLineOffset, SelectionEndLineOffset);
-			if (!Range.IsEmpty())
-			{
-				TextLayout->AddRunRenderer(FTextRunRenderer(SelectionBeginningLineIndex, Range, TextSelectionRunRenderer.ToSharedRef()));
-			}
+			TextLayout->AddRunRenderer(FTextRunRenderer(SelectionBeginningLineIndex, Range, TextSelectionRunRenderer.ToSharedRef()));
 		}
 		else
 		{
@@ -1164,26 +1163,17 @@ void SMultiLineEditableText::UpdateCursorHighlight()
 				if ( LineIndex == SelectionBeginningLineIndex )
 				{
 					const FTextRange Range(SelectionBeginningLineOffset, Lines[LineIndex].Text->Len());
-					if (!Range.IsEmpty())
-					{
-						TextLayout->AddRunRenderer(FTextRunRenderer(LineIndex, Range, TextSelectionRunRenderer.ToSharedRef()));
-					}
+					TextLayout->AddRunRenderer(FTextRunRenderer(LineIndex, Range, TextSelectionRunRenderer.ToSharedRef()));
 				}
 				else if ( LineIndex == SelectionEndLineIndex )
 				{
 					const FTextRange Range(0, SelectionEndLineOffset);
-					if (!Range.IsEmpty())
-					{
-						TextLayout->AddRunRenderer(FTextRunRenderer(LineIndex, Range, TextSelectionRunRenderer.ToSharedRef()));
-					}
+					TextLayout->AddRunRenderer(FTextRunRenderer(LineIndex, Range, TextSelectionRunRenderer.ToSharedRef()));
 				}
 				else
 				{
 					const FTextRange Range(0, Lines[LineIndex].Text->Len());
-					if (!Range.IsEmpty())
-					{
-						TextLayout->AddRunRenderer(FTextRunRenderer(LineIndex, Range, TextSelectionRunRenderer.ToSharedRef()));
-					}
+					TextLayout->AddRunRenderer(FTextRunRenderer(LineIndex, Range, TextSelectionRunRenderer.ToSharedRef()));
 				}
 			}
 		}
