Removed some outdated comments

diff --git a/Engine/Source/Runtime/Slate/Private/Widgets/Text/SRichTextBlock.cpp b/Engine/Source/Runtime/Slate/Private/Widgets/Text/SRichTextBlock.cpp
--- a/Engine/Source/Runtime/Slate/Private/Widgets/Text/SRichTextBlock.cpp
+++ b/Engine/Source/Runtime/Slate/Private/Widgets/Text/SRichTextBlock.cpp
@@ -49,9 +49,6 @@ int32 SRichTextBlock::OnPaint( const FPaintArgs& Args, const FGeometry& Allotted
 
 FVector2D SRichTextBlock::ComputeDesiredSize(float LayoutScaleMultiplier) const
 {
-	// todo: jdale - The scale needs to be passed to ComputeDesiredSize
-	//const float Scale = CachedScale;
-
 	// ComputeDesiredSize will also update the text layout cache if required
 	const FVector2D TextSize = TextLayoutCache->ComputeDesiredSize(
 		FTextBlockLayout::FWidgetArgs(BoundText, HighlightText, WrapTextAt, AutoWrapText, Margin, LineHeightPercentage, Justification), 
diff --git a/Engine/Source/Runtime/Slate/Private/Widgets/Text/STextBlock.cpp b/Engine/Source/Runtime/Slate/Private/Widgets/Text/STextBlock.cpp
--- a/Engine/Source/Runtime/Slate/Private/Widgets/Text/STextBlock.cpp
+++ b/Engine/Source/Runtime/Slate/Private/Widgets/Text/STextBlock.cpp
@@ -165,9 +165,6 @@ FVector2D STextBlock::ComputeDesiredSize(float LayoutScaleMultiplier) const
 {
 #if WITH_FANCY_TEXT
 
-	// todo: jdale - The scale needs to be passed to ComputeDesiredSize
-	//const float Scale = CachedScale;
-
 	// ComputeDesiredSize will also update the text layout cache if required
 	const FVector2D TextSize = TextLayoutCache->ComputeDesiredSize(
 		FTextBlockLayout::FWidgetArgs(BoundText, HighlightText, WrapTextAt, AutoWrapText, Margin, LineHeightPercentage, Justification), 
