STextBlock, SRichTextBlock, and SMultiLineEditableText now scale their margin appropriately

UE-15873 - Zooming in UMG preview doesn't scale text block margin

diff --git a/Engine/Source/Runtime/Slate/Private/Widgets/Text/SMultiLineEditableText.cpp b/Engine/Source/Runtime/Slate/Private/Widgets/Text/SMultiLineEditableText.cpp
--- a/Engine/Source/Runtime/Slate/Private/Widgets/Text/SMultiLineEditableText.cpp
+++ b/Engine/Source/Runtime/Slate/Private/Widgets/Text/SMultiLineEditableText.cpp
@@ -2592,8 +2592,6 @@ void SMultiLineEditableText::CacheDesiredSize(float LayoutScaleMultiplier)
 	// Get the wrapping width and font to see if they have changed
 	float WrappingWidth = WrapTextAt.Get();
 
-	const FMargin& OurMargin = Margin.Get();
-
 	// Text wrapping can either be used defined (WrapTextAt), automatic (AutoWrapText), or a mixture of both
 	// Take whichever has the smallest value (>1)
 	if(AutoWrapText.Get() && CachedSize.X >= 1.0f)
@@ -2603,7 +2601,7 @@ void SMultiLineEditableText::CacheDesiredSize(float LayoutScaleMultiplier)
 
 	TextLayout->SetScale( LayoutScaleMultiplier );
 	TextLayout->SetWrappingWidth( WrappingWidth );
-	TextLayout->SetMargin( OurMargin );
+	TextLayout->SetMargin( Margin.Get() * TextLayout->GetScale() );
 	TextLayout->SetLineHeightPercentage( LineHeightPercentage.Get() );
 	TextLayout->SetJustification( Justification.Get() );
 	TextLayout->SetVisibleRegion( CachedSize, ScrollOffset * TextLayout->GetScale() );
diff --git a/Engine/Source/Runtime/Slate/Private/Widgets/Text/TextBlockLayout.cpp b/Engine/Source/Runtime/Slate/Private/Widgets/Text/TextBlockLayout.cpp
--- a/Engine/Source/Runtime/Slate/Private/Widgets/Text/TextBlockLayout.cpp
+++ b/Engine/Source/Runtime/Slate/Private/Widgets/Text/TextBlockLayout.cpp
@@ -29,7 +29,7 @@ FVector2D FTextBlockLayout::ComputeDesiredSize(const FWidgetArgs& InWidgetArgs,
 {
 	SLATE_CYCLE_COUNTER_SCOPE_DETAILED(SLATE_STATS_DETAIL_LEVEL_HI, GSlateTextBlockLayoutComputeDesiredSize);
 	TextLayout->SetWrappingWidth(CalculateWrappingWidth(InWidgetArgs));
-	TextLayout->SetMargin(InWidgetArgs.Margin.Get());
+	TextLayout->SetMargin(InWidgetArgs.Margin.Get() * InScale);
 	TextLayout->SetJustification(InWidgetArgs.Justification.Get());
 	TextLayout->SetLineHeightPercentage(InWidgetArgs.LineHeightPercentage.Get());
 	TextLayout->SetScale(InScale);
