Fixed SDockTab ignoring ToolTipText

ToolTip and ToolTipText are generic properties that can be set on any widget, however SDockTab was only fowarding on ToolTip when building its sub-widgets, ignoring the value set in ToolTipText. This essentially meant that ToolTipText didn't work for SDockTab widgets.

diff --git a/Engine/Source/Runtime/Slate/Private/Widgets/Docking/SDockTab.cpp b/Engine/Source/Runtime/Slate/Private/Widgets/Docking/SDockTab.cpp
--- a/Engine/Source/Runtime/Slate/Private/Widgets/Docking/SDockTab.cpp
+++ b/Engine/Source/Runtime/Slate/Private/Widgets/Docking/SDockTab.cpp
@@ -511,7 +511,8 @@ void SDockTab::Construct( const FArguments& InArgs )
 			[
 				SNew(SHorizontalBox)
 				.Visibility(EVisibility::Visible)
-				.ToolTip( InArgs._ToolTip.IsValid() ? InArgs._ToolTip : SNew( SToolTip ).Text( this, &SDockTab::GetTabLabel ) )
+				.ToolTip( InArgs._ToolTip )
+				.ToolTipText( InArgs._ToolTipText.IsSet() ? InArgs._ToolTipText : TAttribute<FText>( this, &SDockTab::GetTabLabel ) )
 
 				// Tab Icon
 				+ SHorizontalBox::Slot()
