This change is to fix up internal behavour for SMultiLineEditableText to be closer to SEditableText after recent changes.
#Branch UE4
#Proj Slate

diff --git a/Engine/Source/Runtime/Slate/Private/Widgets/Text/SMultiLineEditableText.cpp b/Engine/Source/Runtime/Slate/Private/Widgets/Text/SMultiLineEditableText.cpp
--- a/Engine/Source/Runtime/Slate/Private/Widgets/Text/SMultiLineEditableText.cpp
+++ b/Engine/Source/Runtime/Slate/Private/Widgets/Text/SMultiLineEditableText.cpp
@@ -610,6 +610,12 @@ void SMultiLineEditableText::OnFocusLost( const FFocusEvent& InFocusEvent )
 			}
 		}
 
+		// Clear selection unless activating a new window (otherwise can't copy and past on right click)
+		if (InFocusEvent.GetCause() != EFocusCause::WindowActivate)
+		{
+			ClearSelection();
+		}
+
 		// When focus is lost let anyone who is interested that text was committed
 		// See if user explicitly tabbed away or moved focus
 		ETextCommit::Type TextAction;
@@ -636,11 +642,6 @@ void SMultiLineEditableText::OnFocusLost( const FFocusEvent& InFocusEvent )
 
 		OnTextCommitted.ExecuteIfBound(EditedText, TextAction);
 
-		if(bClearKeyboardFocusOnCommit.Get())
-		{
-			ClearSelection();
-		}
-
 		UpdateCursorHighlight();
 
 		// UpdateCursorHighlight always tries to scroll to the cursor, but we don't want that to happen when we 
@@ -1989,11 +1990,6 @@ FReply SMultiLineEditableText::OnEscape()
 		if(bRevertTextOnEscape.Get() && HasTextChangedFromOriginal())
 		{
 			RestoreOriginalText();
-			// Release input focus
-			if(bClearKeyboardFocusOnCommit.Get())
-			{
-				FSlateApplication::Get().ClearKeyboardFocus(EFocusCause::Cleared);
-			}
 			MyReply = FReply::Handled();
 		}
 	}
@@ -2022,7 +2018,6 @@ void SMultiLineEditableText::OnEnter()
 		// Release input focus
 		if(bClearKeyboardFocusOnCommit.Get())
 		{
-			ClearSelection();
 			FSlateApplication::Get().ClearKeyboardFocus(EFocusCause::Cleared);
 		}
 	}
