Added a way to specify the header padding for an SExpandableArea

This was previously always set to the default padding of an SButton. This is still the default, but you can now override it if you need to.

diff --git a/Engine/Source/Runtime/Slate/Private/Widgets/Layout/SExpandableArea.cpp b/Engine/Source/Runtime/Slate/Private/Widgets/Layout/SExpandableArea.cpp
--- a/Engine/Source/Runtime/Slate/Private/Widgets/Layout/SExpandableArea.cpp
+++ b/Engine/Source/Runtime/Slate/Private/Widgets/Layout/SExpandableArea.cpp
@@ -48,6 +48,7 @@ void SExpandableArea::Construct( const FArguments& InArgs )
 			[
 				SNew( SButton )
 				.ButtonStyle(FCoreStyle::Get(), "NoBorder")
+				.ContentPadding(InArgs._HeaderPadding)
 				.ForegroundColor(FSlateColor::UseForeground())
 				.OnClicked( this, &SExpandableArea::OnHeaderClicked )
 				[
diff --git a/Engine/Source/Runtime/Slate/Public/Widgets/Layout/SExpandableArea.h b/Engine/Source/Runtime/Slate/Public/Widgets/Layout/SExpandableArea.h
--- a/Engine/Source/Runtime/Slate/Public/Widgets/Layout/SExpandableArea.h
+++ b/Engine/Source/Runtime/Slate/Public/Widgets/Layout/SExpandableArea.h
@@ -17,6 +17,7 @@ public:
 		, _AreaTitle( )
 		, _InitiallyCollapsed( false )
 		, _MaxHeight( 0.0f )
+		, _HeaderPadding(FMargin(4.0, 2.0))
 		, _Padding( 1.0f )
 		, _AreaTitleFont( FCoreStyle::Get().GetFontStyle("ExpandableArea.TitleFont") )
 		{}
@@ -37,6 +38,8 @@ public:
 		SLATE_ARGUMENT( bool, InitiallyCollapsed )
 		/** The maximum height of the area */
 		SLATE_ARGUMENT( float, MaxHeight )
+		/** The header padding */
+		SLATE_ATTRIBUTE( FMargin, HeaderPadding )
 		/** The content padding */
 		SLATE_ATTRIBUTE( FMargin, Padding )
 		/** Called when the area is expanded or collapsed */
