Replaced AdditionalCommands with an OnKeyDownHandler in SEditableText, SEditableTextBox, SMultiLineEditableText, and SMultiLineEditableTextBox. Handler is first chance functionality for overriding default functionality of the widget.

diff --git a/Engine/Source/Runtime/Slate/Private/Widgets/Input/SEditableText.cpp b/Engine/Source/Runtime/Slate/Private/Widgets/Input/SEditableText.cpp
--- a/Engine/Source/Runtime/Slate/Private/Widgets/Input/SEditableText.cpp
+++ b/Engine/Source/Runtime/Slate/Private/Widgets/Input/SEditableText.cpp
@@ -73,6 +73,7 @@ void SEditableText::Construct( const FArguments& InArgs )
 	MinDesiredWidth = InArgs._MinDesiredWidth;
 	SelectAllTextOnCommit = InArgs._SelectAllTextOnCommit;
 	VirtualKeyboardType = IsPassword.Get() ? EKeyboardType::Keyboard_Password : InArgs._VirtualKeyboardType;
+	OnKeyDownHandler = InArgs._OnKeyDownHandler;
 
 	// Map UI commands to delegates which are called when the command should be executed
 	UICommandList->MapAction( FGenericCommands::Get().Undo,
@@ -101,13 +102,6 @@ void SEditableText::Construct( const FArguments& InArgs )
 		FExecuteAction::CreateSP( this, &SEditableText::SelectAllText ),
 		FCanExecuteAction::CreateSP( this, &SEditableText::CanExecuteSelectAll ) );
 
-	// Append any additional commands that a consumer of EditableText wants us to be aware of.
-	const TSharedPtr<FUICommandList>& AdditionalCommands = InArgs._AdditionalCommands;
-	if ( AdditionalCommands.IsValid() )
-	{
-		UICommandList->Append( AdditionalCommands.ToSharedRef() );
-	}
-
 	// build context menu extender
 	MenuExtender = MakeShareable(new FExtender);
 	MenuExtender->AddMenuExtension("EditText", EExtensionHook::Before, TSharedPtr<FUICommandList>(), InArgs._ContextMenuExtender);
@@ -1495,18 +1489,29 @@ FReply SEditableText::OnKeyChar( const FGeometry& MyGeometry, const FCharacterEv
 
 FReply SEditableText::OnKeyDown( const FGeometry& MyGeometry, const FKeyEvent& InKeyEvent )
 {
-	FReply Reply = FTextEditHelper::OnKeyDown( InKeyEvent, SharedThis( this ) );
+	FReply Reply = FReply::Unhandled();
+
+	// First call the user defined key handler, there might be overrides to normal functionality
+	if (OnKeyDownHandler.IsBound())
+	{
+		Reply = OnKeyDownHandler.Execute(MyGeometry, InKeyEvent);
+	}
 
-	// Process keybindings if the event wasn't already handled
 	if( !Reply.IsEventHandled() )
 	{
-		if ( UICommandList->ProcessCommandBindings( InKeyEvent ) )
-		{
-			Reply = FReply::Handled();
-		}
-		else
+		Reply = FTextEditHelper::OnKeyDown( InKeyEvent, SharedThis( this ) );
+		
+		if (!Reply.IsEventHandled())
 		{
-			Reply = SLeafWidget::OnKeyDown( MyGeometry, InKeyEvent );
+			// Process keybindings if the event wasn't already handled
+			if ( UICommandList->ProcessCommandBindings( InKeyEvent ) )
+			{
+				Reply = FReply::Handled();
+			}
+			else
+			{
+				Reply = SLeafWidget::OnKeyDown( MyGeometry, InKeyEvent );
+			}
 		}
 
 	}
diff --git a/Engine/Source/Runtime/Slate/Private/Widgets/Input/SEditableTextBox.cpp b/Engine/Source/Runtime/Slate/Private/Widgets/Input/SEditableTextBox.cpp
--- a/Engine/Source/Runtime/Slate/Private/Widgets/Input/SEditableTextBox.cpp
+++ b/Engine/Source/Runtime/Slate/Private/Widgets/Input/SEditableTextBox.cpp
@@ -50,7 +50,7 @@ void SEditableTextBox::Construct( const FArguments& InArgs )
 					.OnTextCommitted( InArgs._OnTextCommitted )
 					.MinDesiredWidth( InArgs._MinDesiredWidth )
 					.SelectAllTextOnCommit( InArgs._SelectAllTextOnCommit )
-					.AdditionalCommands( InArgs._AdditionalCommands )
+					.OnKeyDownHandler( InArgs._OnKeyDownHandler )
 				]
 			]
 		]
@@ -119,6 +119,10 @@ void SEditableTextBox::SetError( const FString& InError )
 	ErrorReporting->SetError( InError );
 }
 
+void SEditableTextBox::SetOnKeyDownHandler(FOnKeyDown InOnKeyDownHandler)
+{
+	EditableText->SetOnKeyDownHandler(InOnKeyDownHandler);
+}
 
 /** @return Border image for the text box based on the hovered and focused state */
 const FSlateBrush* SEditableTextBox::GetBorderImage() const
diff --git a/Engine/Source/Runtime/Slate/Private/Widgets/Input/SMultiLineEditableTextBox.cpp b/Engine/Source/Runtime/Slate/Private/Widgets/Input/SMultiLineEditableTextBox.cpp
--- a/Engine/Source/Runtime/Slate/Private/Widgets/Input/SMultiLineEditableTextBox.cpp
+++ b/Engine/Source/Runtime/Slate/Private/Widgets/Input/SMultiLineEditableTextBox.cpp
@@ -109,7 +109,7 @@ void SMultiLineEditableTextBox::Construct( const FArguments& InArgs )
 					.VScrollBar(VScrollBar)
 					.OnHScrollBarUserScrolled(InArgs._OnHScrollBarUserScrolled)
 					.OnVScrollBarUserScrolled(InArgs._OnVScrollBarUserScrolled)
-					.AdditionalCommands( InArgs._AdditionalCommands)
+					.OnKeyDownHandler( InArgs._OnKeyDownHandler)
 					.ModiferKeyForNewLine(InArgs._ModiferKeyForNewLine)
 				]
 
@@ -305,4 +305,9 @@ void SMultiLineEditableTextBox::Refresh()
 	return EditableText->Refresh();
 }
 
+void SMultiLineEditableTextBox::SetOnKeyDownHandler(FOnKeyDown InOnKeyDownHandler)
+{
+	EditableText->SetOnKeyDownHandler(InOnKeyDownHandler);
+}
+
 #endif //WITH_FANCY_TEXT
diff --git a/Engine/Source/Runtime/Slate/Private/Widgets/Text/SMultiLineEditableText.cpp b/Engine/Source/Runtime/Slate/Private/Widgets/Text/SMultiLineEditableText.cpp
--- a/Engine/Source/Runtime/Slate/Private/Widgets/Text/SMultiLineEditableText.cpp
+++ b/Engine/Source/Runtime/Slate/Private/Widgets/Text/SMultiLineEditableText.cpp
@@ -277,6 +277,7 @@ void SMultiLineEditableText::Construct( const FArguments& InArgs )
 	bRevertTextOnEscape = InArgs._RevertTextOnEscape;
 	OnHScrollBarUserScrolled = InArgs._OnHScrollBarUserScrolled;
 	OnVScrollBarUserScrolled = InArgs._OnVScrollBarUserScrolled;
+	OnKeyDownHandler = InArgs._OnKeyDownHandler;
 
 	Marshaller = InArgs._Marshaller;
 	if (!Marshaller.IsValid())
@@ -344,13 +345,6 @@ void SMultiLineEditableText::Construct( const FArguments& InArgs )
 		FExecuteAction::CreateSP(this, &SMultiLineEditableText::SelectAllText),
 		FCanExecuteAction::CreateSP(this, &SMultiLineEditableText::CanExecuteSelectAll));
 
-	// Append any additional commands that a consumer of MultiLineEditableText wants us to be aware of.
-	const TSharedPtr<FUICommandList>& AdditionalCommands = InArgs._AdditionalCommands;
-	if ( AdditionalCommands.IsValid() )
-	{
-		UICommandList->Append( AdditionalCommands.ToSharedRef() );
-	}
-
 	// build context menu extender
 	MenuExtender = MakeShareable(new FExtender);
 	MenuExtender->AddMenuExtension("EditText", EExtensionHook::Before, TSharedPtr<FUICommandList>(), InArgs._ContextMenuExtender);
@@ -2680,12 +2674,23 @@ FReply SMultiLineEditableText::OnKeyChar( const FGeometry& MyGeometry,const FCha
 
 FReply SMultiLineEditableText::OnKeyDown( const FGeometry& MyGeometry, const FKeyEvent& InKeyEvent )
 {
-	FReply Reply = FTextEditHelper::OnKeyDown( InKeyEvent, SharedThis( this ) );
+	FReply Reply = FReply::Unhandled();
 
-	// Process keybindings if the event wasn't already handled
-	if (!Reply.IsEventHandled() && UICommandList->ProcessCommandBindings(InKeyEvent))
+	// First call the user defined key handler, there might be overrides to normal functionality
+	if (OnKeyDownHandler.IsBound())
 	{
-		Reply = FReply::Handled();
+		Reply = OnKeyDownHandler.Execute(MyGeometry, InKeyEvent);
+	}
+
+	if( !Reply.IsEventHandled() )
+	{
+		Reply = FTextEditHelper::OnKeyDown( InKeyEvent, SharedThis( this ) );
+
+		// Process keybindings if the event wasn't already handled
+		if (!Reply.IsEventHandled() && UICommandList->ProcessCommandBindings(InKeyEvent))
+		{
+			Reply = FReply::Handled();
+		}
 	}
 
 	return Reply;
diff --git a/Engine/Source/Runtime/Slate/Public/Framework/SlateDelegates.h b/Engine/Source/Runtime/Slate/Public/Framework/SlateDelegates.h
--- a/Engine/Source/Runtime/Slate/Public/Framework/SlateDelegates.h
+++ b/Engine/Source/Runtime/Slate/Public/Framework/SlateDelegates.h
@@ -213,3 +213,6 @@ DECLARE_DELEGATE_RetVal_ThreeParams(bool, FAccessSourceCode, const FString& /*Fi
 
 /** A Delegate for an asset object to access */
 DECLARE_DELEGATE_RetVal_OneParam(bool, FAccessAsset, UObject* /*InAsset*/);
+
+/** A delegate for a callback during OnKeyDown from an SWidget */
+DECLARE_DELEGATE_RetVal_TwoParams( FReply, FOnKeyDown, const FGeometry&, const FKeyEvent&);
diff --git a/Engine/Source/Runtime/Slate/Public/Widgets/Input/SEditableText.h b/Engine/Source/Runtime/Slate/Public/Widgets/Input/SEditableText.h
--- a/Engine/Source/Runtime/Slate/Public/Widgets/Input/SEditableText.h
+++ b/Engine/Source/Runtime/Slate/Public/Widgets/Input/SEditableText.h
@@ -33,7 +33,6 @@ public:
 		, _ClearKeyboardFocusOnCommit(true)
 		, _MinDesiredWidth(0.0f)
 		, _SelectAllTextOnCommit( false )
-		, _AdditionalCommands()
 		, _VirtualKeyboardType(EKeyboardType::Keyboard_Default)
 		{}
 
@@ -102,8 +101,8 @@ public:
 		/** Whether to select all text when pressing enter to commit changes */
 		SLATE_ATTRIBUTE( bool, SelectAllTextOnCommit )
 
-		/** Additional UI commands. */
-		SLATE_ARGUMENT(TSharedPtr<FUICommandList>, AdditionalCommands)
+		/** Callback delegate to have first chance handling of the OnKeyDown event */
+		SLATE_EVENT(FOnKeyDown, OnKeyDownHandler)
 
 		/** Menu extender for the right-click context menu */
 		SLATE_EVENT( FMenuExtensionDelegate, ContextMenuExtender )
@@ -202,6 +201,16 @@ public:
 	 */
 	void SetSelectAllTextOnCommit(const TAttribute<bool>& InSelectAllTextOnCommit);
 
+	/**
+	 * Sets the OnKeyDownHandler to provide first chance handling of the OnKeyDown event
+	 *
+	 * @param InOnKeyDownHandler			Delegate to call during OnKeyDown event
+	 */
+	void SetOnKeyDownHandler(FOnKeyDown InOnKeyDownHandler)
+	{
+		OnKeyDownHandler = InOnKeyDownHandler;
+	}
+
 protected:
 
 	friend class FTextInputMethodContext;
@@ -653,4 +662,7 @@ private:
 
 	/** The iterator to use to detect word boundaries */
 	mutable TSharedPtr<IBreakIterator> WordBreakIterator;
+
+	/** Callback delegate to have first chance handling of the OnKeyDown event */
+	FOnKeyDown OnKeyDownHandler;
 };
diff --git a/Engine/Source/Runtime/Slate/Public/Widgets/Input/SEditableTextBox.h b/Engine/Source/Runtime/Slate/Public/Widgets/Input/SEditableTextBox.h
--- a/Engine/Source/Runtime/Slate/Public/Widgets/Input/SEditableTextBox.h
+++ b/Engine/Source/Runtime/Slate/Public/Widgets/Input/SEditableTextBox.h
@@ -25,7 +25,6 @@ public:
 		, _ClearKeyboardFocusOnCommit( true )
 		, _MinDesiredWidth( 0.0f )
 		, _SelectAllTextOnCommit( false )
-		, _AdditionalCommands()
 		, _BackgroundColor()		
 		, _Padding()
 		, _ErrorReporting()
@@ -82,8 +81,8 @@ public:
 		/** Whether to select all text when pressing enter to commit changes */
 		SLATE_ATTRIBUTE( bool, SelectAllTextOnCommit )
 
-		/** Additional UI commands. */
-		SLATE_ARGUMENT(TSharedPtr<FUICommandList>, AdditionalCommands)
+		/** Callback delegate to have first chance handling of the OnKeyDown event */
+		SLATE_EVENT(FOnKeyDown, OnKeyDownHandler)
 
 		/** The color of the background/border around the editable text (overrides Style) */
 		SLATE_ATTRIBUTE( FSlateColor, BackgroundColor )
@@ -209,6 +208,13 @@ public:
 	void SetError( const FText& InError );
 	void SetError( const FString& InError );
 
+	/**
+	 * Sets the OnKeyDownHandler to provide first chance handling of the SEditableText's OnKeyDown event
+	 *
+	 * @param InOnKeyDownHandler			Delegate to call during OnKeyDown event
+	 */
+	void SetOnKeyDownHandler(FOnKeyDown InOnKeyDownHandler);
+
 	// SWidget overrides
 	virtual bool SupportsKeyboardFocus() const override;
 	virtual bool HasKeyboardFocus() const override;
diff --git a/Engine/Source/Runtime/Slate/Public/Widgets/Input/SMultiLineEditableTextBox.h b/Engine/Source/Runtime/Slate/Public/Widgets/Input/SMultiLineEditableTextBox.h
--- a/Engine/Source/Runtime/Slate/Public/Widgets/Input/SMultiLineEditableTextBox.h
+++ b/Engine/Source/Runtime/Slate/Public/Widgets/Input/SMultiLineEditableTextBox.h
@@ -32,7 +32,6 @@ public:
 		, _AlwaysShowScrollbars( false )
 		, _HScrollBar()
 		, _VScrollBar()
-		, _AdditionalCommands()
 		, _WrapTextAt(0.0f)
 		, _AutoWrapText(false)
 		, _SelectAllTextOnCommit( false )
@@ -124,8 +123,8 @@ public:
 		/** Called when the cursor is moved within the text area */
 		SLATE_EVENT( SMultiLineEditableText::FOnCursorMoved, OnCursorMoved )
 
-		/** Additional UI commands. */
-		SLATE_ARGUMENT(TSharedPtr<FUICommandList>, AdditionalCommands)
+		/** Callback delegate to have first chance handling of the OnKeyDown event */
+		SLATE_EVENT(FOnKeyDown, OnKeyDownHandler)
 
 		/** Menu extender for the right-click context menu */
 		SLATE_EVENT( FMenuExtensionDelegate, ContextMenuExtender )
@@ -273,6 +272,13 @@ public:
 	/** Refresh this text box immediately, rather than wait for the usual caching mechanisms to take affect on the text Tick */
 	void Refresh();
 
+	/**
+	 * Sets the OnKeyDownHandler to provide first chance handling of the SMultiLineEditableText's OnKeyDown event
+	 *
+	 * @param InOnKeyDownHandler			Delegate to call during OnKeyDown event
+	 */
+	void SetOnKeyDownHandler(FOnKeyDown InOnKeyDownHandler);
+
 protected:
 
 	/** Editable text widget */
diff --git a/Engine/Source/Runtime/Slate/Public/Widgets/Text/SMultiLineEditableText.h b/Engine/Source/Runtime/Slate/Public/Widgets/Text/SMultiLineEditableText.h
--- a/Engine/Source/Runtime/Slate/Public/Widgets/Text/SMultiLineEditableText.h
+++ b/Engine/Source/Runtime/Slate/Public/Widgets/Text/SMultiLineEditableText.h
@@ -37,7 +37,6 @@ public:
 		, _RevertTextOnEscape(false)
 		, _ClearKeyboardFocusOnCommit(true)
 		, _OnCursorMoved()
-		, _AdditionalCommands()
 		, _ContextMenuExtender()
 		, _ModiferKeyForNewLine(EModifierKey::None)
 	{}
@@ -109,8 +108,8 @@ public:
 		/** Called when the cursor is moved within the text area */
 		SLATE_EVENT(FOnCursorMoved, OnCursorMoved)
 
-		/** Additional UI commands. */
-		SLATE_ARGUMENT(TSharedPtr<FUICommandList>, AdditionalCommands)
+		/** Callback delegate to have first chance handling of the OnKeyDown event */
+		SLATE_EVENT(FOnKeyDown, OnKeyDownHandler)
 
 		/** Menu extender for the right-click context menu */
 		SLATE_EVENT(FMenuExtensionDelegate, ContextMenuExtender)
@@ -187,6 +186,16 @@ public:
 	/** Returns whether the current text varies from the original */
 	bool HasTextChangedFromOriginal() const;
 
+	/**
+	 * Sets the OnKeyDownHandler to provide first chance handling of the OnKeyDown event
+	 *
+	 * @param InOnKeyDownHandler			Delegate to call during OnKeyDown event
+	 */
+	void SetOnKeyDownHandler(FOnKeyDown InOnKeyDownHandler)
+	{
+		OnKeyDownHandler = InOnKeyDownHandler;
+	}
+
 private:
 	
 	/** 
@@ -752,6 +761,9 @@ private:
 
 	/**	The current position of the software cursor */
 	FVector2D SoftwareCursorPosition;
+
+	/** Callback delegate to have first chance handling of the OnKeyDown event */
+	FOnKeyDown OnKeyDownHandler;
 };
 
 
