Slate - Little code cleanup removing some unused variables.

diff --git a/Engine/Source/Runtime/Slate/Private/Framework/Application/SlateApplication.cpp b/Engine/Source/Runtime/Slate/Private/Framework/Application/SlateApplication.cpp
--- a/Engine/Source/Runtime/Slate/Private/Framework/Application/SlateApplication.cpp
+++ b/Engine/Source/Runtime/Slate/Private/Framework/Application/SlateApplication.cpp
@@ -855,9 +855,8 @@ void FSlateApplication::SetCursorPos( const FVector2D& MouseCoordinate )
 FWidgetPath FSlateApplication::LocateWindowUnderMouse( FVector2D ScreenspaceMouseCoordinate, const TArray< TSharedRef< SWindow > >& Windows, bool bIgnoreEnabledStatus )
 {
 	bool bPrevWindowWasModal = false;
-	FArrangedChildren OutWidgetPath(EVisibility::Visible);
 
-	for (int32 WindowIndex = Windows.Num() - 1; WindowIndex >= 0 && OutWidgetPath.Num() == 0; --WindowIndex)
+	for (int32 WindowIndex = Windows.Num() - 1; WindowIndex >= 0; --WindowIndex)
 	{ 
 		const TSharedRef<SWindow>& Window = Windows[WindowIndex];
 
