Fixed a couple of places where the caret could be left in a bad position after changing the editable text

UE-16221 - CRASH: Cancelling filtering on a blueprint graph context menu

diff --git a/Engine/Source/Runtime/Slate/Private/Widgets/Input/SEditableText.cpp b/Engine/Source/Runtime/Slate/Private/Widgets/Input/SEditableText.cpp
--- a/Engine/Source/Runtime/Slate/Private/Widgets/Input/SEditableText.cpp
+++ b/Engine/Source/Runtime/Slate/Private/Widgets/Input/SEditableText.cpp
@@ -170,6 +170,9 @@ void SEditableText::SetTextFromVirtualKeyboard(const FText& InNewText)
 	{
 		EditedText = InNewText;
 
+		// Move the cursor to the end of the string
+		SetCaretPosition(EditedText.ToString().Len());
+
 		// This method is called from the main thread (i.e. not the game thread) of the device with the virtual keyboard
 		// This causes the app to crash on those devices, so we're using polling here to ensure delegates are
 		// fired on the game thread in Tick.
@@ -183,6 +186,9 @@ void SEditableText::RestoreOriginalText()
 	{
 		EditedText = OriginalText;
 
+		// Reset the cursor position
+		SetCaretPosition(0);
+
 		SaveText();
 
 		// Let outsiders know that the text content has been changed
