Fixed CIS error on MacMenu

This allows the Mac to support the "undetermined" state in its check box menu items.

diff --git a/Engine/Source/Runtime/Slate/Private/Framework/MultiBox/Mac/MacMenu.cpp b/Engine/Source/Runtime/Slate/Private/Framework/MultiBox/Mac/MacMenu.cpp
--- a/Engine/Source/Runtime/Slate/Private/Framework/MultiBox/Mac/MacMenu.cpp
+++ b/Engine/Source/Runtime/Slate/Private/Framework/MultiBox/Mac/MacMenu.cpp
@@ -491,20 +491,29 @@ bool FSlateMacMenu::IsMenuItemEnabled(const TSharedRef<const class FMenuEntryBlo
 
 int32 FSlateMacMenu::GetMenuItemState(const TSharedRef<const class FMenuEntryBlock>& Block)
 {
-	bool bIsChecked = false;
 	TSharedPtr<const FUICommandList> ActionList = Block->GetActionList();
 	TSharedPtr<const FUICommandInfo> Action = Block->GetAction();
 	const FUIAction& DirectActions = Block->GetDirectActions();
 
+	ECheckBoxState CheckState = ECheckBoxState::Unchecked;
 	if (ActionList.IsValid() && Action.IsValid())
 	{
-		bIsChecked = ActionList->IsChecked(Action.ToSharedRef());
+		CheckState = ActionList->GetCheckState(Action.ToSharedRef());
 	}
 	else
 	{
 		// There is no action list or action associated with this block via a UI command.  Execute any direct action we have
-		bIsChecked = DirectActions.IsChecked();
+		CheckState = DirectActions.GetCheckState();
 	}
 
-	return bIsChecked ? NSOnState : NSOffState;
+	switch(CheckState)
+	{
+	case ECheckBoxState::Checked:
+		return NSOnState;
+	case ECheckBoxState::Undetermined:
+		return NSMixedState;
+	default:
+		break;
+	}
+	return NSOffState;
 }
