PushMenu() now respects QueryPopupMethod(). All menus now support reusing windows.
UE-14641 - Fix PushMenu() to use QueryPopupMethod()

Pretty big refactor
Adds IMenu as way to identify menus. Replaces referring to menus as SWindows.
Lots of uses of PushMenu() fixed up to match new API

diff --git a/Engine/Plugins/2D/Paper2D/Source/Paper2DEditor/Private/FlipbookEditor/SFlipbookTimeline.cpp b/Engine/Plugins/2D/Paper2D/Source/Paper2DEditor/Private/FlipbookEditor/SFlipbookTimeline.cpp
--- a/Engine/Plugins/2D/Paper2D/Source/Paper2DEditor/Private/FlipbookEditor/SFlipbookTimeline.cpp
+++ b/Engine/Plugins/2D/Paper2D/Source/Paper2DEditor/Private/FlipbookEditor/SFlipbookTimeline.cpp
@@ -268,7 +268,8 @@ FReply SFlipbookTimeline::OnMouseButtonUp(const FGeometry& MyGeometry, const FPo
 	if (MouseEvent.GetEffectingButton() == EKeys::RightMouseButton)
 	{
 		TSharedRef<SWidget> MenuContents = GenerateContextMenu();
-		FSlateApplication::Get().PushMenu(AsShared(), MenuContents, MouseEvent.GetScreenSpacePosition(), FPopupTransitionEffect(FPopupTransitionEffect::ContextMenu));
+		FWidgetPath WidgetPath = MouseEvent.GetEventPath() != nullptr ? *MouseEvent.GetEventPath() : FWidgetPath();
+		FSlateApplication::Get().PushMenu(AsShared(), WidgetPath, MenuContents, MouseEvent.GetScreenSpacePosition(), FPopupTransitionEffect(FPopupTransitionEffect::ContextMenu));
 
 		return FReply::Handled();
 	}
diff --git a/Engine/Plugins/2D/Paper2D/Source/Paper2DEditor/Private/FlipbookEditor/STimelineTrack.cpp b/Engine/Plugins/2D/Paper2D/Source/Paper2DEditor/Private/FlipbookEditor/STimelineTrack.cpp
--- a/Engine/Plugins/2D/Paper2D/Source/Paper2DEditor/Private/FlipbookEditor/STimelineTrack.cpp
+++ b/Engine/Plugins/2D/Paper2D/Source/Paper2DEditor/Private/FlipbookEditor/STimelineTrack.cpp
@@ -179,7 +179,8 @@ FReply SFlipbookKeyframeWidget::KeyframeOnMouseButtonUp(const FGeometry& MyGeome
 	if (MouseEvent.GetEffectingButton() == EKeys::RightMouseButton)
 	{
 		TSharedRef<SWidget> MenuContents = GenerateContextMenu();
-		FSlateApplication::Get().PushMenu(AsShared(), MenuContents, MouseEvent.GetScreenSpacePosition(), FPopupTransitionEffect(FPopupTransitionEffect::ContextMenu));
+		FWidgetPath WidgetPath = MouseEvent.GetEventPath() != nullptr ? *MouseEvent.GetEventPath() : FWidgetPath();
+		FSlateApplication::Get().PushMenu(AsShared(), WidgetPath, MenuContents, MouseEvent.GetScreenSpacePosition(), FPopupTransitionEffect(FPopupTransitionEffect::ContextMenu));
 
 		return FReply::Handled();
 	}
diff --git a/Engine/Source/Developer/FriendsAndChat/Private/FriendsAndChatManager.cpp b/Engine/Source/Developer/FriendsAndChat/Private/FriendsAndChatManager.cpp
--- a/Engine/Source/Developer/FriendsAndChat/Private/FriendsAndChatManager.cpp
+++ b/Engine/Source/Developer/FriendsAndChat/Private/FriendsAndChatManager.cpp
@@ -962,7 +962,7 @@ void FFriendsAndChatManager::SetChatWindowContents(TSharedPtr<SWindow> Window, T
 		];
 
 	Window->SetContent(Content);
-	Window->SetOnWindowActivated(FOnWindowActivated::CreateSP(ChatWindow.Get(), &SChatWindow::HandleWindowActivated));
+	Window->GetOnWindowActivatedEvent().AddSP(ChatWindow.Get(), &SChatWindow::HandleWindowActivated);
 }
 // Actions
 
diff --git a/Engine/Source/Developer/Profiler/Private/Widgets/SDataGraph.cpp b/Engine/Source/Developer/Profiler/Private/Widgets/SDataGraph.cpp
--- a/Engine/Source/Developer/Profiler/Private/Widgets/SDataGraph.cpp
+++ b/Engine/Source/Developer/Profiler/Private/Widgets/SDataGraph.cpp
@@ -1085,7 +1085,7 @@ FReply SDataGraph::OnMouseButtonUp( const FGeometry& MyGeometry, const FPointerE
 		{
 			if( !bIsRMB_Scrolling && bIsValidForMouseClick )
 			{
-				ShowContextMenu( MouseEvent.GetScreenSpacePosition() );
+				ShowContextMenu(MouseEvent.GetScreenSpacePosition(), MouseEvent);
 				Reply = FReply::Handled();
 			}
 			else
@@ -1277,7 +1277,7 @@ FCursorReply SDataGraph::OnCursorQuery( const FGeometry& MyGeometry, const FPoin
 	return FCursorReply::Unhandled();
 }
 
-void SDataGraph::ShowContextMenu( const FVector2D& ScreenSpacePosition )
+void SDataGraph::ShowContextMenu(const FVector2D& ScreenSpacePosition, const FPointerEvent& MouseEvent)
 {
 	TSharedPtr<FUICommandList> ProfilerCommandList = FProfilerManager::Get()->GetCommandList();
 	const FProfilerCommands& ProfilerCommands = FProfilerManager::GetCommands();
@@ -1381,7 +1381,10 @@ void SDataGraph::ShowContextMenu( const FVector2D& ScreenSpacePosition )
 	MenuBuilder.EndSection();
 
 	TSharedRef<SWidget> MenuWidget = MenuBuilder.MakeWidget();
-	FSlateApplication::Get().PushMenu( SharedThis( this ), MenuWidget, ScreenSpacePosition, FPopupTransitionEffect::ContextMenu );
+
+	FWidgetPath EventPath = MouseEvent.GetEventPath() != nullptr ? *MouseEvent.GetEventPath() : FWidgetPath();
+	FSlateApplication::Get().PushMenu(SharedThis(this), EventPath, MenuWidget, ScreenSpacePosition, FPopupTransitionEffect::ContextMenu);
+
 }
 
 void SDataGraph::BindCommands()
diff --git a/Engine/Source/Developer/Profiler/Private/Widgets/SDataGraph.h b/Engine/Source/Developer/Profiler/Private/Widgets/SDataGraph.h
--- a/Engine/Source/Developer/Profiler/Private/Widgets/SDataGraph.h
+++ b/Engine/Source/Developer/Profiler/Private/Widgets/SDataGraph.h
@@ -342,7 +342,7 @@ protected:
 	const int32 CalculateFrameIndex( const FVector2D MousePosition ) const;
 
 
-	void ShowContextMenu( const FVector2D& ScreenSpacePosition );
+	void ShowContextMenu(const FVector2D& ScreenSpacePosition, const FPointerEvent& MouseEvent);
 
 	/**	Binds our UI commands to delegates. */
 	void BindCommands();
diff --git a/Engine/Source/Developer/RealtimeProfiler/Private/SRealtimeProfilerVisualizer.cpp b/Engine/Source/Developer/RealtimeProfiler/Private/SRealtimeProfilerVisualizer.cpp
--- a/Engine/Source/Developer/RealtimeProfiler/Private/SRealtimeProfilerVisualizer.cpp
+++ b/Engine/Source/Developer/RealtimeProfiler/Private/SRealtimeProfilerVisualizer.cpp
@@ -99,11 +99,12 @@ void SRealtimeProfilerVisualizer::RouteBarEventSelectionChanged( int32 Thread, T
 	EventsTree->HandleBarEventSelectionChanged( Thread, Selection );
 }
 
-void SRealtimeProfilerVisualizer::OnBarGraphContextMenu( TSharedPtr< FVisualizerEvent > Selection )
+void SRealtimeProfilerVisualizer::OnBarGraphContextMenu( TSharedPtr< FVisualizerEvent > Selection, const FPointerEvent& InputEvent )
 {
 	SelectedBarGraph = Selection;
 
-	FSlateApplication::Get().PushMenu( SharedThis( this ), MakeBarVisualizerContextMenu(), FSlateApplication::Get().GetCursorPos(), FPopupTransitionEffect::ContextMenu );
+	FWidgetPath WidgetPath = InputEvent.GetEventPath() != nullptr ? *InputEvent.GetEventPath() : FWidgetPath();
+	FSlateApplication::Get().PushMenu(SharedThis(this), WidgetPath, MakeBarVisualizerContextMenu(), FSlateApplication::Get().GetCursorPos(), FPopupTransitionEffect::ContextMenu);
 }
 
 TSharedRef<SWidget> SRealtimeProfilerVisualizer::MakeBarVisualizerContextMenu()
diff --git a/Engine/Source/Developer/RealtimeProfiler/Private/SRealtimeProfilerVisualizer.h b/Engine/Source/Developer/RealtimeProfiler/Private/SRealtimeProfilerVisualizer.h
--- a/Engine/Source/Developer/RealtimeProfiler/Private/SRealtimeProfilerVisualizer.h
+++ b/Engine/Source/Developer/RealtimeProfiler/Private/SRealtimeProfilerVisualizer.h
@@ -52,7 +52,7 @@ private:
 	void RouteBarEventSelectionChanged( int32 Thread, TSharedPtr<FVisualizerEvent> Selection );
 
 	/** Opens a context menu when a bar is right clicked */
-	void OnBarGraphContextMenu( TSharedPtr< FVisualizerEvent > Selection );
+	void OnBarGraphContextMenu( TSharedPtr< FVisualizerEvent > Selection, const FPointerEvent& InputEvent );
 
 	/** Creates bar visualizer context menu for opening new events tree tabs */
 	TSharedRef<SWidget> MakeBarVisualizerContextMenu();
diff --git a/Engine/Source/Developer/SessionFrontend/Private/Widgets/Console/SSessionConsoleShortcutWindow.cpp b/Engine/Source/Developer/SessionFrontend/Private/Widgets/Console/SSessionConsoleShortcutWindow.cpp
--- a/Engine/Source/Developer/SessionFrontend/Private/Widgets/Console/SSessionConsoleShortcutWindow.cpp
+++ b/Engine/Source/Developer/SessionFrontend/Private/Widgets/Console/SSessionConsoleShortcutWindow.cpp
@@ -2,7 +2,7 @@
 
 #include "SessionFrontendPrivatePCH.h"
 #include "STextEntryPopup.h"
-
+#include "IMenu.h"
 
 #define LOCTEXT_NAMESPACE "SSessionConsoleShortcutWindow"
 
@@ -77,8 +77,9 @@ void SSessionConsoleShortcutWindow::HandleEditCommandActionExecute( TSharedPtr<F
 		.SelectAllTextWhenFocused(true)
 		.ClearKeyboardFocusOnCommit(false);
 
-	NameEntryPopupWindow = FSlateApplication::Get().PushMenu(
+	NameEntryMenu = FSlateApplication::Get().PushMenu(
 		SharedThis(this),
+		FWidgetPath(),
 		TextEntry,
 		FSlateApplication::Get().GetCursorPos(),
 		FPopupTransitionEffect(FPopupTransitionEffect::TypeInPopup)
@@ -241,9 +242,9 @@ TSharedRef<ITableRow> SSessionConsoleShortcutWindow::HandleShortcutListViewGener
 
 void SSessionConsoleShortcutWindow::HandleShortcutTextEntryCommitted( const FText& CommandText, ETextCommit::Type CommitInfo )
 {
-	if (NameEntryPopupWindow.IsValid())
+	if (NameEntryMenu.IsValid())
 	{
-		NameEntryPopupWindow->RequestDestroyWindow();
+		NameEntryMenu.Pin()->Dismiss();
 
 		int32 IndexOfShortcut = Shortcuts.IndexOfByKey(EditedShortcut);
 		if (EditedShortcut.IsValid() && (IndexOfShortcut != INDEX_NONE))
diff --git a/Engine/Source/Developer/SessionFrontend/Private/Widgets/Console/SSessionConsoleShortcutWindow.h b/Engine/Source/Developer/SessionFrontend/Private/Widgets/Console/SSessionConsoleShortcutWindow.h
--- a/Engine/Source/Developer/SessionFrontend/Private/Widgets/Console/SSessionConsoleShortcutWindow.h
+++ b/Engine/Source/Developer/SessionFrontend/Private/Widgets/Console/SSessionConsoleShortcutWindow.h
@@ -105,7 +105,7 @@ private:
 	TSharedPtr<FConsoleShortcutData> EditedShortcut;
 
 	/** Reference to owner of the current pop-up. */
-	TSharedPtr<class SWindow> NameEntryPopupWindow;
+	TWeakPtr<class IMenu> NameEntryMenu;
 
 	/** Holds a delegate that is executed when a command is submitted. */
 	FOnSessionConsoleCommandSubmitted OnCommandSubmitted;
diff --git a/Engine/Source/Developer/TaskGraph/Private/SProfileVisualizer.cpp b/Engine/Source/Developer/TaskGraph/Private/SProfileVisualizer.cpp
--- a/Engine/Source/Developer/TaskGraph/Private/SProfileVisualizer.cpp
+++ b/Engine/Source/Developer/TaskGraph/Private/SProfileVisualizer.cpp
@@ -134,11 +134,12 @@ void SProfileVisualizer::RouteBarEventSelectionChanged( int32 Thread, TSharedPtr
 	EventsTree->HandleBarEventSelectionChanged( Thread, Selection );
 }
 
-void SProfileVisualizer::OnBarGraphContextMenu( TSharedPtr< FVisualizerEvent > Selection )
+void SProfileVisualizer::OnBarGraphContextMenu( TSharedPtr< FVisualizerEvent > Selection, const FPointerEvent& InputEvent )
 {
 	SelectedBarGraph = Selection;
 
-	FSlateApplication::Get().PushMenu( SharedThis( this ), MakeBarVisualizerContextMenu(), FSlateApplication::Get().GetCursorPos(), FPopupTransitionEffect::ContextMenu );
+	FWidgetPath WidgetPath = InputEvent.GetEventPath() != nullptr ? *InputEvent.GetEventPath() : FWidgetPath();
+	FSlateApplication::Get().PushMenu(SharedThis(this), WidgetPath, MakeBarVisualizerContextMenu(), FSlateApplication::Get().GetCursorPos(), FPopupTransitionEffect::ContextMenu);
 }
 
 TSharedRef<SWidget> SProfileVisualizer::MakeBarVisualizerContextMenu()
diff --git a/Engine/Source/Developer/TaskGraph/Private/SProfileVisualizer.h b/Engine/Source/Developer/TaskGraph/Private/SProfileVisualizer.h
--- a/Engine/Source/Developer/TaskGraph/Private/SProfileVisualizer.h
+++ b/Engine/Source/Developer/TaskGraph/Private/SProfileVisualizer.h
@@ -62,7 +62,7 @@ private:
 	void RouteBarEventSelectionChanged( int32 Thread, TSharedPtr<FVisualizerEvent> Selection );
 
 	/** Opens a context menu when a bar is right clicked */
-	void OnBarGraphContextMenu( TSharedPtr< FVisualizerEvent > Selection );
+	void OnBarGraphContextMenu( TSharedPtr< FVisualizerEvent > Selection, const FPointerEvent& InputEvent );
 
 	/** Creates bar visualizer context menu for opening new events tree tabs */
 	TSharedRef<SWidget> MakeBarVisualizerContextMenu();
diff --git a/Engine/Source/Developer/TaskGraph/Public/SBarVisualizer.h b/Engine/Source/Developer/TaskGraph/Public/SBarVisualizer.h
--- a/Engine/Source/Developer/TaskGraph/Public/SBarVisualizer.h
+++ b/Engine/Source/Developer/TaskGraph/Public/SBarVisualizer.h
@@ -23,7 +23,7 @@ public:
 	DECLARE_DELEGATE_TwoParams( FOnBarEventSelectionChanged, int32, TSharedPtr< FVisualizerEvent > );
 	
 	/** Delegate used when the user right-clicks on a bar graph */
-	DECLARE_DELEGATE_OneParam( FOnBarGraphContextMenu, TSharedPtr< FVisualizerEvent > );
+	DECLARE_DELEGATE_TwoParams( FOnBarGraphContextMenu, TSharedPtr< FVisualizerEvent >, const FPointerEvent& );
 
 	SLATE_BEGIN_ARGS( SBarVisualizer )
 		: _ProfileData()
diff --git a/Engine/Source/Developer/TreeMap/Public/STreeMap.h b/Engine/Source/Developer/TreeMap/Public/STreeMap.h
--- a/Engine/Source/Developer/TreeMap/Public/STreeMap.h
+++ b/Engine/Source/Developer/TreeMap/Public/STreeMap.h
@@ -174,8 +174,14 @@ protected:
 	void ApplyVisualizationToNodesRecursively( const FTreeMapNodeDataRef& Node, const FLinearColor& DefaultColor, const int32 TreeDepth );
 		
 	/** Displays a context menu at the specified location with options for configuring the tree display */
-	void ShowOptionsMenuAt( const FVector2D& ScreenSpacePosition );
+	DEPRECATED(4.9, "ShowOptionsMenuAt() taking a FVector2D param is deprecated. Use ShowOptionsMenuAt() that takes an FPointerEvent instead.")
+	void ShowOptionsMenuAt(const FVector2D& ScreenSpacePosition);
 
+	/** Displays a context menu at the specified location with options for configuring the tree display */
+	void ShowOptionsMenuAt( const FPointerEvent& InMouseEvent );
+
+private:
+	void ShowOptionsMenuAtInternal(const FVector2D& ScreenSpacePosition, const FWidgetPath& WidgetPath);
 
 private:
 
diff --git a/Engine/Source/Developer/TreeMap/STreeMap.cpp b/Engine/Source/Developer/TreeMap/STreeMap.cpp
--- a/Engine/Source/Developer/TreeMap/STreeMap.cpp
+++ b/Engine/Source/Developer/TreeMap/STreeMap.cpp
@@ -614,7 +614,7 @@ FReply STreeMap::OnMouseButtonUp( const FGeometry& InMyGeometry, const FPointerE
 	else if( InMouseEvent.GetEffectingButton() == EKeys::RightMouseButton )
 	{
 		// Show a pop-up menu!
-		ShowOptionsMenuAt( InMouseEvent.GetScreenSpacePosition() );
+		ShowOptionsMenuAt( InMouseEvent );
 	}
 
 	return Reply;
@@ -1162,8 +1162,7 @@ void STreeMap::StartRenamingNode( const FGeometry& MyGeometry, const FTreeMapNod
 	bIsNamingNewNode = bIsNewNode;
 
 	const bool bFocusImmediately = true;
-	const bool bShouldAutoSize = true;
-	FSlateApplication::Get().PushMenu( AsShared(), RenamerWidget, MyGeometry.LocalToAbsolute( RelativePosition ), FPopupTransitionEffect::None, bFocusImmediately, bShouldAutoSize );
+	FSlateApplication::Get().PushMenu( AsShared(), FWidgetPath(), RenamerWidget, MyGeometry.LocalToAbsolute( RelativePosition ), FPopupTransitionEffect::None, bFocusImmediately );
 
 	// Focus the text box right after we spawn it so that the user can start typing
 	FSlateApplication::Get().SetKeyboardFocus( EditableText, EKeyboardFocusCause::SetDirectly );
@@ -1296,9 +1295,21 @@ void STreeMap::ApplyVisualizationToNodesRecursively( const FTreeMapNodeDataRef&
 
 }
 
+void STreeMap::ShowOptionsMenuAt(const FVector2D& ScreenSpacePosition)
+{
+	// deprecated
+	ShowOptionsMenuAtInternal(ScreenSpacePosition, FWidgetPath());
+}
 
+void STreeMap::ShowOptionsMenuAt(const FPointerEvent& InMouseEvent)
+{
+	FWidgetPath WidgetPath = InMouseEvent.GetEventPath() != nullptr ? *InMouseEvent.GetEventPath() : FWidgetPath();
+	const FVector2D& ScreenSpacePosition = InMouseEvent.GetScreenSpacePosition();
+
+	ShowOptionsMenuAtInternal(ScreenSpacePosition, WidgetPath);
+}
 
-void STreeMap::ShowOptionsMenuAt( const FVector2D& ScreenSpacePosition )
+void STreeMap::ShowOptionsMenuAtInternal(const FVector2D& ScreenSpacePosition, const FWidgetPath& WidgetPath)
 {
 	struct Local
 	{
@@ -1546,8 +1557,7 @@ void STreeMap::ShowOptionsMenuAt( const FVector2D& ScreenSpacePosition )
 			];
 
 		const bool bFocusImmediately = false;
-		const bool bShouldAutoSize = true;
-		FSlateApplication::Get().PushMenu( AsShared(), WindowContent, ScreenSpacePosition, FPopupTransitionEffect::ContextMenu, bFocusImmediately, bShouldAutoSize );
+		FSlateApplication::Get().PushMenu(AsShared(), WidgetPath, WindowContent, ScreenSpacePosition, FPopupTransitionEffect::ContextMenu, bFocusImmediately);
 	}
 }
 
diff --git a/Engine/Source/Editor/Cascade/Private/Cascade.cpp b/Engine/Source/Editor/Cascade/Private/Cascade.cpp
--- a/Engine/Source/Editor/Cascade/Private/Cascade.cpp
+++ b/Engine/Source/Editor/Cascade/Private/Cascade.cpp
@@ -50,6 +50,7 @@
 #include "Distributions/DistributionVectorUniform.h"
 #include "Distributions/DistributionVectorUniformCurve.h"
 #include "UnrealEngine.h"
+#include "IMenu.h"
 
 static const FName Cascade_PreviewViewportTab("Cascade_PreviewViewport");
 static const FName Cascade_EmmitterCanvasTab("Cascade_EmitterCanvas");
@@ -3711,8 +3712,9 @@ void FCascade::OnSetMotionRadius()
 		.SelectAllTextWhenFocused(true)
 		.ClearKeyboardFocusOnCommit(false);
 
-	EntryPopupWindow = FSlateApplication::Get().PushMenu(
+	EntryMenu = FSlateApplication::Get().PushMenu(
 		PreviewViewport.ToSharedRef(),
+		FWidgetPath(),
 		TextEntry,
 		FSlateApplication::Get().GetCursorPos(),
 		FPopupTransitionEffect(FPopupTransitionEffect::TypeInPopup)
@@ -3932,8 +3934,9 @@ void FCascade::OnToggleWireframeSphere()
 			.SelectAllTextWhenFocused(true)
 			.ClearKeyboardFocusOnCommit(false);
 
-		EntryPopupWindow = FSlateApplication::Get().PushMenu(
+		EntryMenu = FSlateApplication::Get().PushMenu(
 			PreviewViewport.ToSharedRef(),
+			FWidgetPath(),
 			TextEntry,
 			FSlateApplication::Get().GetCursorPos(),
 			FPopupTransitionEffect(FPopupTransitionEffect::TypeInPopup)
@@ -4823,8 +4826,9 @@ void FCascade::OnRenameEmitter()
 		.SelectAllTextWhenFocused(true)
 		.ClearKeyboardFocusOnCommit(false);
 
-	EntryPopupWindow = FSlateApplication::Get().PushMenu(
+	EntryMenu = FSlateApplication::Get().PushMenu(
 		EmitterCanvas.ToSharedRef(),
+		FWidgetPath(),
 		TextEntry,
 		FSlateApplication::Get().GetCursorPos(),
 		FPopupTransitionEffect(FPopupTransitionEffect::TypeInPopup)
@@ -5043,9 +5047,9 @@ void FCascade::OnRemoveDuplicateModules()
 
 void FCascade::CloseEntryPopup()
 {
-	if( EntryPopupWindow.IsValid() )
+	if (EntryMenu.IsValid())
 	{
-		EntryPopupWindow.Pin()->RequestDestroyWindow();
+		EntryMenu.Pin()->Dismiss();
 	}
 }
 
diff --git a/Engine/Source/Editor/Cascade/Private/Cascade.h b/Engine/Source/Editor/Cascade/Private/Cascade.h
--- a/Engine/Source/Editor/Cascade/Private/Cascade.h
+++ b/Engine/Source/Editor/Cascade/Private/Cascade.h
@@ -383,7 +383,7 @@ private:
 	TSharedPtr<IDistributionCurveEditor> CurveEditor;
 
 	/** Reference to owner of the current popup */
-	TWeakPtr<SWindow> EntryPopupWindow;
+	TWeakPtr<class IMenu> EntryMenu;
 
 	/** Components used for drawing the particle system in the preview viewport */
 	UCascadeParticleSystemComponent* ParticleSystemComponent;
diff --git a/Engine/Source/Editor/Cascade/Private/CascadeEmitterCanvasClient.cpp b/Engine/Source/Editor/Cascade/Private/CascadeEmitterCanvasClient.cpp
--- a/Engine/Source/Editor/Cascade/Private/CascadeEmitterCanvasClient.cpp
+++ b/Engine/Source/Editor/Cascade/Private/CascadeEmitterCanvasClient.cpp
@@ -1701,6 +1701,7 @@ void FCascadeEmitterCanvasClient::OpenModuleMenu()
 
 	FSlateApplication::Get().PushMenu(
 		CascadeViewportPtr.Pin().ToSharedRef(),
+		FWidgetPath(),
 		BuildMenuWidgetModule(),
 		MouseCursorLocation,
 		FPopupTransitionEffect(FPopupTransitionEffect::ContextMenu)
@@ -1713,6 +1714,7 @@ void FCascadeEmitterCanvasClient::OpenEmitterMenu()
 
 	FSlateApplication::Get().PushMenu(
 		CascadeViewportPtr.Pin().ToSharedRef(),
+		FWidgetPath(),
 		BuildMenuWidgetEmitter(),
 		MouseCursorLocation,
 		FPopupTransitionEffect(FPopupTransitionEffect::ContextMenu)
@@ -1725,6 +1727,7 @@ void FCascadeEmitterCanvasClient::OpenBackgroundMenu()
 
 	FSlateApplication::Get().PushMenu(
 		CascadeViewportPtr.Pin().ToSharedRef(),
+		FWidgetPath(),
 		BuildMenuWidgetBackround(),
 		MouseCursorLocation,
 		FPopupTransitionEffect(FPopupTransitionEffect::ContextMenu)
diff --git a/Engine/Source/Editor/ContentBrowser/Private/ContentBrowserUtils.cpp b/Engine/Source/Editor/ContentBrowser/Private/ContentBrowserUtils.cpp
--- a/Engine/Source/Editor/ContentBrowser/Private/ContentBrowserUtils.cpp
+++ b/Engine/Source/Editor/ContentBrowser/Private/ContentBrowserUtils.cpp
@@ -13,6 +13,7 @@
 #include "SNotificationList.h"
 #include "NotificationManager.h"
 #include "IPluginManager.h"
+#include "IMenu.h"
 
 #define LOCTEXT_NAMESPACE "ContentBrowser"
 
@@ -74,35 +75,31 @@ public:
 
 		const FVector2D ScreenLocation = FVector2D(ScreenAnchor.Left, ScreenAnchor.Top);
 		const bool bFocusImmediately = true;
-		const bool bShouldAutoSize = true;
-		const FVector2D WindowSize = FVector2D::ZeroVector;
 		const FVector2D SummonLocationSize = ScreenAnchor.GetSize();
 
-		TSharedRef<SWindow> PopupWindow = FSlateApplication::Get().PushMenu(
+		TSharedPtr<IMenu> Menu = FSlateApplication::Get().PushMenu(
 			ParentContent,
+			FWidgetPath(),
 			PopupContent,
 			ScreenLocation,
 			FPopupTransitionEffect( FPopupTransitionEffect::TopMenu ),
 			bFocusImmediately,
-			bShouldAutoSize,
-			WindowSize,
-			SummonLocationSize
-			);
+			SummonLocationSize);
 
-		PopupContent->SetWindow(PopupWindow);
+		PopupContent->SetMenu(Menu);
 	}
 
 private:
-	void SetWindow( const TSharedRef<SWindow>& InWindow )
+	void SetMenu(const TSharedPtr<IMenu>& InMenu)
 	{
-		Window = InWindow;
+		Menu = InMenu;
 	}
 
 	FReply OnBorderClicked(const FGeometry& Geometry, const FPointerEvent& MouseEvent)
 	{
-		if ( Window.IsValid() )
+		if (Menu.IsValid())
 		{
-			Window.Pin()->RequestDestroyWindow();
+			Menu.Pin()->Dismiss();
 		}
 
 		return FReply::Handled();
@@ -114,7 +111,7 @@ private:
 	}
 
 private:
-	TWeakPtr<SWindow> Window;
+	TWeakPtr<IMenu> Menu;
 };
 
 /** A miniture confirmation popup for quick yes/no questions */
@@ -195,8 +192,9 @@ public:
 	void OpenPopup(const TSharedRef<SWidget>& ParentContent)
 	{
 		// Show dialog to confirm the delete
-		PopupWindow = FSlateApplication::Get().PushMenu(
+		Menu = FSlateApplication::Get().PushMenu(
 			ParentContent,
+			FWidgetPath(),
 			SharedThis(this),
 			FSlateApplication::Get().GetCursorPos(),
 			FPopupTransitionEffect( FPopupTransitionEffect::TopMenu )
@@ -212,7 +210,10 @@ private:
 			OnYesClicked.Execute();
 		}
 
-		PopupWindow.Pin()->RequestDestroyWindow();
+		if (Menu.IsValid())
+		{
+			Menu.Pin()->Dismiss();
+		}
 
 		return FReply::Handled();
 	}
@@ -225,13 +226,16 @@ private:
 			OnNoClicked.Execute();
 		}
 
-		PopupWindow.Pin()->RequestDestroyWindow();
+		if (Menu.IsValid())
+		{
+			Menu.Pin()->Dismiss();
+		}
 
 		return FReply::Handled();
 	}
 
-	/** The window containing this popup */
-	TWeakPtr<SWindow> PopupWindow;
+	/** The IMenu prepresenting this popup */
+	TWeakPtr<IMenu> Menu;
 
 	/** Delegates for button clicks */
 	FOnClicked OnYesClicked;
diff --git a/Engine/Source/Editor/ContentBrowser/Private/DragDropHandler.cpp b/Engine/Source/Editor/ContentBrowser/Private/DragDropHandler.cpp
--- a/Engine/Source/Editor/ContentBrowser/Private/DragDropHandler.cpp
+++ b/Engine/Source/Editor/ContentBrowser/Private/DragDropHandler.cpp
@@ -150,6 +150,7 @@ void DragDropHandler::HandleAssetsDroppedOnAssetFolder(const TSharedRef<SWidget>
 
 	FSlateApplication::Get().PushMenu(
 		ParentWidget,
+		FWidgetPath(),
 		MenuBuilder.MakeWidget(),
 		FSlateApplication::Get().GetCursorPos(),
 		FPopupTransitionEffect(FPopupTransitionEffect::ContextMenu)
@@ -191,6 +192,7 @@ void DragDropHandler::HandleFoldersDroppedOnAssetFolder(const TSharedRef<SWidget
 
 	FSlateApplication::Get().PushMenu(
 		ParentWidget,
+		FWidgetPath(),
 		MenuBuilder.MakeWidget(),
 		FSlateApplication::Get().GetCursorPos(),
 		FPopupTransitionEffect(FPopupTransitionEffect::ContextMenu)
diff --git a/Engine/Source/Editor/ContentBrowser/Private/SCollectionView.cpp b/Engine/Source/Editor/ContentBrowser/Private/SCollectionView.cpp
--- a/Engine/Source/Editor/ContentBrowser/Private/SCollectionView.cpp
+++ b/Engine/Source/Editor/ContentBrowser/Private/SCollectionView.cpp
@@ -408,6 +408,7 @@ FReply SCollectionView::MakeAddCollectionMenu()
 
 	FSlateApplication::Get().PushMenu(
 		AsShared(),
+		FWidgetPath(),
 		MenuBuilder.MakeWidget(),
 		FSlateApplication::Get().GetCursorPos(),
 		FPopupTransitionEffect( FPopupTransitionEffect::TopMenu )
diff --git a/Engine/Source/Editor/ContentBrowser/Private/SFilterList.cpp b/Engine/Source/Editor/ContentBrowser/Private/SFilterList.cpp
--- a/Engine/Source/Editor/ContentBrowser/Private/SFilterList.cpp
+++ b/Engine/Source/Editor/ContentBrowser/Private/SFilterList.cpp
@@ -460,7 +460,8 @@ FReply SFilterList::OnMouseButtonUp( const FGeometry& MyGeometry, const FPointer
 			if ( MenuContent.IsValid() )
 			{
 				FVector2D SummonLocation = MouseEvent.GetScreenSpacePosition();
-				FSlateApplication::Get().PushMenu( AsShared(), MenuContent.ToSharedRef(), SummonLocation, FPopupTransitionEffect( FPopupTransitionEffect::ContextMenu ) );
+				FWidgetPath WidgetPath = MouseEvent.GetEventPath() != nullptr ? *MouseEvent.GetEventPath() : FWidgetPath();
+				FSlateApplication::Get().PushMenu(AsShared(), WidgetPath, MenuContent.ToSharedRef(), SummonLocation, FPopupTransitionEffect(FPopupTransitionEffect::ContextMenu));
 			}
 
 			return Reply;
diff --git a/Engine/Source/Editor/DetailCustomizations/Private/AnimTransitionNodeDetails.cpp b/Engine/Source/Editor/DetailCustomizations/Private/AnimTransitionNodeDetails.cpp
--- a/Engine/Source/Editor/DetailCustomizations/Private/AnimTransitionNodeDetails.cpp
+++ b/Engine/Source/Editor/DetailCustomizations/Private/AnimTransitionNodeDetails.cpp
@@ -224,6 +224,7 @@ FReply FAnimTransitionNodeDetails::OnPromoteToSharedClick(bool RuleShare)
 		// Show dialog to enter new event name
 		FSlateApplication::Get().PushMenu(
 			Parent.ToSharedRef(),
+			FWidgetPath(),
 			TextEntry,
 			FSlateApplication::Get().GetCursorPos(),
 			FPopupTransitionEffect( FPopupTransitionEffect::TypeInPopup )
diff --git a/Engine/Source/Editor/DistCurveEditor/Private/SDistributionCurveEditor.cpp b/Engine/Source/Editor/DistCurveEditor/Private/SDistributionCurveEditor.cpp
--- a/Engine/Source/Editor/DistCurveEditor/Private/SDistributionCurveEditor.cpp
+++ b/Engine/Source/Editor/DistCurveEditor/Private/SDistributionCurveEditor.cpp
@@ -13,6 +13,7 @@
 #include "SNotificationList.h"
 #include "NotificationManager.h"
 #include "Engine/InterpCurveEdSetup.h"
+#include "IMenu.h"
 
 #define LOCTEXT_NAMESPACE "CurveEditor"
 
@@ -531,8 +532,9 @@ void SDistributionCurveEditor::OnSetTime()
 		.SelectAllTextWhenFocused(true)
 		.ClearKeyboardFocusOnCommit(false);
 
-	EntryPopupWindow = FSlateApplication::Get().PushMenu(
+	EntryMenu = FSlateApplication::Get().PushMenu(
 		SharedThis(this),
+		FWidgetPath(),
 		TextEntry,
 		FSlateApplication::Get().GetCursorPos(),
 		FPopupTransitionEffect(FPopupTransitionEffect::TypeInPopup)
@@ -555,8 +557,9 @@ void SDistributionCurveEditor::OnSetValue()
 		.SelectAllTextWhenFocused(true)
 		.ClearKeyboardFocusOnCommit(false);
 
-	EntryPopupWindow = FSlateApplication::Get().PushMenu(
+	EntryMenu = FSlateApplication::Get().PushMenu(
 		SharedThis(this),
+		FWidgetPath(),
 		TextEntry,
 		FSlateApplication::Get().GetCursorPos(),
 		FPopupTransitionEffect(FPopupTransitionEffect::TypeInPopup)
@@ -775,8 +778,9 @@ void SDistributionCurveEditor::OnScaleTimes(ECurveScaleScope::Type Scope)
 		.SelectAllTextWhenFocused(true)
 		.ClearKeyboardFocusOnCommit(false);
 
-	EntryPopupWindow = FSlateApplication::Get().PushMenu(
+	EntryMenu = FSlateApplication::Get().PushMenu(
 		SharedThis(this),
+		FWidgetPath(),
 		TextEntry,
 		FSlateApplication::Get().GetCursorPos(),
 		FPopupTransitionEffect(FPopupTransitionEffect::TypeInPopup)
@@ -811,8 +815,9 @@ void SDistributionCurveEditor::OnScaleValues(ECurveScaleScope::Type Scope)
 		.SelectAllTextWhenFocused(true)
 		.ClearKeyboardFocusOnCommit(false);
 
-	EntryPopupWindow = FSlateApplication::Get().PushMenu(
+	EntryMenu = FSlateApplication::Get().PushMenu(
 		SharedThis(this),
+		FWidgetPath(),
 		TextEntry,
 		FSlateApplication::Get().GetCursorPos(),
 		FPopupTransitionEffect(FPopupTransitionEffect::TypeInPopup)
@@ -966,8 +971,9 @@ void SDistributionCurveEditor::OnCreateTab()
 		.OnTextCommitted(this, &SDistributionCurveEditor::TabNameCommitted)
 		.ClearKeyboardFocusOnCommit(false);
 
-	EntryPopupWindow = FSlateApplication::Get().PushMenu(
+	EntryMenu = FSlateApplication::Get().PushMenu(
 		SharedThis(this),
+		FWidgetPath(),
 		TextEntry,
 		FSlateApplication::Get().GetCursorPos(),
 		FPopupTransitionEffect(FPopupTransitionEffect::TypeInPopup)
@@ -1003,6 +1009,7 @@ void SDistributionCurveEditor::OpenLabelMenu()
 
 	FSlateApplication::Get().PushMenu(
 		SharedThis( this ),
+		FWidgetPath(),
 		BuildMenuWidgetLabel(),
 		MouseCursorLocation,
 		FPopupTransitionEffect(FPopupTransitionEffect::ContextMenu)
@@ -1015,6 +1022,7 @@ void SDistributionCurveEditor::OpenKeyMenu()
 
 	FSlateApplication::Get().PushMenu(
 		SharedThis( this ),
+		FWidgetPath(),
 		BuildMenuWidgetKey(),
 		MouseCursorLocation,
 		FPopupTransitionEffect(FPopupTransitionEffect::ContextMenu)
@@ -1027,6 +1035,7 @@ void SDistributionCurveEditor::OpenGeneralMenu()
 
 	FSlateApplication::Get().PushMenu(
 		SharedThis( this ),
+		FWidgetPath(),
 		BuildMenuWidgetGeneral(),
 		MouseCursorLocation,
 		FPopupTransitionEffect(FPopupTransitionEffect::ContextMenu)
@@ -1039,6 +1048,7 @@ void SDistributionCurveEditor::OpenCurveMenu()
 
 	FSlateApplication::Get().PushMenu(
 		SharedThis( this ),
+		FWidgetPath(),
 		BuildMenuWidgetCurve(),
 		MouseCursorLocation,
 		FPopupTransitionEffect(FPopupTransitionEffect::ContextMenu)
@@ -1047,9 +1057,9 @@ void SDistributionCurveEditor::OpenCurveMenu()
 
 void SDistributionCurveEditor::CloseEntryPopup()
 {
-	if( EntryPopupWindow.IsValid() )
+	if (EntryMenu.IsValid())
 	{
-		EntryPopupWindow.Pin()->RequestDestroyWindow();
+		EntryMenu.Pin()->Dismiss();
 	}
 }
 
diff --git a/Engine/Source/Editor/DistCurveEditor/Private/SDistributionCurveEditor.h b/Engine/Source/Editor/DistCurveEditor/Private/SDistributionCurveEditor.h
--- a/Engine/Source/Editor/DistCurveEditor/Private/SDistributionCurveEditor.h
+++ b/Engine/Source/Editor/DistCurveEditor/Private/SDistributionCurveEditor.h
@@ -161,7 +161,7 @@ private:
 	TSharedPtr<FCurveEditorSharedData> SharedData;
 
 	/** Reference to owner of the current popup */
-	TWeakPtr<SWindow> EntryPopupWindow;
+	TWeakPtr<IMenu> EntryMenu;
 
 	/** Tabs dropdown */
 	TSharedPtr<STextComboBox> TabNamesComboBox;
diff --git a/Engine/Source/Editor/GameplayTagsEditor/Private/GameplayTagContainerCustomization.cpp b/Engine/Source/Editor/GameplayTagsEditor/Private/GameplayTagContainerCustomization.cpp
--- a/Engine/Source/Editor/GameplayTagsEditor/Private/GameplayTagContainerCustomization.cpp
+++ b/Engine/Source/Editor/GameplayTagsEditor/Private/GameplayTagContainerCustomization.cpp
@@ -154,7 +154,7 @@ FReply FGameplayTagContainerCustomization::OnEditButtonClicked()
 		.PropertyHandle( StructPropertyHandle )
 	];
 
-	GameplayTagWidgetWindow->SetOnWindowDeactivated( FOnWindowDeactivated::CreateRaw( this, &FGameplayTagContainerCustomization::OnGameplayTagWidgetWindowDeactivate ) );
+	GameplayTagWidgetWindow->GetOnWindowDeactivatedEvent().AddRaw(this, &FGameplayTagContainerCustomization::OnGameplayTagWidgetWindowDeactivate);
 
 	IMainFrameModule& MainFrameModule = FModuleManager::LoadModuleChecked<IMainFrameModule>(TEXT("MainFrame"));
 	if (MainFrameModule.GetParentWindow().IsValid())
diff --git a/Engine/Source/Editor/GraphEditor/Private/SGraphPanel.cpp b/Engine/Source/Editor/GraphEditor/Private/SGraphPanel.cpp
--- a/Engine/Source/Editor/GraphEditor/Private/SGraphPanel.cpp
+++ b/Engine/Source/Editor/GraphEditor/Private/SGraphPanel.cpp
@@ -1060,6 +1060,7 @@ TSharedPtr<SWidget> SGraphPanel::SummonContextMenu(const FVector2D& WhereToSummo
 		
 		FSlateApplication::Get().PushMenu(
 			AsShared(),
+			FWidgetPath(),
 			MenuContent,
 			WhereToSummon,
 			FPopupTransitionEffect( FPopupTransitionEffect::ContextMenu )
diff --git a/Engine/Source/Editor/IntroTutorials/Private/STutorialButton.cpp b/Engine/Source/Editor/IntroTutorials/Private/STutorialButton.cpp
--- a/Engine/Source/Editor/IntroTutorials/Private/STutorialButton.cpp
+++ b/Engine/Source/Editor/IntroTutorials/Private/STutorialButton.cpp
@@ -252,8 +252,9 @@ FReply STutorialButton::OnMouseButtonDown(const FGeometry& MyGeometry, const FPo
 			FUIAction(FExecuteAction::CreateSP(this, &STutorialButton::LaunchBrowser))
 			);
 
+		FWidgetPath WidgetPath = MouseEvent.GetEventPath() != nullptr ? *MouseEvent.GetEventPath() : FWidgetPath();
 
-		FSlateApplication::Get().PushMenu(SharedThis(this), MenuBuilder.MakeWidget(), FSlateApplication::Get().GetCursorPos(), FPopupTransitionEffect::ContextMenu);
+		FSlateApplication::Get().PushMenu(SharedThis(this), WidgetPath, MenuBuilder.MakeWidget(), FSlateApplication::Get().GetCursorPos(), FPopupTransitionEffect::ContextMenu);
 	}
 	return FReply::Handled();
 }
diff --git a/Engine/Source/Editor/IntroTutorials/Private/STutorialContent.cpp b/Engine/Source/Editor/IntroTutorials/Private/STutorialContent.cpp
--- a/Engine/Source/Editor/IntroTutorials/Private/STutorialContent.cpp
+++ b/Engine/Source/Editor/IntroTutorials/Private/STutorialContent.cpp
@@ -301,7 +301,9 @@ FReply STutorialContent::OnMouseButtonUp( const FGeometry& MyGeometry, const FPo
 {
 	if (!bIsStandalone && MouseEvent.GetEffectingButton() == EKeys::RightMouseButton)
 	{
-		FSlateApplication::Get().PushMenu(AsShared(), HandleGetMenuContent(), MouseEvent.GetScreenSpacePosition(), FPopupTransitionEffect(FPopupTransitionEffect::ContextMenu));
+		FWidgetPath WidgetPath = MouseEvent.GetEventPath() != nullptr ? *MouseEvent.GetEventPath() : FWidgetPath();
+
+		FSlateApplication::Get().PushMenu(AsShared(), WidgetPath, HandleGetMenuContent(), MouseEvent.GetScreenSpacePosition(), FPopupTransitionEffect(FPopupTransitionEffect::ContextMenu));
 		return FReply::Handled();
 	}
 	
diff --git a/Engine/Source/Editor/Kismet/Private/BPDelegateDragDropAction.cpp b/Engine/Source/Editor/Kismet/Private/BPDelegateDragDropAction.cpp
--- a/Engine/Source/Editor/Kismet/Private/BPDelegateDragDropAction.cpp
+++ b/Engine/Source/Editor/Kismet/Private/BPDelegateDragDropAction.cpp
@@ -130,7 +130,7 @@ FReply FKismetDelegateDragDropAction::DroppedOnPanel(const TSharedRef< SWidget >
 			}
 		}
 		MenuBuilder.EndSection();
-		FSlateApplication::Get().PushMenu(Panel, MenuBuilder.MakeWidget(), ScreenPosition, FPopupTransitionEffect( FPopupTransitionEffect::ContextMenu));
+		FSlateApplication::Get().PushMenu(Panel, FWidgetPath(), MenuBuilder.MakeWidget(), ScreenPosition, FPopupTransitionEffect( FPopupTransitionEffect::ContextMenu));
 	}
 	return FReply::Handled();
 }
diff --git a/Engine/Source/Editor/Kismet/Private/BPVariableDragDropAction.cpp b/Engine/Source/Editor/Kismet/Private/BPVariableDragDropAction.cpp
--- a/Engine/Source/Editor/Kismet/Private/BPVariableDragDropAction.cpp
+++ b/Engine/Source/Editor/Kismet/Private/BPVariableDragDropAction.cpp
@@ -513,6 +513,7 @@ FReply FKismetVariableDragDropAction::DroppedOnPanel( const TSharedRef< SWidget
 				// Show dialog to choose getter vs setter
 				FSlateApplication::Get().PushMenu(
 					PanelWidget,
+					FWidgetPath(),
 					MenuBuilder.MakeWidget(),
 					ScreenPosition,
 					FPopupTransitionEffect( FPopupTransitionEffect::ContextMenu)
diff --git a/Engine/Source/Editor/Kismet/Private/BlueprintDetailsCustomization.cpp b/Engine/Source/Editor/Kismet/Private/BlueprintDetailsCustomization.cpp
--- a/Engine/Source/Editor/Kismet/Private/BlueprintDetailsCustomization.cpp
+++ b/Engine/Source/Editor/Kismet/Private/BlueprintDetailsCustomization.cpp
@@ -420,7 +420,7 @@ void FBlueprintVarActionDetails::CustomizeDetails( IDetailLayoutBuilder& DetailL
 
 	Category.AddCustomRow( LOCTEXT("CategoryLabel", "Category") )
 		.Visibility(GetPropertyOwnerBlueprint()? EVisibility::Visible : EVisibility::Hidden)
-		.NameContent()
+	.NameContent()
 	[
 		SNew(STextBlock)
 		.Text( LOCTEXT("CategoryLabel", "Category") )
@@ -561,7 +561,7 @@ void FBlueprintVarActionDetails::CustomizeDetails( IDetailLayoutBuilder& DetailL
 	TSharedPtr<SToolTip> ReplicationTooltip = IDocumentation::Get()->CreateToolTip(LOCTEXT("VariableReplicate_Tooltip", "Should this Variable be replicated over the network?"), NULL, DocLink, TEXT("Replication"));
 
 	Category.AddCustomRow( LOCTEXT("VariableReplicationLabel", "Replication") )
-		.Visibility(TAttribute<EVisibility>(this, &FBlueprintVarActionDetails::ReplicationVisibility))
+	.Visibility(TAttribute<EVisibility>(this, &FBlueprintVarActionDetails::ReplicationVisibility))
 	.NameContent()
 	[
 		SNew(STextBlock)
@@ -4817,6 +4817,7 @@ void FBlueprintComponentDetails::OnBrowseSocket()
 					// Pop up a combo box to pick socket from mesh
 					FSlateApplication::Get().PushMenu(
 						Editor.ToSharedRef(),
+						FWidgetPath(),
 						SNew(SSocketChooserPopup)
 						.SceneComponent( ParentSceneComponent )
 						.OnSocketChosen( this, &FBlueprintComponentDetails::OnSocketSelection ),
diff --git a/Engine/Source/Editor/Kismet/Private/BlueprintEditor.cpp b/Engine/Source/Editor/Kismet/Private/BlueprintEditor.cpp
--- a/Engine/Source/Editor/Kismet/Private/BlueprintEditor.cpp
+++ b/Engine/Source/Editor/Kismet/Private/BlueprintEditor.cpp
@@ -97,6 +97,8 @@
 #include "Engine/LevelStreaming.h"
 #include "Engine/LevelStreamingKismet.h"
 
+#include "IMenu.h"
+
 #define LOCTEXT_NAMESPACE "BlueprintEditor"
 
 /////////////////////////////////////////////////////
@@ -2683,6 +2685,7 @@ void FBlueprintEditor::NavigateToChildGraph_Clicked()
 			// Display a child jump list
 			FSlateApplication::Get().PushMenu( 
 				GetToolkitHost()->GetParentWidget(),
+				FWidgetPath(),
 				SNew(SChildGraphPicker, CurrentGraph),
 				FSlateApplication::Get().GetCursorPos(), // summon location
 				FPopupTransitionEffect( FPopupTransitionEffect::TypeInPopup )
@@ -3843,49 +3846,51 @@ void FBlueprintEditor::OnChangePinType()
 		// Grab the root pin, that is what we want to edit
 		UEdGraphPin* RootPin = SelectedPin;
 		while(RootPin->ParentPin)
-		{
+			{
 			RootPin = RootPin->ParentPin;
 		}
 
-		const UEdGraphSchema_K2* Schema = GetDefault<UEdGraphSchema_K2>();
+				const UEdGraphSchema_K2* Schema = GetDefault<UEdGraphSchema_K2>();
 
-		// If this is the index node of the select node, we need to use the index list of types
-		UK2Node_Select* SelectNode = Cast<UK2Node_Select>(SelectedPin->GetOwningNode());
-		if (SelectNode && SelectNode->GetIndexPin() == SelectedPin)
-		{
-			TSharedRef<SCompoundWidget> PinChange = SNew(SPinTypeSelector, FGetPinTypeTree::CreateUObject(Schema, &UEdGraphSchema_K2::GetVariableIndexTypeTree))
+				// If this is the index node of the select node, we need to use the index list of types
+				UK2Node_Select* SelectNode = Cast<UK2Node_Select>(SelectedPin->GetOwningNode());
+				if (SelectNode && SelectNode->GetIndexPin() == SelectedPin)
+				{
+					TSharedRef<SCompoundWidget> PinChange = SNew(SPinTypeSelector, FGetPinTypeTree::CreateUObject(Schema, &UEdGraphSchema_K2::GetVariableIndexTypeTree))
 				.TargetPinType(this, &FBlueprintEditor::OnGetPinType, RootPin)
-				.OnPinTypeChanged(this, &FBlueprintEditor::OnChangePinTypeFinished, SelectedPin)
-				.Schema(Schema)
-				.bAllowExec(false)
-				.IsEnabled(true)
-				.bAllowArrays(false);
-
-			PinTypeChangePopupWindow = FSlateApplication::Get().PushMenu( 
-				GetToolkitHost()->GetParentWidget(), // Parent widget should be k2 not the menu thats open or it will be closed when the menu is dismissed
-				PinChange,
-				FSlateApplication::Get().GetCursorPos(), // summon location
-				FPopupTransitionEffect(FPopupTransitionEffect::TypeInPopup)
-				);
-		}
-		else
-		{
-			TSharedRef<SCompoundWidget> PinChange = SNew(SPinTypeSelector, FGetPinTypeTree::CreateUObject(Schema, &UEdGraphSchema_K2::GetVariableTypeTree))
+						.OnPinTypeChanged(this, &FBlueprintEditor::OnChangePinTypeFinished, SelectedPin)
+						.Schema(Schema)
+						.bAllowExec(false)
+						.IsEnabled(true)
+						.bAllowArrays(false);
+
+					PinTypeChangeMenu = FSlateApplication::Get().PushMenu(
+						GetToolkitHost()->GetParentWidget(), // Parent widget should be k2 not the menu thats open or it will be closed when the menu is dismissed
+						FWidgetPath(),
+						PinChange,
+						FSlateApplication::Get().GetCursorPos(), // summon location
+						FPopupTransitionEffect(FPopupTransitionEffect::TypeInPopup)
+						);
+				}
+				else
+				{
+					TSharedRef<SCompoundWidget> PinChange = SNew(SPinTypeSelector, FGetPinTypeTree::CreateUObject(Schema, &UEdGraphSchema_K2::GetVariableTypeTree))
 				.TargetPinType(this, &FBlueprintEditor::OnGetPinType, RootPin)
-				.OnPinTypeChanged(this, &FBlueprintEditor::OnChangePinTypeFinished, SelectedPin)
-				.Schema(Schema)
-				.bAllowExec(false)
-				.IsEnabled(true)
-				.bAllowArrays(false);
-
-			PinTypeChangePopupWindow = FSlateApplication::Get().PushMenu( 
-				GetToolkitHost()->GetParentWidget(), // Parent widget should be k2 not the menu thats open or it will be closed when the menu is dismissed
-				PinChange,
-				FSlateApplication::Get().GetCursorPos(), // summon location
-				FPopupTransitionEffect(FPopupTransitionEffect::TypeInPopup)
-				);
-		}
-	}
+						.OnPinTypeChanged(this, &FBlueprintEditor::OnChangePinTypeFinished, SelectedPin)
+						.Schema(Schema)
+						.bAllowExec(false)
+						.IsEnabled(true)
+						.bAllowArrays(false);
+
+					PinTypeChangeMenu = FSlateApplication::Get().PushMenu(
+						GetToolkitHost()->GetParentWidget(), // Parent widget should be k2 not the menu thats open or it will be closed when the menu is dismissed
+						FWidgetPath(),
+						PinChange,
+						FSlateApplication::Get().GetCursorPos(), // summon location
+						FPopupTransitionEffect(FPopupTransitionEffect::TypeInPopup)
+						);
+				}
+			}
 }
 
 FEdGraphPinType FBlueprintEditor::OnGetPinType(UEdGraphPin* SelectedPin) const
@@ -3899,14 +3904,14 @@ void FBlueprintEditor::OnChangePinTypeFinished(const FEdGraphPinType& PinType, U
 	{
 		InSelectedPin->PinType = PinType;
 		if (UK2Node_Select* SelectNode = Cast<UK2Node_Select>(InSelectedPin->GetOwningNode()))
-		{
+			{
 			SelectNode->ChangePinType(InSelectedPin);
 		}
 	}
 
-	if (PinTypeChangePopupWindow.IsValid())
+	if (PinTypeChangeMenu.IsValid())
 	{
-		PinTypeChangePopupWindow.Pin()->RequestDestroyWindow();
+		PinTypeChangeMenu.Pin()->Dismiss();
 	}
 }
 
diff --git a/Engine/Source/Editor/Kismet/Private/FindInBlueprints.cpp b/Engine/Source/Editor/Kismet/Private/FindInBlueprints.cpp
--- a/Engine/Source/Editor/Kismet/Private/FindInBlueprints.cpp
+++ b/Engine/Source/Editor/Kismet/Private/FindInBlueprints.cpp
@@ -1102,6 +1102,7 @@ void SFindInBlueprints::ConditionallyAddCacheBar()
 
 					FSlateApplication::Get().PushMenu(
 						InParentWidget.Pin().ToSharedRef(),
+						FWidgetPath(),
 						DisplayWidget,
 						FSlateApplication::Get().GetCursorPos(),
 						FPopupTransitionEffect(FPopupTransitionEffect::TypeInPopup)
diff --git a/Engine/Source/Editor/Kismet/Private/SSCSEditor.cpp b/Engine/Source/Editor/Kismet/Private/SSCSEditor.cpp
--- a/Engine/Source/Editor/Kismet/Private/SSCSEditor.cpp
+++ b/Engine/Source/Editor/Kismet/Private/SSCSEditor.cpp
@@ -2273,13 +2273,17 @@ FReply SSCS_RowWidget::HandleOnAcceptDrop( const FDragDropEvent& DragDropEvent,
 			break;
 
 		case FSCSRowDragDropOp::DropAction_AttachToOrMakeNewRoot:
-			check(DragRowOp->SourceNodes.Num() == 1);
-			FSlateApplication::Get().PushMenu(
-				SharedThis(this),
-				BuildSceneRootDropActionMenu(DragRowOp->SourceNodes[0]).ToSharedRef(),
-				FSlateApplication::Get().GetCursorPos(),
-				FPopupTransitionEffect(FPopupTransitionEffect::TypeInPopup)
-			);
+			{
+				check(DragRowOp->SourceNodes.Num() == 1);
+				FWidgetPath WidgetPath = DragDropEvent.GetEventPath() != nullptr ? *DragDropEvent.GetEventPath() : FWidgetPath();
+				FSlateApplication::Get().PushMenu(
+					SharedThis(this),
+					WidgetPath,
+					BuildSceneRootDropActionMenu(DragRowOp->SourceNodes[0]).ToSharedRef(),
+					FSlateApplication::Get().GetCursorPos(),
+					FPopupTransitionEffect(FPopupTransitionEffect::TypeInPopup)
+				);
+			}
 			break;
 
 		case FSCSRowDragDropOp::DropAction_None:
diff --git a/Engine/Source/Editor/Kismet/Public/BlueprintEditor.h b/Engine/Source/Editor/Kismet/Public/BlueprintEditor.h
--- a/Engine/Source/Editor/Kismet/Public/BlueprintEditor.h
+++ b/Engine/Source/Editor/Kismet/Public/BlueprintEditor.h
@@ -1020,7 +1020,7 @@ protected:
 	TSharedPtr<class INameValidatorInterface> NameEntryValidator;
 
 	/** Reference to owner of the pin type change popup */
-	TWeakPtr<class SWindow> PinTypeChangePopupWindow;
+	TWeakPtr<class IMenu> PinTypeChangeMenu;
 
 	/** The toolbar builder class */
 	TSharedPtr<class FBlueprintEditorToolbar> Toolbar;
diff --git a/Engine/Source/Editor/KismetWidgets/Private/SScrubWidget.cpp b/Engine/Source/Editor/KismetWidgets/Private/SScrubWidget.cpp
--- a/Engine/Source/Editor/KismetWidgets/Private/SScrubWidget.cpp
+++ b/Engine/Source/Editor/KismetWidgets/Private/SScrubWidget.cpp
@@ -242,7 +242,7 @@ FReply SScrubWidget::OnMouseButtonUp( const FGeometry& MyGeometry, const FPointe
 
 		if( !bMouseMovedDuringPanning )
 		{
-			CreateContextMenu( NewValue );
+			CreateContextMenu(NewValue, MouseEvent);
 		}
 		return FReply::Handled().ReleaseMouseCapture();
 	}
@@ -411,7 +411,7 @@ FCursorReply SScrubWidget::OnCursorQuery( const FGeometry& MyGeometry, const FPo
 	return FCursorReply::Unhandled();
 }
 
-void SScrubWidget::CreateContextMenu(float CurrentFrameTime)
+void SScrubWidget::CreateContextMenu(float CurrentFrameTime, const FPointerEvent& MouseEvent)
 {
 	if ((OnCropAnimSequence.IsBound() || OnReZeroAnimSequence.IsBound() || OnAddAnimSequence.IsBound()) && (SequenceLength.Get() >= MINIMUM_ANIMATION_LENGTH))
 	{
@@ -480,7 +480,9 @@ void SScrubWidget::CreateContextMenu(float CurrentFrameTime)
 		}
 		MenuBuilder.EndSection();
 
-		FSlateApplication::Get().PushMenu( SharedThis( this ), MenuBuilder.MakeWidget(), FSlateApplication::Get().GetCursorPos(), FPopupTransitionEffect( FPopupTransitionEffect::ContextMenu ) );
+		FWidgetPath WidgetPath = MouseEvent.GetEventPath() != nullptr ? *MouseEvent.GetEventPath() : FWidgetPath();
+
+		FSlateApplication::Get().PushMenu(SharedThis(this), WidgetPath, MenuBuilder.MakeWidget(), FSlateApplication::Get().GetCursorPos(), FPopupTransitionEffect(FPopupTransitionEffect::ContextMenu));
 	}
 }
 
diff --git a/Engine/Source/Editor/KismetWidgets/Public/SScrubWidget.h b/Engine/Source/Editor/KismetWidgets/Public/SScrubWidget.h
--- a/Engine/Source/Editor/KismetWidgets/Public/SScrubWidget.h
+++ b/Engine/Source/Editor/KismetWidgets/Public/SScrubWidget.h
@@ -91,7 +91,7 @@ protected:
 private:
 
 	/** Function to create context menu to display anim sequence editing options */
-	void CreateContextMenu( float CurrentFrameTime );
+	void CreateContextMenu( float CurrentFrameTime, const FPointerEvent& MouseEvent );
 
 	/** Function to crop animation sequence before/after selected frame */
 	void OnSequenceCropped( bool bFromStart, float CurrentFrameTime );
diff --git a/Engine/Source/Editor/LandscapeEditor/Private/LandscapeEditorDetailCustomization_TargetLayers.cpp b/Engine/Source/Editor/LandscapeEditor/Private/LandscapeEditorDetailCustomization_TargetLayers.cpp
--- a/Engine/Source/Editor/LandscapeEditor/Private/LandscapeEditorDetailCustomization_TargetLayers.cpp
+++ b/Engine/Source/Editor/LandscapeEditor/Private/LandscapeEditorDetailCustomization_TargetLayers.cpp
@@ -915,7 +915,9 @@ FReply SLandscapeEditorSelectableBorder::OnMouseButtonUp(const FGeometry& MyGeom
 			TSharedPtr<SWidget> Content = OnContextMenuOpening.Execute();
 			if (Content.IsValid())
 			{
-				FSlateApplication::Get().PushMenu(SharedThis(this), Content.ToSharedRef(), MouseEvent.GetScreenSpacePosition(), FPopupTransitionEffect(FPopupTransitionEffect::ContextMenu));
+				FWidgetPath WidgetPath = MouseEvent.GetEventPath() != nullptr ? *MouseEvent.GetEventPath() : FWidgetPath();
+
+				FSlateApplication::Get().PushMenu(SharedThis(this), WidgetPath, Content.ToSharedRef(), MouseEvent.GetScreenSpacePosition(), FPopupTransitionEffect(FPopupTransitionEffect::ContextMenu));
 			}
 
 			return FReply::Handled().ReleaseMouseCapture();
diff --git a/Engine/Source/Editor/LevelEditor/Private/LevelEditorActions.cpp b/Engine/Source/Editor/LevelEditor/Private/LevelEditorActions.cpp
--- a/Engine/Source/Editor/LevelEditor/Private/LevelEditorActions.cpp
+++ b/Engine/Source/Editor/LevelEditor/Private/LevelEditorActions.cpp
@@ -421,6 +421,7 @@ void FLevelEditorActionCallbacks::AttachToActor(AActor* ParentActorPtr)
 		// Create as context menu
 		FSlateApplication::Get().PushMenu(
 			LevelEditor.ToSharedRef(),
+			FWidgetPath(),
 			SNew(SSocketChooserPopup)
 			.SceneComponent( ComponentWithSockets )
 			.OnSocketChosen_Static( &FLevelEditorActionCallbacks::AttachToSocketSelection, ParentActorPtr ),
diff --git a/Engine/Source/Editor/LevelEditor/Private/LevelEditorContextMenu.cpp b/Engine/Source/Editor/LevelEditor/Private/LevelEditorContextMenu.cpp
--- a/Engine/Source/Editor/LevelEditor/Private/LevelEditorContextMenu.cpp
+++ b/Engine/Source/Editor/LevelEditor/Private/LevelEditorContextMenu.cpp
@@ -469,8 +469,12 @@ void FLevelEditorContextMenu::SummonMenu( const TSharedRef< SLevelEditor >& Leve
 		// @todo: Should actually use the location from a click event instead!
 		const FVector2D MouseCursorLocation = FSlateApplication::Get().GetCursorPos();
 	
-		TWeakPtr< SWindow > ContextMenuWindow = FSlateApplication::Get().PushMenu(
-			LevelEditor->GetActiveViewport().ToSharedRef(), MenuWidget.ToSharedRef(), MouseCursorLocation, FPopupTransitionEffect( FPopupTransitionEffect::ContextMenu ) );
+		FSlateApplication::Get().PushMenu(
+			LevelEditor->GetActiveViewport().ToSharedRef(),
+			FWidgetPath(),
+			MenuWidget.ToSharedRef(),
+			MouseCursorLocation,
+			FPopupTransitionEffect( FPopupTransitionEffect::ContextMenu ) );
 	}
 }
 
diff --git a/Engine/Source/Editor/LevelEditor/Private/SLevelViewport.cpp b/Engine/Source/Editor/LevelEditor/Private/SLevelViewport.cpp
--- a/Engine/Source/Editor/LevelEditor/Private/SLevelViewport.cpp
+++ b/Engine/Source/Editor/LevelEditor/Private/SLevelViewport.cpp
@@ -775,8 +775,15 @@ bool SLevelViewport::HandlePlaceDraggedObjects(const FGeometry& MyGeometry, cons
 		}
 		else if ( bAllAssetWereLoaded && DroppedObjects.Num() > 0 )
 		{
-			TWeakPtr< SWindow > ContextMenuWindow = FSlateApplication::Get().PushMenu( SharedThis( this ),
-				BuildViewportDragDropContextMenu(), DragDropEvent.GetScreenSpacePosition(), FPopupTransitionEffect( FPopupTransitionEffect::ContextMenu ) );
+			FWidgetPath WidgetPath = DragDropEvent.GetEventPath() != nullptr ? *DragDropEvent.GetEventPath() : FWidgetPath();
+
+			FSlateApplication::Get().PushMenu(
+				SharedThis( this ),
+				WidgetPath,
+				BuildViewportDragDropContextMenu(),
+				DragDropEvent.GetScreenSpacePosition(),
+				FPopupTransitionEffect( FPopupTransitionEffect::ContextMenu ) );
+
 			bDropSuccessful = true;
 		}
 
diff --git a/Engine/Source/Editor/Matinee/Private/Matinee.h b/Engine/Source/Editor/Matinee/Private/Matinee.h
--- a/Engine/Source/Editor/Matinee/Private/Matinee.h
+++ b/Engine/Source/Editor/Matinee/Private/Matinee.h
@@ -1476,7 +1476,7 @@ public:
 	 */
 	void GenericTextEntryModeless(const FText& DialogText, const FText& DefaultText, FOnTextCommitted OnTextComitted);
 	/** Closes the popup created by GenericTextEntryModal or GenericTextEntryModeless*/
-	void CloseEntryPopupWindow();
+	void CloseEntryPopupMenu();
 private:
 	/** Popup menu function for retrieving a new name from the user. */	
 	void GetNewNamePopup( const FText& InDialogTitle, const FText& InDialogCaption, const FText& InDefaultText, const FText& InOriginalName, FOnTextCommitted OnTextCommited );
@@ -1518,7 +1518,7 @@ private:
 	void BindCommands();
 
 	/** Generic Popup Entry */
-	TWeakPtr<SWindow> EntryPopupWindow;
+	TWeakPtr<class IMenu> EntryPopupMenu;
 
 	/** ToolBar delegates */
 	const FSlateBrush* GetToolbarInterpSpeedIcon();
diff --git a/Engine/Source/Editor/Matinee/Private/MatineeFilterButton.cpp b/Engine/Source/Editor/Matinee/Private/MatineeFilterButton.cpp
--- a/Engine/Source/Editor/Matinee/Private/MatineeFilterButton.cpp
+++ b/Engine/Source/Editor/Matinee/Private/MatineeFilterButton.cpp
@@ -36,7 +36,9 @@ FReply SMatineeFilterButton::OnMouseButtonUp( const FGeometry& MyGeometry, const
 			TSharedPtr<SWidget> MenuContent = OnContextMenuOpening.Execute();
 			if( MenuContent.IsValid() )
 			{
-				FSlateApplication::Get().PushMenu( AsShared(), MenuContent.ToSharedRef(), MouseEvent.GetScreenSpacePosition(), FPopupTransitionEffect( FPopupTransitionEffect::ContextMenu ) );
+				FWidgetPath WidgetPath = MouseEvent.GetEventPath() != nullptr ? *MouseEvent.GetEventPath() : FWidgetPath();
+
+				FSlateApplication::Get().PushMenu(AsShared(), WidgetPath, MenuContent.ToSharedRef(), MouseEvent.GetScreenSpacePosition(), FPopupTransitionEffect(FPopupTransitionEffect::ContextMenu));
 			}
 		}
 
diff --git a/Engine/Source/Editor/Matinee/Private/MatineeKeyReduction.cpp b/Engine/Source/Editor/Matinee/Private/MatineeKeyReduction.cpp
--- a/Engine/Source/Editor/Matinee/Private/MatineeKeyReduction.cpp
+++ b/Engine/Source/Editor/Matinee/Private/MatineeKeyReduction.cpp
@@ -207,7 +207,7 @@ FReply SMatineeKeyReduction::OnOK()
 
 	Matinee->InterpEdTrans->EndSpecial();
 
-	Matinee->CloseEntryPopupWindow();
+	Matinee->CloseEntryPopupMenu();
 
 	return FReply::Handled();
 }
@@ -260,8 +260,9 @@ void FMatinee::ReduceKeys()
 			TSharedRef<SMatineeKeyReduction> ParameterDialog = 
 				SNew(SMatineeKeyReduction, this, Track, TrackInst, IntervalStart, IntervalEnd);
 
-			EntryPopupWindow = FSlateApplication::Get().PushMenu(
+			EntryPopupMenu = FSlateApplication::Get().PushMenu(
 				ToolkitHost.Pin()->GetParentWidget(),
+				FWidgetPath(),
 				ParameterDialog,
 				FSlateApplication::Get().GetCursorPos(),
 				FPopupTransitionEffect(FPopupTransitionEffect::TypeInPopup)
diff --git a/Engine/Source/Editor/Matinee/Private/MatineeMenus.cpp b/Engine/Source/Editor/Matinee/Private/MatineeMenus.cpp
--- a/Engine/Source/Editor/Matinee/Private/MatineeMenus.cpp
+++ b/Engine/Source/Editor/Matinee/Private/MatineeMenus.cpp
@@ -60,6 +60,7 @@
 #include "Sound/SoundBase.h"
 #include "Camera/CameraAnim.h"
 
+#include "IMenu.h"
 
 #define LOCTEXT_NAMESPACE "MatineeMenus"
 
@@ -120,21 +121,22 @@ void FMatinee::GenericTextEntryModeless(const FText& DialogText, const FText& De
 		.OnTextCommitted(OnTextComitted)
 		.ClearKeyboardFocusOnCommit(false)
 		.SelectAllTextWhenFocused(true)
-		.MaxWidth(1024.0f)
-		;
-	EntryPopupWindow = FSlateApplication::Get().PushMenu(
+		.MaxWidth(1024.0f);
+
+	EntryPopupMenu = FSlateApplication::Get().PushMenu(
 		ToolkitHost.Pin()->GetParentWidget(),
+		FWidgetPath(),
 		TextEntryPopup,
 		FSlateApplication::Get().GetCursorPos(),
 		FPopupTransitionEffect(FPopupTransitionEffect::TypeInPopup)
 		);
 }
 
-void FMatinee::CloseEntryPopupWindow()
+void FMatinee::CloseEntryPopupMenu()
 {
-	if (EntryPopupWindow.IsValid())
+	if (EntryPopupMenu.IsValid())
 	{
-		EntryPopupWindow.Pin()->RequestDestroyWindow();
+		EntryPopupMenu.Pin()->Dismiss();
 	}
 }
 
@@ -171,7 +173,7 @@ void FMatinee::OnNewNamePopupTextCommitted(const FText& InText, ETextCommit::Typ
 {
 	if (CommitInfo == ETextCommit::OnEnter)
 	{	
-		CloseEntryPopupWindow();
+		CloseEntryPopupMenu();
 		if (!InText.IsEmpty())
 		{
 			// Make sure there are no spaces!
@@ -836,7 +838,7 @@ void FMatinee::StretchSection( bool bUseSelectedOnly )
 
 void FMatinee::OnStretchSectionTextEntered(const FText& InText, ETextCommit::Type CommitInfo, float SectionStart, float SectionEnd, float CurrentSectionLength )
 {
-	CloseEntryPopupWindow();
+	CloseEntryPopupMenu();
 	if (CommitInfo == ETextCommit::OnEnter)
 	{
 		double dNewSectionLength = FCString::Atod(*InText.ToString());
@@ -978,7 +980,7 @@ void FMatinee::OnMenuInsertSpace()
 
 void FMatinee::OnInsertSpaceTextEntry(const FText& InText, ETextCommit::Type CommitInfo)
 {
-	CloseEntryPopupWindow();
+	CloseEntryPopupMenu();
 	if (CommitInfo == ETextCommit::OnEnter)
 	{
 		bool bIsNumber = InText.IsNumeric();
@@ -1821,7 +1823,7 @@ bool FMatinee::CanGroupCreateTab() const
 
 void FMatinee::OnContextGroupCreateTabTextCommitted(const FText& InText, ETextCommit::Type CommitInfo)
 {
-	CloseEntryPopupWindow();	
+	CloseEntryPopupMenu();	
 	if (CommitInfo == ETextCommit::OnEnter)
 	{
 		// Create a new tab.
@@ -2215,7 +2217,7 @@ void FMatinee::OnContextSetKeyTime()
 
 void FMatinee::OnContextSetKeyTimeTextCommitted(const FText& InText, ETextCommit::Type CommitInfo, FInterpEdSelKey* SelKey, UInterpTrack* Track)
 {
-	CloseEntryPopupWindow();
+	CloseEntryPopupMenu();
 	if (CommitInfo == ETextCommit::OnEnter)
 	{
 		double dNewTime = FCString::Atod(*InText.ToString());
@@ -2275,7 +2277,7 @@ void FMatinee::OnContextSetValue()
 
 void FMatinee::OnContextSetValueTextCommitted(const FText& InText, ETextCommit::Type CommitInfo, FInterpEdSelKey* SelKey, UInterpTrackFloatBase* FloatTrack)
 {
-	CloseEntryPopupWindow();
+	CloseEntryPopupMenu();
 	if (CommitInfo == ETextCommit::OnEnter)
 	{
 		double dNewVal = FCString::Atod(*InText.ToString());
@@ -2496,8 +2498,9 @@ void FMatinee::OnSetMoveKeyLookupGroup()
 			.OnTextChosen(this, &FMatinee::OnSetMoveKeyLookupGroupTextChosen, SelKey.KeyIndex, MoveTrack, MoveTrackAxis)
 			;
 
-		EntryPopupWindow = FSlateApplication::Get().PushMenu( 
+		EntryPopupMenu = FSlateApplication::Get().PushMenu(
 			ToolkitHost.Pin()->GetParentWidget(),
+			FWidgetPath(),
 			TextEntryPopup,
 			FSlateApplication::Get().GetCursorPos(),
 			FPopupTransitionEffect(FPopupTransitionEffect::TypeInPopup)
@@ -2516,7 +2519,7 @@ void FMatinee::OnSetMoveKeyLookupGroupTextChosen(const FString& ChosenText, int3
 	{
 		MoveTrackAxis->SetLookupKeyGroupName( KeyIndex, KeyframeLookupGroup );
 	}
-	CloseEntryPopupWindow();
+	CloseEntryPopupMenu();
 }
 
 
@@ -2583,7 +2586,7 @@ void FMatinee::OnContextRenameEventKey()
 
 void FMatinee::OnContextRenameEventKeyTextCommitted(const FText& InText, ETextCommit::Type CommitInfo, FName EventNameToChange)
 { 
-	CloseEntryPopupWindow();
+	CloseEntryPopupMenu();
 	if (CommitInfo == ETextCommit::OnEnter)
 	{
 		FString TempString = InText.ToString().Left(NAME_SIZE);
@@ -2690,7 +2693,7 @@ void FMatinee::OnSetAnimOffset( bool bInEndOffset )
 
 void FMatinee::OnSetAnimOffsetTextCommitted(const FText& InText, ETextCommit::Type CommitInfo, FInterpEdSelKey* SelKey, UInterpTrackAnimControl* AnimTrack, bool bEndOffset)
 {
-	CloseEntryPopupWindow();
+	CloseEntryPopupMenu();
 	if (CommitInfo == ETextCommit::OnEnter)
 	{
 		double dNewOffset = FCString::Atod(*InText.ToString());
@@ -2739,7 +2742,7 @@ void FMatinee::OnSetAnimPlayRate()
 
 void FMatinee::OnSetAnimPlayRateTextCommitted(const FText& InText, ETextCommit::Type CommitInfo, FInterpEdSelKey* SelKey, UInterpTrackAnimControl* AnimTrack)
 {
-	CloseEntryPopupWindow();
+	CloseEntryPopupMenu();
 	if (CommitInfo == ETextCommit::OnEnter)
 	{
 		double dNewRate = FCString::Atod(*InText.ToString());
@@ -2873,7 +2876,7 @@ void FMatinee::ExportCameraAnimationNameCommitted(const FText& InAnimationPackag
 		}
 	}
 
-	CloseEntryPopupWindow();
+	CloseEntryPopupMenu();
 }
 
 void FMatinee::OnContextSaveAsCameraAnimation()
@@ -2893,8 +2896,9 @@ void FMatinee::OnContextSaveAsCameraAnimation()
 		.OnTextCommitted(this, &FMatinee::ExportCameraAnimationNameCommitted)
 		.ClearKeyboardFocusOnCommit( false );
 
-	EntryPopupWindow = FSlateApplication::Get().PushMenu(
+	EntryPopupMenu = FSlateApplication::Get().PushMenu(
 		ToolkitHost.Pin()->GetParentWidget(),
+		FWidgetPath(),
 		TextEntry,
 		FSlateApplication::Get().GetCursorPos(),
 		FPopupTransitionEffect( FPopupTransitionEffect::TypeInPopup )
@@ -3139,7 +3143,7 @@ void FMatinee::OnSetSoundVolume()
 
 void FMatinee::OnSetSoundVolumeTextEntered( const FText& InText, ETextCommit::Type CommitInfo, TArray<int32> SoundTrackKeyIndices)
 {
-	CloseEntryPopupWindow();
+	CloseEntryPopupMenu();
 	if( CommitInfo ==  ETextCommit::OnEnter )
 	{
 		double NewVolume = FCString::Atod(*InText.ToString());
@@ -3321,7 +3325,7 @@ void FMatinee::OnKeyContext_SetMasterVolume()
 
 void FMatinee::OnKeyContext_SetMasterVolumeTextCommitted(const FText& InText, ETextCommit::Type CommitInfo, TArray<int32> SoundTrackKeyIndices)
 {
-	CloseEntryPopupWindow();
+	CloseEntryPopupMenu();
 	if (CommitInfo == ETextCommit::OnEnter)
 	{
 		double NewVolume = FCString::Atod(*InText.ToString());
@@ -3401,7 +3405,7 @@ void FMatinee::OnKeyContext_SetMasterPitch()
 }
 void FMatinee::OnKeyContext_SetMasterPitchTextCommitted(const FText& InText, ETextCommit::Type CommitInfo, TArray<int32> SoundTrackKeyIndices)
 {
-	CloseEntryPopupWindow();
+	CloseEntryPopupMenu();
 	if (CommitInfo == ETextCommit::OnEnter)
 	{
 		double NewPitch = FCString::Atod(*InText.ToString());
@@ -3456,7 +3460,7 @@ void FMatinee::OnParticleReplayKeyContext_SetClipIDNumber()
 
 void FMatinee::OnParticleReplayKeyContext_SetClipIDNumberTextCommitted(const FText& InText, ETextCommit::Type CommitInfo, FParticleReplayTrackKey* ParticleReplayKey)
 {
-	CloseEntryPopupWindow();
+	CloseEntryPopupMenu();
 	if (CommitInfo == ETextCommit::OnEnter)
 	{
 		int32 NewClipIDNumber = FCString::Atoi(*InText.ToString());
@@ -3501,7 +3505,7 @@ void FMatinee::OnParticleReplayKeyContext_SetDuration()
 
 void FMatinee::OnParticleReplayKeyContext_SetDurationTextCommitted(const FText& InText, ETextCommit::Type CommitInfo, FParticleReplayTrackKey* ParticleReplayKey)
 {
-	CloseEntryPopupWindow();
+	CloseEntryPopupMenu();
 	if (CommitInfo == ETextCommit::OnEnter)
 	{
 		int32 NewDuration = FCString::Atoi(*InText.ToString());
@@ -3554,7 +3558,7 @@ void FMatinee::OnContextDirKeyTransitionTime()
 
 void FMatinee::OnContextDirKeyTransitionTimeTextCommitted(const FText& InText, ETextCommit::Type CommitInfo, FInterpEdSelKey* SelKey, UInterpTrackDirector* DirTrack)
 {
-	CloseEntryPopupWindow();
+	CloseEntryPopupMenu();
 	if (CommitInfo == ETextCommit::OnEnter)
 	{
 		double dNewTime = FCString::Atod(*InText.ToString());
@@ -3595,7 +3599,7 @@ void FMatinee::OnContextDirKeyRenameCameraShot()
 
 void FMatinee::OnContextDirKeyRenameCameraShotTextCommitted(const FText& InText, ETextCommit::Type CommitInfo, FInterpEdSelKey* SelKey, UInterpTrackDirector* DirTrack)
 {
-	CloseEntryPopupWindow();
+	CloseEntryPopupMenu();
 	if (CommitInfo == ETextCommit::OnEnter)
 	{
 		uint32 NewShot = FCString::Atoi(*InText.ToString());    
@@ -4752,7 +4756,7 @@ void FMatinee::ScaleTranslationByAmount( const FText& InText, ETextCommit::Type
 {
 	if (CommitInfo == ETextCommit::OnEnter)
 	{
-		CloseEntryPopupWindow();
+		CloseEntryPopupMenu();
 		float Amount = FCString::Atof(*InText.ToString());
 		const bool bIsNumber = InText.IsNumeric();
 		if (bIsNumber)
diff --git a/Engine/Source/Editor/Matinee/Private/MatineeTrackHelpers.cpp b/Engine/Source/Editor/Matinee/Private/MatineeTrackHelpers.cpp
--- a/Engine/Source/Editor/Matinee/Private/MatineeTrackHelpers.cpp
+++ b/Engine/Source/Editor/Matinee/Private/MatineeTrackHelpers.cpp
@@ -25,6 +25,7 @@
 #include "STextEntryPopup.h"
 #include "Engine/Selection.h"
 #include "Sound/SoundBase.h"
+#include "IMenu.h"
 
 #define LOCTEXT_NAMESPACE "MatineeTrackHelpers"
 
@@ -63,7 +64,7 @@ static UAnimSequence*	KeyframeAddAnimSequence = NULL;
 static USoundBase*		KeyframeAddSound = NULL;
 static FName			TrackAddPropName = NAME_None;
 static FName			AnimSlotName = NAME_None;
-static TWeakPtr< class SWindow > EntryPopupWindow;
+static TWeakPtr< class IMenu > EntryMenu;
 
 /**
  * Sets the global property name to use for newly created property tracks
@@ -273,8 +274,9 @@ bool UMatineeTrackAnimControlHelper::PreCreateKeyframe( UInterpTrack *Track, flo
 			}
 			MenuBuilder.EndSection();
 
-			EntryPopupWindow = FSlateApplication::Get().PushMenu(
+			EntryMenu = FSlateApplication::Get().PushMenu(
 				Parent.ToSharedRef(),
+				FWidgetPath(),
 				MenuBuilder.MakeWidget(),
 				FSlateApplication::Get().GetCursorPos(),
 				FPopupTransitionEffect(FPopupTransitionEffect::TypeInPopup)
@@ -291,9 +293,9 @@ bool UMatineeTrackAnimControlHelper::PreCreateKeyframe( UInterpTrack *Track, flo
 
 void UMatineeTrackAnimControlHelper::OnAddKeyTextEntry(const FAssetData& AssetData, IMatineeBase* Matinee, UInterpTrack* Track)
 {
-	if( EntryPopupWindow.IsValid() )
+	if (EntryMenu.IsValid())
 	{
-		EntryPopupWindow.Pin()->RequestDestroyWindow();
+		EntryMenu.Pin()->Dismiss();
 	}
 
 	UObject* SelectedObject = AssetData.GetAsset();
@@ -349,8 +351,9 @@ bool UMatineeTrackDirectorHelper::PreCreateKeyframe( UInterpTrack *Track, float
 		TSharedPtr< SWindow > Parent = FSlateApplication::Get().GetActiveTopLevelWindow();
 		if ( Parent.IsValid() )
 		{
-			EntryPopupWindow = FSlateApplication::Get().PushMenu(
+			EntryMenu = FSlateApplication::Get().PushMenu(
 				Parent.ToSharedRef(),
+				FWidgetPath(),
 				TextEntryPopup,
 				FSlateApplication::Get().GetCursorPos(),
 				FPopupTransitionEffect(FPopupTransitionEffect::TypeInPopup)
@@ -363,7 +366,10 @@ bool UMatineeTrackDirectorHelper::PreCreateKeyframe( UInterpTrack *Track, float
 
 void UMatineeTrackDirectorHelper::OnAddKeyTextEntry(const FString& ChosenText, IMatineeBase* Matinee, UInterpTrack* Track)
 {
-	EntryPopupWindow.Pin()->RequestDestroyWindow();
+	if (EntryMenu.IsValid())
+	{
+		EntryMenu.Pin()->Dismiss();
+	}
 	
 	KeyframeAddDataName = FName( *ChosenText );
 	Matinee->FinishAddKey(Track,true);
@@ -407,8 +413,9 @@ bool UMatineeTrackEventHelper::PreCreateKeyframe( UInterpTrack *Track, float Key
 	TSharedPtr< SWindow > Parent = FSlateApplication::Get().GetActiveTopLevelWindow();
 	if ( Parent.IsValid() )
 	{
-		EntryPopupWindow = FSlateApplication::Get().PushMenu(
+		EntryMenu = FSlateApplication::Get().PushMenu(
 			Parent.ToSharedRef(),
+			FWidgetPath(),
 			TextEntryPopup,
 			FSlateApplication::Get().GetCursorPos(),
 			FPopupTransitionEffect(FPopupTransitionEffect::TypeInPopup)
@@ -420,9 +427,9 @@ bool UMatineeTrackEventHelper::PreCreateKeyframe( UInterpTrack *Track, float Key
 
 void UMatineeTrackEventHelper::OnAddKeyTextEntry(const FText& ChosenText, ETextCommit::Type CommitInfo, IMatineeBase* Matinee, UInterpTrack* Track)
 {
-	if (EntryPopupWindow.IsValid())
+	if (EntryMenu.IsValid())
 	{
-		EntryPopupWindow.Pin()->RequestDestroyWindow();
+		EntryMenu.Pin()->Dismiss();
 	}
 
 	if (CommitInfo == ETextCommit::OnEnter)
@@ -786,8 +793,9 @@ bool UMatineeTrackToggleHelper::PreCreateKeyframe( UInterpTrack *Track, float Ke
 	TSharedPtr< SWindow > Parent = FSlateApplication::Get().GetActiveTopLevelWindow();
 	if ( Parent.IsValid() )
 	{
-		EntryPopupWindow = FSlateApplication::Get().PushMenu(
+		EntryMenu = FSlateApplication::Get().PushMenu(
 			Parent.ToSharedRef(),
+			FWidgetPath(),
 			TextEntryPopup,
 			FSlateApplication::Get().GetCursorPos(),
 			FPopupTransitionEffect(FPopupTransitionEffect::TypeInPopup)
@@ -799,9 +807,9 @@ bool UMatineeTrackToggleHelper::PreCreateKeyframe( UInterpTrack *Track, float Ke
 
 void UMatineeTrackToggleHelper::OnAddKeyTextEntry(const FString& ChosenText, IMatineeBase* Matinee, UInterpTrack* Track)
 {
-	if (EntryPopupWindow.IsValid())
+	if (EntryMenu.IsValid())
 	{
-		EntryPopupWindow.Pin()->RequestDestroyWindow();
+		EntryMenu.Pin()->Dismiss();
 	}
 
 	KeyframeAddDataName = FName(*ChosenText);
@@ -1137,8 +1145,9 @@ bool UMatineeTrackVisibilityHelper::PreCreateKeyframe( UInterpTrack *Track, floa
 	TSharedPtr< SWindow > Parent = FSlateApplication::Get().GetActiveTopLevelWindow();
 	if ( Parent.IsValid() )
 	{
-		EntryPopupWindow = FSlateApplication::Get().PushMenu(
+		EntryMenu = FSlateApplication::Get().PushMenu(
 			Parent.ToSharedRef(),
+			FWidgetPath(),
 			TextEntryPopup,
 			FSlateApplication::Get().GetCursorPos(),
 			FPopupTransitionEffect(FPopupTransitionEffect::TypeInPopup)
@@ -1150,9 +1159,9 @@ bool UMatineeTrackVisibilityHelper::PreCreateKeyframe( UInterpTrack *Track, floa
 
 void UMatineeTrackVisibilityHelper::OnAddKeyTextEntry(const FString& ChosenText, IMatineeBase* Matinee, UInterpTrack* Track)
 {
-	if (EntryPopupWindow.IsValid())
+	if (EntryMenu.IsValid())
 	{
-		EntryPopupWindow.Pin()->RequestDestroyWindow();
+		EntryMenu.Pin()->Dismiss();
 	}
 
 	KeyframeAddDataName = FName(*ChosenText);
diff --git a/Engine/Source/Editor/Matinee/Private/MatineeViewportClient.cpp b/Engine/Source/Editor/Matinee/Private/MatineeViewportClient.cpp
--- a/Engine/Source/Editor/Matinee/Private/MatineeViewportClient.cpp
+++ b/Engine/Source/Editor/Matinee/Private/MatineeViewportClient.cpp
@@ -801,12 +801,12 @@ bool FMatineeViewportClient::InputKey(FViewport* Viewport, int32 ControllerId, F
 						TSharedPtr< SWindow > Parent = FSlateApplication::Get().GetActiveTopLevelWindow(); 
 						if ( Parent.IsValid() )
 						{
-								FSlateApplication::Get().PushMenu(
+							FSlateApplication::Get().PushMenu(
 								Parent.ToSharedRef(),
+								FWidgetPath(),
 								Menu.ToSharedRef(), 
 								FSlateApplication::Get().GetCursorPos(), 
-								FPopupTransitionEffect(FPopupTransitionEffect::ContextMenu)
-							);
+								FPopupTransitionEffect(FPopupTransitionEffect::ContextMenu));
 						}
 					}
 				}
diff --git a/Engine/Source/Editor/Persona/Private/SAnimCurvePanel.cpp b/Engine/Source/Editor/Persona/Private/SAnimCurvePanel.cpp
--- a/Engine/Source/Editor/Persona/Private/SAnimCurvePanel.cpp
+++ b/Engine/Source/Editor/Persona/Private/SAnimCurvePanel.cpp
@@ -412,6 +412,7 @@ FReply SCurveEdTrack::OnContextMenu()
 	{
 		FFloatCurve* Curve = (FFloatCurve*)(CurveInterface->CurveData);
 		FSlateApplication::Get().PushMenu(SharedThis(this),
+										  FWidgetPath(),
 										  PanelShared->CreateCurveContextMenu(Curve),
 										  FSlateApplication::Get().GetCursorPos(),
 										  FPopupTransitionEffect(FPopupTransitionEffect::TypeInPopup));
@@ -591,6 +592,7 @@ FReply SAnimCurvePanel::AddButtonClicked()
 	// Show dialog to enter new track name
 	FSlateApplication::Get().PushMenu(
 		SharedThis( this ),
+		FWidgetPath(),
 		MenuBuilder.MakeWidget(),
 		FSlateApplication::Get().GetCursorPos(),
 		FPopupTransitionEffect( FPopupTransitionEffect::TypeInPopup)
@@ -694,6 +696,7 @@ FReply SAnimCurvePanel::OnContextMenu()
 		);
 
 		FSlateApplication::Get().PushMenu(	SharedThis(this),
+			FWidgetPath(),
 			MenuBuilder.MakeWidget(),
 			FSlateApplication::Get().GetCursorPos(),
 			FPopupTransitionEffect(FPopupTransitionEffect::TypeInPopup) );
@@ -1101,6 +1104,7 @@ void SAnimCurvePanel::CreateNewMetadataEntryClicked()
 	FSlateApplication& SlateApp = FSlateApplication::Get();
 	SlateApp.PushMenu(
 		AsShared(),
+		FWidgetPath(),
 		TextEntry,
 		SlateApp.GetCursorPos(),
 		FPopupTransitionEffect::TypeInPopup
@@ -1136,6 +1140,7 @@ void SAnimCurvePanel::CreateNewCurveClicked()
 	FSlateApplication& SlateApp = FSlateApplication::Get();
 	SlateApp.PushMenu(
 		AsShared(),
+		FWidgetPath(),
 		TextEntry,
 		SlateApp.GetCursorPos(),
 		FPopupTransitionEffect::TypeInPopup
diff --git a/Engine/Source/Editor/Persona/Private/SAnimMontagePanel.cpp b/Engine/Source/Editor/Persona/Private/SAnimMontagePanel.cpp
--- a/Engine/Source/Editor/Persona/Private/SAnimMontagePanel.cpp
+++ b/Engine/Source/Editor/Persona/Private/SAnimMontagePanel.cpp
@@ -382,6 +382,7 @@ void SAnimMontagePanel::OnNewSectionClicked(float DataPosX)
 	// Show dialog to enter new event name
 	FSlateApplication::Get().PushMenu(
 		AsShared(), // Menu being summoned from a menu that is closing: Parent widget should be k2 not the menu thats open or it will be closed when the menu is dismissed
+		FWidgetPath(),
 		TextEntry,
 		FSlateApplication::Get().GetCursorPos(),
 		FPopupTransitionEffect( FPopupTransitionEffect::TypeInPopup )
diff --git a/Engine/Source/Editor/Persona/Private/SAnimNotifyPanel.cpp b/Engine/Source/Editor/Persona/Private/SAnimNotifyPanel.cpp
--- a/Engine/Source/Editor/Persona/Private/SAnimNotifyPanel.cpp
+++ b/Engine/Source/Editor/Persona/Private/SAnimNotifyPanel.cpp
@@ -2580,9 +2580,10 @@ TSharedPtr<SWidget> SAnimNotifyTrack::SummonContextMenu(const FGeometry& MyGeome
 		}
 	}
 
+	FWidgetPath WidgetPath = MouseEvent.GetEventPath() != nullptr ? *MouseEvent.GetEventPath() : FWidgetPath();
+
 	// Display the newly built menu
-	TWeakPtr<SWindow> ContextMenuWindow =
-		FSlateApplication::Get().PushMenu( SharedThis( this ), MenuBuilder.MakeWidget(), CursorPos, FPopupTransitionEffect( FPopupTransitionEffect::ContextMenu ));
+	FSlateApplication::Get().PushMenu(SharedThis(this), WidgetPath, MenuBuilder.MakeWidget(), CursorPos, FPopupTransitionEffect(FPopupTransitionEffect::ContextMenu));
 
 	return TSharedPtr<SWidget>();
 }
@@ -2830,6 +2831,7 @@ void SAnimNotifyTrack::OnSetTriggerWeightNotifyClicked(int32 NotifyIndex)
 
 		FSlateApplication::Get().PushMenu(
 			AsShared(), // Menu being summoned from a menu that is closing: Parent widget should be k2 not the menu thats open or it will be closed when the menu is dismissed
+			FWidgetPath(),
 			TextEntry,
 			FSlateApplication::Get().GetCursorPos(),
 			FPopupTransitionEffect( FPopupTransitionEffect::TypeInPopup )
@@ -2852,6 +2854,7 @@ void SAnimNotifyTrack::OnSetDurationNotifyClicked(int32 NotifyIndex)
 
 		FSlateApplication::Get().PushMenu(
 			AsShared(), // Menu being summoned from a menu that is closing: Parent widget should be k2 not the menu thats open or it will be closed when the menu is dismissed
+			FWidgetPath(),
 			TextEntry,
 			FSlateApplication::Get().GetCursorPos(),
 			FPopupTransitionEffect( FPopupTransitionEffect::TypeInPopup )
@@ -2890,6 +2893,7 @@ void SAnimNotifyTrack::OnNewNotifyClicked()
 	// Show dialog to enter new event name
 	FSlateApplication::Get().PushMenu(
 		AsShared(), // Menu being summoned from a menu that is closing: Parent widget should be k2 not the menu thats open or it will be closed when the menu is dismissed
+		FWidgetPath(),
 		TextEntry,
 		FSlateApplication::Get().GetCursorPos(),
 		FPopupTransitionEffect( FPopupTransitionEffect::TypeInPopup )
@@ -3348,6 +3352,7 @@ void SAnimNotifyTrack::OnSetNotifyTimeClicked(int32 NotifyIndex)
 
 		FSlateApplication::Get().PushMenu(
 			AsShared(), // Menu being summoned from a menu that is closing: Parent widget should be k2 not the menu thats open or it will be closed when the menu is dismissed
+			FWidgetPath(),
 			TextEntry,
 			FSlateApplication::Get().GetCursorPos(),
 			FPopupTransitionEffect(FPopupTransitionEffect::TypeInPopup)
@@ -3399,6 +3404,7 @@ void SAnimNotifyTrack::OnSetNotifyFrameClicked(int32 NotifyIndex)
 
 		FSlateApplication::Get().PushMenu(
 			AsShared(), // Menu being summoned from a menu that is closing: Parent widget should be k2 not the menu thats open or it will be closed when the menu is dismissed
+			FWidgetPath(),
 			TextEntry,
 			FSlateApplication::Get().GetCursorPos(),
 			FPopupTransitionEffect(FPopupTransitionEffect::TypeInPopup)
diff --git a/Engine/Source/Editor/Persona/Private/SAnimTimingPanel.cpp b/Engine/Source/Editor/Persona/Private/SAnimTimingPanel.cpp
--- a/Engine/Source/Editor/Persona/Private/SAnimTimingPanel.cpp
+++ b/Engine/Source/Editor/Persona/Private/SAnimTimingPanel.cpp
@@ -346,6 +346,7 @@ FReply SAnimTimingPanel::OnContextMenu()
 	Builder.EndSection();
 
 	FSlateApplication::Get().PushMenu(SharedThis(this),
+		FWidgetPath(),
 		Builder.MakeWidget(),
 		FSlateApplication::Get().GetCursorPos(),
 		FPopupTransitionEffect(FPopupTransitionEffect::ContextMenu));
diff --git a/Engine/Source/Editor/Persona/Private/SAnimTrackCurvePanel.cpp b/Engine/Source/Editor/Persona/Private/SAnimTrackCurvePanel.cpp
--- a/Engine/Source/Editor/Persona/Private/SAnimTrackCurvePanel.cpp
+++ b/Engine/Source/Editor/Persona/Private/SAnimTrackCurvePanel.cpp
@@ -481,6 +481,7 @@ FReply STransformCurveEdTrack::OnContextMenu()
 	if(PanelShared.IsValid())
 	{
 		FSlateApplication::Get().PushMenu(SharedThis(this),
+										  FWidgetPath(),
 										  PanelShared->CreateCurveContextMenu(CurveUid),
 										  FSlateApplication::Get().GetCursorPos(),
 										  FPopupTransitionEffect(FPopupTransitionEffect::TypeInPopup));
diff --git a/Engine/Source/Editor/Persona/Private/SAnimationBlendSpaceBase.cpp b/Engine/Source/Editor/Persona/Private/SAnimationBlendSpaceBase.cpp
--- a/Engine/Source/Editor/Persona/Private/SAnimationBlendSpaceBase.cpp
+++ b/Engine/Source/Editor/Persona/Private/SAnimationBlendSpaceBase.cpp
@@ -783,7 +783,7 @@ FReply SBlendSpaceWidget::OnMouseButtonUp( const FGeometry& MyGeometry, const FP
 	int32 HighlightedSampleIndex = GetHighlightedSample(GetWindowRectFromGeometry(MyGeometry));
 	if (HighlightedSampleIndex!=INDEX_NONE)
 	{
-		ShowPopupEditWindow(HighlightedSampleIndex);
+		ShowPopupEditWindow(HighlightedSampleIndex, MouseEvent);
 		return FReply::Handled();
 	}
 
@@ -1121,15 +1121,18 @@ int32 SBlendSpaceWidget::GetHighlightedSample(const FSlateRect& WindowRect) cons
 	return HighlightedSampleIndex;
 }
 
-void SBlendSpaceWidget::ShowPopupEditWindow(int32 SampleIndex)
+void SBlendSpaceWidget::ShowPopupEditWindow(int32 SampleIndex, const FPointerEvent& MouseEvent)
 {
 	if (CachedSamples.IsValidIndex(SampleIndex))
 	{
 		FBlendSample Sample = CachedSamples[SampleIndex];
 
+		FWidgetPath WidgetPath = MouseEvent.GetEventPath() != nullptr ? *MouseEvent.GetEventPath() : FWidgetPath();
+
 		// Show dialog to enter new function name
 		FSlateApplication::Get().PushMenu(
 			SharedThis( this ),
+			WidgetPath,
 			SNew(SSampleEditPopUp)
 			.BlendSpace(BlendSpace)
 			.BlendSpaceEditorWidget(this)
diff --git a/Engine/Source/Editor/Persona/Private/SAnimationBlendSpaceBase.h b/Engine/Source/Editor/Persona/Private/SAnimationBlendSpaceBase.h
--- a/Engine/Source/Editor/Persona/Private/SAnimationBlendSpaceBase.h
+++ b/Engine/Source/Editor/Persona/Private/SAnimationBlendSpaceBase.h
@@ -295,7 +295,7 @@ protected:
 	virtual FReply UpdateLastMousePosition( const FGeometry& MyGeometry, const FVector2D& ScreenSpacePosition, bool bClampToWindowRect = false, bool bSnapToGrid = false  ) PURE_VIRTUAL(SBlendSpaceWidget::UpdateLastMousePosition, return FReply::Handled(););
 
 	/** Display sample editing popup */
-	void ShowPopupEditWindow(int32 SampleIndex);
+	void ShowPopupEditWindow(int32 SampleIndex, const FPointerEvent& MouseEvent);
 
 	/** Handler for sample editing popup to delete selected sample */
 	FReply DeleteSample(int32 SampleIndex);
diff --git a/Engine/Source/Editor/Persona/Private/SRetargetSourceWindow.cpp b/Engine/Source/Editor/Persona/Private/SRetargetSourceWindow.cpp
--- a/Engine/Source/Editor/Persona/Private/SRetargetSourceWindow.cpp
+++ b/Engine/Source/Editor/Persona/Private/SRetargetSourceWindow.cpp
@@ -566,6 +566,7 @@ void SRetargetSourceWindow::OnAddRetargetSource()
 
 	FSlateApplication::Get().PushMenu(
 		AsShared(),
+		FWidgetPath(),
 		Widget,
 		FSlateApplication::Get().GetCursorPos(),
 		FPopupTransitionEffect( FPopupTransitionEffect::TopMenu )
diff --git a/Engine/Source/Editor/Persona/Private/SSkeletonSlotNames.cpp b/Engine/Source/Editor/Persona/Private/SSkeletonSlotNames.cpp
--- a/Engine/Source/Editor/Persona/Private/SSkeletonSlotNames.cpp
+++ b/Engine/Source/Editor/Persona/Private/SSkeletonSlotNames.cpp
@@ -400,6 +400,7 @@ void SSkeletonSlotNames::OnAddSlot()
 	// Show dialog to enter new track name
 	FSlateApplication::Get().PushMenu(
 		SharedThis(this),
+		FWidgetPath(),
 		TextEntry,
 		FSlateApplication::Get().GetCursorPos(),
 		FPopupTransitionEffect(FPopupTransitionEffect::TypeInPopup)
@@ -416,6 +417,7 @@ void SSkeletonSlotNames::OnAddGroup()
 	// Show dialog to enter new track name
 	FSlateApplication::Get().PushMenu(
 		SharedThis(this),
+		FWidgetPath(),
 		TextEntry,
 		FSlateApplication::Get().GetCursorPos(),
 		FPopupTransitionEffect(FPopupTransitionEffect::TypeInPopup)
@@ -909,6 +911,7 @@ void SSkeletonSlotNames::OnRenameSlot(FName CurrentName)
 	// Show dialog to enter new track name
 	FSlateApplication::Get().PushMenu(
 		SharedThis(this),
+		FWidgetPath(),
 		TextEntry,
 		FSlateApplication::Get().GetCursorPos(),
 		FPopupTransitionEffect(FPopupTransitionEffect::TypeInPopup)
diff --git a/Engine/Source/Editor/Persona/Private/SSkeletonSmartNameManager.cpp b/Engine/Source/Editor/Persona/Private/SSkeletonSmartNameManager.cpp
--- a/Engine/Source/Editor/Persona/Private/SSkeletonSmartNameManager.cpp
+++ b/Engine/Source/Editor/Persona/Private/SSkeletonSmartNameManager.cpp
@@ -311,6 +311,7 @@ void SSkeletonSmartNameManager::OnAddClicked()
 	FSlateApplication& SlateApp = FSlateApplication::Get();
 	SlateApp.PushMenu(
 		AsShared(),
+		FWidgetPath(),
 		TextEntry,
 		SlateApp.GetCursorPos(),
 		FPopupTransitionEffect::TypeInPopup
diff --git a/Engine/Source/Editor/Persona/Private/STrack.cpp b/Engine/Source/Editor/Persona/Private/STrack.cpp
--- a/Engine/Source/Editor/Persona/Private/STrack.cpp
+++ b/Engine/Source/Editor/Persona/Private/STrack.cpp
@@ -740,14 +740,16 @@ TSharedPtr<SWidget> STrack::SummonContextMenu(const FGeometry& MyGeometry, const
 	}
 
 	// Build the menu if we actually added anything to it
-	TSharedPtr<SWindow> ContextMenuWindow = NULL;
+	TSharedPtr<SWidget> MenuContent;
 	if(SummonedContextMenu)
 	{
-		ContextMenuWindow =
-			FSlateApplication::Get().PushMenu( SharedThis( this ), MenuBuilder.MakeWidget(), CursorPos, FPopupTransitionEffect( FPopupTransitionEffect::ContextMenu ));
+		MenuContent = MenuBuilder.MakeWidget();
+
+		FWidgetPath WidgetPath = MouseEvent.GetEventPath() != nullptr ? *MouseEvent.GetEventPath() : FWidgetPath();
+		FSlateApplication::Get().PushMenu(SharedThis(this), WidgetPath, MenuContent.ToSharedRef(), CursorPos, FPopupTransitionEffect(FPopupTransitionEffect::ContextMenu));
 	}
 
-	return ContextMenuWindow;
+	return MenuContent;
 }
 
 /** Return true if we should snap and set out position to its position */
diff --git a/Engine/Source/Editor/PhAT/Private/PhATPreviewViewportClient.cpp b/Engine/Source/Editor/PhAT/Private/PhATPreviewViewportClient.cpp
--- a/Engine/Source/Editor/PhAT/Private/PhATPreviewViewportClient.cpp
+++ b/Engine/Source/Editor/PhAT/Private/PhATPreviewViewportClient.cpp
@@ -710,6 +710,7 @@ void FPhATEdPreviewViewportClient::OpenBodyMenu()
 
 		FSlateApplication::Get().PushMenu(
 			ParentWidget.ToSharedRef(),
+			FWidgetPath(),
 			MenuWidget.ToSharedRef(),
 			MouseCursorLocation,
 			FPopupTransitionEffect(FPopupTransitionEffect::ContextMenu)
@@ -728,6 +729,7 @@ void FPhATEdPreviewViewportClient::OpenConstraintMenu()
 
 		FSlateApplication::Get().PushMenu(
 			ParentWidget.ToSharedRef(),
+			FWidgetPath(),
 			MenuWidget.ToSharedRef(),
 			MouseCursorLocation,
 			FPopupTransitionEffect(FPopupTransitionEffect::ContextMenu)
diff --git a/Engine/Source/Editor/PhAT/Private/SPhATPreviewToolbar.cpp b/Engine/Source/Editor/PhAT/Private/SPhATPreviewToolbar.cpp
--- a/Engine/Source/Editor/PhAT/Private/SPhATPreviewToolbar.cpp
+++ b/Engine/Source/Editor/PhAT/Private/SPhATPreviewToolbar.cpp
@@ -10,6 +10,7 @@
 #include "SPhATPreviewToolbar.h"
 #include "Editor/UnrealEd//Public/SEditorViewport.h"
 #include "SPhatPreviewViewport.h"
+#include "PhATPreviewViewportClient.h"
 
 #define LOCTEXT_NAMESPACE "PhatViewportToolBar"
 
diff --git a/Engine/Source/Editor/PhAT/Private/SPhATPreviewViewport.h b/Engine/Source/Editor/PhAT/Private/SPhATPreviewViewport.h
--- a/Engine/Source/Editor/PhAT/Private/SPhATPreviewViewport.h
+++ b/Engine/Source/Editor/PhAT/Private/SPhATPreviewViewport.h
@@ -5,6 +5,7 @@
 #include "SEditorViewport.h"
 
 class FPhAT;
+class FPhATEdPreviewViewportClient;
 
 /*-----------------------------------------------------------------------------
    SPhATViewport
diff --git a/Engine/Source/Editor/PropertyEditor/Private/SDetailTableRowBase.h b/Engine/Source/Editor/PropertyEditor/Private/SDetailTableRowBase.h
--- a/Engine/Source/Editor/PropertyEditor/Private/SDetailTableRowBase.h
+++ b/Engine/Source/Editor/PropertyEditor/Private/SDetailTableRowBase.h
@@ -34,7 +34,9 @@ public:
 
 			if( bShouldOpenMenu )
 			{
-				FSlateApplication::Get().PushMenu(AsShared(), MenuBuilder.MakeWidget(), MouseEvent.GetScreenSpacePosition(), FPopupTransitionEffect::ContextMenu);
+				FWidgetPath WidgetPath = MouseEvent.GetEventPath() != nullptr ? *MouseEvent.GetEventPath() : FWidgetPath();
+
+				FSlateApplication::Get().PushMenu(AsShared(), WidgetPath, MenuBuilder.MakeWidget(), MouseEvent.GetScreenSpacePosition(), FPopupTransitionEffect::ContextMenu);
 
 				return FReply::Handled();
 			}
diff --git a/Engine/Source/Editor/SceneOutliner/Private/ActorTreeItem.cpp b/Engine/Source/Editor/SceneOutliner/Private/ActorTreeItem.cpp
--- a/Engine/Source/Editor/SceneOutliner/Private/ActorTreeItem.cpp
+++ b/Engine/Source/Editor/SceneOutliner/Private/ActorTreeItem.cpp
@@ -139,6 +139,7 @@ void FActorDropTarget::OnDrop(FDragDropPayload& DraggedObjects, UWorld& World, c
 			// Create the popup
 			FSlateApplication::Get().PushMenu(
 				DroppedOnWidget,
+				FWidgetPath(),
 				SNew(SSocketChooserPopup)
 				.SceneComponent( Component )
 				.OnSocketChosen_Static(&FActorDropTarget::PerformAttachment, Actor, MoveTemp(DraggedActors) ),
diff --git a/Engine/Source/Editor/Sequencer/Private/SAnimationOutlinerView.cpp b/Engine/Source/Editor/Sequencer/Private/SAnimationOutlinerView.cpp
--- a/Engine/Source/Editor/Sequencer/Private/SAnimationOutlinerView.cpp
+++ b/Engine/Source/Editor/Sequencer/Private/SAnimationOutlinerView.cpp
@@ -111,8 +111,11 @@ FReply SAnimationOutlinerTreeNode::OnMouseButtonUp( const FGeometry& MyGeometry,
 		TSharedPtr<SWidget> MenuContent = DisplayNode->OnSummonContextMenu(MyGeometry, MouseEvent);
 		if (MenuContent.IsValid())
 		{
+			FWidgetPath WidgetPath = MouseEvent.GetEventPath() != nullptr ? *MouseEvent.GetEventPath() : FWidgetPath();
+
 			FSlateApplication::Get().PushMenu(
 				AsShared(),
+				WidgetPath,
 				MenuContent.ToSharedRef(),
 				MouseEvent.GetScreenSpacePosition(),
 				FPopupTransitionEffect( FPopupTransitionEffect::ContextMenu )
diff --git a/Engine/Source/Editor/Sequencer/Private/SSection.cpp b/Engine/Source/Editor/Sequencer/Private/SSection.cpp
--- a/Engine/Source/Editor/Sequencer/Private/SSection.cpp
+++ b/Engine/Source/Editor/Sequencer/Private/SSection.cpp
@@ -555,8 +555,11 @@ FReply SSection::OnMouseButtonUp( const FGeometry& MyGeometry, const FPointerEve
 			TSharedPtr<SWidget> MenuContent = OnSummonContextMenu( MyGeometry, MouseEvent );
 			if (MenuContent.IsValid())
 			{
+				FWidgetPath WidgetPath = MouseEvent.GetEventPath() != nullptr ? *MouseEvent.GetEventPath() : FWidgetPath();
+
 				FSlateApplication::Get().PushMenu(
 					AsShared(),
+					WidgetPath,
 					MenuContent.ToSharedRef(),
 					MouseEvent.GetScreenSpacePosition(),
 					FPopupTransitionEffect( FPopupTransitionEffect::ContextMenu )
diff --git a/Engine/Source/Editor/Sequencer/Private/Sequencer.cpp b/Engine/Source/Editor/Sequencer/Private/Sequencer.cpp
--- a/Engine/Source/Editor/Sequencer/Private/Sequencer.cpp
+++ b/Engine/Source/Editor/Sequencer/Private/Sequencer.cpp
@@ -36,6 +36,7 @@
 #include "GenericCommands.h"
 #include "Engine/BlueprintGeneratedClass.h"
 #include "Engine/Selection.h"
+#include "IMenu.h"
 
 #define LOCTEXT_NAMESPACE "Sequencer"
 
@@ -337,8 +338,9 @@ void FSequencer::RenameShot(UMovieSceneSection* ShotSection)
 		.OnTextCommitted(this, &FSequencer::RenameShotCommitted, ShotSection)
 		.ClearKeyboardFocusOnCommit( false );
 	
-	NameEntryPopupWindow = FSlateApplication::Get().PushMenu(
+	NameEntryPopupMenu = FSlateApplication::Get().PushMenu(
 		SequencerWidget.ToSharedRef(),
+		FWidgetPath(),
 		TextEntry,
 		FSlateApplication::Get().GetCursorPos(),
 		FPopupTransitionEffect( FPopupTransitionEffect::TypeInPopup )
@@ -402,9 +404,9 @@ void FSequencer::RenameShotCommitted(const FText& RenameText, ETextCommit::Type
 		ShotSection->SetTitle(RenameText);
 	}
 
-	if (NameEntryPopupWindow.IsValid())
+	if (NameEntryPopupMenu.IsValid())
 	{
-		NameEntryPopupWindow.Pin()->RequestDestroyWindow();
+		NameEntryPopupMenu.Pin()->Dismiss();
 	}
 }
 
diff --git a/Engine/Source/Editor/Sequencer/Private/Sequencer.h b/Engine/Source/Editor/Sequencer/Private/Sequencer.h
--- a/Engine/Source/Editor/Sequencer/Private/Sequencer.h
+++ b/Engine/Source/Editor/Sequencer/Private/Sequencer.h
@@ -353,7 +353,7 @@ private:
 	TSharedPtr< class SSequencer > SequencerWidget;
 	
 	/** Reference to owner of the current popup */
-	TWeakPtr<class SWindow> NameEntryPopupWindow;
+	TWeakPtr<class IMenu> NameEntryPopupMenu;
 
 	/** The asset editor that created this Sequencer if any */
 	TWeakPtr<IToolkitHost> ToolkitHost;
diff --git a/Engine/Source/Editor/StatsViewer/Private/StatsPages/StaticMeshLightingInfoStatsPage.cpp b/Engine/Source/Editor/StatsViewer/Private/StatsPages/StaticMeshLightingInfoStatsPage.cpp
--- a/Engine/Source/Editor/StatsViewer/Private/StatsPages/StaticMeshLightingInfoStatsPage.cpp
+++ b/Engine/Source/Editor/StatsViewer/Private/StatsPages/StaticMeshLightingInfoStatsPage.cpp
@@ -9,6 +9,7 @@
 #include "ScopedTransaction.h"
 #include "STextEntryPopup.h"
 #include "Engine/LevelStreaming.h"
+#include "IMenu.h"
 
 #define LOCTEXT_NAMESPACE "Editor.StatsViewer.StaticMeshLightingInfo"
 
@@ -462,9 +463,9 @@ void FStaticMeshLightingInfoStatsPage::OnResolutionCommitted(const FText& Resolu
 	}
 
 	// Remove the popup window
-	if (ResolutionEntryPopupWindow.IsValid())
+	if (ResolutionEntryMenu.IsValid())
 	{
-		ResolutionEntryPopupWindow->RequestDestroyWindow();
+		ResolutionEntryMenu.Pin()->Dismiss();
 	}
 
 	// Do we have a valid resolution to use
@@ -497,8 +498,9 @@ void FStaticMeshLightingInfoStatsPage::GetUserSetStaticLightmapResolution(TWeakP
 			.OnTextCommitted( FOnTextCommitted::CreateSP( this, &FStaticMeshLightingInfoStatsPage::OnResolutionCommitted, InParentStatsViewer, bSwap ) )
 			.ClearKeyboardFocusOnCommit( false );
 
-		ResolutionEntryPopupWindow = FSlateApplication::Get().PushMenu(
+		ResolutionEntryMenu = FSlateApplication::Get().PushMenu(
 			InParentStatsViewer.Pin().ToSharedRef(),
+			FWidgetPath(),
 			TextEntry,
 			FSlateApplication::Get().GetCursorPos(),
 			FPopupTransitionEffect( FPopupTransitionEffect::TypeInPopup )
diff --git a/Engine/Source/Editor/StatsViewer/Private/StatsPages/StaticMeshLightingInfoStatsPage.h b/Engine/Source/Editor/StatsViewer/Private/StatsPages/StaticMeshLightingInfoStatsPage.h
--- a/Engine/Source/Editor/StatsViewer/Private/StatsPages/StaticMeshLightingInfoStatsPage.h
+++ b/Engine/Source/Editor/StatsViewer/Private/StatsPages/StaticMeshLightingInfoStatsPage.h
@@ -124,6 +124,6 @@ private:
 	TSharedPtr<SWidget> CustomWidget;
 
 	/** Reference to owner of the current popup */
-	mutable TSharedPtr<class SWindow> ResolutionEntryPopupWindow;
+	mutable TWeakPtr<class IMenu> ResolutionEntryMenu;
 };
 
diff --git a/Engine/Source/Editor/UMGEditor/Private/Designer/SDesignerView.cpp b/Engine/Source/Editor/UMGEditor/Private/Designer/SDesignerView.cpp
--- a/Engine/Source/Editor/UMGEditor/Private/Designer/SDesignerView.cpp
+++ b/Engine/Source/Editor/UMGEditor/Private/Designer/SDesignerView.cpp
@@ -1419,7 +1419,8 @@ void SDesignerView::ShowContextMenu(const FGeometry& MyGeometry, const FPointerE
 	if ( MenuContent.IsValid() )
 	{
 		FVector2D SummonLocation = MouseEvent.GetScreenSpacePosition();
-		FSlateApplication::Get().PushMenu(AsShared(), MenuContent.ToSharedRef(), SummonLocation, FPopupTransitionEffect(FPopupTransitionEffect::ContextMenu));
+		FWidgetPath WidgetPath = MouseEvent.GetEventPath() != nullptr ? *MouseEvent.GetEventPath() : FWidgetPath();
+		FSlateApplication::Get().PushMenu(AsShared(), WidgetPath, MenuContent.ToSharedRef(), SummonLocation, FPopupTransitionEffect(FPopupTransitionEffect::ContextMenu));
 	}
 }
 
diff --git a/Engine/Source/Editor/UnrealEd/Private/ComponentVisualizerManager.cpp b/Engine/Source/Editor/UnrealEd/Private/ComponentVisualizerManager.cpp
--- a/Engine/Source/Editor/UnrealEd/Private/ComponentVisualizerManager.cpp
+++ b/Engine/Source/Editor/UnrealEd/Private/ComponentVisualizerManager.cpp
@@ -20,6 +20,7 @@ bool FComponentVisualizerManager::HandleClick(FLevelEditorViewportClient* InView
 				TSharedPtr<SWidget> MenuParent = ParentLevelEditor->GetActiveViewportInterface()->GetViewportWidget().Pin();
 				FSlateApplication::Get().PushMenu(
 					MenuParent.ToSharedRef(),
+					FWidgetPath(),
 					MenuWidget.ToSharedRef(),
 					FSlateApplication::Get().GetCursorPos(),
 					FPopupTransitionEffect(FPopupTransitionEffect::ContextMenu));
diff --git a/Engine/Source/Editor/UnrealEd/Private/Kismet2/BlueprintEditorUtils.cpp b/Engine/Source/Editor/UnrealEd/Private/Kismet2/BlueprintEditorUtils.cpp
--- a/Engine/Source/Editor/UnrealEd/Private/Kismet2/BlueprintEditorUtils.cpp
+++ b/Engine/Source/Editor/UnrealEd/Private/Kismet2/BlueprintEditorUtils.cpp
@@ -2108,8 +2108,8 @@ void FBlueprintEditorUtils::MarkBlueprintAsStructurallyModified(UBlueprint* Blue
 		{
 			if (!Blueprint->bIsRegeneratingOnLoad)
 			{
-				GetDerivedClasses(SkelClass, ChildrenOfClass, false);
-			}
+			GetDerivedClasses(SkelClass, ChildrenOfClass, false);
+		}
 		}
 
 		{
@@ -5561,8 +5561,8 @@ void FBlueprintEditorUtils::ConformImplementedEvents(UBlueprint* Blueprint)
 								// Fix up the event signature
 								if (TargetFunction != nullptr)
 								{
-									EventNode->EventReference.SetFromField<UFunction>(TargetFunction, false);
-								}
+								EventNode->EventReference.SetFromField<UFunction>(TargetFunction, false);
+							}
 								else
 								{
 									EventNode->EventReference.SetExternalMember(EventFuncName, Blueprint->GeneratedClass);
@@ -6854,23 +6854,26 @@ void FBlueprintEditorUtils::OpenReparentBlueprintMenu( const TArray< UBlueprint*
 
 	TSharedRef<SWidget> ClassPicker = ConstructBlueprintParentClassPicker(Blueprints, OnPicked);
 
-	TSharedRef<SBorder> ClassPickerBorder = 
-		SNew(SBorder)
-		.BorderImage(FEditorStyle::GetBrush("Menu.Background"))
+	TSharedRef<SBox> ClassPickerBox = 
+		SNew(SBox)
+		.WidthOverride(280)
+		.HeightOverride(400)
 		[
-			ClassPicker
+			SNew(SBorder)
+			.BorderImage(FEditorStyle::GetBrush("Menu.Background"))
+			[
+				ClassPicker
+			]
 		];
 
 	// Show dialog to choose new parent class
 	FSlateApplication::Get().PushMenu(
 		ParentContent,
-		ClassPickerBorder,
+		FWidgetPath(),
+		ClassPickerBox,
 		FSlateApplication::Get().GetCursorPos(),
 		FPopupTransitionEffect( FPopupTransitionEffect::ContextMenu),
-		true,
-		false,
-		FVector2D(280, 400)
-		);
+		true);
 }
 
 /** Filter class for ClassPicker handling allowed interfaces for a Blueprint */
diff --git a/Engine/Source/Editor/UnrealEd/Private/Kismet2/KismetDebugUtilities.cpp b/Engine/Source/Editor/UnrealEd/Private/Kismet2/KismetDebugUtilities.cpp
--- a/Engine/Source/Editor/UnrealEd/Private/Kismet2/KismetDebugUtilities.cpp
+++ b/Engine/Source/Editor/UnrealEd/Private/Kismet2/KismetDebugUtilities.cpp
@@ -306,6 +306,7 @@ void FKismetDebugUtilities::OnScriptException(const UObject* ActiveObject, const
 
 							FSlateApplication::Get().PushMenu(
 								FSlateApplication::Get().GetActiveTopLevelWindow().ToSharedRef(),
+								FWidgetPath(),
 								DisplayWidget,
 								FSlateApplication::Get().GetCursorPos(),
 								FPopupTransitionEffect(FPopupTransitionEffect::TypeInPopup)
diff --git a/Engine/Source/Editor/UnrealEd/Private/SColorGradientEditor.cpp b/Engine/Source/Editor/UnrealEd/Private/SColorGradientEditor.cpp
--- a/Engine/Source/Editor/UnrealEd/Private/SColorGradientEditor.cpp
+++ b/Engine/Source/Editor/UnrealEd/Private/SColorGradientEditor.cpp
@@ -415,7 +415,7 @@ FReply SColorGradientEditor::OnMouseButtonUp( const FGeometry& MyGeometry, const
 			// Didnt move the mouse too far, open a context menu
 			if( DistanceDragged < DragThresholdDist && SelectedStop.IsValid( *CurveOwner ) )
 			{
-				OpenGradientStopContextMenu( MouseEvent.GetScreenSpacePosition() );
+				OpenGradientStopContextMenu( MouseEvent );
 			}
 
 			DistanceDragged = 0;
@@ -449,8 +449,10 @@ void SColorGradientEditor::SetCurveOwner( FCurveOwnerInterface* InCurveOwner )
 }
 
 
-void SColorGradientEditor::OpenGradientStopContextMenu( FVector2D Location )
+void SColorGradientEditor::OpenGradientStopContextMenu(const FPointerEvent& MouseEvent)
 {
+	const FVector2D& Location = MouseEvent.GetScreenSpacePosition();
+
 	FMenuBuilder GradientStopMenu( true, NULL );
 
 	
@@ -522,7 +524,8 @@ void SColorGradientEditor::OpenGradientStopContextMenu( FVector2D Location )
 		// Add a Remove option
 		GradientStopMenu.AddMenuEntry( LOCTEXT("RemoveGradientStop", "Remove Stop"), LOCTEXT("RemoveGradientStopTooltip", "Removes the selected gradient stop"), FSlateIcon(), RemoveStopAction );
 
-		FSlateApplication::Get().PushMenu( AsShared(), GradientStopMenu.MakeWidget(), Location, FPopupTransitionEffect::ContextMenu );
+		FWidgetPath WidgetPath = MouseEvent.GetEventPath() != nullptr ? *MouseEvent.GetEventPath() : FWidgetPath();
+		FSlateApplication::Get().PushMenu(AsShared(), WidgetPath, GradientStopMenu.MakeWidget(), Location, FPopupTransitionEffect::ContextMenu);
 
 		FSlateApplication::Get().SetKeyboardFocus( WidgetToFocus.ToSharedRef() );
 	}
@@ -569,7 +572,7 @@ void SColorGradientEditor::OpenGradientStopColorPicker()
 			]
 		];
 
-		FSlateApplication::Get().PushMenu( SharedThis( this ), AlphaSlider, ContextMenuPosition, FPopupTransitionEffect::TypeInPopup );
+		FSlateApplication::Get().PushMenu( SharedThis( this ), FWidgetPath(), AlphaSlider, ContextMenuPosition, FPopupTransitionEffect::TypeInPopup );
 	}
 	else
 	{
diff --git a/Engine/Source/Editor/UnrealEd/Private/SCurveEditor.cpp b/Engine/Source/Editor/UnrealEd/Private/SCurveEditor.cpp
--- a/Engine/Source/Editor/UnrealEd/Private/SCurveEditor.cpp
+++ b/Engine/Source/Editor/UnrealEd/Private/SCurveEditor.cpp
@@ -469,14 +469,17 @@ void SCurveEditor::PushWarningMenu( FVector2D Position, const FText& Message )
 
 	FSlateApplication::Get().PushMenu(
 		SharedThis( this ),
+		FWidgetPath(),
 		WarningMessageText->AsWidget(),
 		Position,
 		FPopupTransitionEffect( FPopupTransitionEffect::ContextMenu));
 }
 
 
-void SCurveEditor::PushKeyMenu(const FGeometry& InMyGeometry, FVector2D Position)
+void SCurveEditor::PushKeyMenu(const FGeometry& InMyGeometry, const FPointerEvent& InMouseEvent)
 {
+	FVector2D Position = InMouseEvent.GetScreenSpacePosition();
+
 	FMenuBuilder MenuBuilder(true, Commands.ToSharedRef());
 
 	if (SelectedKeys.Num() == 0)
@@ -503,8 +506,11 @@ void SCurveEditor::PushKeyMenu(const FGeometry& InMyGeometry, FVector2D Position
 	}
 	MenuBuilder.EndSection(); //CurveEditorInterpolation
 
+	FWidgetPath WidgetPath = InMouseEvent.GetEventPath() != nullptr ? *InMouseEvent.GetEventPath() : FWidgetPath();
+
 	FSlateApplication::Get().PushMenu(
 		SharedThis( this ),
+		WidgetPath,
 		MenuBuilder.MakeWidget(),
 		Position,
 		FPopupTransitionEffect( FPopupTransitionEffect::ContextMenu));
@@ -1494,58 +1500,58 @@ void SCurveEditor::ProcessClick(const FGeometry& InMyGeometry, const FPointerEve
 	const bool bShiftDown = InMouseEvent.IsShiftDown();
 
 	FSelectedCurveKey HitKey = HitTestKeys(InMyGeometry, InMouseEvent.GetScreenSpacePosition());
-		if (bLeftMouseButton)
+
+	if (bLeftMouseButton)
+	{
+		// If the user left clicked a key, update selection based on modifier key state.
+		if (HitKey.IsValid())
 		{
-			// If the user left clicked a key, update selection based on modifier key state.
-			if (HitKey.IsValid())
+			if (!IsKeySelected(HitKey))
 			{
-				if (!IsKeySelected(HitKey))
+				if (!bControlDown)
 				{
-					if (!bControlDown)
-					{
-						EmptySelection();
-					}
-					AddToSelection(HitKey);
-				}
-				else if (bControlDown)
-				{
-					RemoveFromSelection(HitKey);
+					EmptySelection();
 				}
+				AddToSelection(HitKey);
 			}
-			else
+			else if (bControlDown)
 			{
-				// If the user didn't click a key, add a new one if shift is held down, or try to select a curve.
-				if (bShiftDown)
-				{
-					AddNewKey(InMyGeometry, InMouseEvent.GetScreenSpacePosition());
-				}
-				else
-				{
-					// clicking on background clears key selection
-					EmptySelection();
-				}
+				RemoveFromSelection(HitKey);
 			}
 		}
-		else if (bRightMouseButton)
+		else
 		{
-			// If the user right clicked, handle opening context menus.
-			if (HitKey.IsValid())
+			// If the user didn't click a key, add a new one if shift is held down, or try to select a curve.
+			if (bShiftDown)
 			{
-				// Make sure key is selected in readiness for context menu
-				if (!IsKeySelected(HitKey))
-				{
-					EmptySelection();
-					AddToSelection(HitKey);
-				}
-			PushKeyMenu(InMyGeometry, InMouseEvent.GetScreenSpacePosition());
+				AddNewKey(InMyGeometry, InMouseEvent.GetScreenSpacePosition());
 			}
 			else
 			{
-				FVector2D ScreenPos = InMouseEvent.GetScreenSpacePosition();
-				if (!HitTestCubicTangents(InMyGeometry, ScreenPos).IsValid())
-				{
-					CreateContextMenu(InMyGeometry, ScreenPos);
-				}
+				// clicking on background clears key selection
+				EmptySelection();
+			}
+		}
+	}
+	else if (bRightMouseButton)
+	{
+		// If the user right clicked, handle opening context menus.
+		if (HitKey.IsValid())
+		{
+			// Make sure key is selected in readiness for context menu
+			if (!IsKeySelected(HitKey))
+			{
+				EmptySelection();
+				AddToSelection(HitKey);
+			}
+			PushKeyMenu(InMyGeometry, InMouseEvent);
+		}
+		else
+		{
+			if (!HitTestCubicTangents(InMyGeometry, InMouseEvent.GetScreenSpacePosition()).IsValid())
+			{
+				CreateContextMenu(InMyGeometry, InMouseEvent);
+			}
 			else
 			{
 				EmptySelection();
@@ -2018,8 +2024,10 @@ bool SCurveEditor::IsSnappingEnabled()
 	return bSnappingEnabled.Get();
 }
 
-void SCurveEditor::CreateContextMenu(const FGeometry& InMyGeometry, const FVector2D& ScreenPosition)
+void SCurveEditor::CreateContextMenu(const FGeometry& InMyGeometry, const FPointerEvent& InMouseEvent)
 {
+	const FVector2D& ScreenPosition = InMouseEvent.GetScreenSpacePosition();
+
 	const bool CloseAfterSelection = true;
 	FMenuBuilder MenuBuilder( CloseAfterSelection, NULL );
 
@@ -2080,7 +2088,9 @@ void SCurveEditor::CreateContextMenu(const FGeometry& InMyGeometry, const FVecto
 	}
 
 	MenuBuilder.EndSection();
-	FSlateApplication::Get().PushMenu( SharedThis( this ), MenuBuilder.MakeWidget(), FSlateApplication::Get().GetCursorPos(), FPopupTransitionEffect( FPopupTransitionEffect::ContextMenu ) );
+
+	FWidgetPath WidgetPath = InMouseEvent.GetEventPath() != nullptr ? *InMouseEvent.GetEventPath() : FWidgetPath();
+	FSlateApplication::Get().PushMenu(SharedThis(this), WidgetPath, MenuBuilder.MakeWidget(), FSlateApplication::Get().GetCursorPos(), FPopupTransitionEffect(FPopupTransitionEffect::ContextMenu));
 }
 
 void SCurveEditor::OnCreateExternalCurveClicked()
diff --git a/Engine/Source/Editor/UnrealEd/Private/Toolkits/GlobalEditorCommonCommands.cpp b/Engine/Source/Editor/UnrealEd/Private/Toolkits/GlobalEditorCommonCommands.cpp
--- a/Engine/Source/Editor/UnrealEd/Private/Toolkits/GlobalEditorCommonCommands.cpp
+++ b/Engine/Source/Editor/UnrealEd/Private/Toolkits/GlobalEditorCommonCommands.cpp
@@ -5,9 +5,9 @@
 #include "OutputLogModule.h"
 #include "LevelEditor.h"
 #include "SDockTab.h"
-
 #include "SGlobalOpenAssetDialog.h"
 #include "SGlobalTabSwitchingDialog.h"
+#include "IMenu.h"
 
 #define LOCTEXT_NAMESPACE "GlobalEditorCommonCommands"
 
@@ -69,7 +69,7 @@ void FGlobalEditorCommonCommands::OnPressedCtrlTab(TSharedPtr<FUICommandInfo> Tr
 		// Create the contents of the popup
 		TSharedRef<SWidget> ActualWidget = SNew(SGlobalTabSwitchingDialog, TabListSize, *TriggeringCommand->GetActiveChord());
 
-		TSharedPtr<SWindow> NewWindow = OpenPopup(ActualWidget, TabListSize);
+		OpenPopupMenu(ActualWidget, TabListSize);
 	}
 }
 
@@ -86,10 +86,20 @@ void FGlobalEditorCommonCommands::OnSummonedAssetPicker()
 	MenuBuilder.AddWidget(ActualWidget, FText::GetEmpty(), /*bNoIndent=*/ true);
 	MenuBuilder.EndSection();
 
-	OpenPopup(MenuBuilder.MakeWidget(), AssetPickerSize);
+	OpenPopupMenu(MenuBuilder.MakeWidget(), AssetPickerSize);
 }
 
 TSharedPtr<SWindow> FGlobalEditorCommonCommands::OpenPopup(TSharedRef<SWidget> WindowContents, const FVector2D& PopupDesiredSize)
+{
+	TSharedPtr<IMenu> Menu = OpenPopupMenu(WindowContents, PopupDesiredSize);
+	if (Menu.IsValid())
+	{
+		return Menu->GetOwnedWindow();
+	}
+	return TSharedPtr<SWindow>();
+}
+
+TSharedPtr<IMenu> FGlobalEditorCommonCommands::OpenPopupMenu(TSharedRef<SWidget> WindowContents, const FVector2D& PopupDesiredSize)
 {
 	// Determine where the pop-up should open
 	TSharedPtr<SWindow> ParentWindow = FSlateApplication::Get().GetActiveTopLevelWindow();
@@ -109,12 +119,10 @@ TSharedPtr<SWindow> FGlobalEditorCommonCommands::OpenPopup(TSharedRef<SWidget> W
 
 		// Open the pop-up
 		FPopupTransitionEffect TransitionEffect(FPopupTransitionEffect::None);
-		TSharedRef<SWindow> PopupWindow = FSlateApplication::Get().PushMenu(ParentWindow.ToSharedRef(), WindowContents, WindowPosition, TransitionEffect, /*bFocusImmediately=*/ true);
-		
-		return PopupWindow;
+		return FSlateApplication::Get().PushMenu(ParentWindow.ToSharedRef(), FWidgetPath(), WindowContents, WindowPosition, TransitionEffect, /*bFocusImmediately=*/ true);
 	}
 
-	return TSharedPtr<SWindow>();
+	return TSharedPtr<IMenu>();
 }
 
 static void CloseDebugConsole()
diff --git a/Engine/Source/Editor/UnrealEd/Private/UnrealEdSrv.cpp b/Engine/Source/Editor/UnrealEd/Private/UnrealEdSrv.cpp
--- a/Engine/Source/Editor/UnrealEd/Private/UnrealEdSrv.cpp
+++ b/Engine/Source/Editor/UnrealEd/Private/UnrealEdSrv.cpp
@@ -212,7 +212,8 @@ private:
 			}
 			NewMenu.EndSection();
 
-			FSlateApplication::Get().PushMenu( SharedThis( this ), NewMenu.MakeWidget(), MouseEvent.GetScreenSpacePosition(), FPopupTransitionEffect( FPopupTransitionEffect::None ) );
+			FWidgetPath WidgetPath = MouseEvent.GetEventPath() != nullptr ? *MouseEvent.GetEventPath() : FWidgetPath();
+			FSlateApplication::Get().PushMenu(SharedThis(this), WidgetPath, NewMenu.MakeWidget(), MouseEvent.GetScreenSpacePosition(), FPopupTransitionEffect(FPopupTransitionEffect::None));
 
 			return FReply::Handled();
 		}
diff --git a/Engine/Source/Editor/UnrealEd/Public/SColorGradientEditor.h b/Engine/Source/Editor/UnrealEd/Public/SColorGradientEditor.h
--- a/Engine/Source/Editor/UnrealEd/Public/SColorGradientEditor.h
+++ b/Engine/Source/Editor/UnrealEd/Public/SColorGradientEditor.h
@@ -64,9 +64,9 @@ private:
 	/**
 	 * Opens a context menu with options for the selected gradient stop
 	 *
-	 * @param Location	The location in screen space where the menu should appear
+	 * @param MouseEvent	The mouse event that raised this handler (containing the location and widget path of the click)
 	 */
-	void OpenGradientStopContextMenu( FVector2D Location );
+	void OpenGradientStopContextMenu(const FPointerEvent& MouseEvent);
 
 	/**
 	 * Opens a color picker to change the color of the selected stop
diff --git a/Engine/Source/Editor/UnrealEd/Public/SCurveEditor.h b/Engine/Source/Editor/UnrealEd/Public/SCurveEditor.h
--- a/Engine/Source/Editor/UnrealEd/Public/SCurveEditor.h
+++ b/Engine/Source/Editor/UnrealEd/Public/SCurveEditor.h
@@ -352,7 +352,7 @@ private:
 	bool GetInputEditEnabled() const;
 
 	/** Function to create context menu on mouse right click*/
-	void CreateContextMenu(const FGeometry& InMyGeometry, const FVector2D& ScreenPosition);
+	void CreateContextMenu(const FGeometry& InMyGeometry, const FPointerEvent& InMouseEvent);
 
 	/** Callback function called when item is select in the context menu */
 	void OnCreateExternalCurveClicked();
@@ -429,7 +429,7 @@ private:
 	void	PushWarningMenu(FVector2D Position, const FText& Message);
 
 	/** Create context Menu for key interpolation settings*/
-	void	PushKeyMenu(const FGeometry& InMyGeometry, FVector2D Position);
+	void	PushKeyMenu(const FGeometry& InMyGeometry, const FPointerEvent& InMouseEvent);
 
 	/** Called when the user selects the interpolation mode */
 	void	OnSelectInterpolationMode(ERichCurveInterpMode InterpMode, ERichCurveTangentMode TangentMode);
diff --git a/Engine/Source/Editor/UnrealEd/Public/Toolkits/GlobalEditorCommonCommands.h b/Engine/Source/Editor/UnrealEd/Public/Toolkits/GlobalEditorCommonCommands.h
--- a/Engine/Source/Editor/UnrealEd/Public/Toolkits/GlobalEditorCommonCommands.h
+++ b/Engine/Source/Editor/UnrealEd/Public/Toolkits/GlobalEditorCommonCommands.h
@@ -19,7 +19,10 @@ protected:
 	static void OnSummonedAssetPicker();
 	static void OnSummonedConsoleCommandBox();
 
+	DEPRECATED(4.9, "OpenPopup() returning a window is deprecated. Use OpenPopupMenu() that returns an IMenu.")
 	static TSharedPtr<SWindow> OpenPopup(TSharedRef<SWidget> WindowContents, const FVector2D& PopupDesiredSize);
+
+	static TSharedPtr<IMenu> OpenPopupMenu(TSharedRef<SWidget> WindowContents, const FVector2D& PopupDesiredSize);
 public:
 	TSharedPtr<FUICommandInfo> FindInContentBrowser;
 	TSharedPtr<FUICommandInfo> SummonControlTabNavigation;
diff --git a/Engine/Source/Editor/WorldBrowser/Private/Tiles/SWorldComposition.cpp b/Engine/Source/Editor/WorldBrowser/Private/Tiles/SWorldComposition.cpp
--- a/Engine/Source/Editor/WorldBrowser/Private/Tiles/SWorldComposition.cpp
+++ b/Engine/Source/Editor/WorldBrowser/Private/Tiles/SWorldComposition.cpp
@@ -9,6 +9,7 @@
 #include "SWorldTileItem.h"
 #include "SWorldLayers.h"
 #include "WorldTileCollectionModel.h"
+#include "IMenu.h"
 
 #define LOCTEXT_NAMESPACE "WorldBrowser"
 
@@ -657,8 +658,11 @@ protected:
 		WorldModel->BuildWorldCompositionMenu(MenuBuilder);
 		TSharedPtr<SWidget> MenuWidget = MenuBuilder.MakeWidget();
 
+		FWidgetPath WidgetPath = MouseEvent.GetEventPath() != nullptr ? *MouseEvent.GetEventPath() : FWidgetPath();
+
 		FSlateApplication::Get().PushMenu(
 			AsShared(),
+			WidgetPath,
 			MenuWidget.ToSharedRef(),
 			MouseEvent.GetScreenSpacePosition(),
 			FPopupTransitionEffect( FPopupTransitionEffect::ContextMenu )
@@ -985,7 +989,7 @@ void SWorldComposition::OnBrowseWorld(UWorld* InWorld)
 	ContentParent->SetContent(SNullWidget::NullWidget);
 	LayersListWrapBox = nullptr;
 	NewLayerButton = nullptr;
-	NewLayerPopupWindow = nullptr;
+	NewLayerMenu.Reset();
 	GridView = nullptr;
 	TileWorldModel = nullptr;
 			
@@ -1210,8 +1214,9 @@ FReply SWorldComposition::NewLayer_Clicked()
 		.OnCreateLayer(this, &SWorldComposition::CreateNewLayer)
 		.DefaultName(LOCTEXT("Layer_DefaultName", "MyLayer").ToString());
 
-	NewLayerPopupWindow = FSlateApplication::Get().PushMenu(
+	NewLayerMenu = FSlateApplication::Get().PushMenu(
 		this->AsShared(),
+		FWidgetPath(),
 		CreateLayerWidget,
 		FSlateApplication::Get().GetCursorPos(),
 		FPopupTransitionEffect( FPopupTransitionEffect::TypeInPopup )
@@ -1225,9 +1230,9 @@ FReply SWorldComposition::CreateNewLayer(const FWorldTileLayer& NewLayer)
 	TileWorldModel->AddManagedLayer(NewLayer);
 	PopulateLayersList();
 	
-	if (NewLayerPopupWindow.IsValid())
+	if (NewLayerMenu.IsValid())
 	{
-		NewLayerPopupWindow->RequestDestroyWindow();
+		NewLayerMenu.Pin()->Dismiss();
 	}
 		
 	return FReply::Handled();
diff --git a/Engine/Source/Editor/WorldBrowser/Private/Tiles/SWorldComposition.h b/Engine/Source/Editor/WorldBrowser/Private/Tiles/SWorldComposition.h
--- a/Engine/Source/Editor/WorldBrowser/Private/Tiles/SWorldComposition.h
+++ b/Engine/Source/Editor/WorldBrowser/Private/Tiles/SWorldComposition.h
@@ -51,6 +51,6 @@ private:
 	TSharedPtr<SBorder>							ContentParent;
 	TSharedPtr<SWrapBox>						LayersListWrapBox;
 	TSharedPtr<SButton>							NewLayerButton;
-	TSharedPtr<class SWindow>					NewLayerPopupWindow;
+	TWeakPtr<class IMenu>						NewLayerMenu;
 	TSharedPtr<class SWorldCompositionGrid>		GridView;
 };
diff --git a/Engine/Source/Runtime/AppFramework/Private/Framework/Testing/STableViewTesting.cpp b/Engine/Source/Runtime/AppFramework/Private/Framework/Testing/STableViewTesting.cpp
--- a/Engine/Source/Runtime/AppFramework/Private/Framework/Testing/STableViewTesting.cpp
+++ b/Engine/Source/Runtime/AppFramework/Private/Framework/Testing/STableViewTesting.cpp
@@ -444,6 +444,7 @@ public:
 		FSlateApplication::Get().PushMenu
 		(
 			SharedThis( this ),
+			FWidgetPath(),
 			GetPopupContent(),
 			SpawnLocation,
 			FPopupTransitionEffect( FPopupTransitionEffect::ContextMenu )
@@ -603,6 +604,7 @@ public:
 		FSlateApplication::Get().PushMenu
 		(
 			SharedThis( this ),
+			FWidgetPath(),
 			GetPopupContent(),
 			SpawnLocation,
 			FPopupTransitionEffect( FPopupTransitionEffect::ContextMenu )
diff --git a/Engine/Source/Runtime/AppFramework/Private/Framework/Testing/SWidgetGallery.cpp b/Engine/Source/Runtime/AppFramework/Private/Framework/Testing/SWidgetGallery.cpp
--- a/Engine/Source/Runtime/AppFramework/Private/Framework/Testing/SWidgetGallery.cpp
+++ b/Engine/Source/Runtime/AppFramework/Private/Framework/Testing/SWidgetGallery.cpp
@@ -1083,6 +1083,7 @@ private:
 
 		FSlateApplication::Get().PushMenu(
 			SharedThis(this),
+			FWidgetPath(),
 			SNew(STextComboPopup)
 				.TextOptions(TextOptions)
 				.OnTextChosen(this, &SWidgetGallery::HandleTextComboPopupTextChosen),
diff --git a/Engine/Source/Runtime/AppFramework/Private/Widgets/Testing/STestSuite.cpp b/Engine/Source/Runtime/AppFramework/Private/Widgets/Testing/STestSuite.cpp
--- a/Engine/Source/Runtime/AppFramework/Private/Widgets/Testing/STestSuite.cpp
+++ b/Engine/Source/Runtime/AppFramework/Private/Widgets/Testing/STestSuite.cpp
@@ -23,6 +23,7 @@
 #include "SResponsiveGridPanel.h"
 #include "SColorPicker.h"
 #include "INotificationWidget.h"
+#include "IMenu.h"
 
 #define LOCTEXT_NAMESPACE "STestSuite"
 
@@ -1371,6 +1372,7 @@ struct RichTextHelper
 
 		FSlateApplication::Get().PushMenu(
 			ParentWidget, // Parent widget should be TestSuite, not the menu thats open or it will be closed when the menu is dismissed
+			FWidgetPath(),
 			Widget,
 			FSlateApplication::Get().GetCursorPos(), // summon location
 			FPopupTransitionEffect(FPopupTransitionEffect::ContextMenu)
@@ -2142,8 +2144,9 @@ public:
 			.OnTextCommitted( this, &STextEditTest::OnPopupTextCommitted ) 
 			.HintText( DefaultText );
 
-		PopupWindow = FSlateApplication::Get().PushMenu( 
+		PopupMenu = FSlateApplication::Get().PushMenu(
 			AsShared(), // Parent widget should be TestSyuite, not the menu thats open or it will be closed when the menu is dismissed
+			FWidgetPath(),
 			TextEntry,
 			FSlateApplication::Get().GetCursorPos(), // summon location
 			FPopupTransitionEffect( FPopupTransitionEffect::TypeInPopup )
@@ -2167,9 +2170,9 @@ public:
 		if ( (CommitInfo == ETextCommit::OnEnter) && (NewText.ToString().Len() == 3) )
 		{
 			// manually close menu on validated committal
-			if ( PopupWindow.IsValid() )
+			if (PopupMenu.IsValid())
 			{		
-				PopupWindow->RequestDestroyWindow();
+				PopupMenu.Pin()->Dismiss();
 			}
 		}
 	}
@@ -2193,7 +2196,7 @@ protected:
 
 	TSharedPtr<STextEntryPopup> PopupInput;
 
-	TSharedPtr<SWindow> PopupWindow;
+	TWeakPtr<IMenu> PopupMenu;
 
 	TSharedPtr<SInlineEditableTextBlock> InlineEditableTextBlock;
 	FText InlineEditableText;
diff --git a/Engine/Source/Runtime/Slate/Private/Framework/Application/Menu.cpp b/Engine/Source/Runtime/Slate/Private/Framework/Application/Menu.cpp
new file mode 100644
--- /dev/null
+++ b/Engine/Source/Runtime/Slate/Private/Framework/Application/Menu.cpp
@@ -0,0 +1,92 @@
+// Copyright 1998-2015 Epic Games, Inc. All Rights Reserved.
+
+#include "SlatePrivatePCH.h"
+#include "Menu.h"
+
+FMenuBase::FMenuBase(TSharedRef<SWidget> InContent)
+	: Content(InContent)
+	, bDismissing(false)
+{
+}
+
+FMenuInWindow::FMenuInWindow(TSharedRef<SWindow> InWindow, TSharedRef<SWidget> InContent)
+	: FMenuBase(InContent)
+	, Window(InWindow)
+{
+}
+
+TSharedPtr<SWindow> FMenuInWindow::GetParentWindow() const
+{
+	// Return the menu's window
+	return Window.Pin();
+}
+
+void FMenuInWindow::Dismiss()
+{
+	if (!bDismissing)
+	{
+		bDismissing = true;
+		OnMenuDismissed.Broadcast(AsShared());
+
+		// Close the window
+		// Will cause the window destroy code to call back into the stack to clean up
+		TSharedPtr<SWindow> WindowPinned = Window.Pin();
+		if (WindowPinned.IsValid())
+		{
+			WindowPinned->RequestDestroyWindow();
+		}
+	}
+}
+
+
+FMenuInPopup::FMenuInPopup(TSharedRef<SWidget> InContent)
+	: FMenuBase(InContent)
+{
+}
+
+TSharedPtr<SWindow> FMenuInPopup::GetParentWindow() const
+{
+	// Return the menu's window
+	return FSlateApplication::Get().GetVisibleMenuWindow();
+}
+
+void FMenuInPopup::Dismiss()
+{
+	if (!bDismissing)
+	{
+		bDismissing = true;
+		OnMenuDismissed.Broadcast(AsShared());
+	}
+}
+
+
+FMenuInHostWidget::FMenuInHostWidget(TSharedRef<IMenuHost> InHost, const TSharedRef<SWidget>& InContent)
+	: FMenuBase(InContent)
+	, MenuHost(InHost)
+{
+}
+
+TSharedPtr<SWindow> FMenuInHostWidget::GetParentWindow() const
+{
+	// Return the menu's window
+	TSharedPtr<IMenuHost> HostPinned = MenuHost.Pin();
+	if (HostPinned.IsValid())
+	{
+		return HostPinned->GetMenuWindow();
+	}
+	return TSharedPtr<SWindow>();
+}
+
+void FMenuInHostWidget::Dismiss()
+{
+	if (!bDismissing)
+	{
+		bDismissing = true;
+		TSharedPtr<IMenuHost> HostPinned = MenuHost.Pin();
+		if (HostPinned.IsValid())
+		{
+			HostPinned->OnMenuDismissed();
+		}
+		OnMenuDismissed.Broadcast(AsShared());
+	}
+}
diff --git a/Engine/Source/Runtime/Slate/Private/Framework/Application/Menu.h b/Engine/Source/Runtime/Slate/Private/Framework/Application/Menu.h
new file mode 100644
--- /dev/null
+++ b/Engine/Source/Runtime/Slate/Private/Framework/Application/Menu.h
@@ -0,0 +1,70 @@
+// Copyright 1998-2015 Epic Games, Inc. All Rights Reserved.
+
+#pragma once
+
+#include "IMenu.h"
+
+/**
+ * Represents the base class of popup menus.
+ */
+class FMenuBase : public IMenu, public TSharedFromThis<FMenuBase>
+{
+public:
+	virtual FOnMenuDismissed& GetOnMenuDismissed() override { return OnMenuDismissed; }
+	virtual TSharedPtr<SWidget> GetContent() const override { return Content; }
+
+protected:
+	FMenuBase(TSharedRef<SWidget> InContent);
+
+	FOnMenuDismissed OnMenuDismissed;
+	TSharedRef<SWidget> Content;
+	bool bDismissing;
+};
+
+/**
+ * Represents a popup menu that is shown in its own SWindow.
+ */
+class FMenuInWindow : public FMenuBase
+{
+public:
+	FMenuInWindow(TSharedRef<SWindow> InWindow, TSharedRef<SWidget> InContent);
+
+	virtual EPopupMethod GetPopupMethod() const override { return EPopupMethod::CreateNewWindow; }
+	virtual TSharedPtr<SWindow> GetParentWindow() const override;
+	virtual TSharedPtr<SWindow> GetOwnedWindow() const override { return GetParentWindow(); }
+	virtual void Dismiss() override;
+
+private:
+	TWeakPtr<SWindow> Window;
+};
+
+/**
+ * Represents a popup menu that is shown in the same window as the widget that summons it.
+ */
+class FMenuInPopup : public FMenuBase
+{
+public:
+	FMenuInPopup(TSharedRef<SWidget> InContent);
+
+	virtual EPopupMethod GetPopupMethod() const { return EPopupMethod::UseCurrentWindow; }
+	virtual TSharedPtr<SWindow> GetParentWindow() const;
+	virtual TSharedPtr<SWindow> GetOwnedWindow() const { return TSharedPtr<SWindow>(); }
+	virtual void Dismiss() override;
+};
+
+/**
+* Represents a popup menu that is shown in a host widget (such as a menu anchor).
+*/
+class FMenuInHostWidget : public FMenuBase
+{
+public:
+	FMenuInHostWidget(TSharedRef<IMenuHost> InHost, const TSharedRef<SWidget>& InContent);
+
+	virtual EPopupMethod GetPopupMethod() const { return EPopupMethod::UseCurrentWindow; }
+	virtual TSharedPtr<SWindow> GetParentWindow() const;
+	virtual TSharedPtr<SWindow> GetOwnedWindow() const { return TSharedPtr<SWindow>(); }
+	virtual void Dismiss() override;
+
+private:
+	TWeakPtr<IMenuHost> MenuHost;
+};
\ No newline at end of file
diff --git a/Engine/Source/Runtime/Slate/Private/Framework/Application/MenuStack.cpp b/Engine/Source/Runtime/Slate/Private/Framework/Application/MenuStack.cpp
--- a/Engine/Source/Runtime/Slate/Private/Framework/Application/MenuStack.cpp
+++ b/Engine/Source/Runtime/Slate/Private/Framework/Application/MenuStack.cpp
@@ -1,19 +1,72 @@
 // Copyright 1998-2015 Epic Games, Inc. All Rights Reserved.
 
 #include "SlatePrivatePCH.h"
+#include "Menu.h"
+#include "SPopup.h"
 
+#define LOCTEXT_NAMESPACE "MenuStack"
 
 namespace FMenuStackDefs
 {
 	/** Maximum size of menus as a fraction of the work area height */
 	const float MaxMenuScreenHeightFraction = 0.8f;
+	const float AnimationDuration = 0.15f;
 }
 
+/** Overlay widget class used to hold menu contents and display them on top of the current window */
+class SMenuPanel : public SOverlay
+{
+public:
+	SLATE_BEGIN_ARGS(SMenuPanel)
+	{
+		_Visibility = EVisibility::SelfHitTestInvisible;
+	}
+
+	SLATE_END_ARGS()
+
+	void Construct(const FArguments& InArgs)
+	{
+		SOverlay::Construct(SOverlay::FArguments());
+	}
+
+	void PushMenu(TSharedRef<FMenuBase> InMenu, const FVector2D& InLocation)
+	{
+		check(InMenu->GetContent().IsValid());
+
+		TSharedPtr<SWindow> ParentWindow = InMenu->GetParentWindow();
+		check(ParentWindow.IsValid());
+
+		// Transform InLocation into a position local to this panel (assumes the panel is in an overlay that covers the whole of the panel window) 
+		FVector2D PanelInScreen = ParentWindow->GetRectInScreen().GetTopLeft();
+		FVector2D PanelInWindow = ParentWindow->GetLocalToScreenTransform().Inverse().TransformPoint(PanelInScreen);
+		FVector2D LocationInWindow = ParentWindow->GetLocalToScreenTransform().Inverse().TransformPoint(InLocation);
+		FVector2D LocationInPanel = LocationInWindow - PanelInWindow;
+
+		// Add the new menu into a slot on this panel and set the padding so that its position is correct
+		AddSlot()
+		.HAlign(HAlign_Left)
+		.VAlign(VAlign_Top)
+		.Padding(LocationInPanel.X, LocationInPanel.Y, 0, 0)
+		[
+			InMenu->GetContent().ToSharedRef()
+		];
+
+		// Make sure that the menu will remove itself from the panel when dismissed
+		InMenu->GetOnMenuDismissed().AddSP(this, &SMenuPanel::OnMenuClosed);
+	}
+
+	void OnMenuClosed(TSharedRef<IMenu> InMenu)
+	{
+		RemoveSlot(InMenu->GetContent().ToSharedRef());
+	}
+};
+
 
 namespace
 {
-	/** Widget that wraps any popup menu created through FMenuStack::PushMenu to provide default key handling */
+	/** Widget that wraps any menu created in FMenuStack to provide default key handling, focus tracking and helps us spot menus in widget paths */
 	DECLARE_DELEGATE_RetVal_OneParam(FReply, FOnKeyDown, FKey)
+	DECLARE_DELEGATE_OneParam(FOnMenuLostFocus, const FWidgetPath&)
 	class SMenuContentWrapper : public SCompoundWidget
 	{
 	public:
@@ -24,18 +77,29 @@ namespace
 
 			SLATE_DEFAULT_SLOT(FArguments, MenuContent)
 			SLATE_EVENT(FOnKeyDown, OnKeyDown)
+			SLATE_EVENT(FOnMenuLostFocus, OnMenuLostFocus)
 		SLATE_END_ARGS()
 
 		/** Construct this widget */
 		void Construct(const FArguments& InArgs)
 		{
 			OnKeyDownDelegate = InArgs._OnKeyDown;
+			OnMenuLostFocus = InArgs._OnMenuLostFocus;
 			ChildSlot
 			[
 				InArgs._MenuContent.Widget
 			];
 		}
 
+		virtual void OnFocusChanging(const FWeakWidgetPath& PreviousFocusPath, const FWidgetPath& NewWidgetPath) override
+		{
+			// if focus changed and this menu content had focus (or one of its children did) then inform the stack via the OnMenuLostFocus event
+			if (OnMenuLostFocus.IsBound() && PreviousFocusPath.ContainsWidget(AsShared()))
+			{
+				return OnMenuLostFocus.Execute(NewWidgetPath);
+			}
+		}
+
 	private:
 
 		/** This widget must support keyboard focus */
@@ -56,6 +120,9 @@ namespace
 
 		/** Delegate to forward keys down events on the menu */
 		FOnKeyDown OnKeyDownDelegate;
+
+		/** Delegate to inform the stack that a menu has lost focus and might need to be closed */
+		FOnMenuLostFocus OnMenuLostFocus;
 	};
 
 
@@ -76,244 +143,583 @@ namespace
 
 TSharedRef<SWindow> FMenuStack::PushMenu( const TSharedRef<SWindow>& ParentWindow, const TSharedRef<SWidget>& InContent, const FVector2D& SummonLocation, const FPopupTransitionEffect& TransitionEffect, const bool bFocusImmediately, const bool bShouldAutoSize, const FVector2D& WindowSize, const FVector2D& SummonLocationSize)
 {
-	// Only enable window position/size transitions if we're running at a decent frame rate
-	const bool bAllowAnimations = FSlateApplication::Get().AreMenuAnimationsEnabled() && FSlateApplication::Get().IsRunningAtTargetFrameRate();
-
-	// Calc the max height available on screen for the menu
-	FVector2D ExpectedSize = WindowSize;
-	FSlateRect AnchorRect(SummonLocation, SummonLocation + SummonLocationSize);
-	FSlateRect WorkArea = FSlateApplication::Get().GetWorkArea( AnchorRect );
-	float MaxHeight = FMenuStackDefs::MaxMenuScreenHeightFraction * WorkArea.GetSize().Y;
-	ExpectedSize.Y = FMath::Min(ExpectedSize.Y, MaxHeight);
-
-	// Wrap menu content in a box that limits its maximum height
-	TSharedRef<SWidget> WrappedContent = SNew(SVerticalBox)
-		+SVerticalBox::Slot()
-		.MaxHeight(MaxHeight)
+	// Deprecated method that is no longer called by its deprecated counterpart in FSlateApplication so it will only create a dummy message warning the caller not to use this method
+	// Backwards compatibility of the API is taken care of by FSlateApplication.
+	ensureMsg(false, TEXT("PushMenu() returning an SWindow is deprecated. Use Push() that returns an IMenu instead."));
+
+	FSlateRect Anchor(SummonLocation, SummonLocation + SummonLocationSize);
+	FVector2D ExpectedSize(300, 200);
+	FVector2D AdjustedSummonLocation = FSlateApplication::Get().CalculatePopupWindowPosition(Anchor, ExpectedSize, Orient_Vertical);
+
+	TSharedRef<SWindow> NewMenuWindow =
+		SNew(SWindow)
+		.IsPopupWindow(true)
+		.SizingRule(bShouldAutoSize ? ESizingRule::Autosized : ESizingRule::UserSized)
+		.ScreenPosition(AdjustedSummonLocation)
+		.AutoCenter(EAutoCenter::None)
+		.ClientSize(ExpectedSize)
+		.FocusWhenFirstShown(true)
+		.ActivateWhenFirstShown(true)
 		[
-			SNew(SMenuContentWrapper)
-			.OnKeyDown_Static(&OnMenuKeyDown)
-			.MenuContent()
+			SNew(SVerticalBox)
+			+ SVerticalBox::Slot()
+			.VAlign(VAlign_Center)
+			.HAlign(HAlign_Center)
 			[
-				InContent
+				SNew(STextBlock)
+				.Text(LOCTEXT("PushMenuDeprecatedWarning", "WARNING: PushMenu() returning an SWindow is deprecated. Use Push() that returns an IMenu instead."))
 			]
 		];
 
-	// Adjust the position of popup windows so they do not go out of the visible area of the monitor(s)	
-	if( bShouldAutoSize )
+	FSlateApplication::Get().AddWindowAsNativeChild(NewMenuWindow, ParentWindow, true);
+
+	return NewMenuWindow;
+}
+
+TSharedRef<IMenu> FMenuStack::Push(const FWidgetPath& InOwnerPath, const TSharedRef<SWidget>& InContent, const FVector2D& SummonLocation, const FPopupTransitionEffect& TransitionEffect, const bool bFocusImmediately, const FVector2D& SummonLocationSize, TOptional<EPopupMethod> InMethod)
+{
+	FSlateRect Anchor(SummonLocation, SummonLocation + SummonLocationSize);
+	TSharedPtr<IMenu> ParentMenu;
+
+	if (HasMenus())
 	{
-		// @todo slate: Assumes that popup is not Scaled up or down from application scale.
-		WrappedContent->SlatePrepass( FSlateApplication::Get().GetApplicationScale() );
-		// @todo slate: Doesn't take into account potential window border size
-		ExpectedSize = WrappedContent->GetDesiredSize();
+		// Find a menu in the stack in InOwnerPath to determine the level to insert this 
+		ParentMenu = FindMenuInWidgetPath(InOwnerPath);
+		check(HostWindow.IsValid());
 	}
-	EOrientation Orientation = (TransitionEffect.SlideDirection == FPopupTransitionEffect::SubMenu) ? Orient_Horizontal : Orient_Vertical;
-	FSlateRect Anchor(SummonLocation.X, SummonLocation.Y, SummonLocation.X + SummonLocationSize.X, SummonLocation.Y + SummonLocationSize.Y);
-	FVector2D AdjustedSummonLocation = FSlateApplication::Get().CalculatePopupWindowPosition( Anchor, ExpectedSize, Orientation );
 
-	const float AnimationDuration = 0.15f;
+	if (!ParentMenu.IsValid())
+	{
+		// pushing a new root menu (leave ParentMenu invalid)
+
+		// The active method is determined when a new root menu is pushed
+		ActiveMethod = InMethod.IsSet() ? InMethod.GetValue() : QueryPopupMethod(InOwnerPath);
+
+		// The host window is determined when a new root menu is pushed
+		SetHostWindow(InOwnerPath.GetWindow());
+	}
+
+	return PushInternal(ParentMenu, InContent, Anchor, TransitionEffect, bFocusImmediately);
+}
+
+TSharedRef<IMenu> FMenuStack::Push(const TSharedPtr<IMenu>& InParentMenu, const TSharedRef<SWidget>& InContent, const FVector2D& SummonLocation, const FPopupTransitionEffect& TransitionEffect, const bool bFocusImmediately, const FVector2D& SummonLocationSize)
+{
+	check(Stack.Contains(InParentMenu));
+	check(HostWindow.IsValid());
+
+	FSlateRect Anchor(SummonLocation, SummonLocation + SummonLocationSize);
+
+	return PushInternal(InParentMenu, InContent, Anchor, TransitionEffect, bFocusImmediately);
+}
+
+TSharedRef<IMenu> FMenuStack::PushHosted(const FWidgetPath& InOwnerPath, const TSharedRef<IMenuHost>& InMenuHost, const TSharedRef<SWidget>& InContent, TSharedPtr<SWidget>& OutWrappedContent, const FPopupTransitionEffect& TransitionEffect)
+{
+	TSharedPtr<IMenu> ParentMenu;
+
+	if (HasMenus())
+	{
+		// Find a menu in the stack in InOwnerPath to determine the level to insert this 
+		ParentMenu = FindMenuInWidgetPath(InOwnerPath);
+		check(HostWindow.IsValid());
+	}
+
+	if (!ParentMenu.IsValid())
+	{
+		// pushing a new root menu (leave ParentMenu invalid)
+
+		// The active method is determined when a new root menu is pushed and hosted menus are always UseCurrentWindow
+		ActiveMethod = EPopupMethod::UseCurrentWindow;
+
+		// The host window is determined when a new root menu is pushed
+		SetHostWindow(InOwnerPath.GetWindow());
+	}
+
+	return PushHosted(ParentMenu, InMenuHost, InContent, OutWrappedContent, TransitionEffect);
+}
+
+TSharedRef<IMenu> FMenuStack::PushHosted(const TSharedPtr<IMenu>& InParentMenu, const TSharedRef<IMenuHost>& InMenuHost, const TSharedRef<SWidget>& InContent, TSharedPtr<SWidget>& OutWrappedContent, const FPopupTransitionEffect& TransitionEffect)
+{
+	check(HostWindow.IsValid());
+
+	// Create a FMenuInHostWidget
+	TSharedRef<SWidget> WrappedContent = WrapContent(InContent);
+	TSharedRef<FMenuInHostWidget> OutMenu = MakeShareable(new FMenuInHostWidget(InMenuHost, WrappedContent));
+	PendingNewMenu = OutMenu;
+
+	// Set the returned content - this must be drawn by the hosting widget until the menu gets dismissed and calls IMenuHost::OnMenuDismissed on its host
+	OutWrappedContent = WrappedContent;
+
+	// Register to get a callback when it's dismissed - to fixup stack
+	OutMenu->GetOnMenuDismissed().AddRaw(this, &FMenuStack::OnMenuDestroyed);
+
+	PostPush(InParentMenu, OutMenu);
+
+	PendingNewMenu.Reset();
+
+	return OutMenu;
+}
+
+TSharedRef<IMenu> FMenuStack::PushInternal(const TSharedPtr<IMenu>& InParentMenu, const TSharedRef<SWidget>& InContent, FSlateRect Anchor, const FPopupTransitionEffect& TransitionEffect, const bool bFocusImmediately)
+{
+	FPrePushArgs PrePushArgs;
+	PrePushArgs.Content = InContent;
+	PrePushArgs.Anchor = Anchor;
+	PrePushArgs.TransitionEffect = TransitionEffect;
+	PrePushArgs.bFocusImmediately = bFocusImmediately;
+
+	// Pre-push stage
+	//   Determines correct layout
+	//   Wraps content
+	//   Other common setup steps needed by all (non-hosted) menus
+	const FPrePushResults PrePushResults = PrePush(PrePushArgs);
+
+	// Menu object creation stage
+	TSharedRef<FMenuBase> OutMenu = ActiveMethod == EPopupMethod::CreateNewWindow
+		? PushNewWindow(InParentMenu, PrePushResults)
+		: PushPopup(InParentMenu, PrePushResults);
+
+	// Post-push stage
+	//   Updates the stack and content map member variables
+	PostPush(InParentMenu, OutMenu);
+
+	PendingNewMenu.Reset();
+
+	return OutMenu;
+}
+
+FMenuStack::FPrePushResults FMenuStack::PrePush(const FPrePushArgs& InArgs)
+{
+	FPrePushResults OutResults;
+	OutResults.bFocusImmediately = InArgs.bFocusImmediately;
+
+	// Only enable window position/size transitions if we're running at a decent frame rate
+	OutResults.bAllowAnimations = FSlateApplication::Get().AreMenuAnimationsEnabled() && FSlateApplication::Get().IsRunningAtTargetFrameRate();
+
+	// Calc the max height available on screen for the menu
+	float MaxHeight;
+	if (ActiveMethod == EPopupMethod::CreateNewWindow)
+	{
+		FSlateRect WorkArea = FSlateApplication::Get().GetWorkArea(InArgs.Anchor);
+		MaxHeight = FMenuStackDefs::MaxMenuScreenHeightFraction * WorkArea.GetSize().Y;
+	}
+	else
+	{
+		MaxHeight = FMenuStackDefs::MaxMenuScreenHeightFraction * HostWindow->GetClientSizeInScreen().Y;
+	}
+
+	bool bAnchorSetsMinWidth = InArgs.TransitionEffect.SlideDirection == FPopupTransitionEffect::ComboButton;
+
+	// Wrap menu content in a box needed for various sizing and tracking purposes
+	FOptionalSize OptionalMinWidth = bAnchorSetsMinWidth ? InArgs.Anchor.GetSize().X : FOptionalSize();
+	FOptionalSize OptionalMinHeight = MaxHeight;
+
+	// Wrap content in an SPopup before the rest of the wrapping process - should make menus appear on top using deferred presentation
+	TSharedRef<SWidget> TempContent = SNew(SPopup)[InArgs.Content.ToSharedRef()];
+
+	OutResults.WrappedContent = WrapContent(TempContent, OptionalMinWidth, OptionalMinHeight);
+
+	// @todo slate: Assumes that popup is not Scaled up or down from application scale.
+	OutResults.WrappedContent->SlatePrepass(FSlateApplication::Get().GetApplicationScale());
+	// @todo slate: Doesn't take into account potential window border size
+	OutResults.ExpectedSize = OutResults.WrappedContent->GetDesiredSize();
+
+	EOrientation Orientation = (InArgs.TransitionEffect.SlideDirection == FPopupTransitionEffect::SubMenu) ? Orient_Horizontal : Orient_Vertical;
+
+	// Calculate the correct position for the menu based on the popup method
+	if (ActiveMethod == EPopupMethod::CreateNewWindow)
+	{
+		// Places the menu's window in the work area
+		OutResults.AnimStartLocation = OutResults.AnimFinalLocation = FSlateApplication::Get().CalculatePopupWindowPosition(InArgs.Anchor, OutResults.ExpectedSize, Orientation);
+	}
+	else
+	{
+		// Places the menu's content in the host window
+		const FVector2D ProposedPlacement(
+			Orientation == Orient_Horizontal ? InArgs.Anchor.Right : InArgs.Anchor.Left,
+			Orientation == Orient_Horizontal ? InArgs.Anchor.Top : InArgs.Anchor.Bottom);
+
+		OutResults.AnimStartLocation = OutResults.AnimFinalLocation =
+			ComputePopupFitInRect(InArgs.Anchor, FSlateRect(ProposedPlacement, ProposedPlacement + OutResults.ExpectedSize), Orientation, HostWindow->GetClientRectInScreen());
+	}
 
 	// Start the pop-up menu at an offset location, then animate it to its target location over time
-	FVector2D AnimationStartLocation = FVector2D( AdjustedSummonLocation.X, AdjustedSummonLocation.Y );
-	if( bAllowAnimations )
+	// @todo: Anims aren't supported or attempted - this is legacy code left in in case we reinstate menu anims
+	if (OutResults.bAllowAnimations)
 	{
-		const bool bSummonRight = AdjustedSummonLocation.X >= SummonLocation.X;
-		const bool bSummonBelow = AdjustedSummonLocation.Y >= SummonLocation.Y;
+		const bool bSummonRight = OutResults.AnimFinalLocation.X >= OutResults.AnimStartLocation.X;
+		const bool bSummonBelow = OutResults.AnimFinalLocation.Y >= OutResults.AnimStartLocation.Y;
 		const int32 SummonDirectionX = bSummonRight ? 1 : -1;
 		const int32 SummonDirectionY = bSummonBelow ? 1 : -1;
 
-		switch( TransitionEffect.SlideDirection )
+		switch (InArgs.TransitionEffect.SlideDirection)
 		{
-			case FPopupTransitionEffect::None:
-				// No sliding
-				break;
-
-			case FPopupTransitionEffect::ComboButton:
-				AnimationStartLocation.Y = FMath::Max( AnimationStartLocation.Y + 30.0f * SummonDirectionY, 0.0f );
-				break;
-
-			case FPopupTransitionEffect::TopMenu:
-				AnimationStartLocation.Y = FMath::Max( AnimationStartLocation.Y + 60.0f * SummonDirectionY, 0.0f );
-				break;
-
-			case FPopupTransitionEffect::SubMenu:
-				AnimationStartLocation.X += 60.0f * SummonDirectionX;
-				break;
-
-			case FPopupTransitionEffect::TypeInPopup:
-				AnimationStartLocation.Y = FMath::Max( AnimationStartLocation.Y + 30.0f * SummonDirectionY, 0.0f );
-				break;
-
-			case FPopupTransitionEffect::ContextMenu:
-				AnimationStartLocation.X += 30.0f * SummonDirectionX;
-				AnimationStartLocation.Y += 50.0f * SummonDirectionY;
-				break;
+		case FPopupTransitionEffect::None:
+			// No sliding
+			break;
+
+		case FPopupTransitionEffect::ComboButton:
+			OutResults.AnimStartLocation.Y = FMath::Max(OutResults.AnimStartLocation.Y + 30.0f * SummonDirectionY, 0.0f);
+			break;
+
+		case FPopupTransitionEffect::TopMenu:
+			OutResults.AnimStartLocation.Y = FMath::Max(OutResults.AnimStartLocation.Y + 60.0f * SummonDirectionY, 0.0f);
+			break;
+
+		case FPopupTransitionEffect::SubMenu:
+			OutResults.AnimStartLocation.X += 60.0f * SummonDirectionX;
+			break;
+
+		case FPopupTransitionEffect::TypeInPopup:
+			OutResults.AnimStartLocation.Y = FMath::Max(OutResults.AnimStartLocation.Y + 30.0f * SummonDirectionY, 0.0f);
+			break;
+
+		case FPopupTransitionEffect::ContextMenu:
+			OutResults.AnimStartLocation.X += 30.0f * SummonDirectionX;
+			OutResults.AnimStartLocation.Y += 50.0f * SummonDirectionY;
+			break;
 		}
 	}
 
+	// Release the mouse so that context can be properly restored upon closing menus.  See CL 1411833 before changing this.
+	if (InArgs.bFocusImmediately)
+	{
+		FSlateApplication::Get().ReleaseMouseCapture();
+	}
+
+	return OutResults;
+}
+
+TSharedRef<FMenuBase> FMenuStack::PushNewWindow(TSharedPtr<IMenu> InParentMenu, const FPrePushResults& InPrePushResults)
+{
+	check(ActiveMethod == EPopupMethod::CreateNewWindow);
+
 	// Start pop-up windows out transparent, then fade them in over time
-	const EWindowTransparency Transparency(bAllowAnimations ? EWindowTransparency::PerWindow : EWindowTransparency::None);
-	const float InitialWindowOpacity = bAllowAnimations ? 0.0f : 1.0f;
+	const EWindowTransparency Transparency(InPrePushResults.bAllowAnimations ? EWindowTransparency::PerWindow : EWindowTransparency::None);
+	const float InitialWindowOpacity = InPrePushResults.bAllowAnimations ? 0.0f : 1.0f;
 	const float TargetWindowOpacity = 1.0f;
 
-	//Release the mouse so that context can be properly restored upon closing NewMenuWindow.  See CL 1411833 before changing this.
-	if ( bFocusImmediately ) 
-	{ 
-		FSlateApplication::Get().ReleaseMouseCapture(); 
-	}
-
 	// Create a new window for the menu
 	TSharedRef<SWindow> NewMenuWindow = SNew(SWindow)
-		.IsPopupWindow( true )
-		.SizingRule( bShouldAutoSize ? ESizingRule::Autosized : ESizingRule::UserSized )
-		.ScreenPosition( AnimationStartLocation )
-		.AutoCenter( EAutoCenter::None )
-		.ClientSize( ExpectedSize )
+		.IsPopupWindow(true)
+		.SizingRule(ESizingRule::Autosized)
+		.ScreenPosition(InPrePushResults.AnimStartLocation)
+		.AutoCenter(EAutoCenter::None)
+		.ClientSize(InPrePushResults.ExpectedSize)
 		.InitialOpacity(InitialWindowOpacity)
-		.SupportsTransparency( Transparency )
-		.FocusWhenFirstShown( bFocusImmediately )
-		.ActivateWhenFirstShown( bFocusImmediately )
+		.SupportsTransparency(Transparency)
+		.FocusWhenFirstShown(InPrePushResults.bFocusImmediately)
+		.ActivateWhenFirstShown(InPrePushResults.bFocusImmediately)
 		[
-			WrappedContent
+			InPrePushResults.WrappedContent.ToSharedRef()
 		];
 
-	if (bFocusImmediately)
+	PendingNewWindow = NewMenuWindow;
+
+	if (InPrePushResults.bFocusImmediately)
 	{
 		// Focus the content rather than just the window
-		NewMenuWindow->SetWidgetToFocusOnActivate(InContent);
+		NewMenuWindow->SetWidgetToFocusOnActivate(InPrePushResults.WrappedContent);
 	}
 
-	// Find the location where the parent of the new menu resides.  The new menu's location in the stack is the level below the parent
-	int32 LocIndex = FindLocationInStack( ParentWindow ) + 1;
-	if( !Stack.IsValidIndex( LocIndex ) )
+	TSharedRef<FMenuInWindow> Menu = MakeShareable(new FMenuInWindow(NewMenuWindow, InPrePushResults.WrappedContent.ToSharedRef()));
+	PendingNewMenu = Menu;
+
+	TSharedPtr<SWindow> ParentWindow;
+	if (InParentMenu.IsValid())
 	{
-		// If the new level doesn't exist yet create it now
-		Stack.Add( FMenuWindowList() );
-		// This is to ensure that the parent window is always one level above the new window being created
-		check( Stack.IsValidIndex( LocIndex ) );
+		ParentWindow = InParentMenu->GetParentWindow();
 	}
-
+	else
 	{
-		// Make a copy of the windows array in case it changes size while we are operating on it
-		FMenuWindowList WindowsInStackCopy = Stack[LocIndex];
+		ParentWindow = HostWindow;
+	}
 
-		// Remove all other windows on the same level from the stack as they are no longer valid
-		// This will also remove the children of any of those windows
-		for( int32 WindowIndex = 0; WindowIndex < WindowsInStackCopy.Num(); ++WindowIndex )
-		{
-			WindowsInStackCopy[WindowIndex]->RequestDestroyWindow();
-		}
+	FSlateApplication::Get().AddWindowAsNativeChild(NewMenuWindow, ParentWindow.ToSharedRef(), true);
+
+	// Kick off the intro animation!
+	// @todo: Anims aren't supported or attempted - this is legacy code left in in case we reinstate menu anims
+	if (InPrePushResults.bAllowAnimations)
+	{
+		FCurveSequence Sequence;
+		Sequence.AddCurve(0, FMenuStackDefs::AnimationDuration, ECurveEaseFunction::CubicOut);
+		NewMenuWindow->MorphToPosition(Sequence, TargetWindowOpacity, InPrePushResults.AnimFinalLocation);
 	}
 
-	// Destroying windows above may cause the stack to get cleared or modified when activating another window.
-	if( !Stack.IsValidIndex( LocIndex ) )
+	PendingNewWindow.Reset();
+
+	return Menu;
+}
+
+TSharedRef<FMenuBase> FMenuStack::PushPopup(TSharedPtr<IMenu> InParentMenu, const FPrePushResults& InPrePushResults)
+{
+	check(ActiveMethod == EPopupMethod::UseCurrentWindow);
+
+	// Create a FMenuInPopup
+	check(InPrePushResults.WrappedContent.IsValid());
+	TSharedRef<FMenuInPopup> Menu = MakeShareable(new FMenuInPopup(InPrePushResults.WrappedContent.ToSharedRef()));
+	PendingNewMenu = Menu;
+
+	// Register to get callback when it's dismissed - to fixup stack
+	Menu->GetOnMenuDismissed().AddRaw(this, &FMenuStack::OnMenuDestroyed);
+
+	// Add it to a slot on the menus panel widget
+	HostWindowPopupPanel->PushMenu(Menu, InPrePushResults.AnimFinalLocation);
+
+	if (InPrePushResults.bFocusImmediately)
 	{
-		// Re-find the location in the stack to make sure we add to the correct window list
-		LocIndex = FindLocationInStack( ParentWindow ) + 1;
-		if( !Stack.IsValidIndex( LocIndex ) )
-		{
-			// If the new level doesn't exist yet create it now
-			Stack.Add( FMenuWindowList() );
-			// This is to ensure that the parent window is always one level above the new window being created
-			check( Stack.IsValidIndex( LocIndex ) );
-		}
+		FSlateApplication::Get().SetKeyboardFocus(InPrePushResults.WrappedContent, EFocusCause::SetDirectly);
 	}
 
-	// Add the new window to menu window array at the required stack index
-	Stack[LocIndex].Add( NewMenuWindow );
+	return Menu;
+}
+
+void FMenuStack::PostPush(TSharedPtr<IMenu> InParentMenu, TSharedRef<FMenuBase> InMenu)
+{
+	// Determine at which index we insert this new menu in the stack
+	int32 InsertIndex = 0;
+	if (InParentMenu.IsValid())
+	{
+		int32 ParentIndex = Stack.IndexOfByKey(InParentMenu);
+		check(ParentIndex != INDEX_NONE);
 
-	FSlateApplication::Get().AddWindowAsNativeChild( NewMenuWindow, ParentWindow, true );
+		InsertIndex = ParentIndex + 1;
+	}
 
-	// Kick off the intro animation!
-	if( bAllowAnimations )
+	// Insert before dismissing others to stop the stack accidentally emptying briefly mid-push and reseting some state
+	Stack.Insert(InMenu, InsertIndex);
+	CachedContentMap.Add(InMenu->GetContent(), InMenu);
+
+	// Dismiss menus after the insert point
+	if (Stack.Num() > InsertIndex + 1)
 	{
-		FCurveSequence Sequence;
-		Sequence.AddCurve( 0, AnimationDuration, ECurveEaseFunction::CubicOut );
-		NewMenuWindow->MorphToPosition( Sequence, TargetWindowOpacity, AdjustedSummonLocation );
+		DismissFrom(Stack[InsertIndex + 1]);
+
+		// tidy the stack data now (it will happen via callbacks from the dismissed menus but that might be delayed)
+		for (int32 StackIndex = Stack.Num() - 1; StackIndex > InsertIndex; --StackIndex)
+		{
+			CachedContentMap.Remove(Stack[StackIndex]->GetContent());
+			Stack.RemoveAt(StackIndex);
+		}
 	}
 
 	// When a new menu is pushed, if we are not already in responsive mode for Slate UI, enter it now
 	// to ensure the menu is responsive in low FPS situations
-	if( !ThrottleHandle.IsValid() )
+	if (!ThrottleHandle.IsValid())
 	{
 		ThrottleHandle = FSlateThrottleManager::Get().EnterResponsiveMode();
 	}
-	return NewMenuWindow;
 }
 
+EPopupMethod FMenuStack::QueryPopupMethod(const FWidgetPath& PathToQuery)
+{
+	for (int32 WidgetIndex = PathToQuery.Widgets.Num() - 1; WidgetIndex >= 0; --WidgetIndex)
+	{
+		TOptional<EPopupMethod> PopupMethod = PathToQuery.Widgets[WidgetIndex].Widget->OnQueryPopupMethod();
+		if (PopupMethod.IsSet())
+		{
+			return PopupMethod.GetValue();
+		}
+	}
+
+	return EPopupMethod::CreateNewWindow;
+}
+
+void FMenuStack::Dismiss(int32 FirstStackIndexToRemove)
+{
+	// deprecated
+	DismissInternal(FirstStackIndexToRemove);
+}
 
-void FMenuStack::Dismiss( int32 LastStackIndex )
+void FMenuStack::DismissFrom(const TSharedPtr<IMenu>& InFromMenu)
+{
+	int32 Index = Stack.IndexOfByKey(InFromMenu);
+	if (Index != INDEX_NONE)
+	{
+		DismissInternal(Index);
+	}
+}
+
+void FMenuStack::DismissAll()
+{
+	const int32 TopLevel = 0;
+	DismissInternal(TopLevel);
+}
+
+void FMenuStack::DismissInternal(int32 FirstStackIndexToRemove)
 {
 	// Dismiss the stack in reverse order so we destroy children before parents (causes focusing issues if done the other way around)
-	for( int32 StackIndex = Stack.Num()-1; StackIndex >= LastStackIndex; --StackIndex )
+	for (int32 StackIndex = Stack.Num() - 1; StackIndex >= FirstStackIndexToRemove; --StackIndex)
+	{
+		Stack[StackIndex]->Dismiss();
+	}
+}
+
+void FMenuStack::SetHostWindow(TSharedPtr<SWindow> InWindow)
+{
+	if (HostWindow != InWindow)
 	{
-		const TArray< TSharedPtr< SWindow> >& MenuWindows = TArray< TSharedPtr<SWindow> >(Stack[StackIndex]);
-		for( int32 MenuIndex = 0; MenuIndex < MenuWindows.Num(); ++MenuIndex )
+		// If the host window is changing, remove the popup panel from the previous host
+		if (HostWindow.IsValid() && HostWindowPopupPanel.IsValid())
 		{
-			MenuWindows[MenuIndex]->RequestDestroyWindow();
+			HostWindow->RemoveOverlaySlot(HostWindowPopupPanel.ToSharedRef());
+			HostWindowPopupPanel.Reset();
+		}
+
+		HostWindow = InWindow;
+
+		// If the host window is changing, add a popup panel to the new host
+		if (HostWindow.IsValid())
+		{
+			// setup a popup layer on the new window
+			HostWindow->AddOverlaySlot()
+			[
+				SAssignNew(HostWindowPopupPanel, SMenuPanel)
+			];
 		}
 	}
+}
 
-	if (LastStackIndex == 0)
+void FMenuStack::OnMenuDestroyed(TSharedRef<IMenu> InMenu)
+{
+	int32 Index = Stack.IndexOfByKey(InMenu);
+	if (Index != INDEX_NONE)
 	{
-		// Empty everything if this is a full dismissal
-		Stack.Empty();
+		// Dismiss menus below this one
+		for (int32 StackIndex = Stack.Num() - 1; StackIndex > Index; --StackIndex)
+		{
+			Stack[StackIndex]->Dismiss();	// this will cause OnMenuDestroyed() to re-enter
+		}
+
+		// Clean up the stack and content map arrays
+		for (int32 StackIndex = Stack.Num() - 1; StackIndex >= Index; --StackIndex)
+		{
+			CachedContentMap.Remove(Stack[StackIndex]->GetContent());
+			Stack.RemoveAt(StackIndex);
+		}
 
 		// Leave responsive mode once the last menu closes
-		if (ThrottleHandle.IsValid())
+		if (Stack.Num() == 0)
 		{
-			FSlateThrottleManager::Get().LeaveResponsiveMode(ThrottleHandle);
+			if (ThrottleHandle.IsValid())
+			{
+				FSlateThrottleManager::Get().LeaveResponsiveMode(ThrottleHandle);
+			}
+
+			SetHostWindow(TSharedPtr<SWindow>());
 		}
 	}
 }
 
-
-void FMenuStack::RemoveWindow( TSharedRef<SWindow> WindowToRemove )
+void FMenuStack::OnMenuContentLostFocus(const FWidgetPath& InFocussedPath)
 {
-	// A window was requested to be destroyed, so make sure it's not in the pop-up menu stack to avoid it
-	// becoming a parent to a freshly-created window!
-	int32 Location = FindLocationInStack( WindowToRemove );
-
-	if( Location != INDEX_NONE )
+	// Window activation messages will make menus collapse when in CreateNewWindow mode
+	// In UseCurrentWindow mode we must look for focus moving from menus
+	if (ActiveMethod == EPopupMethod::UseCurrentWindow && HasMenus() && !PendingNewMenu.IsValid())
 	{
-		if( Stack.IsValidIndex( Location+1 ) )
+		// If focus is switching determine which of our menus it is in, if any
+		TSharedPtr<IMenu> FocussedMenu = FindMenuInWidgetPath(InFocussedPath);
+
+		if (FocussedMenu.IsValid())
 		{
-			// Remove all levels below the window being removed
-			Stack.RemoveAt( Location+1, Stack.Num()-(Location+1) );
+			// dismiss menus below FocussedMenu
+			int32 FocussedIndex = Stack.IndexOfByKey(FocussedMenu);
+			check(FocussedIndex != INDEX_NONE);
+
+			if (Stack.Num() > FocussedIndex + 1)
+			{
+				DismissFrom(Stack[FocussedIndex + 1]);
+			}
 		}
+		else
+		{
+			// Focus has moved outside all menus - collapse the stack
+			DismissAll();
+		}
+	}
+}
 
-		// Remove the actual window
-		Stack[Location].Remove( WindowToRemove );
+TSharedRef<SWidget> FMenuStack::WrapContent(TSharedRef<SWidget> InContent, FOptionalSize OptionalMinWidth, FOptionalSize OptionalMinHeight)
+{
+	// Wrap menu content in a box that limits its maximum height
+	// and in a SMenuContentWrapper that handles key presses and focus changes.
+	return SNew(SMenuContentWrapper)
+		.OnKeyDown_Static(&OnMenuKeyDown)
+		.OnMenuLostFocus_Raw(this, &FMenuStack::OnMenuContentLostFocus)
+		.MenuContent()
+		[
+			SNew(SBox)
+			.MinDesiredWidth(OptionalMinWidth)
+			.MaxDesiredHeight(OptionalMinHeight)
+			[
+				InContent
+			]
+		];
+}
 
-		// Remove list if empty
-		if (Stack[Location].Num() == 0)
+TSharedPtr<IMenu> FMenuStack::FindMenuInWidgetPath(const FWidgetPath& PathToQuery) const
+{
+	for (int32 PathIndex = PathToQuery.Widgets.Num() - 1; PathIndex >= 0; --PathIndex)
+	{
+		TSharedPtr<const SWidget> Widget = PathToQuery.Widgets[PathIndex].Widget;
+		const TSharedPtr<FMenuBase>* FoundMenu = CachedContentMap.Find(Widget);
+		if (FoundMenu != nullptr)
 		{
-			Stack.RemoveAt(Location);
+			return *FoundMenu;
 		}
 	}
+	return TSharedPtr<IMenu>();
+}
+
+void FMenuStack::RemoveWindow( TSharedRef<SWindow> WindowToRemove )
+{
+	// deprecated
+	OnWindowDestroyed(WindowToRemove);
+}
 
-	// Leave responsive mode once the last menu closes
-	if (Stack.Num() == 0 && ThrottleHandle.IsValid())
+void FMenuStack::OnWindowDestroyed(TSharedRef<SWindow> InWindow)
+{
+	if (HostWindow == InWindow)
 	{
-		FSlateThrottleManager::Get().LeaveResponsiveMode( ThrottleHandle );
+		// If the host window is destroyed, collapse the whole stack and reset all state
+		Stack.Empty();
+		CachedContentMap.Empty();
+
+		SetHostWindow(TSharedPtr<SWindow>());
 	}
-}
+	else
+	{
+		// A window was requested to be destroyed, so make sure it's not in the menu stack to avoid it
+		// becoming a parent to a freshly-created window!
+		TSharedPtr<IMenu> Menu = FindMenuFromWindow(InWindow);
 
+		if (Menu.IsValid())
+		{
+			OnMenuDestroyed(Menu.ToSharedRef());
+		}
+	}
+}
 
 void FMenuStack::OnWindowActivated( TSharedRef<SWindow> ActivatedWindow )
 {
-	if (Stack.Num() > 0)
+	if (ActivatedWindow != PendingNewWindow && HasMenus())
 	{
-		// See if the activated window is in the menu stack, if it is not then the current menu should be dismissed
-		int32 Location = FindLocationInStack(ActivatedWindow);
+		TWeakPtr<IMenu> ActivatedMenu = FindMenuFromWindow(ActivatedWindow);
 
-		if (Location == INDEX_NONE)
+		if (ActivatedMenu.IsValid())
 		{
-			Dismiss();
+			// Dismiss menus below ActivatedMenu
+			int32 ActivatedIndex = Stack.IndexOfByKey(ActivatedMenu);
+			check(ActivatedIndex != INDEX_NONE);
+
+			if (Stack.Num() > ActivatedIndex + 1)
+			{
+				DismissFrom(Stack[ActivatedIndex + 1]);
+			}
 		}
 		else
 		{
-			// If we clicked a window in the stack, we'll let the menu itself figure out what to do.
+			// Activated a window that isn't a menu - collapse the stack
+			DismissAll();
 		}
 	}
 }
@@ -321,41 +727,98 @@ void FMenuStack::OnWindowActivated( TSharedRef<SWindow> ActivatedWindow )
 
 int32 FMenuStack::FindLocationInStack( TSharedPtr<SWindow> WindowToFind ) const
 {
-	// Search all levels of the stack
-	for( int32 StackIndex = 0; StackIndex < Stack.Num(); ++StackIndex )
+	// deprecated
+	if (WindowToFind.IsValid())
 	{
-		// Search all windows in each level of the stack
-		const TArray<TSharedPtr<SWindow>>& MenuWindows = TArray<TSharedPtr<SWindow>>(Stack[StackIndex]);
+		TWeakPtr<IMenu> Menu = FindMenuFromWindow(WindowToFind.ToSharedRef());
+
+		return Stack.IndexOfByKey(Menu);
+	}
+
+	// The window was not found
+	return INDEX_NONE;
+}
+
+TSharedPtr<IMenu> FMenuStack::FindMenuFromWindow(TSharedRef<SWindow> InWindow) const
+{
+	const TSharedPtr<FMenuBase>* FoundMenu = Stack.FindByPredicate([InWindow](TSharedPtr<FMenuBase> Menu) { return Menu->GetOwnedWindow() == InWindow; });
+	if (FoundMenu != nullptr)
+	{
+		return *FoundMenu;
+	}
+	return TSharedPtr<IMenu>();
+}
+
+FSlateRect FMenuStack::GetToolTipForceFieldRect(TSharedRef<IMenu> InMenu, const FWidgetPath& PathContainingMenu) const
+{
+	FSlateRect ForceFieldRect(0, 0, 0, 0);
 
-		for( int32 MenuIndex = 0; MenuIndex < MenuWindows.Num(); ++MenuIndex )
+	int32 StackLevel = Stack.IndexOfByKey(InMenu);
+
+	if (StackLevel != INDEX_NONE)
+	{
+		for (int32 StackLevelIndex = StackLevel + 1; StackLevelIndex < Stack.Num(); ++StackLevelIndex)
 		{
-			if( MenuWindows[MenuIndex] == WindowToFind )
+			TSharedPtr<SWidget> MenuContent = Stack[StackLevelIndex]->GetContent();
+			if (MenuContent.IsValid())
 			{
-				return StackIndex;
+				FWidgetPath WidgetPath = PathContainingMenu.GetPathDownTo(MenuContent.ToSharedRef());
+				if (!WidgetPath.IsValid())
+				{
+					FSlateApplication::Get().GeneratePathToWidgetChecked(MenuContent.ToSharedRef(), WidgetPath);
+				}
+				if (WidgetPath.IsValid())
+				{
+					const FGeometry& ContentGeometry = WidgetPath.Widgets.Last().Geometry;
+					ForceFieldRect = ForceFieldRect.Expand(ContentGeometry.GetClippingRect());
+				}
 			}
 		}
 	}
 
-	// The window was not found
-	return INDEX_NONE;
+	return ForceFieldRect;
 }
 
+bool FMenuStack::HasMenus() const
+{
+	return Stack.Num() > 0;
+}
 
 bool FMenuStack::HasOpenSubMenus( const TSharedRef<SWindow>& Window ) const
 {
-	int32 LocationInStack = FindLocationInStack( Window );
-
-	return LocationInStack != INDEX_NONE ? Stack.Num() > LocationInStack + 1 : false;
+	// deprecated
+	TWeakPtr<IMenu> Menu = FindMenuFromWindow(Window);
+	if (Menu.IsValid())
+	{
+		return Menu != Stack.Last();
+	}
+	return false;
 }
 
+bool FMenuStack::HasOpenSubMenus(TSharedPtr<IMenu> InMenu) const
+{
+	int32 StackIndex = Stack.IndexOfByKey(InMenu);
+	return StackIndex != INDEX_NONE && StackIndex < Stack.Num() - 1;
+}
 
 int32 FMenuStack::GetNumStackLevels() const
 {
 	return Stack.Num();
 }
 
+TSharedPtr<SWindow> FMenuStack::GetHostWindow() const
+{
+	return HostWindow;
+}
+
 
 FMenuWindowList& FMenuStack::GetWindowsAtStackLevel( const int32 StackLevelIndex )
 {
-	return Stack[StackLevelIndex];
+	// deprecated
+	// can't support this so return an empty list
+	ensure(false);
+	static FMenuWindowList EmptyList;
+	return EmptyList;
 }
+
+#undef LOCTEXT_NAMESPACE
diff --git a/Engine/Source/Runtime/Slate/Private/Framework/Application/SlateApplication.cpp b/Engine/Source/Runtime/Slate/Private/Framework/Application/SlateApplication.cpp
--- a/Engine/Source/Runtime/Slate/Private/Framework/Application/SlateApplication.cpp
+++ b/Engine/Source/Runtime/Slate/Private/Framework/Application/SlateApplication.cpp
@@ -1767,45 +1767,146 @@ TSharedRef<SWindow> FSlateApplication::AddWindowAsNativeChild( TSharedRef<SWindo
 
 TSharedRef<SWindow> FSlateApplication::PushMenu( const TSharedRef<SWidget>& InParentContent, const TSharedRef<SWidget>& InContent, const FVector2D& SummonLocation, const FPopupTransitionEffect& TransitionEffect, const bool bFocusImmediately, const bool bShouldAutoSize, const FVector2D& WindowSize, const FVector2D& SummonLocationSize )
 {
-	FWidgetPath WidgetPath;
-	GeneratePathToWidgetChecked( InParentContent, WidgetPath );
-
 #if !(UE_BUILD_SHIPPING && WITH_EDITOR)
 	// The would-be parent of the new menu being pushed is about to be destroyed.  Any children added to an about to be destroyed window will also be destroyed
+	FWidgetPath WidgetPath;
+	GeneratePathToWidgetChecked(InParentContent, WidgetPath);
 	if (IsWindowInDestroyQueue( WidgetPath.GetWindow() ))
 	{
 		UE_LOG(LogSlate, Warning, TEXT("FSlateApplication::PushMenu() called when parent window queued for destroy. New menu will be destroyed."));
 	}
 #endif
 
-	return MenuStack.PushMenu( WidgetPath.GetWindow(), InContent, SummonLocation, TransitionEffect, bFocusImmediately, bShouldAutoSize, WindowSize, SummonLocationSize );
+	// This is a DEPRECATED version of PushMenu(). It creates a new menu but forces the popup method to EPopupMethod::CreateNewWindow
+
+	// Wrap content in a fixed size box if autosizing isn't wanted
+	TSharedRef<SWidget> WappedContent = InContent;
+	if (!bShouldAutoSize)
+	{
+		WappedContent = SNew(SBox)
+			.WidthOverride(WindowSize.X)
+			.HeightOverride(WindowSize.Y)
+			[
+				InContent
+			];
+	}
+
+	TSharedPtr<IMenu> Menu = PushMenu(InParentContent, FWidgetPath(), WappedContent, SummonLocation, TransitionEffect, bFocusImmediately, SummonLocationSize, EPopupMethod::CreateNewWindow);
+	TSharedPtr<SWindow> Window = Menu->GetOwnedWindow();
+	check(Window.IsValid());
+
+	return Window.ToSharedRef();
+}
+
+TSharedPtr<IMenu> FSlateApplication::PushMenu(const TSharedRef<SWidget>& InParentWidget, const FWidgetPath& InOwnerPath, const TSharedRef<SWidget>& InContent, const FVector2D& SummonLocation, const FPopupTransitionEffect& TransitionEffect, const bool bFocusImmediately, const FVector2D& SummonLocationSize, TOptional<EPopupMethod> Method)
+{
+	// Caller supplied a valid path? Pass it to the menu stack.
+	if (InOwnerPath.IsValid())
+	{
+		return MenuStack.Push(InOwnerPath, InContent, SummonLocation, TransitionEffect, bFocusImmediately, SummonLocationSize, Method);
+	}
+
+	// If the caller doesn't specify a valid event path we'll generate one from InParentWidget
+	FWidgetPath WidgetPath;
+	if (GeneratePathToWidgetUnchecked(InParentWidget, WidgetPath))
+	{
+		return MenuStack.Push(WidgetPath, InContent, SummonLocation, TransitionEffect, bFocusImmediately, SummonLocationSize, Method);
+	}
+
+	return TSharedPtr<IMenu>();
+}
+
+TSharedPtr<IMenu> FSlateApplication::PushMenu(const TSharedPtr<IMenu>& InParentMenu, const TSharedRef<SWidget>& InContent, const FVector2D& SummonLocation, const FPopupTransitionEffect& TransitionEffect, const bool bFocusImmediately, const FVector2D& SummonLocationSize)
+{
+	return MenuStack.Push(InParentMenu, InContent, SummonLocation, TransitionEffect, bFocusImmediately, SummonLocationSize);
+}
+
+TSharedPtr<IMenu> FSlateApplication::PushHostedMenu(const TSharedRef<SWidget>& InParentWidget, const FWidgetPath& InOwnerPath, const TSharedRef<IMenuHost>& InMenuHost, const TSharedRef<SWidget>& InContent, TSharedPtr<SWidget>& OutWrappedContent, const FPopupTransitionEffect& TransitionEffect)
+{
+	// Caller supplied a valid path? Pass it to the menu stack.
+	if (InOwnerPath.IsValid())
+	{
+		return MenuStack.PushHosted(InOwnerPath, InMenuHost, InContent, OutWrappedContent, TransitionEffect);
+	}
+
+	// If the caller doesn't specify a valid event path we'll generate one from InParentWidget
+	FWidgetPath WidgetPath;
+	if (GeneratePathToWidgetUnchecked(InParentWidget, WidgetPath))
+	{
+		return MenuStack.PushHosted(WidgetPath, InMenuHost, InContent, OutWrappedContent, TransitionEffect);
+	}
+
+	return TSharedPtr<IMenu>();
+}
+
+TSharedPtr<IMenu> FSlateApplication::PushHostedMenu(const TSharedPtr<IMenu>& InParentMenu, const TSharedRef<IMenuHost>& InMenuHost, const TSharedRef<SWidget>& InContent, TSharedPtr<SWidget>& OutWrappedContent, const FPopupTransitionEffect& TransitionEffect)
+{
+	return MenuStack.PushHosted(InParentMenu, InMenuHost, InContent, OutWrappedContent, TransitionEffect);
 }
 
 bool FSlateApplication::HasOpenSubMenus( TSharedRef<SWindow> Window ) const
 {
-	return MenuStack.HasOpenSubMenus(Window);
+	// deprecated
+	TSharedPtr<IMenu> Menu = MenuStack.FindMenuFromWindow(Window);
+	if (Menu.IsValid())
+	{
+		return MenuStack.HasOpenSubMenus(Menu);
+	}
+	return false;
+}
+
+bool FSlateApplication::HasOpenSubMenus(TSharedPtr<IMenu> InMenu) const
+{
+	return MenuStack.HasOpenSubMenus(InMenu);
 }
 
 bool FSlateApplication::AnyMenusVisible() const
 {
-	return MenuStack.GetNumStackLevels() > 0;
+	return MenuStack.HasMenus();
+}
+
+TSharedPtr<SWindow> FSlateApplication::GetVisibleMenuWindow() const
+{
+	return MenuStack.GetHostWindow();
 }
 
 void FSlateApplication::DismissAllMenus()
 {
-	MenuStack.Dismiss();
+	MenuStack.DismissAll();
 }
 
 void FSlateApplication::DismissMenu( TSharedRef<SWindow> MenuWindowToDismiss )
 {
-	int32 Location = MenuStack.FindLocationInStack( MenuWindowToDismiss );
-	// Dismiss everything starting at the window to dismiss
-	MenuStack.Dismiss( Location );
+	// deprecated
+	TSharedPtr<IMenu> Menu = MenuStack.FindMenuFromWindow(MenuWindowToDismiss);
+	if (Menu.IsValid())
+	{
+		DismissMenu(Menu);
+	}
+}
+
+void FSlateApplication::DismissMenu(const TSharedPtr<IMenu>& InFromMenu)
+{
+	MenuStack.DismissFrom(InFromMenu);
+}
+
+void FSlateApplication::DismissMenuByWidget(const TSharedRef<SWidget>& InWidgetInMenu)
+{
+	FWidgetPath WidgetPath;
+	if (GeneratePathToWidgetUnchecked(InWidgetInMenu, WidgetPath))
+	{
+		TSharedPtr<IMenu> Menu = MenuStack.FindMenuInWidgetPath(WidgetPath);
+		if (Menu.IsValid())
+		{
+			MenuStack.DismissFrom(Menu);
+		}
+	}
 }
 
 int32 FSlateApplication::GetLocationInMenuStack( TSharedRef<SWindow> WindowToFind ) const
 {
-	return MenuStack.FindLocationInStack( WindowToFind );
+	// DEPRECATED
+	return MenuStack.FindLocationInStack(WindowToFind);
 }
 
 
@@ -2326,7 +2427,7 @@ void FSlateApplication::DestroyWindowsImmediately()
 		}
 
 		// Any window being destroyed should be removed from the menu stack if its in it
-		MenuStack.RemoveWindow( CurrentWindow );
+		MenuStack.OnWindowDestroyed(CurrentWindow);
 
 		// Perform actual cleanup of the window
 		PrivateDestroyWindow( CurrentWindow );
@@ -2892,22 +2993,11 @@ void FSlateApplication::UpdateToolTip( bool AllowSpawningOfNewToolTips )
 	// spawned from menu items in a different window, for example.
 	if( bHaveForceFieldRect && WidgetsToQueryForToolTip.IsValid() )
 	{
-		const int32 MenuStackLevel = MenuStack.FindLocationInStack( WidgetsToQueryForToolTip.GetWindow() );
+		TSharedPtr<IMenu> MenuInPath = MenuStack.FindMenuInWidgetPath(WidgetsToQueryForToolTip);
 
-		// Also check widgets in pop-up menus owned by this window
-		for( int32 CurStackLevel = MenuStackLevel + 1; CurStackLevel < MenuStack.GetNumStackLevels(); ++CurStackLevel )
+		if (MenuInPath.IsValid())
 		{
-			FMenuWindowList& Windows = MenuStack.GetWindowsAtStackLevel( CurStackLevel );
-
-			for( FMenuWindowList::TConstIterator WindowIt( Windows ); WindowIt; ++WindowIt )
-			{
-				TSharedPtr< SWindow > CurWindow = *WindowIt;
-				if( CurWindow.IsValid() )
-				{
-					const FGeometry& WindowGeometry = CurWindow->GetWindowGeometryInScreen();
-					ForceFieldRect = ForceFieldRect.Expand( WindowGeometry.GetClippingRect() );
-				}
-			}
+			ForceFieldRect = ForceFieldRect.Expand(MenuStack.GetToolTipForceFieldRect(MenuInPath.ToSharedRef(), WidgetsToQueryForToolTip));
 		}
 	}
 
diff --git a/Engine/Source/Runtime/Slate/Private/Framework/Docking/SDockingArea.cpp b/Engine/Source/Runtime/Slate/Private/Framework/Docking/SDockingArea.cpp
--- a/Engine/Source/Runtime/Slate/Private/Framework/Docking/SDockingArea.cpp
+++ b/Engine/Source/Runtime/Slate/Private/Framework/Docking/SDockingArea.cpp
@@ -251,7 +251,7 @@ void SDockingArea::SetParentWindow( const TSharedRef<SWindow>& NewParentWindow )
 	}
 
 	ParentWindowPtr = NewParentWindow;
-	NewParentWindow->SetOnWindowActivated( FOnWindowActivated::CreateSP( this, &SDockingArea::OnOwningWindowActivated ) );
+	NewParentWindow->GetOnWindowActivatedEvent().AddSP(this, &SDockingArea::OnOwningWindowActivated);
 }
 
 TSharedPtr<FTabManager::FLayoutNode> SDockingArea::GatherPersistentLayout() const
diff --git a/Engine/Source/Runtime/Slate/Private/Framework/Docking/SDockingTabStack.cpp b/Engine/Source/Runtime/Slate/Private/Framework/Docking/SDockingTabStack.cpp
--- a/Engine/Source/Runtime/Slate/Private/Framework/Docking/SDockingTabStack.cpp
+++ b/Engine/Source/Runtime/Slate/Private/Framework/Docking/SDockingTabStack.cpp
@@ -490,9 +490,10 @@ void SDockingTabStack::CloseAllButForegroundTab(ETabsToClose TabsToClose)
 
 FReply SDockingTabStack::TabWellRightClicked( const FGeometry& TabWellGeometry, const FPointerEvent& MouseEvent )
 {
-	if ( MouseEvent.GetEffectingButton() == EKeys::RightMouseButton )
+	if (MouseEvent.GetEffectingButton() == EKeys::RightMouseButton)
 	{
-		FSlateApplication::Get().PushMenu( SharedThis( this ), MakeContextMenu(), FSlateApplication::Get().GetCursorPos(), FPopupTransitionEffect::ContextMenu );
+		FWidgetPath WidgetPath = MouseEvent.GetEventPath() != nullptr ? *MouseEvent.GetEventPath() : FWidgetPath();
+		FSlateApplication::Get().PushMenu(AsShared(), WidgetPath, MakeContextMenu(), FSlateApplication::Get().GetCursorPos(), FPopupTransitionEffect::ContextMenu);
 		return FReply::Handled();
 	}
 	else
diff --git a/Engine/Source/Runtime/Slate/Private/Framework/MultiBox/SMenuEntryBlock.cpp b/Engine/Source/Runtime/Slate/Private/Framework/MultiBox/SMenuEntryBlock.cpp
--- a/Engine/Source/Runtime/Slate/Private/Framework/MultiBox/SMenuEntryBlock.cpp
+++ b/Engine/Source/Runtime/Slate/Private/Framework/MultiBox/SMenuEntryBlock.cpp
@@ -874,8 +874,7 @@ void SMenuEntryBlock::OnClicked( bool bCheckBoxClicked )
 			if( MenuEntryBlock->bCloseSelfOnly )
 			{
 				// Close only this menu and its children
-				TSharedRef<SWindow> ParentContextMenuWindow = FSlateApplication::Get().FindWidgetWindow( AsShared() ).ToSharedRef();
-				FSlateApplication::Get().DismissMenu( ParentContextMenuWindow );
+				FSlateApplication::Get().DismissMenuByWidget(AsShared());
 			}
 			else
 			{
diff --git a/Engine/Source/Runtime/Slate/Private/Framework/Text/TextEditHelper.cpp b/Engine/Source/Runtime/Slate/Private/Framework/Text/TextEditHelper.cpp
--- a/Engine/Source/Runtime/Slate/Private/Framework/Text/TextEditHelper.cpp
+++ b/Engine/Source/Runtime/Slate/Private/Framework/Text/TextEditHelper.cpp
@@ -421,7 +421,8 @@ FReply FTextEditHelper::OnMouseButtonUp( const FGeometry& MyGeometry, const FPoi
 			if ( MyGeometry.IsUnderLocation( InMouseEvent.GetScreenSpacePosition() ) )
 			{
 				// Right clicked, so summon a context menu if the cursor is within the widget
-				TextEditor->SummonContextMenu(InMouseEvent.GetScreenSpacePosition(), InMouseEvent.GetWindow());
+				FWidgetPath WidgetPath = InMouseEvent.GetEventPath() != nullptr ? *InMouseEvent.GetEventPath() : FWidgetPath();
+				TextEditor->SummonContextMenu(InMouseEvent.GetScreenSpacePosition(), InMouseEvent.GetWindow(), WidgetPath);
 			}
 
 			// Release mouse capture
diff --git a/Engine/Source/Runtime/Slate/Private/Widgets/Input/SCheckBox.cpp b/Engine/Source/Runtime/Slate/Private/Widgets/Input/SCheckBox.cpp
--- a/Engine/Source/Runtime/Slate/Private/Widgets/Input/SCheckBox.cpp
+++ b/Engine/Source/Runtime/Slate/Private/Widgets/Input/SCheckBox.cpp
@@ -171,8 +171,11 @@ FReply SCheckBox::OnMouseButtonDown( const FGeometry& MyGeometry, const FPointer
 	}
 	else if ( MouseEvent.GetEffectingButton() == EKeys::RightMouseButton && OnGetMenuContent.IsBound() )
 	{
+		FWidgetPath WidgetPath = MouseEvent.GetEventPath() != nullptr ? *MouseEvent.GetEventPath() : FWidgetPath();
+
 		FSlateApplication::Get().PushMenu(
 			AsShared(),
+			WidgetPath,
 			OnGetMenuContent.Execute(),
 			MouseEvent.GetScreenSpacePosition(),
 			FPopupTransitionEffect( FPopupTransitionEffect::ContextMenu )
diff --git a/Engine/Source/Runtime/Slate/Private/Widgets/Input/SComboButton.cpp b/Engine/Source/Runtime/Slate/Private/Widgets/Input/SComboButton.cpp
--- a/Engine/Source/Runtime/Slate/Private/Widgets/Input/SComboButton.cpp
+++ b/Engine/Source/Runtime/Slate/Private/Widgets/Input/SComboButton.cpp
@@ -111,7 +111,7 @@ FReply SComboButton::OnButtonClicked()
 
 void SComboButton::SetMenuContent(TSharedRef<SWidget> InContent)
 {
-	MenuContent = 
+	WrappedContent = MenuContent =
 		SNew(SBorder)
 		.BorderImage(MenuBorderBrush)
 		.Padding(MenuBorderPadding)
diff --git a/Engine/Source/Runtime/Slate/Private/Widgets/Input/SEditableText.cpp b/Engine/Source/Runtime/Slate/Private/Widgets/Input/SEditableText.cpp
--- a/Engine/Source/Runtime/Slate/Private/Widgets/Input/SEditableText.cpp
+++ b/Engine/Source/Runtime/Slate/Private/Widgets/Input/SEditableText.cpp
@@ -4,6 +4,7 @@
 #include "BreakIterator.h"
 #include "TextEditHelper.h"
 #include "GenericCommands.h"
+#include "IMenu.h"
 
 /** A pointer to the editable text currently under the mouse cursor. nullptr when there isn't one. */
 SEditableText* SEditableText::EditableTextUnderCursor = nullptr;
@@ -1799,8 +1800,7 @@ void SEditableText::SaveText()
 	}
 }
 
-
-void SEditableText::SummonContextMenu(const FVector2D& InLocation, TSharedPtr<SWindow> ParentWindow)
+void SEditableText::SummonContextMenu(const FVector2D& InLocation, TSharedPtr<SWindow> ParentWindow, const FWidgetPath& EventPath)
 {
 	TSharedPtr<SWidget> MenuContentWidget;
 
@@ -1866,13 +1866,13 @@ void SEditableText::SummonContextMenu(const FVector2D& InLocation, TSharedPtr<SW
 		{
 			MenuParent = StaticCastSharedRef<SWidget>(ParentWindow.ToSharedRef());
 		}
-		TSharedPtr< SWindow > ContextMenuWindow = FSlateApplication::Get().PushMenu(MenuParent, MenuContentWidget.ToSharedRef(), InLocation, FPopupTransitionEffect(FPopupTransitionEffect::ContextMenu), bFocusImmediately);
+		TSharedPtr< IMenu > ContextMenu = FSlateApplication::Get().PushMenu(MenuParent, EventPath, MenuContentWidget.ToSharedRef(), InLocation, FPopupTransitionEffect(FPopupTransitionEffect::ContextMenu), bFocusImmediately);
 
-		// Make sure the window is valid.  It's possible for the parent to already be in the destroy queue, for example if the editable text was configured to dismiss it's window during OnTextCommitted.
-		if (ContextMenuWindow.IsValid())
+		// Make sure the menu is valid.  It's possible for the parent to already be in the destroy queue, for example if the editable text was configured to dismiss it's window during OnTextCommitted.
+		if (ContextMenu.IsValid())
 		{
-			ContextMenuWindow->SetOnWindowClosed(FOnWindowClosed::CreateSP(this, &SEditableText::OnWindowClosed));
-			ActiveContextMenu.SummonSucceeded(ContextMenuWindow.ToSharedRef());
+			ContextMenu->GetOnMenuDismissed().AddSP(this, &SEditableText::OnContextMenuClosed);
+			ActiveContextMenu.SummonSucceeded(ContextMenu.ToSharedRef());
 		}
 		else
 		{
@@ -1991,7 +1991,7 @@ FString SEditableText::GetStringToRender() const
 	return VisibleText;
 }
 
-void SEditableText::OnWindowClosed(const TSharedRef<SWindow>&)
+void SEditableText::OnContextMenuClosed(TSharedRef<IMenu> Menu)
 {
 	// Note: We don't reset the ActiveContextMenu here, as Slate hasn't yet finished processing window focus events, and we need 
 	// to know that the window is still available for OnFocusReceived and OnFocusLost even though it's about to be destroyed
diff --git a/Engine/Source/Runtime/Slate/Private/Widgets/Input/SMenuAnchor.cpp b/Engine/Source/Runtime/Slate/Private/Widgets/Input/SMenuAnchor.cpp
--- a/Engine/Source/Runtime/Slate/Private/Widgets/Input/SMenuAnchor.cpp
+++ b/Engine/Source/Runtime/Slate/Private/Widgets/Input/SMenuAnchor.cpp
@@ -94,6 +94,7 @@ void SMenuAnchor::Construct( const FArguments& InArgs )
 		InArgs._Content.Widget
 	];
 	MenuContent = InArgs._MenuContent;
+	WrappedContent = InArgs._MenuContent;
 	OnGetMenuContent = InArgs._OnGetMenuContent;
 	OnMenuOpenChanged = InArgs._OnMenuOpenChanged;
 	Placement = InArgs._Placement;
@@ -160,6 +161,7 @@ void SMenuAnchor::Tick( const FGeometry& AllottedGeometry, const double InCurren
 			LocalPlacement.Orientation, WindowRectLocalSpace);
 
 		LocalPopupPosition = FittedPlacement;
+		ScreenPopupPosition = AllottedGeometry.GetAccumulatedLayoutTransform().TransformPoint(LocalPopupPosition);
 	}
 
 	/** The tick is ending, so the window was not dismissed this tick. */
@@ -234,30 +236,6 @@ int32 SMenuAnchor::OnPaint( const FPaintArgs& Args, const FGeometry& AllottedGeo
 	return LayerId;
 }
 
-
-void SMenuAnchor::RequestClosePopupWindow( const TSharedRef<SWindow>& PopupWindow )
-{
-	bDismissedThisTick = true;
-	if (ensure(IsOpenViaCreatedWindow()))
-	{
-		FSlateApplication::Get().RequestDestroyWindow(PopupWindow);
-		
-		if (OnMenuOpenChanged.IsBound())
-		{
-			OnMenuOpenChanged.Execute(false);
-		}
-	}
-}
-
-void SMenuAnchor::OnClickedOutsidePopup()
-{
-	if (IsOpen())
-	{
-		bDismissedThisTick = true;
-		SetIsOpen(false);
-	}
-}
-
 bool SMenuAnchor::IsOpenAndReusingWindow() const
 {
 	return MethodInUse.IsSet() && MethodInUse.GetValue() == EPopupMethod::UseCurrentWindow;
@@ -280,6 +258,7 @@ void SMenuAnchor::SetContent(TSharedRef<SWidget> InContent)
 void SMenuAnchor::SetMenuContent(TSharedRef<SWidget> InMenuContent)
 {
 	MenuContent = InMenuContent;
+	WrappedContent = InMenuContent;	// wrapping, if any will happen when the menu is opened
 }
 
 EPopupMethod QueryPopupMethod(const FWidgetPath& PathToQuery)
@@ -369,28 +348,45 @@ void SMenuAnchor::SetIsOpen( bool InIsOpen, const bool bFocusMenu )
 					if (MethodInUse == EPopupMethod::CreateNewWindow)
 					{
 						// Open the pop-up
-						TSharedRef<SWindow> PopupWindow = FSlateApplication::Get().PushMenu( AsShared(), MenuContentRef, NewPosition, TransitionEffect, bFocusMenu, false, NewWindowSize, SummonLocationSize );
-						PopupWindow->SetRequestDestroyWindowOverride( FRequestDestroyWindowOverride::CreateSP( this, &SMenuAnchor::RequestClosePopupWindow ) );
-						PopupWindowPtr = PopupWindow;
+						TSharedPtr<IMenu> NewMenu = FSlateApplication::Get().PushMenu(AsShared(), MyWidgetPath, MenuContentRef, NewPosition, TransitionEffect, bFocusMenu, SummonLocationSize, MethodInUse);
+
+						PopupMenuPtr = NewMenu;
+						check(NewMenu.IsValid() && NewMenu->GetOwnedWindow().IsValid());
+						NewMenu->GetOnMenuDismissed().AddSP(this, &SMenuAnchor::OnMenuClosed);
+						PopupWindowPtr = NewMenu->GetOwnedWindow();
 					}
 					else
 					{
 						// We are re-using the current window instead of creating a new one.
 						// The popup will be presented via an overlay service.
 						ensure(MethodInUse == EPopupMethod::UseCurrentWindow);
-						PopupWindowPtr = MyWidgetPath.GetWindow();
+						TSharedRef<SWindow> PopupWindow = MyWidgetPath.GetWindow();
+						PopupWindowPtr = PopupWindow;
+
+						if (bFocusMenu)
+						{
+							FSlateApplication::Get().ReleaseMouseCapture();
+						}
+
+						TSharedRef<SMenuAnchor> SharedThis = StaticCastSharedRef<SMenuAnchor>(AsShared());
+
+						TSharedPtr<IMenu> NewMenu = FSlateApplication::Get().PushHostedMenu(
+							SharedThis, MyWidgetPath, SharedThis, MenuContentRef, WrappedContent, TransitionEffect);
+
+						PopupMenuPtr = NewMenu;
+						check(NewMenu.IsValid());
+						check(NewMenu->GetParentWindow().ToSharedRef() == PopupWindow);
+						check(WrappedContent.IsValid());
+
 						Children[1]
 						[
-							MenuContentRef
+							WrappedContent.ToSharedRef()
 						];
 
-						// We want to support dismissing the popup widget when the user clicks outside it.
-						OnClickedOutsidePopupDelegateHandle = FSlateApplication::Get().GetPopupSupport().RegisterClickNotification( MenuContentRef, FOnClickedOutside::CreateSP( this, &SMenuAnchor::OnClickedOutsidePopup ) );
-					}
-
-					if ( bFocusMenu )
-					{
-						FSlateApplication::Get().SetKeyboardFocus( MenuContentRef, EFocusCause::SetDirectly );
+						if (bFocusMenu)
+						{
+							FSlateApplication::Get().SetKeyboardFocus(MenuContentRef, EFocusCause::SetDirectly);
+						}
 					}
 				}
 			}
@@ -398,38 +394,35 @@ void SMenuAnchor::SetIsOpen( bool InIsOpen, const bool bFocusMenu )
 		else
 		{
 			// CLOSE POPUP
-
-			if (IsOpenAndReusingWindow())
+			if (PopupMenuPtr.IsValid())
 			{
-				// Window reuse case; remove child widget.
-				Children[1].DetachWidget();
-				FSlateApplication::Get().GetPopupSupport().UnregisterClickNotification( OnClickedOutsidePopupDelegateHandle );
+				PopupMenuPtr.Pin()->Dismiss();
 			}
 			else
 			{
-				// New window case
-				ensure(IsOpenViaCreatedWindow());
-				TSharedPtr<SWindow> PopupWindow = PopupWindowPtr.Pin();
-				if (PopupWindow.IsValid())
-				{
-					// Request that the popup be closed.
-					PopupWindow->RequestDestroyWindow();
-				}
+				PopupWindowPtr.Reset();
+				MethodInUse.Reset();
 			}
+		}
+	}
+}
 
-			if (OnMenuOpenChanged.IsBound())
-			{
-				OnMenuOpenChanged.Execute(false);
-			}
+void SMenuAnchor::OnMenuClosed(TSharedRef<IMenu> InMenu)
+{
+	bDismissedThisTick = true;
+	MethodInUse = TOptional<EPopupMethod>();
+	PopupMenuPtr.Reset();
+	PopupWindowPtr.Reset();
 
-			MethodInUse = TOptional<EPopupMethod>();
-		}
+	if (OnMenuOpenChanged.IsBound())
+	{
+		OnMenuOpenChanged.Execute(false);
 	}
 }
 
 bool SMenuAnchor::IsOpen() const
 {
-	return (MethodInUse.IsSet() && PopupWindowPtr.Pin().IsValid()) || IsOpenAndReusingWindow();
+	return MethodInUse.IsSet() && PopupMenuPtr.IsValid();
 }
 
 bool SMenuAnchor::ShouldOpenDueToClick() const
@@ -440,10 +433,15 @@ bool SMenuAnchor::ShouldOpenDueToClick() const
 FVector2D SMenuAnchor::GetMenuPosition() const
 {
 	FVector2D Pos(0,0);
-	if( PopupWindowPtr.IsValid() )
+
+	if (IsOpenViaCreatedWindow() && PopupWindowPtr.IsValid())
 	{
 		Pos = PopupWindowPtr.Pin()->GetPositionInScreen();
 	}
+	else if (IsOpenAndReusingWindow() && PopupMenuPtr.IsValid())
+	{
+		Pos = ScreenPopupPosition;
+	}
 	
 	return Pos;
 }
@@ -451,15 +449,29 @@ FVector2D SMenuAnchor::GetMenuPosition() const
 bool SMenuAnchor::HasOpenSubMenus() const
 {
 	bool Result = false;
-	if( PopupWindowPtr.IsValid() )
+	if (PopupMenuPtr.IsValid())
 	{
-		Result = FSlateApplication::Get().HasOpenSubMenus( PopupWindowPtr.Pin().ToSharedRef() );
+		Result = FSlateApplication::Get().HasOpenSubMenus(PopupMenuPtr.Pin());
 	}
 	return Result;
 }
 
+TSharedPtr<SWindow> SMenuAnchor::GetMenuWindow() const
+{
+	return IsOpen() ? PopupWindowPtr.Pin() : TSharedPtr<SWindow>();
+}
+
+void SMenuAnchor::OnMenuDismissed()
+{
+	if (PopupMenuPtr.IsValid())
+	{
+		OnMenuClosed(PopupMenuPtr.Pin().ToSharedRef());
+	}
+}
+
 SMenuAnchor::SMenuAnchor()
 	: MenuContent( SNullWidget::NullWidget )
+	, WrappedContent(SNullWidget::NullWidget)
 	, bDismissedThisTick( false )
 	, Method()
 	, MethodInUse()
@@ -469,17 +481,13 @@ SMenuAnchor::SMenuAnchor()
 
 SMenuAnchor::~SMenuAnchor()
 {
-	TSharedPtr<SWindow> PopupWindow = PopupWindowPtr.Pin();
-	if (PopupWindow.IsValid())
+	if (PopupMenuPtr.IsValid())
 	{
-		// Close the Popup window.
-		if (IsOpenViaCreatedWindow())
-		{
-			// Request that the popup be closed.
-			PopupWindow->RequestDestroyWindow();
-		}
-		
-		// We no longer have a popup open, so reset all the tracking state associated.
-		PopupWindowPtr.Reset();
+		PopupMenuPtr.Pin()->Dismiss();
 	}
+		
+	// We no longer have a popup open, so reset all the tracking state associated.
+	PopupMenuPtr.Reset();
+	PopupWindowPtr.Reset();
+	MethodInUse.Reset();
 }
diff --git a/Engine/Source/Runtime/Slate/Private/Widgets/Text/SMultiLineEditableText.cpp b/Engine/Source/Runtime/Slate/Private/Widgets/Text/SMultiLineEditableText.cpp
--- a/Engine/Source/Runtime/Slate/Private/Widgets/Text/SMultiLineEditableText.cpp
+++ b/Engine/Source/Runtime/Slate/Private/Widgets/Text/SMultiLineEditableText.cpp
@@ -8,6 +8,7 @@
 #include "PlainTextLayoutMarshaller.h"
 #include "TextBlockLayout.h"
 #include "GenericCommands.h"
+#include "IMenu.h"
 
 void SMultiLineEditableText::FCursorInfo::SetCursorLocationAndCalculateAlignment(const TSharedPtr<FTextLayout>& InTextLayout, const FTextLocation& InCursorPosition)
 {
@@ -2320,7 +2321,7 @@ TSharedRef< SWidget > SMultiLineEditableText::GetWidget()
 	return SharedThis( this );
 }
 
-void SMultiLineEditableText::SummonContextMenu(const FVector2D& InLocation, TSharedPtr<SWindow> ParentWindow)
+void SMultiLineEditableText::SummonContextMenu(const FVector2D& InLocation, TSharedPtr<SWindow> ParentWindow, const FWidgetPath& EventPath)
 {
 	TSharedPtr<SWidget> MenuContentWidget;
 
@@ -2386,13 +2387,13 @@ void SMultiLineEditableText::SummonContextMenu(const FVector2D& InLocation, TSha
 		{
 			MenuParent = StaticCastSharedRef<SWidget>(ParentWindow.ToSharedRef());
 		}
-		TSharedPtr< SWindow > ContextMenuWindow = FSlateApplication::Get().PushMenu(MenuParent, MenuContentWidget.ToSharedRef(), InLocation, FPopupTransitionEffect(FPopupTransitionEffect::ContextMenu), bFocusImmediately);
+		TSharedPtr< IMenu > ContextMenu = FSlateApplication::Get().PushMenu(MenuParent, EventPath, MenuContentWidget.ToSharedRef(), InLocation, FPopupTransitionEffect(FPopupTransitionEffect::ContextMenu), bFocusImmediately);
 
 		// Make sure the window is valid.  It's possible for the parent to already be in the destroy queue, for example if the editable text was configured to dismiss it's window during OnTextCommitted.
-		if (ContextMenuWindow.IsValid())
+		if (ContextMenu.IsValid())
 		{
-			ContextMenuWindow->SetOnWindowClosed(FOnWindowClosed::CreateSP(this, &SMultiLineEditableText::OnWindowClosed));
-			ActiveContextMenu.SummonSucceeded(ContextMenuWindow.ToSharedRef());
+			ContextMenu->GetOnMenuDismissed().AddSP(this, &SMultiLineEditableText::OnContextMenuClosed);
+			ActiveContextMenu.SummonSucceeded(ContextMenu.ToSharedRef());
 		}
 		else
 		{
@@ -2401,7 +2402,7 @@ void SMultiLineEditableText::SummonContextMenu(const FVector2D& InLocation, TSha
 	}
 }
 
-void SMultiLineEditableText::OnWindowClosed(const TSharedRef<SWindow>&)
+void SMultiLineEditableText::OnContextMenuClosed(TSharedRef<IMenu> Menu)
 {
 	// Note: We don't reset the ActiveContextMenu here, as Slate hasn't yet finished processing window focus events, and we need 
 	// to know that the window is still available for OnFocusReceived and OnFocusLost even though it's about to be destroyed
diff --git a/Engine/Source/Runtime/Slate/Private/Widgets/Views/SHeaderRow.cpp b/Engine/Source/Runtime/Slate/Private/Widgets/Views/SHeaderRow.cpp
--- a/Engine/Source/Runtime/Slate/Private/Widgets/Views/SHeaderRow.cpp
+++ b/Engine/Source/Runtime/Slate/Private/Widgets/Views/SHeaderRow.cpp
@@ -208,7 +208,7 @@ public:
 	{
 		if ( MouseEvent.GetEffectingButton() == EKeys::RightMouseButton && ContextMenuContent != SNullWidget::NullWidget )
 		{
-			OpenContextMenu( MouseEvent.GetScreenSpacePosition() );
+			OpenContextMenu( MouseEvent );
 			return FReply::Handled();
 		}
 
@@ -355,11 +355,14 @@ private:
 		return FReply::Handled();
 	}
 
-	void OpenContextMenu(const FVector2D& SummonLocation)
+	void OpenContextMenu(const FPointerEvent& MouseEvent)
 	{
 		if ( ContextMenuContent != SNullWidget::NullWidget )
 		{
-			FSlateApplication::Get().PushMenu( AsShared(), ContextMenuContent, SummonLocation, FPopupTransitionEffect( FPopupTransitionEffect::ContextMenu ) );
+			const FVector2D& SummonLocation = MouseEvent.GetScreenSpacePosition();
+			FWidgetPath WidgetPath = MouseEvent.GetEventPath() != nullptr ? *MouseEvent.GetEventPath() : FWidgetPath();
+
+			FSlateApplication::Get().PushMenu(AsShared(), WidgetPath, ContextMenuContent, SummonLocation, FPopupTransitionEffect(FPopupTransitionEffect::ContextMenu));
 		}
 	}
 
diff --git a/Engine/Source/Runtime/Slate/Private/Widgets/Views/STableViewBase.cpp b/Engine/Source/Runtime/Slate/Private/Widgets/Views/STableViewBase.cpp
--- a/Engine/Source/Runtime/Slate/Private/Widgets/Views/STableViewBase.cpp
+++ b/Engine/Source/Runtime/Slate/Private/Widgets/Views/STableViewBase.cpp
@@ -377,7 +377,7 @@ FReply STableViewBase::OnMouseButtonUp( const FGeometry& MyGeometry, const FPoin
 	if ( MouseEvent.GetEffectingButton() == EKeys::RightMouseButton )
 	{
 
-		OnRightMouseButtonUp( MouseEvent.GetScreenSpacePosition() );
+		OnRightMouseButtonUp( MouseEvent );
 
 		FReply Reply = FReply::Handled().ReleaseMouseCapture();
 		bShowSoftwareCursor = false;
@@ -744,8 +744,9 @@ int32 STableViewBase::GetNumItemsWide() const
 	return 1;
 }
 
-void STableViewBase::OnRightMouseButtonUp(const FVector2D& SummonLocation)
+void STableViewBase::OnRightMouseButtonUp(const FPointerEvent& MouseEvent)
 {
+	const FVector2D& SummonLocation = MouseEvent.GetScreenSpacePosition();
 	const bool bShouldOpenContextMenu = !IsRightClickScrolling();
 	const bool bContextMenuOpeningBound = OnContextMenuOpening.IsBound();
 
@@ -757,7 +758,9 @@ void STableViewBase::OnRightMouseButtonUp(const FVector2D& SummonLocation)
 		if( MenuContent.IsValid() )
 		{
 			bShowSoftwareCursor = false;
-			FSlateApplication::Get().PushMenu( AsShared(), MenuContent.ToSharedRef(), SummonLocation, FPopupTransitionEffect( FPopupTransitionEffect::ContextMenu ) );
+
+			FWidgetPath WidgetPath = MouseEvent.GetEventPath() != nullptr ? *MouseEvent.GetEventPath() : FWidgetPath();
+			FSlateApplication::Get().PushMenu(AsShared(), WidgetPath, MenuContent.ToSharedRef(), SummonLocation, FPopupTransitionEffect(FPopupTransitionEffect::ContextMenu));
 
 
 		}
diff --git a/Engine/Source/Runtime/Slate/Public/Framework/Application/IMenu.h b/Engine/Source/Runtime/Slate/Public/Framework/Application/IMenu.h
new file mode 100644
--- /dev/null
+++ b/Engine/Source/Runtime/Slate/Public/Framework/Application/IMenu.h
@@ -0,0 +1,26 @@
+// Copyright 1998-2015 Epic Games, Inc. All Rights Reserved.
+
+#pragma once
+
+/**
+ * Represents a popup menu.
+ */
+class IMenu
+{
+public:
+	DECLARE_MULTICAST_DELEGATE_OneParam(FOnMenuDismissed, TSharedRef<IMenu> /*DismissedMenu*/);
+
+	virtual EPopupMethod GetPopupMethod() const = 0;
+	virtual TSharedPtr<SWindow> GetParentWindow() const = 0;
+	virtual TSharedPtr<SWindow> GetOwnedWindow() const = 0;
+	virtual TSharedPtr<SWidget> GetContent() const = 0;
+	virtual FOnMenuDismissed& GetOnMenuDismissed() = 0;
+	virtual void Dismiss() = 0;
+};
+
+class IMenuHost
+{
+public:
+	virtual TSharedPtr<SWindow> GetMenuWindow() const = 0;
+	virtual void OnMenuDismissed() = 0;
+};
diff --git a/Engine/Source/Runtime/Slate/Public/Framework/Application/MenuStack.h b/Engine/Source/Runtime/Slate/Public/Framework/Application/MenuStack.h
--- a/Engine/Source/Runtime/Slate/Public/Framework/Application/MenuStack.h
+++ b/Engine/Source/Runtime/Slate/Public/Framework/Application/MenuStack.h
@@ -2,12 +2,19 @@
 
 #pragma once
 
+// forward declarations
+class IMenu;
+class FMenuBase;
+class IMenuHost;
+class SMenuPanel;
+
+typedef TArray<TSharedPtr<SWindow>> FMenuWindowList;	// deprecated
 
 #if UE_BUILD_DEBUG
 	// Debugging is harder with the inline allocator
-	typedef TArray<TSharedPtr<SWindow>> FMenuWindowList;
+	typedef TArray<TSharedPtr<IMenu>> FMenuList;
 #else
-	typedef TArray<TSharedPtr<SWindow>, TInlineAllocator<4>> FMenuWindowList;
+	typedef TArray<TSharedPtr<IMenu>, TInlineAllocator<4>> FMenuList;
 #endif
 
 
@@ -44,19 +51,17 @@ struct FPopupTransitionEffect
 	{ }
 };
 
-
 /**
- * Represents a stack of open menus.  The last item in the stack is the top most menu
- * Each stack level can have any number of windows opened on that level but they all must be parented to the same window
- * which must reside in the previous level of the stack
+ * Represents a stack of open menus. The last item in the stack is the top most menu
+ * Menus are described as IMenus. Implementations of IMenu can control how the menus are created and presented (e.g. in their own window)
  */
 class FMenuStack
 {
 public:
-
 	/**
-	 * Pushes a new menu onto the stack.  Automatically removes windows from the stack that are no longer valid
+	 * [Deprecated] Pushes a new menu onto the stack.  Automatically removes windows from the stack that are no longer valid
 	 * Invalid windows are those that are not parents of the window menu being pushed and any windows on the same stack level
+	 * This is the old method for creating menus.
 	 *
 	 * @param InParentWindow		The parent window of the menu
 	 * @param InContent				The menu's content
@@ -67,60 +72,328 @@ public:
 	 * @param WindowSize			When bShouldAutoSize=false, this must be set to the size of the window to be created
 	 * @param SummonLocationSize	An optional size around the summon location which describes an area in which the menu may not appear
 	 */
+	DEPRECATED(4.9, "PushMenu() returning an SWindow is deprecated. Use Push() that returns an IMenu instead.")
 	TSharedRef<SWindow> PushMenu( const TSharedRef<SWindow>& InParentWindow, const TSharedRef<SWidget>& InContent, const FVector2D& SummonLocation, const FPopupTransitionEffect& TransitionEffect, const bool bFocusImmediately = true, const bool bShouldAutoSize = true, const FVector2D& WindowSize = FVector2D::ZeroVector, const FVector2D& SummonLocationSize = FVector2D::ZeroVector );
 
 	/**
-	 * Dismisses the menu stack up to a certian level (by default removes all levels in the stack)
+	 * Pushes a new menu onto the stack. The widget path will be searched for existing menus and the new menu will be parented appropriately.
+	 * Menus are always auto-sized. Use fixed-size content if a fixed size is required.
+	 * 
+	 * @param InOwnerPath			The widget path for the parent widget of this menu.
+	 * @param InContent				The menu's content
+	 * @param SummonLocation		Location in screen-space where the menu should appear
+	 * @param TransitionEffect		Animation to use when the popup appears
+	 * @param bFocusImmediately		Should the popup steal focus when shown?
+	 * @param SummonLocationSize	An optional size around the summon location which describes an area in which the menu may not appear
+	 * @param InMethod				An optional popup method that will override the default method for the widgets in InOwnerPath
+	 */
+	TSharedRef<IMenu> Push(const FWidgetPath& InOwnerPath, const TSharedRef<SWidget>& InContent, const FVector2D& SummonLocation, const FPopupTransitionEffect& TransitionEffect, const bool bFocusImmediately = true, const FVector2D& SummonLocationSize = FVector2D::ZeroVector, TOptional<EPopupMethod> InMethod = TOptional<EPopupMethod>());
+	
+	/**
+	 * Pushes a new child menu onto the stack.
+	 * Menus are always auto-sized. Use fixed-size content if a fixed size is required.
+	 * 
+	 * @param InParentMenu			The parent menu for this menu
+	 * @param InContent				The menu's content
+	 * @param SummonLocation		Location in screen-space where the menu should appear
+	 * @param TransitionEffect		Animation to use when the popup appears
+	 * @param bFocusImmediately		Should the popup steal focus when shown?
+	 * @param SummonLocationSize	An optional size around the summon location which describes an area in which the menu may not appear
+	 */
+	TSharedRef<IMenu> Push(const TSharedPtr<IMenu>& InParentMenu, const TSharedRef<SWidget>& InContent, const FVector2D& SummonLocation, const FPopupTransitionEffect& TransitionEffect, const bool bFocusImmediately = true, const FVector2D& SummonLocationSize = FVector2D::ZeroVector);
+
+	/**
+	 * Pushes a new menu onto the stack that is drawn by an external host widget.
+	 * The widget path will be searched for existing menus and the new menu will be parented appropriately.
+	 * Menus are always auto-sized. Use fixed-size content if a fixed size is required.
+	 * 
+	 * @param InOwnerPath			The widget path for the parent widget of this menu
+	 * @param InMenuHost			The host widget that draws the menu's content
+	 * @param InContent				The menu's content
+	 * @param OutWrappedContent		Returns the InContent wrapped with widgets needed by the menu stack system. This is what should be drawn by the host after this call.
+	 * @param TransitionEffect		Animation to use when the popup appears
+	 */
+	TSharedRef<IMenu> PushHosted(const FWidgetPath& InOwnerPath, const TSharedRef<IMenuHost>& InMenuHost, const TSharedRef<SWidget>& InContent, TSharedPtr<SWidget>& OutWrappedContent, const FPopupTransitionEffect& TransitionEffect);
+	
+	/**
+	 * Pushes a new child menu onto the stack that is drawn by an external host widget.
+	 * Menus are always auto-sized. Use fixed-size content if a fixed size is required.
+	 * 
+	 * @param InParentMenu			The parent menu for this menu
+	 * @param InMenuHost			The host widget that draws the menu's content
+	 * @param InContent				The menu's content
+	 * @param OutWrappedContent		Returns the InContent wrapped with widgets needed by the menu stack system. This is what should be drawn by the host after this call.
+	 * @param TransitionEffect		Animation to use when the popup appears
+	 */
+	TSharedRef<IMenu> PushHosted(const TSharedPtr<IMenu>& InParentMenu, const TSharedRef<IMenuHost>& InMenuHost, const TSharedRef<SWidget>& InContent, TSharedPtr<SWidget>& OutWrappedContent, const FPopupTransitionEffect& TransitionEffect);
+
+	/**
+	 * [Deprecated] Dismisses the menu stack up to a certain level (by default removes all levels in the stack)
 	 * Dismisses in reverse order so to keep just the first level you would pass in 1
 	 *
 	 * @param FirstStackIndexToRemove	The first level of the stack to remove
 	 */
+	DEPRECATED(4.9, "Dismiss() taking a stack index is deprecated. Use DismissFrom() or DismissAll().")
 	void Dismiss( int32 FirstStackIndexToRemove = 0 );
 	
+	/**
+	 * Dismisses the menu stack including InFromMenu and all its child menus
+	 * Dismisses in reverse order (children first)
+	 *
+	 * @param InFromMenu		The menu to remove along with its children
+	 */
+	void DismissFrom(const TSharedPtr<IMenu>& InFromMenu);
+
+	/**
+	 * Dismisses the entire menu stack
+	 */
+	void DismissAll();
+
 	/**
 	 * Removes a window from the stack and all of its children.  Assumes the window is cleaned up elsewhere
 	 *
 	 * @param WindowToRemove	The window to remove (will remove all children of this window as well)	
 	 */
+	DEPRECATED(4.9, "RemoveWindow() is deprecated. Do not use.")
 	void RemoveWindow( TSharedRef<SWindow> WindowToRemove );
 
 	/**
-	 * Notifies the stack that a new window was activated.  The menu stack will be dismissed if the activated window is not in the menu stack
+	 * Called by the application when any window is destroyed.
+	 *
+	 * @param InWindow	The window that was destroyed
+	 */
+	void OnWindowDestroyed(TSharedRef<SWindow> InWindow);
+
+	/**
+	 * Notifies the stack that a new window was activated.  The menu stack will be dismissed if the activated window is not a menu in the stack
 	 *
 	 * @param ActivatedWindow	The window that was just activated
 	 */
 	void OnWindowActivated( TSharedRef<SWindow> ActivatedWindow );
 
 	/**
-	 * Finds a window in the stack
+	 * [Deprecated] Finds a window in the stack.  (Deprecated but not using the deprecation macro. This is to avoid warnings as this is still called from a deprecated FSlateApplication method)
+	 * This is deprecated because menus are no longer necessarily windows. Callers should also not need to know the stack levels by index. Menus are identified using the IMenu interface.
 	 *
-	 * @param WindowToFind	The window to look fort
-	 * @param The level of the stack where the window is located or INDEX_NONE if the window could not be found
+	 * @param	WindowToFind	The window to look for
+	 * @return	The level of the stack where the window is located or INDEX_NONE if the window could not be found
 	 */
 	int32 FindLocationInStack( TSharedPtr<SWindow> WindowToFind ) const;
 
 	/**
+	 * Finds a menu in the stack that owns InWindow.
+	 * Although this method isn't deprecated, it is intended for private use inside FMenuStack
+	 * and its use as a public method should be avoided. Menus should be identified in client code as IMenu interfaces.
+	 *
+	 * @param	InWindow	The window to look for
+	 * @return	The menu in the stack that owns InWindow, or an invalid ptr if not found
+	 */
+	TSharedPtr<IMenu> FindMenuFromWindow(TSharedRef<SWindow> InWindow) const;
+
+	/**
+	 * Searches from bottom to top of the widget path for a menu in the stack.
+	 * If the widget at the end of PathToQuery is inside a menu, there will be a menu content wrapper widget in the path.
+	 *
+	 * @param	PathToQuery		The widget path to search for a menu
+	 * @return	The menu in the stack that contains the widget at the end of PathToQuery, or an invalid ptr if not found
+	 */
+	TSharedPtr<IMenu> FindMenuInWidgetPath(const FWidgetPath& PathToQuery) const;
+
+	/**
+	 * Called by the application when showing tooltips. It prevents tooltips from drawing over menu items.
+	 *
+	 * @param InMenu				A menu in the stack, used to generate the force field rect
+	 * @param PathContainingMenu	A widget path containing InMenu. This could be generated but the application has the path so it helps performance to pass it in here. 
+	 * @return						A rectangle enclosing the menu's children (created from the clipping rects in the geometry in PathContainingMenu)
+	 */
+	FSlateRect GetToolTipForceFieldRect(TSharedRef<IMenu> InMenu, const FWidgetPath& PathContainingMenu) const;
+
+	/**
+	 * [Deprecated]
 	 * @return	Returns the size of the menu stack
 	 */
+	DEPRECATED(4.9, "GetNumStackLevels() is deprecated. Consider HasMenus().")
 	int32 GetNumStackLevels( ) const;
 
-	/** @return Returns whether the window has child menus. */
+	/**
+	 * @return	Returns the window that is the parent of everything in the stack, if any. If the stack is empty, returns an invalid ptr.
+	 */
+	TSharedPtr<SWindow> GetHostWindow() const;
+
+	/**
+	* @return	True if the stack has one or more menus in it, false if it is empty.
+	*/
+	bool HasMenus() const;
+
+	/**
+	 * [Deprecated] This is deprecated because menus are no longer necessarily windows. Menus are identified using the IMenu interface.
+	 *
+	 * @return Returns whether the window has child menus. If the window is not a menu, returns false.
+	 */
+	DEPRECATED(4.9, "HasOpenSubMenus() taking a window is deprecated. Use HasOpenSubMenus() taking an IMenu as a parameter.")
 	bool HasOpenSubMenus( const TSharedRef<SWindow>& Window ) const;
 
 	/**
-	 * Returns all of the windows at the specified stack level index
+	 * @return Returns whether the menu has child menus. If the menu isn't in the stack, returns false.
+	 */
+	bool HasOpenSubMenus(TSharedPtr<IMenu> InMenu) const;
+
+	/**
+	 * [Deprecated] Returned all of the windows at the specified stack level index. Always returns an empty list because it's deprecated and the return type is not supported.
 	 *
 	 * @param	StackLevelIndex	Index into the menu stack
 	 *
-	 * @return	List of windows at this stack level
+	 * @return	List of windows at this stack level (returns an empty list because it's deprecated)
 	 */
+	DEPRECATED(4.9, "GetWindowsAtStackLevel() is deprecated. Do not use.")
 	FMenuWindowList& GetWindowsAtStackLevel( const int32 StackLevelIndex );
 
 private:
+	/**
+	 * Queries the widgets in the path for a menu popup method.
+	 *
+	 * @param	PathToQuery	The path to query
+	 *
+	 * @return	The popup method chosen by a widget in the path. Defaults to EPopupMethod::CreateNewWindow if no widgets in the path have a preference.
+	 */
+	EPopupMethod QueryPopupMethod(const FWidgetPath& PathToQuery);
+
+	/**
+	 * Called by public Dismiss methods. Dismisses all menus in the stack from FirstStackIndexToRemove and below.
+	 *
+	 * @param	FirstStackIndexToRemove	Index into the stack from which menus are dismissed.
+	 */
+	void DismissInternal(int32 FirstStackIndexToRemove);
+
+	/**
+	 * Handles changes to the host window. Makes sure overlay slots are added/removed as appropriate.
+	 *
+	 * @param	InWindow	The new host window.
+	 */
+	void SetHostWindow(TSharedPtr<SWindow> InWindow);
+
+	/**
+	 * Callback method called by all menus when they are dismissed.  Handles changes to the stack state when a menu is destroyed.
+	 *
+	 * @param	InMenu	The menu that was destroyed.
+	 */
+	void OnMenuDestroyed(TSharedRef<IMenu> InMenu);
 
-	/** The stack of windows, each level in the stack has an arbitrary number of siblings.  All siblings must belong to the same parent in the previous level */
-	TArray<FMenuWindowList> Stack;
+	/**
+	 * Callback method called by all menus when they lose focus.  Handles dismissing menus that have lost focus.
+	 *
+	 * @param	InFocussedPath	The new focus path.
+	 */
+	void OnMenuContentLostFocus(const FWidgetPath& InFocussedPath);
+
+	/**
+	 * Called by all menus during creation. Handles wrapping content in SMenuContentWrapper that acts as a marker in widget paths and in a box that sets minimum size.
+	 *
+	 * @param	InContent			The unwrapped content.
+	 * @param	OptionalMinWidth	Optional minimum width for the wrapped content.
+	 * @param	OptionalMinHeight	Optional minimum height for the wrapped content.
+	 */
+	TSharedRef<SWidget> WrapContent(TSharedRef<SWidget> InContent, FOptionalSize OptionalMinWidth = FOptionalSize(), FOptionalSize OptionalMinHeight = FOptionalSize());
+
+	/** Contains all the options passed to the pre-push stage of the menu creation process */
+	struct FPrePushArgs
+	{
+		FPrePushArgs() : TransitionEffect(FPopupTransitionEffect::None) {}
+
+		TSharedPtr<SWidget> Content;
+		FSlateRect Anchor;
+		FPopupTransitionEffect TransitionEffect;
+		bool bFocusImmediately;
+	};
+
+	/** Contains all the options returned from the pre-push stage of the menu creation process */
+	struct FPrePushResults
+	{
+		TSharedPtr<SWidget> WrappedContent;
+		FVector2D AnimStartLocation;
+		FVector2D AnimFinalLocation;
+		bool bAnchorSetsMinWidth;
+		FVector2D ExpectedSize;
+		bool bAllowAnimations;
+		bool bFocusImmediately;
+	};
+
+	/**
+	 * This is the pre-push stage of menu creation.
+	 * It is responsible for wrapping content and calculating menu positioning. This stage is not used by "hosted" menus.
+	 *
+	 * @param	InArgs	Pre-push options.
+	 *
+	 * @return			Pre-push results.
+	 */
+	FPrePushResults PrePush(const FPrePushArgs& InArgs);
+
+	/**
+	 * This is the common code used during menu creation. This stage is not used by "hosted" menus.
+	 * It handles pre-push, actual menu creation and post-push stack updates.
+	 *
+	 * @param InParentMenu			The parent menu for this menu
+	 * @param InContent				The menu's content
+	 * @param Anchor				The anchor location for the menu (rect around the parent widget that is summoning the menu)
+	 * @param TransitionEffect		Animation to use when the popup appears
+	 * @param bFocusImmediately		Should the popup steal focus when shown?
+	 *
+	 * @return			The newly created menu.
+	 */
+	TSharedRef<IMenu> PushInternal(const TSharedPtr<IMenu>& InParentMenu, const TSharedRef<SWidget>& InContent, FSlateRect Anchor, const FPopupTransitionEffect& TransitionEffect, const bool bFocusImmediately);
+
+	/**
+	 * This is the actual menu object creation method for FMenuInWindow menus.
+	 * It creates a new SWindow and new FMenuInWindow that uses it.
+	 *
+	 * @param InParentMenu			The parent menu for this menu
+	 * @param InPrePushResults		The results of the pre-push stage
+	 *
+	 * @return			The newly created FMenuInWindow menu.
+	 */
+	TSharedRef<FMenuBase> PushNewWindow(TSharedPtr<IMenu> InParentMenu, const FPrePushResults& InPrePushResults);
+
+	/**
+	 * This is the actual menu object creation method for FMenuInPopup menus.
+	 * It creates a new FMenuInWindow and adds it into the overlay layer on the host window.
+	 *
+	 * @param InParentMenu			The parent menu for this menu
+	 * @param InPrePushResults		The results of the pre-push stage
+	 *
+	 * @return			The newly created FMenuInPopup menu.
+	 */
+	TSharedRef<FMenuBase> PushPopup(TSharedPtr<IMenu> InParentMenu, const FPrePushResults& InPrePushResults);
+
+	/**
+	 * This is the post-push stage of menu creation.
+	 * It is responsible for updating the stack with the new menu and removing any that it replaces.
+	 *
+	 * @param InParentMenu			The parent menu for this menu
+	 * @param InMenu				A newly created menu
+	 */
+	void PostPush(TSharedPtr<IMenu> InParentMenu, TSharedRef<FMenuBase> InMenu);
+
+	/** The popup method currently used by the whole stack. It can only use one at a time */
+	EPopupMethod ActiveMethod;
+
+	/** The parent window of the root menu in the stack. NOT the actual menu window if it's a CreateNewWindow */
+	TSharedPtr<SWindow> HostWindow;
+
+	/** The parent window of the root menu in the stack. NOT the actual menu window if it's a CreateNewWindow */
+	TSharedPtr<SMenuPanel> HostWindowPopupPanel;
+
+	/** The array of menus in the stack */
+	TArray<TSharedPtr<class FMenuBase>> Stack;
+
+	/** Maps top-level content widgets (should always be SMenuContentWrappers) to menus in the stack */
+	TMap<TSharedPtr<const SWidget>, TSharedPtr<class FMenuBase>> CachedContentMap;
 	
 	/** Handle to a throttle request made to ensure the menu is responsive in low FPS situations */
 	FThrottleRequest ThrottleHandle;
+
+	/** Temporary ptr to a new window created during the menu creation process. Nulled before the Push() call returns. Stops activation of the new window collapsing the stack. */
+	TSharedPtr<SWindow> PendingNewWindow;
+
+	/** Temporary ptr to a new menu created during the menu creation process. Nulled before the Push() call returns. Stops it collapsing the stack when it gets focus. */
+	TSharedPtr<class FMenuBase> PendingNewMenu;
 };
diff --git a/Engine/Source/Runtime/Slate/Public/Framework/Application/SlateApplication.h b/Engine/Source/Runtime/Slate/Public/Framework/Application/SlateApplication.h
--- a/Engine/Source/Runtime/Slate/Public/Framework/Application/SlateApplication.h
+++ b/Engine/Source/Runtime/Slate/Public/Framework/Application/SlateApplication.h
@@ -260,7 +260,7 @@ public:
 	TSharedRef<SWindow> AddWindowAsNativeChild( TSharedRef<SWindow> InSlateWindow, TSharedRef<SWindow> InParentWindow, const bool bShowImmediately = true );
 
 	/**
-	 * Creates a new Menu window and add it to the menu stack.
+	 * Creates a new Menu window and adds it to the menu stack.
 	 *
 	 * @param InWindow				The parent of the menu.  If there is already an open menu this parent must exist in the menu stack or the menu stack is dismissed and a new one started
 	 * @param Content				The content to be placed inside the new window
@@ -271,15 +271,76 @@ public:
 	 * @param WindowSize			When bShouldAutoSize=false, this must be set to the size of the window to be created
 	 * @param SummonLocationSize	An optional rect which describes an area in which the menu may not appear
 	 */
+	DEPRECATED(4.9, "PushMenu() returning a window is deprecated. Use the new version of PushMenu() that returns an IMenu.")
 	TSharedRef<SWindow> PushMenu( const TSharedRef<SWidget>& InParentContent, const TSharedRef<SWidget>& InContent, const FVector2D& SummonLocation, const FPopupTransitionEffect& TransitionEffect, const bool bFocusImmediately = true, const bool bShouldAutoSize = true, const FVector2D& WindowSize = FVector2D::ZeroVector, const FVector2D& SummonLocationSize = FVector2D::ZeroVector );
 
+	/**
+	 * Creates a new Menu and adds it to the menu stack.
+	 * Menus are always auto-sized. Use fixed-size content if a fixed size is required.
+	 *
+	 * @param InParentWidget		The parent of the menu. If the stack isn't empty, PushMenu will attempt to determine the stack level for the new menu by looking of an open menu in the parent's path.
+	 * @param InOwnerPath			Optional full widget path of the parent if one is available. If an invalid path is given PushMenu will attempt to generate a path to the InParentWidget
+	 * @param InContent				The content to be placed inside the new menu
+	 * @param SummonLocation		The location where this menu should be summoned
+	 * @param TransitionEffect		Animation to use when the popup appears
+	 * @param bFocusImmediately		Should the popup steal focus when shown?
+	 * @param SummonLocationSize	An optional rect which describes an area in which the menu may not appear
+	 * @param Method				An optional popup method override. If not set, the widgets in the InOwnerPath will be queried for this.
+	 */
+	TSharedPtr<IMenu> PushMenu(const TSharedRef<SWidget>& InParentWidget, const FWidgetPath& InOwnerPath, const TSharedRef<SWidget>& InContent, const FVector2D& SummonLocation, const FPopupTransitionEffect& TransitionEffect, const bool bFocusImmediately = true, const FVector2D& SummonLocationSize = FVector2D::ZeroVector, TOptional<EPopupMethod> Method = TOptional<EPopupMethod>());
+
+	/**
+	 * Creates a new Menu and adds it to the menu stack under the specified parent menu.
+	 * Menus are always auto-sized. Use fixed-size content if a fixed size is required.
+	 *
+	 * @param InParentMenu			The parent of the menu. Must be a valid menu in the stack.
+	 * @param InContent				The content to be placed inside the new menu
+	 * @param SummonLocation		The location where this menu should be summoned
+	 * @param TransitionEffect		Animation to use when the popup appears
+	 * @param bFocusImmediately		Should the popup steal focus when shown?
+	 * @param SummonLocationSize	An optional rect which describes an area in which the menu may not appear
+	 */
+	TSharedPtr<IMenu> PushMenu(const TSharedPtr<IMenu>& InParentMenu, const TSharedRef<SWidget>& InContent, const FVector2D& SummonLocation, const FPopupTransitionEffect& TransitionEffect, const bool bFocusImmediately = true, const FVector2D& SummonLocationSize = FVector2D::ZeroVector);
+
+	/**
+	 * Creates a new hosted Menu and adds it to the menu stack.
+	 * Hosted menus are drawn by an external host widget.
+	 * 
+	 * @param InParentMenu			The parent of the menu. Must be a valid menu in the stack.
+	 * @param InOwnerPath			Optional full widget path of the parent if one is available. If an invalid path is given PushMenu will attempt to generate a path to the InParentWidget
+	 * @param InMenuHost			The host widget that draws the menu's content
+	 * @param InContent				The content to be placed inside the new menu
+	 * @param OutWrappedContent		Returns the InContent wrapped with widgets needed by the menu stack system. This is what should be drawn by the host after this call.
+	 * @param TransitionEffect		Animation to use when the popup appears
+	 */
+	TSharedPtr<IMenu> PushHostedMenu(const TSharedRef<SWidget>& InParentWidget, const FWidgetPath& InOwnerPath, const TSharedRef<IMenuHost>& InMenuHost, const TSharedRef<SWidget>& InContent, TSharedPtr<SWidget>& OutWrappedContent, const FPopupTransitionEffect& TransitionEffect);
+	
+	/**
+	 * Creates a new hosted child Menu and adds it to the menu stack under the specified parent menu.
+	 * Hosted menus are drawn by an external host widget.
+	 * 
+	 * @param InParentMenu			The parent menu for this menu
+	 * @param InMenuHost			The host widget that draws the menu's content
+	 * @param InContent				The menu's content
+	 * @param OutWrappedContent		Returns the InContent wrapped with widgets needed by the menu stack system. This is what should be drawn by the host after this call.
+	 * @param TransitionEffect		Animation to use when the popup appears
+	 */	
+	TSharedPtr<IMenu> PushHostedMenu(const TSharedPtr<IMenu>& InParentMenu, const TSharedRef<IMenuHost>& InMenuHost, const TSharedRef<SWidget>& InContent, TSharedPtr<SWidget>& OutWrappedContent, const FPopupTransitionEffect& TransitionEffect);
+
 	/** @return Returns whether the window has child menus. */
+	DEPRECATED(4.9, "HasOpenSubMenus() taking a window is deprecated. Use HasOpenSubMenus() taking an IMenu as a parameter.")
 	bool HasOpenSubMenus( TSharedRef<SWindow> Window ) const;
 
+	/** @return Returns whether the menu has child menus. */
+	bool HasOpenSubMenus(TSharedPtr<IMenu> InMenu) const;
+
 	/** @return	Returns true if there are any pop-up menus summoned */
 	bool AnyMenusVisible() const;
 
-	/** Dismisses any open menus */
+	/** @return	Returns a ptr to the window that is currently the host of the menu stack or null if no menus are visible */
+	TSharedPtr<SWindow> GetVisibleMenuWindow() const;
+
+	/** Dismisses all open menus */
 	void DismissAllMenus();
 
 	/**
@@ -287,14 +348,30 @@ public:
 	 *
 	 * @param MenuWindowToDismiss	The window to dismiss, any children, grandchildren etc will also be dismissed
 	 */
+	DEPRECATED(4.9, "DismissMenu() taking a window is deprecated. Use DismissMenu() taking an IMenu, DismissMenuByWidget() or DismissAll().")
 	void DismissMenu( TSharedRef<SWindow> MenuWindowToDismiss );
 
+	/**
+	 * Dismisses a menu and all its children
+	 *
+	 * @param InFromMenu	The menu to dismiss, any children, grandchildren etc will also be dismissed
+	 */
+	void DismissMenu(const TSharedPtr<IMenu>& InFromMenu);
+
+	/**
+	 * Dismisses a menu and all its children. The menu is determined by looking for menus in the parent chain of the widget.
+	 *
+	 * @param InWidgetInMenu	The widget whose path is search upwards for a menu. That menu will then be dismissed.
+	 */
+	void DismissMenuByWidget(const TSharedRef<SWidget>& InWidgetInMenu);
+
 	/**
 	 * Finds the window in the menu stack
 	 *
 	 * @param WindowToFind	The window to look for
 	 * @return The level in the stack  that the window is in or INDEX_NONE if it is not found
 	 */
+	DEPRECATED(4.9, "GetLocationInMenuStack() is deprecated. Shouldn't be needed.")
 	int32 GetLocationInMenuStack( TSharedRef<SWindow> WindowToFind ) const;
 
 	/**
diff --git a/Engine/Source/Runtime/Slate/Public/Widgets/Input/SEditableText.h b/Engine/Source/Runtime/Slate/Public/Widgets/Input/SEditableText.h
--- a/Engine/Source/Runtime/Slate/Public/Widgets/Input/SEditableText.h
+++ b/Engine/Source/Runtime/Slate/Public/Widgets/Input/SEditableText.h
@@ -367,7 +367,7 @@ public:
 	virtual void Undo() override;
 	virtual void Redo() override;
 	virtual TSharedRef< SWidget > GetWidget() override;
-	virtual void SummonContextMenu(const FVector2D& InLocation, TSharedPtr<SWindow> ParentWindow = TSharedPtr<SWindow>()) override;
+	virtual void SummonContextMenu(const FVector2D& InLocation, TSharedPtr<SWindow> ParentWindow, const FWidgetPath& EventPath) override;
 	virtual void LoadText() override;
 	// END ITextEditorWidget interface
 
@@ -508,9 +508,9 @@ protected:
 	void SetCaretPosition( int32 Position );
 
 	/**
-	 * Called when the content menu window is closed
+	 * Called when the context menu is closed
 	 */
-	void OnWindowClosed(const TSharedRef<SWindow>&);
+	void OnContextMenuClosed(TSharedRef<IMenu>);
 
 private:
 	/** @return Whether the editable text should appear focused */
diff --git a/Engine/Source/Runtime/Slate/Public/Widgets/Input/SMenuAnchor.h b/Engine/Source/Runtime/Slate/Public/Widgets/Input/SMenuAnchor.h
--- a/Engine/Source/Runtime/Slate/Public/Widgets/Input/SMenuAnchor.h
+++ b/Engine/Source/Runtime/Slate/Public/Widgets/Input/SMenuAnchor.h
@@ -2,6 +2,7 @@
 
 #pragma once
 #include "STextBlock.h"
+#include "IMenu.h"
 
 /** Notification when popup is opened/closed. */
 DECLARE_DELEGATE_OneParam(FOnIsOpenChanged, bool)
@@ -21,7 +22,7 @@ enum class EPopupMethod : uint8
  * A PopupAnchor summons a Popup relative to its content.
  * Summoning a popup relative to the cursor is accomplished via the application.
  */
-class SLATE_API SMenuAnchor : public SPanel
+class SLATE_API SMenuAnchor : public SPanel, public IMenuHost
 {
 public:
 	DEPRECATED(4.7, "You probably do not need this setting any more. See OnQueryPopupMethod() in SWidget.h.")
@@ -90,6 +91,11 @@ public:
 	/** @return Whether this menu has open submenus */
 	bool HasOpenSubMenus() const;
 
+	// IMenuHost interface
+	virtual TSharedPtr<SWindow> GetMenuWindow() const override;
+	virtual void OnMenuDismissed() override;
+	// End of IMenuHost interface
+
 protected:
 	// SWidget interface
 	virtual void Tick( const FGeometry& AllottedGeometry, const double InCurrentTime, const float InDeltaTime ) override;
@@ -99,24 +105,34 @@ protected:
 	virtual int32 OnPaint( const FPaintArgs& Args, const FGeometry& AllottedGeometry, const FSlateRect& MyClippingRect, FSlateWindowElementList& OutDrawElements, int32 LayerId, const FWidgetStyle& InWidgetStyle, bool bParentEnabled ) const override;
 	// End of SWidget interface
 
-	/** Invoked when the popup window is being closed by the application */
-	void RequestClosePopupWindow( const TSharedRef<SWindow>& PopupWindow );
-
-	/** Invoked when a click occurred outside the widget that is our popup; used for popup auto-dismissal */
-	void OnClickedOutsidePopup();
-
 	/** @return true if the popup is currently open and reusing an existing window */
 	bool IsOpenAndReusingWindow() const;
 
 	/** @return true if the popup is currently open and we created a dedicated window for it */
 	bool IsOpenViaCreatedWindow() const;
 
+	/** Handler/callback called by menus created by this anchor, when they are dismissed */
+	void OnMenuClosed(TSharedRef<IMenu> InMenu);
+
 protected:
-	/** A pointer to the window presenting this popup. Pointer is null when the popup is not visible */
+	/**
+	 * A pointer to the window presenting this popup.
+	 * Can be the window created to hold a menu or the window containing this anchor if the menu is drawn as a child of the this anchor.
+	 * Pointer is null when a popup is not visible.
+	 */
 	TWeakPtr<SWindow> PopupWindowPtr;
 
+	/**
+	 * An interface pointer to the menu object presenting this popup.
+	 * Pointer is null when a popup is not visible.
+	 */
+	TWeakPtr<IMenu> PopupMenuPtr;
+
 	/** Static menu content to use when the delegate used when OnGetMenuContent is not defined. */
 	TSharedPtr<SWidget> MenuContent;
+
+	/** MenuContent plus any extra wrapping widgets needed by the menu infrastructure. */
+	TSharedPtr<SWidget> WrappedContent;
 	
 	/** Callback invoked when the popup is being summoned */
 	FOnGetContent OnGetMenuContent;
@@ -144,8 +160,11 @@ protected:
 	 */
 	FVector2D LocalPopupPosition;
 
-	TPanelChildren<FSimpleSlot> Children;
+	/**
+	 * Screen-space version of LocalPopupPosition, also cached in Tick.
+	 * Used to satisfy calls to GetMenuPosition() for menus that are reusing the current window.
+	 */
+	FVector2D ScreenPopupPosition;
 
-	/** Handle to the registered OnClickedOutsidePopup delegate */
-	FDelegateHandle OnClickedOutsidePopupDelegateHandle;
+	TPanelChildren<FSimpleSlot> Children;
 };
diff --git a/Engine/Source/Runtime/Slate/Public/Widgets/Text/ITextEditorWidget.h b/Engine/Source/Runtime/Slate/Public/Widgets/Text/ITextEditorWidget.h
--- a/Engine/Source/Runtime/Slate/Public/Widgets/Text/ITextEditorWidget.h
+++ b/Engine/Source/Runtime/Slate/Public/Widgets/Text/ITextEditorWidget.h
@@ -112,42 +112,50 @@ class FActiveTextEditContextMenu
 public:
 	FActiveTextEditContextMenu()
 		: bIsPendingSummon(false)
-		, ActiveMenuWindow()
+		, ActiveMenu()
 	{
 	}
 
 	/** Check to see whether this context is valid (either pending or active) */
 	bool IsValid() const
 	{
-		return bIsPendingSummon || ActiveMenuWindow.IsValid();
+		return bIsPendingSummon || ActiveMenu.IsValid();
 	}
 
 	/** Called to reset the active context menu state */
 	void Reset()
 	{
 		bIsPendingSummon = false;
-		ActiveMenuWindow.Reset();
+		ActiveMenu.Reset();
 	}
 
 	/** Called before you summon your context menu */
 	void PrepareToSummon()
 	{
 		bIsPendingSummon = true;
-		ActiveMenuWindow.Reset();
+		ActiveMenu.Reset();
+	}
+
+	DEPRECATED(4.9, "SummonSucceeded() taking an SWindow parameter is deprecated. Use the new version of SummonSucceeded() takes an IMenu.")
+	void SummonSucceeded(const TSharedRef<SWindow>& InWindow)
+	{
+		// deprecated
+		bIsPendingSummon = false;
+		ActiveMenu.Reset();
 	}
 
 	/** Called when you've successfully summoned your context menu */
-	void SummonSucceeded(const TSharedRef<SWindow>& InActiveMenuWindow)
+	void SummonSucceeded(const TSharedRef<IMenu>& InMenu)
 	{
 		bIsPendingSummon = false;
-		ActiveMenuWindow = InActiveMenuWindow;
+		ActiveMenu = InMenu;
 	}
 
 	/** Called if your context menu summon fails */
 	void SummonFailed()
 	{
 		bIsPendingSummon = false;
-		ActiveMenuWindow.Reset();
+		ActiveMenu.Reset();
 	}
 
 private:
@@ -155,7 +163,7 @@ private:
 	bool bIsPendingSummon;
 
 	/** Handle to the active context menu (if any) */
-	TWeakPtr<SWindow> ActiveMenuWindow;
+	TWeakPtr<IMenu> ActiveMenu;
 };
 
 class SLATE_API ITextEditorWidget 
@@ -307,7 +315,13 @@ class SLATE_API ITextEditorWidget
 
 	virtual TSharedRef< SWidget > GetWidget() = 0;
 
-	virtual void SummonContextMenu(const FVector2D& InLocation, TSharedPtr<SWindow> ParentWindow = TSharedPtr<SWindow>()) = 0;
+	DEPRECATED(4.9, "SummonContextMenu() without a FWidgetPath parameter is deprecated. Use the new version of SummonContextMenu().")
+	virtual void SummonContextMenu(const FVector2D& InLocation, TSharedPtr<SWindow> ParentWindow = TSharedPtr<SWindow>())
+	{
+		SummonContextMenu(InLocation, ParentWindow, FWidgetPath());
+	}
+
+	virtual void SummonContextMenu(const FVector2D& InLocation, TSharedPtr<SWindow> ParentWindow, const FWidgetPath& EventPath) = 0;
 
 	virtual void LoadText() = 0;
 };
diff --git a/Engine/Source/Runtime/Slate/Public/Widgets/Text/SMultiLineEditableText.h b/Engine/Source/Runtime/Slate/Public/Widgets/Text/SMultiLineEditableText.h
--- a/Engine/Source/Runtime/Slate/Public/Widgets/Text/SMultiLineEditableText.h
+++ b/Engine/Source/Runtime/Slate/Public/Widgets/Text/SMultiLineEditableText.h
@@ -450,7 +450,7 @@ private:
 	virtual void Undo() override;
 	virtual void Redo() override;
 	virtual TSharedRef< SWidget > GetWidget() override;
-	virtual void SummonContextMenu(const FVector2D& InLocation, TSharedPtr<SWindow> ParentWindow = TSharedPtr<SWindow>()) override;
+	virtual void SummonContextMenu(const FVector2D& InLocation, TSharedPtr<SWindow> ParentWindow, const FWidgetPath& EventPath) override;
 	virtual void LoadText() override;
 	// END ITextEditorWidget interface
 
@@ -503,7 +503,7 @@ protected:
 
 private:
 	
-	void OnWindowClosed(const TSharedRef<SWindow>&);
+	void OnContextMenuClosed(TSharedRef<IMenu>);
 
 	/** Remember where the cursor was when we started selecting. */
 	void BeginSelecting( const FTextLocation& Endpoint );
diff --git a/Engine/Source/Runtime/Slate/Public/Widgets/Views/SListView.h b/Engine/Source/Runtime/Slate/Public/Widgets/Views/SListView.h
--- a/Engine/Source/Runtime/Slate/Public/Widgets/Views/SListView.h
+++ b/Engine/Source/Runtime/Slate/Public/Widgets/Views/SListView.h
@@ -671,7 +671,7 @@ public:
 
 	virtual void Private_OnItemRightClicked( ItemType TheItem, const FPointerEvent& MouseEvent ) override
 	{
-		this->OnRightMouseButtonUp( MouseEvent.GetScreenSpacePosition() );
+		this->OnRightMouseButtonUp( MouseEvent );
 	}
 
 	virtual bool Private_OnItemClicked(ItemType TheItem) override
diff --git a/Engine/Source/Runtime/Slate/Public/Widgets/Views/STableViewBase.h b/Engine/Source/Runtime/Slate/Public/Widgets/Views/STableViewBase.h
--- a/Engine/Source/Runtime/Slate/Public/Widgets/Views/STableViewBase.h
+++ b/Engine/Source/Runtime/Slate/Public/Widgets/Views/STableViewBase.h
@@ -178,7 +178,7 @@ protected:
 	/**
 	 * Opens a context menu as the result of a right click if OnContextMenuOpening is bound and we are not right click scrolling.
 	 */
-	virtual void OnRightMouseButtonUp(const FVector2D& SummonLocation);
+	virtual void OnRightMouseButtonUp(const FPointerEvent& MouseEvent);
 	
 	/**
 	 * Get the scroll rate in items that best approximates a constant physical scroll rate.
diff --git a/Engine/Source/Runtime/SlateCore/Private/Widgets/SWindow.cpp b/Engine/Source/Runtime/SlateCore/Private/Widgets/SWindow.cpp
--- a/Engine/Source/Runtime/SlateCore/Private/Widgets/SWindow.cpp
+++ b/Engine/Source/Runtime/SlateCore/Private/Widgets/SWindow.cpp
@@ -1038,11 +1038,13 @@ bool SWindow::AppearsInTaskbar() const
 
 void SWindow::SetOnWindowActivated( const FOnWindowActivated& InDelegate )
 {
+	// deprecated
 	OnWindowActivated = InDelegate;
 }
 
 void SWindow::SetOnWindowDeactivated( const FOnWindowDeactivated& InDelegate )
 {
+	// deprecated
 	OnWindowDeactivated = InDelegate;
 }
 
@@ -1313,7 +1315,8 @@ bool SWindow::OnIsActiveChanged( const FWindowActivateEvent& ActivateEvent )
 	const bool bWasDeactivated = ( ActivateEvent.GetActivationType() == FWindowActivateEvent::EA_Deactivate );
 	if ( bWasDeactivated )
 	{
-		OnWindowDeactivated.ExecuteIfBound();
+		OnWindowDeactivated.ExecuteIfBound();	// deprecated
+		WindowDeactivatedEvent.Broadcast();
 
 		const EWindowMode::Type WindowMode = GetWindowMode();
 		// If the window is not fullscreen, we do not want to automatically recapture the mouse unless an external UI such as Steam is open. Fullscreen windows we do.
@@ -1357,7 +1360,8 @@ bool SWindow::OnIsActiveChanged( const FWindowActivateEvent& ActivateEvent )
 			}
 		}
 
-		OnWindowActivated.ExecuteIfBound();
+		OnWindowActivated.ExecuteIfBound();	// deprecated
+		WindowActivatedEvent.Broadcast();
 	}
 
 	return true;
diff --git a/Engine/Source/Runtime/SlateCore/Public/Input/Events.h b/Engine/Source/Runtime/SlateCore/Public/Input/Events.h
--- a/Engine/Source/Runtime/SlateCore/Public/Input/Events.h
+++ b/Engine/Source/Runtime/SlateCore/Public/Input/Events.h
@@ -329,6 +329,11 @@ public:
 		EventPath = &InEventPath;
 	}
 
+	const FWidgetPath* GetEventPath() const
+	{
+		return EventPath;
+	}
+
 	SLATECORE_API virtual FText ToText() const;
 	
 	/** Is this event a pointer event (touch or cursor). */
@@ -345,7 +350,7 @@ protected:
 	// True if this key was auto-repeated.
 	bool bIsRepeat;
 
-	// Events are sent along paths. See GetEventPath().
+	// Events are sent along paths. See (GetEventPath).
 	const FWidgetPath* EventPath;
 };
 
diff --git a/Engine/Source/Runtime/SlateCore/Public/Widgets/SWidget.h b/Engine/Source/Runtime/SlateCore/Public/Widgets/SWidget.h
--- a/Engine/Source/Runtime/SlateCore/Public/Widgets/SWidget.h
+++ b/Engine/Source/Runtime/SlateCore/Public/Widgets/SWidget.h
@@ -288,7 +288,7 @@ public:
 	virtual FReply OnMouseMove(const FGeometry& MyGeometry, const FPointerEvent& MouseEvent);
 	
 	/**
-	 * The system will use this event to notify a widget that the cursor has entered it. This event is NOT bubbled.
+	 * The system will use this event to notify a widget that the cursor has entered it. This event is uses a custom bubble strategy.
 	 *
 	 * @param MyGeometry The Geometry of the widget receiving the event
 	 * @param MouseEvent Information about the input event
@@ -296,7 +296,7 @@ public:
 	virtual void OnMouseEnter(const FGeometry& MyGeometry, const FPointerEvent& MouseEvent);
 	
 	/**
-	 * The system will use this event to notify a widget that the cursor has left it. This event is NOT bubbled.
+	 * The system will use this event to notify a widget that the cursor has left it. This event is uses a custom bubble strategy.
 	 *
 	 * @param MouseEvent Information about the input event
 	 */
diff --git a/Engine/Source/Runtime/SlateCore/Public/Widgets/SWindow.h b/Engine/Source/Runtime/SlateCore/Public/Widgets/SWindow.h
--- a/Engine/Source/Runtime/SlateCore/Public/Widgets/SWindow.h
+++ b/Engine/Source/Runtime/SlateCore/Public/Widgets/SWindow.h
@@ -7,9 +7,11 @@ class FHittestGrid;
 
 /** Notification that a window has been activated */
 DECLARE_DELEGATE( FOnWindowActivated );
+DECLARE_MULTICAST_DELEGATE( FOnWindowActivatedEvent );
 
 /** Notification that a window has been deactivated */
 DECLARE_DELEGATE( FOnWindowDeactivated );
+DECLARE_MULTICAST_DELEGATE( FOnWindowDeactivatedEvent );
 
 /** Notification that a window is about to be closed */
 DECLARE_DELEGATE_OneParam( FOnWindowClosed, const TSharedRef<SWindow>& );
@@ -434,11 +436,19 @@ public:
 	bool AppearsInTaskbar() const;
 
 	/** Sets the delegate to execute when the window is activated */
+	DEPRECATED(4.9, "SetOnWindowActivated() is deprecated. Use GetOnWindowActivatedEvent() and subscribe to the multicast event.")
 	void SetOnWindowActivated( const FOnWindowActivated& InDelegate );
 
+	/** Gets the multicast delegate executed when the window is deactivated */
+	FOnWindowActivatedEvent& GetOnWindowActivatedEvent() { return WindowActivatedEvent; }
+
 	/** Sets the delegate to execute when the window is deactivated */
+	DEPRECATED(4.9, "SetOnWindowDeactivated() is deprecated. Use GetOnWindowDeactivatedEvent() and subscribe to the multicast event.")
 	void SetOnWindowDeactivated( const FOnWindowDeactivated& InDelegate );
 
+	/** Gets the multicast delegate executed when the window is deactivated */
+	FOnWindowDeactivatedEvent& GetOnWindowDeactivatedEvent() { return WindowDeactivatedEvent; }
+
 	/** Sets the delegate to execute right before the window is closed */
 	void SetOnWindowClosed( const FOnWindowClosed& InDelegate );
 
@@ -890,10 +900,12 @@ private:
 	
 	/** Invoked when the window has been activated. */
 	FOnWindowActivated OnWindowActivated;
-	
+	FOnWindowActivatedEvent WindowActivatedEvent;
+
 	/** Invoked when the window has been deactivated. */
 	FOnWindowDeactivated OnWindowDeactivated;
-	
+	FOnWindowDeactivatedEvent WindowDeactivatedEvent;
+
 	/** Invoked when the window is about to be closed. */
 	FOnWindowClosed OnWindowClosed;
 
