Fix CIS error on Mac

diff --git a/Engine/Source/Runtime/Slate/Private/Framework/Application/Menu.h b/Engine/Source/Runtime/Slate/Private/Framework/Application/Menu.h
--- a/Engine/Source/Runtime/Slate/Private/Framework/Application/Menu.h
+++ b/Engine/Source/Runtime/Slate/Private/Framework/Application/Menu.h
@@ -28,6 +28,7 @@ class FMenuInWindow : public FMenuBase
 {
 public:
 	FMenuInWindow(TSharedRef<SWindow> InWindow, TSharedRef<SWidget> InContent);
+	virtual ~FMenuInWindow() {}
 
 	virtual EPopupMethod GetPopupMethod() const override { return EPopupMethod::CreateNewWindow; }
 	virtual TSharedPtr<SWindow> GetParentWindow() const override;
@@ -45,6 +46,7 @@ class FMenuInPopup : public FMenuBase
 {
 public:
 	FMenuInPopup(TSharedRef<SWidget> InContent);
+	virtual ~FMenuInPopup() {}
 
 	virtual EPopupMethod GetPopupMethod() const { return EPopupMethod::UseCurrentWindow; }
 	virtual TSharedPtr<SWindow> GetParentWindow() const;
@@ -59,6 +61,7 @@ class FMenuInHostWidget : public FMenuBase
 {
 public:
 	FMenuInHostWidget(TSharedRef<IMenuHost> InHost, const TSharedRef<SWidget>& InContent);
+	virtual ~FMenuInHostWidget() {}
 
 	virtual EPopupMethod GetPopupMethod() const { return EPopupMethod::UseCurrentWindow; }
 	virtual TSharedPtr<SWindow> GetParentWindow() const;
