Fix CIS error on XboxOne_Mono

diff --git a/Engine/Source/Runtime/Slate/Private/Framework/Application/MenuStack.cpp b/Engine/Source/Runtime/Slate/Private/Framework/Application/MenuStack.cpp
--- a/Engine/Source/Runtime/Slate/Private/Framework/Application/MenuStack.cpp
+++ b/Engine/Source/Runtime/Slate/Private/Framework/Application/MenuStack.cpp
@@ -3,6 +3,7 @@
 #include "SlatePrivatePCH.h"
 #include "Menu.h"
 #include "SPopup.h"
+#include "SMenuAnchor.h"
 
 #define LOCTEXT_NAMESPACE "MenuStack"
 
@@ -193,7 +194,7 @@ TSharedRef<IMenu> FMenuStack::Push(const FWidgetPath& InOwnerPath, const TShared
 		// pushing a new root menu (leave ParentMenu invalid)
 
 		// The active method is determined when a new root menu is pushed
-		ActiveMethod = InMethod.IsSet() ? InMethod.GetValue() : QueryPopupMethod(InOwnerPath);
+		ActiveMethod = InMethod.IsSet() ? InMethod : QueryPopupMethod(InOwnerPath);
 
 		// The host window is determined when a new root menu is pushed
 		SetHostWindow(InOwnerPath.GetWindow());
@@ -274,7 +275,7 @@ TSharedRef<IMenu> FMenuStack::PushInternal(const TSharedPtr<IMenu>& InParentMenu
 	const FPrePushResults PrePushResults = PrePush(PrePushArgs);
 
 	// Menu object creation stage
-	TSharedRef<FMenuBase> OutMenu = ActiveMethod == EPopupMethod::CreateNewWindow
+	TSharedRef<FMenuBase> OutMenu = ActiveMethod.Get(EPopupMethod::CreateNewWindow) == EPopupMethod::CreateNewWindow
 		? PushNewWindow(InParentMenu, PrePushResults)
 		: PushPopup(InParentMenu, PrePushResults);
 
@@ -297,7 +298,7 @@ FMenuStack::FPrePushResults FMenuStack::PrePush(const FPrePushArgs& InArgs)
 
 	// Calc the max height available on screen for the menu
 	float MaxHeight;
-	if (ActiveMethod == EPopupMethod::CreateNewWindow)
+	if (ActiveMethod.Get(EPopupMethod::CreateNewWindow) == EPopupMethod::CreateNewWindow)
 	{
 		FSlateRect WorkArea = FSlateApplication::Get().GetWorkArea(InArgs.Anchor);
 		MaxHeight = FMenuStackDefs::MaxMenuScreenHeightFraction * WorkArea.GetSize().Y;
@@ -326,7 +327,7 @@ FMenuStack::FPrePushResults FMenuStack::PrePush(const FPrePushArgs& InArgs)
 	EOrientation Orientation = (InArgs.TransitionEffect.SlideDirection == FPopupTransitionEffect::SubMenu) ? Orient_Horizontal : Orient_Vertical;
 
 	// Calculate the correct position for the menu based on the popup method
-	if (ActiveMethod == EPopupMethod::CreateNewWindow)
+	if (ActiveMethod.Get(EPopupMethod::CreateNewWindow) == EPopupMethod::CreateNewWindow)
 	{
 		// Places the menu's window in the work area
 		OutResults.AnimStartLocation = OutResults.AnimFinalLocation = FSlateApplication::Get().CalculatePopupWindowPosition(InArgs.Anchor, OutResults.ExpectedSize, Orientation);
@@ -391,7 +392,7 @@ FMenuStack::FPrePushResults FMenuStack::PrePush(const FPrePushArgs& InArgs)
 
 TSharedRef<FMenuBase> FMenuStack::PushNewWindow(TSharedPtr<IMenu> InParentMenu, const FPrePushResults& InPrePushResults)
 {
-	check(ActiveMethod == EPopupMethod::CreateNewWindow);
+	check(ActiveMethod.GetValue() == EPopupMethod::CreateNewWindow);
 
 	// Start pop-up windows out transparent, then fade them in over time
 	const EWindowTransparency Transparency(InPrePushResults.bAllowAnimations ? EWindowTransparency::PerWindow : EWindowTransparency::None);
@@ -452,7 +453,7 @@ TSharedRef<FMenuBase> FMenuStack::PushNewWindow(TSharedPtr<IMenu> InParentMenu,
 
 TSharedRef<FMenuBase> FMenuStack::PushPopup(TSharedPtr<IMenu> InParentMenu, const FPrePushResults& InPrePushResults)
 {
-	check(ActiveMethod == EPopupMethod::UseCurrentWindow);
+	check(ActiveMethod.GetValue() == EPopupMethod::UseCurrentWindow);
 
 	// Create a FMenuInPopup
 	check(InPrePushResults.WrappedContent.IsValid());
@@ -614,7 +615,7 @@ void FMenuStack::OnMenuContentLostFocus(const FWidgetPath& InFocussedPath)
 {
 	// Window activation messages will make menus collapse when in CreateNewWindow mode
 	// In UseCurrentWindow mode we must look for focus moving from menus
-	if (ActiveMethod == EPopupMethod::UseCurrentWindow && HasMenus() && !PendingNewMenu.IsValid())
+	if (ActiveMethod.Get(EPopupMethod::CreateNewWindow) == EPopupMethod::UseCurrentWindow && HasMenus() && !PendingNewMenu.IsValid())
 	{
 		// If focus is switching determine which of our menus it is in, if any
 		TSharedPtr<IMenu> FocussedMenu = FindMenuInWidgetPath(InFocussedPath);
diff --git a/Engine/Source/Runtime/Slate/Public/Framework/Application/MenuStack.h b/Engine/Source/Runtime/Slate/Public/Framework/Application/MenuStack.h
--- a/Engine/Source/Runtime/Slate/Public/Framework/Application/MenuStack.h
+++ b/Engine/Source/Runtime/Slate/Public/Framework/Application/MenuStack.h
@@ -374,7 +374,7 @@ private:
 	void PostPush(TSharedPtr<IMenu> InParentMenu, TSharedRef<FMenuBase> InMenu);
 
 	/** The popup method currently used by the whole stack. It can only use one at a time */
-	EPopupMethod ActiveMethod;
+	TOptional<EPopupMethod> ActiveMethod;
 
 	/** The parent window of the root menu in the stack. NOT the actual menu window if it's a CreateNewWindow */
 	TSharedPtr<SWindow> HostWindow;
