Slate/UMG - Exposing the scroll into view function in a much better way.  We no longer allow/expect the user to pass in geometry for the scrollbox, instead they simply request the widget to scroll into view and on the next layout pass we make that happen.  Also exposing that to blueprints for the UMG version of the scrollbox.

diff --git a/Engine/Source/Runtime/Slate/Private/Widgets/Layout/SScrollBox.cpp b/Engine/Source/Runtime/Slate/Private/Widgets/Layout/SScrollBox.cpp
--- a/Engine/Source/Runtime/Slate/Private/Widgets/Layout/SScrollBox.cpp
+++ b/Engine/Source/Runtime/Slate/Private/Widgets/Layout/SScrollBox.cpp
@@ -165,6 +165,7 @@ void SScrollBox::Construct( const FArguments& InArgs )
 	ConsumeMouseWheel = InArgs._ConsumeMouseWheel;
 	TickScrollDelta = 0;
 	AllowOverscroll = InArgs._AllowOverscroll;
+	bAnimateScrollingWidgetIntoView = false;
 
 	if (InArgs._ExternalScrollbar.IsValid())
 	{
@@ -211,45 +212,44 @@ void SScrollBox::ConstructVerticalLayout()
 {
 	TSharedPtr<SHorizontalBox> PanelAndScrollbar;
 	this->ChildSlot
+	[
+		SAssignNew(PanelAndScrollbar, SHorizontalBox)
+
+		+ SHorizontalBox::Slot()
+		.FillWidth(1)
 		[
-			SAssignNew(PanelAndScrollbar, SHorizontalBox)
+			SNew(SOverlay)
 
-			+ SHorizontalBox::Slot()
-			.FillWidth(1)
+			+ SOverlay::Slot()
+			.Padding(FMargin(0.0f, 0.0f, 0.0f, 1.0f))
 			[
-				SNew(SOverlay)
-
-				+ SOverlay::Slot()
-				.Padding(FMargin(0.0f, 0.0f, 0.0f, 1.0f))
-				[
-					// Scroll panel that presents the scrolled content
-					ScrollPanel.ToSharedRef()
-				]
-
-				+ SOverlay::Slot()
-					.HAlign(HAlign_Fill)
-					.VAlign(VAlign_Top)
-					[
-						// Shadow: Hint to scroll up
-						SNew(SImage)
-						.Visibility(EVisibility::HitTestInvisible)
-						.ColorAndOpacity(this, &SScrollBox::GetStartShadowOpacity)
-						.Image(&Style->TopShadowBrush)
-					]
-
-				+ SOverlay::Slot()
-					.HAlign(HAlign_Fill)
-					.VAlign(VAlign_Bottom)
-					[
-						// Shadow: a hint to scroll down
-						SNew(SImage)
-						.Visibility(EVisibility::HitTestInvisible)
-						.ColorAndOpacity(this, &SScrollBox::GetEndShadowOpacity)
-						.Image(&Style->BottomShadowBrush)
-					]
+				// Scroll panel that presents the scrolled content
+				ScrollPanel.ToSharedRef()
 			]
 
-		];
+			+ SOverlay::Slot()
+			.HAlign(HAlign_Fill)
+			.VAlign(VAlign_Top)
+			[
+				// Shadow: Hint to scroll up
+				SNew(SImage)
+				.Visibility(EVisibility::HitTestInvisible)
+				.ColorAndOpacity(this, &SScrollBox::GetStartShadowOpacity)
+				.Image(&Style->TopShadowBrush)
+			]
+
+			+ SOverlay::Slot()
+			.HAlign(HAlign_Fill)
+			.VAlign(VAlign_Bottom)
+			[
+				// Shadow: a hint to scroll down
+				SNew(SImage)
+				.Visibility(EVisibility::HitTestInvisible)
+				.ColorAndOpacity(this, &SScrollBox::GetEndShadowOpacity)
+				.Image(&Style->BottomShadowBrush)
+			]
+		]
+	];
 
 	if (!bScrollBarIsExternal)
 	{
@@ -265,45 +265,44 @@ void SScrollBox::ConstructHorizontalLayout()
 {
 	TSharedPtr<SVerticalBox> PanelAndScrollbar;
 	this->ChildSlot
+	[
+		SAssignNew(PanelAndScrollbar, SVerticalBox)
+
+		+ SVerticalBox::Slot()
+		.FillHeight(1)
 		[
-			SAssignNew(PanelAndScrollbar, SVerticalBox)
+			SNew(SOverlay)
+
+			+ SOverlay::Slot()
+			.Padding(FMargin(0.0f, 0.0f, 1.0f, 0.0f))
+			[
+				// Scroll panel that presents the scrolled content
+				ScrollPanel.ToSharedRef()
+			]
 
-			+ SVerticalBox::Slot()
-			.FillHeight(1)
+			+ SOverlay::Slot()
+			.HAlign(HAlign_Left)
+			.VAlign(VAlign_Fill)
 			[
-				SNew(SOverlay)
-
-				+ SOverlay::Slot()
-				.Padding(FMargin(0.0f, 0.0f, 1.0f, 0.0f))
-				[
-					// Scroll panel that presents the scrolled content
-					ScrollPanel.ToSharedRef()
-				]
-
-				+ SOverlay::Slot()
-					.HAlign(HAlign_Left)
-					.VAlign(VAlign_Fill)
-					[
-						// Shadow: Hint to left
-						SNew(SImage)
-						.Visibility(EVisibility::HitTestInvisible)
-						.ColorAndOpacity(this, &SScrollBox::GetStartShadowOpacity)
-						.Image(&Style->LeftShadowBrush)
-					]
-
-				+ SOverlay::Slot()
-					.HAlign(HAlign_Right)
-					.VAlign(VAlign_Fill)
-					[
-						// Shadow: a hint to scroll right
-						SNew(SImage)
-						.Visibility(EVisibility::HitTestInvisible)
-						.ColorAndOpacity(this, &SScrollBox::GetEndShadowOpacity)
-						.Image(&Style->RightShadowBrush)
-					]
+				// Shadow: Hint to left
+				SNew(SImage)
+				.Visibility(EVisibility::HitTestInvisible)
+				.ColorAndOpacity(this, &SScrollBox::GetStartShadowOpacity)
+				.Image(&Style->LeftShadowBrush)
 			]
 
-		];
+			+ SOverlay::Slot()
+			.HAlign(HAlign_Right)
+			.VAlign(VAlign_Fill)
+			[
+				// Shadow: a hint to scroll right
+				SNew(SImage)
+				.Visibility(EVisibility::HitTestInvisible)
+				.ColorAndOpacity(this, &SScrollBox::GetEndShadowOpacity)
+				.Image(&Style->RightShadowBrush)
+			]
+		]
+	];
 
 	if (!bScrollBarIsExternal)
 	{
@@ -369,6 +368,15 @@ void SScrollBox::ScrollToEnd()
 	bScrollToEnd = true;
 }
 
+void SScrollBox::ScrollDescendantIntoView(const TSharedPtr<SWidget>& WidgetToFind, bool InAnimateScroll)
+{
+	WidgetToScrollIntoView = WidgetToFind;
+	bAnimateScrollingWidgetIntoView = InAnimateScroll;
+
+	// This will force the active timer system to wakeup for a frame to ensure we tick at least once.
+	RegisterActiveTimer(0.f, FWidgetActiveTimerDelegate::CreateSP(this, &SScrollBox::UpdateInertialScroll));
+}
+
 bool SScrollBox::ScrollDescendantIntoView(const FGeometry& MyGeometry, const TSharedPtr<SWidget>& WidgetToFind, bool InAnimateScroll)
 {
 	// We need to safely find the one WidgetToFind among our descendants.
@@ -526,6 +534,13 @@ void SScrollBox::Tick( const FGeometry& AllottedGeometry, const double InCurrent
 {
 	CachedGeometry = AllottedGeometry;
 
+	// If we needed a widget to be scrolled into view, make that happen.
+	if ( WidgetToScrollIntoView.IsValid() )
+	{
+		ScrollDescendantIntoView(AllottedGeometry, WidgetToScrollIntoView, bAnimateScrollingWidgetIntoView);
+		WidgetToScrollIntoView.Reset();
+	}
+
 	const FGeometry ScrollPanelGeometry = FindChildGeometry( AllottedGeometry, ScrollPanel.ToSharedRef() );
 	const float ContentSize = GetScrollComponentFromVector(ScrollPanel->GetDesiredSize());
 
diff --git a/Engine/Source/Runtime/Slate/Public/Widgets/Layout/SScrollBox.h b/Engine/Source/Runtime/Slate/Public/Widgets/Layout/SScrollBox.h
--- a/Engine/Source/Runtime/Slate/Public/Widgets/Layout/SScrollBox.h
+++ b/Engine/Source/Runtime/Slate/Public/Widgets/Layout/SScrollBox.h
@@ -111,12 +111,10 @@ public:
 	/** 
 	 * Attempt to scroll a widget into view, will safely handle non-descendant widgets 
 	 *
-	 * @param MyGeometry   The geometry of this widget.
 	 * @param WidgetToFind The widget whose geometry we wish to discover.
 	 * @param InAnimateScroll	Whether or not to animate the scroll
-	 * @return true if descendant
 	 */
-	bool ScrollDescendantIntoView(const FGeometry& MyGeometry, const TSharedPtr<SWidget>& WidgetToFind, bool InAnimateScroll = true);
+	void ScrollDescendantIntoView(const TSharedPtr<SWidget>& WidgetToFind, bool InAnimateScroll = true);
 
 	/** Get the current orientation of the scrollbox. */
 	EOrientation GetOrientation();
@@ -182,6 +180,9 @@ private:
 	/** Scroll offset that the user asked for. We will clamp it before actually scrolling there. */
 	float DesiredScrollOffset;
 
+	/** Scrolls or begins scrolling a widget into view, only valid to call when we have layout geometry. */
+	bool ScrollDescendantIntoView(const FGeometry& MyGeometry, const TSharedPtr<SWidget>& WidgetToFind, bool InAnimateScroll);
+
 	/**
 	 * Scroll the view by ScrollAmount given its currently AllottedGeometry.
 	 *
@@ -254,6 +255,12 @@ private:
 	/** Cached geometry for use with the active timer */
 	FGeometry CachedGeometry;
 
+	/** The widget to scroll into view next frame. */
+	TSharedPtr<SWidget> WidgetToScrollIntoView;
+
+	/** Whether we should scroll the widget into view over time. */
+	bool bAnimateScrollingWidgetIntoView : 1;
+
 	/**	Whether the software cursor should be drawn in the viewport */
 	bool bShowSoftwareCursor : 1;
 
diff --git a/Engine/Source/Runtime/UMG/Private/Components/ScrollBox.cpp b/Engine/Source/Runtime/UMG/Private/Components/ScrollBox.cpp
--- a/Engine/Source/Runtime/UMG/Private/Components/ScrollBox.cpp
+++ b/Engine/Source/Runtime/UMG/Private/Components/ScrollBox.cpp
@@ -128,6 +128,15 @@ void UScrollBox::ScrollToEnd()
 	}
 }
 
+void UScrollBox::ScrollWidgetIntoView(UWidget* WidgetToFind, bool AnimateScroll)
+{
+	TSharedPtr<SWidget> SlateWidgetToFind = WidgetToFind->GetCachedWidget();
+	if ( MyScrollBox.IsValid() )
+	{
+		MyScrollBox->ScrollDescendantIntoView(SlateWidgetToFind, AnimateScroll);
+	}
+}
+
 void UScrollBox::PostLoad()
 {
 	Super::PostLoad();
diff --git a/Engine/Source/Runtime/UMG/Public/Components/ScrollBox.h b/Engine/Source/Runtime/UMG/Public/Components/ScrollBox.h
--- a/Engine/Source/Runtime/UMG/Public/Components/ScrollBox.h
+++ b/Engine/Source/Runtime/UMG/Public/Components/ScrollBox.h
@@ -76,6 +76,10 @@ public:
 	UFUNCTION(BlueprintCallable, Category="Widget")
 	void ScrollToEnd();
 
+	/** Scrolls the ScrollBox to the widget during the next layout pass. */
+	UFUNCTION(BlueprintCallable, Category="Widget")
+	void ScrollWidgetIntoView(UWidget* WidgetToFind, bool AnimateScroll = true);
+
 	// UWidget interface
 	virtual void SynchronizeProperties() override;
 	// End of UWidget interface
