Fixed bad string access in SEditableText::IsAtWordStart

UE-16760 - Crash occurs using ctrl/cmd + right arrow in text fields after copying text

FString was recently changed to prevent you accessing the null terminator from within the string (the character at index Len()), and this code was still trying to access that character, and was triggering an assert.

I've added a guard to SEditableText::IsAtWordStart to prevent this, and also checked that the other string access was suitably guarded (it was).

diff --git a/Engine/Source/Runtime/Slate/Private/Widgets/Input/SEditableText.cpp b/Engine/Source/Runtime/Slate/Private/Widgets/Input/SEditableText.cpp
--- a/Engine/Source/Runtime/Slate/Private/Widgets/Input/SEditableText.cpp
+++ b/Engine/Source/Runtime/Slate/Private/Widgets/Input/SEditableText.cpp
@@ -1700,7 +1700,7 @@ bool SEditableText::IsAtWordStart( const int32 Location ) const
 {
 	const FString& EditedTextString = EditedText.ToString();
 
-	if (Location > 0)
+	if (Location > 0 && Location < EditedTextString.Len())
 	{
 		const TCHAR CharBeforeCursor = EditedTextString[Location - 1];
 		const TCHAR CharAtCursor = EditedTextString[Location];
