Fixed window activation on Mac in case where mouse down event changes focus to a new widget

diff --git a/Engine/Source/Runtime/Slate/Private/Framework/Application/SlateApplication.cpp b/Engine/Source/Runtime/Slate/Private/Framework/Application/SlateApplication.cpp
--- a/Engine/Source/Runtime/Slate/Private/Framework/Application/SlateApplication.cpp
+++ b/Engine/Source/Runtime/Slate/Private/Framework/Application/SlateApplication.cpp
@@ -4143,7 +4143,8 @@ bool FSlateApplication::ProcessMouseButtonDownEvent( const TSharedPtr< FGenericW
 			// If none of the widgets requested keyboard focus to be set (or set the keyboard focus explicitly), set it to the leaf-most widget under the mouse.
 			// On Mac we prevent the OS from activating the window on mouse down, so we have full control and can activate only if there's nothing draggable under the mouse cursor.
 			const bool bFocusChangedByEventHandler = PreviouslyFocusedWidget != GetKeyboardFocusedWidget();
-			if ((!Reply.GetUserFocusRecepient().IsValid() || (PLATFORM_MAC && MouseEvent.GetEffectingButton() == EKeys::LeftMouseButton && !DragDetector.DetectDragForWidget.IsValid())) && !bFocusChangedByEventHandler)
+			if ((!Reply.GetUserFocusRecepient().IsValid() || (PLATFORM_MAC && MouseEvent.GetEffectingButton() == EKeys::LeftMouseButton && !DragDetector.DetectDragForWidget.IsValid()))
+				&& (!bFocusChangedByEventHandler || (PLATFORM_MAC && ActiveWindow == nullptr)))
 			{
 				// The event handler for OnMouseButton down may have altered the widget hierarchy.
 				// Refresh the previously-cached widget path.
