Fixed Windows compile error in SlateApplication

diff --git a/Engine/Source/Runtime/Slate/Private/Framework/Application/SlateApplication.cpp b/Engine/Source/Runtime/Slate/Private/Framework/Application/SlateApplication.cpp
--- a/Engine/Source/Runtime/Slate/Private/Framework/Application/SlateApplication.cpp
+++ b/Engine/Source/Runtime/Slate/Private/Framework/Application/SlateApplication.cpp
@@ -4105,6 +4105,9 @@ bool FSlateApplication::ProcessMouseButtonDownEvent( const TSharedPtr< FGenericW
 
 #if PLATFORM_MAC
 			NSWindow* ActiveWindow = [NSApp keyWindow];
+			const bool bNeedToActivateWindow = (ActiveWindow == nullptr);
+#else
+			const bool bNeedToActivateWindow = false;
 #endif
 			PopupSupport.SendNotifications( WidgetsUnderCursor );
 
@@ -4144,7 +4147,7 @@ bool FSlateApplication::ProcessMouseButtonDownEvent( const TSharedPtr< FGenericW
 			// On Mac we prevent the OS from activating the window on mouse down, so we have full control and can activate only if there's nothing draggable under the mouse cursor.
 			const bool bFocusChangedByEventHandler = PreviouslyFocusedWidget != GetKeyboardFocusedWidget();
 			if ((!Reply.GetUserFocusRecepient().IsValid() || (PLATFORM_MAC && MouseEvent.GetEffectingButton() == EKeys::LeftMouseButton && !DragDetector.DetectDragForWidget.IsValid()))
-				&& (!bFocusChangedByEventHandler || (PLATFORM_MAC && ActiveWindow == nullptr)))
+				&& (!bFocusChangedByEventHandler || bNeedToActivateWindow))
 			{
 				// The event handler for OnMouseButton down may have altered the widget hierarchy.
 				// Refresh the previously-cached widget path.
