Slate - Text margins are now consistantly calculated at all zoom levels and no longer do an immediate calculation, they instead just invalidate layout like all other operations do.  Additionally there is now a step in the UpdateLayout pass after justification and flow have been calculated, not doing it last leads to a lot of bugs and inconsistances.  This checkin fixes the issue with zoom negatively affecting margins as well as the double application of margins to the justification of text, additional offsets were being calculated into the offset location of the text.  Also, because margin is now part of the layout calculation and no longer some crazy immediate thing that happens, the same operation now occurs every time, instead of sometimes looking correctly until layout got invalidated by some other means.

diff --git a/Engine/Source/Runtime/Slate/Private/Framework/Text/TextLayout.cpp b/Engine/Source/Runtime/Slate/Private/Framework/Text/TextLayout.cpp
--- a/Engine/Source/Runtime/Slate/Private/Framework/Text/TextLayout.cpp
+++ b/Engine/Source/Runtime/Slate/Private/Framework/Text/TextLayout.cpp
@@ -234,7 +234,7 @@ void FTextLayout::CreateLineViewBlocks( int32 LineModelIndex, const int32 StopIn
 	FVector2D LineSize( ForceInitToZero );
 	
 	// Use a negative scroll offset since positive scrolling moves things negatively in screen space
-	FVector2D CurrentOffset( Margin.Left - ScrollOffset.X, Margin.Top + TextLayoutSize.Height - ScrollOffset.Y );
+	FVector2D CurrentOffset(-ScrollOffset.X, TextLayoutSize.Height - ScrollOffset.Y);
 
 	if ( OutSoftLine.Num() > 0 )
 	{
@@ -281,7 +281,7 @@ void FTextLayout::JustifyLayout()
 		return;
 	}
 
-	const float LayoutWidthNoMargin = FMath::Max(TextLayoutSize.DrawWidth, ViewSize.X * Scale) - (Margin.GetTotalSpaceAlong<Orient_Horizontal>() * Scale);
+	const float LayoutWidthNoMargin = FMath::Max(TextLayoutSize.DrawWidth, ViewSize.X * Scale) - ( Margin.GetTotalSpaceAlong<Orient_Horizontal>() * Scale );
 
 	for (FLineView& LineView : LineViews)
 	{
@@ -323,13 +323,29 @@ void FTextLayout::FlowLayout()
 	{
 		FlowLineLayout(LineModelIndex, WrappingDrawWidth, SoftLine);
 	}
+}
 
+void FTextLayout::MarginLayout()
+{
 	// Add on the margins to the layout size
 	const float MarginWidth = Margin.GetTotalSpaceAlong<Orient_Horizontal>() * Scale;
 	const float MarginHeight = Margin.GetTotalSpaceAlong<Orient_Vertical>() * Scale;
 	TextLayoutSize.DrawWidth += MarginWidth;
 	TextLayoutSize.WrappedWidth += MarginWidth;
 	TextLayoutSize.Height += MarginHeight;
+
+	// Adjust the lines to be offset
+	FVector2D OffsetAdjustment = FVector2D(Margin.Left, Margin.Top) * Scale;
+	for ( int32 LineViewIndex = 0; LineViewIndex < LineViews.Num(); LineViewIndex++ )
+	{
+		FLineView& LineView = LineViews[LineViewIndex];
+
+		for ( int32 BlockIndex = 0; BlockIndex < LineView.Blocks.Num(); BlockIndex++ )
+		{
+			const TSharedRef< ILayoutBlock >& Block = LineView.Blocks[BlockIndex];
+			Block->SetLocationOffset(Block->GetLocationOffset() + OffsetAdjustment);
+		}
+	}
 }
 
 void FTextLayout::FlowLineLayout(const int32 LineModelIndex, const float WrappingDrawWidth, TArray<TSharedRef<ILayoutBlock>>& SoftLine)
@@ -634,6 +650,7 @@ void FTextLayout::UpdateLayout()
 
 	FlowLayout();
 	JustifyLayout();
+	MarginLayout();
 
 	EndLayout();
 
@@ -1682,35 +1699,8 @@ void FTextLayout::SetMargin( const FMargin& InMargin )
 		return;
 	}
 
-	const FMargin PreviousMargin = Margin;
 	Margin = InMargin;
-
-	if ( WrappingWidth > 0 )
-	{
-		// Since we are wrapping our text we'll need to rebuild the view as the actual wrapping width includes the margin.
-		DirtyFlags |= EDirtyState::Layout;
-	}
-	else
-	{
-		FVector2D OffsetAdjustment( Margin.Left - PreviousMargin.Left, Margin.Top - PreviousMargin.Top );
-		OffsetAdjustment *= Scale;
-		for (int32 LineViewIndex = 0; LineViewIndex < LineViews.Num(); LineViewIndex++)
-		{
-			FLineView& LineView = LineViews[ LineViewIndex ];
-
-			for (int32 BlockIndex = 0; BlockIndex < LineView.Blocks.Num(); BlockIndex++)
-			{
-				const TSharedRef< ILayoutBlock >& Block = LineView.Blocks[ BlockIndex ];
-				Block->SetLocationOffset( Block->GetLocationOffset() + OffsetAdjustment );
-			}
-		}
-
-		const float MarginDeltaWidth = ( Margin.GetTotalSpaceAlong<Orient_Horizontal>() - PreviousMargin.GetTotalSpaceAlong<Orient_Horizontal>() ) * Scale;
-		const float MarginDeltaHeight = ( Margin.GetTotalSpaceAlong<Orient_Vertical>() - PreviousMargin.GetTotalSpaceAlong<Orient_Vertical>() ) * Scale;
-		TextLayoutSize.DrawWidth += MarginDeltaWidth;
-		TextLayoutSize.WrappedWidth += MarginDeltaWidth;
-		TextLayoutSize.Height += MarginDeltaHeight;
-	}
+	DirtyFlags |= EDirtyState::Layout;
 }
 
 FMargin FTextLayout::GetMargin() const
diff --git a/Engine/Source/Runtime/Slate/Private/Widgets/Text/TextBlockLayout.cpp b/Engine/Source/Runtime/Slate/Private/Widgets/Text/TextBlockLayout.cpp
--- a/Engine/Source/Runtime/Slate/Private/Widgets/Text/TextBlockLayout.cpp
+++ b/Engine/Source/Runtime/Slate/Private/Widgets/Text/TextBlockLayout.cpp
@@ -28,11 +28,11 @@ FTextBlockLayout::FTextBlockLayout(FTextBlockStyle InDefaultTextStyle, TSharedRe
 FVector2D FTextBlockLayout::ComputeDesiredSize(const FWidgetArgs& InWidgetArgs, const float InScale, const FTextBlockStyle& InTextStyle)
 {
 	SLATE_CYCLE_COUNTER_SCOPE_DETAILED(SLATE_STATS_DETAIL_LEVEL_HI, GSlateTextBlockLayoutComputeDesiredSize);
+	TextLayout->SetScale(InScale);
 	TextLayout->SetWrappingWidth(CalculateWrappingWidth(InWidgetArgs));
-	TextLayout->SetMargin(InWidgetArgs.Margin.Get() * InScale);
+	TextLayout->SetMargin(InWidgetArgs.Margin.Get());
 	TextLayout->SetJustification(InWidgetArgs.Justification.Get());
 	TextLayout->SetLineHeightPercentage(InWidgetArgs.LineHeightPercentage.Get());
-	TextLayout->SetScale(InScale);
 
 	// Has the style used for this text block changed?
 	if(!IsStyleUpToDate(InTextStyle))
diff --git a/Engine/Source/Runtime/Slate/Public/Framework/Text/TextLayout.h b/Engine/Source/Runtime/Slate/Public/Framework/Text/TextLayout.h
--- a/Engine/Source/Runtime/Slate/Public/Framework/Text/TextLayout.h
+++ b/Engine/Source/Runtime/Slate/Public/Framework/Text/TextLayout.h
@@ -447,6 +447,7 @@ private:
 	float GetWrappingDrawWidth() const;
 
 	void FlowLayout();
+	void MarginLayout();
 
 	void FlowLineLayout(const int32 LineModelIndex, const float WrappingDrawWidth, TArray<TSharedRef<ILayoutBlock>>& SoftLine);
 
