SGridPanel row and column fills can not be specified as attributes

diff --git a/Engine/Source/Runtime/Slate/Private/Widgets/Layout/SGridPanel.cpp b/Engine/Source/Runtime/Slate/Private/Widgets/Layout/SGridPanel.cpp
--- a/Engine/Source/Runtime/Slate/Private/Widgets/Layout/SGridPanel.cpp
+++ b/Engine/Source/Runtime/Slate/Private/Widgets/Layout/SGridPanel.cpp
@@ -105,6 +105,38 @@ int32 SGridPanel::OnPaint( const FPaintArgs& Args, const FGeometry& AllottedGeom
 	return MaxLayerId;
 }
 
+void CalculateStretchedCellSizes(TArray<float>& OutSizes, float AllotedSize, const TArray<float>& InDesiredSizes, const TArray<TAttribute<float>>& Coefficients)
+{
+	const int32 NumCoefficients = Coefficients.Num();
+	float CoefficientTotal = 0.f;
+
+	for(int32 Index=0; Index < InDesiredSizes.Num(); ++Index)
+	{
+		const float Coefficient = Index < NumCoefficients ? Coefficients[Index].Get(0) : 0;
+
+		// Compute the total space available for stretchy columns.
+		if (Coefficient == 0)
+		{
+			AllotedSize -= InDesiredSizes[Index];
+		}
+		else
+		{
+			// Compute the denominator for dividing up the stretchy column space
+			CoefficientTotal += Coefficient;
+		}
+
+	}
+
+	for(int32 Index=0; Index < InDesiredSizes.Num(); ++Index)
+	{
+		const float Coefficient = Index < NumCoefficients ? Coefficients[Index].Get(0) : 0;
+
+		// Figure out how big each column needs to be
+		OutSizes[Index] = Coefficient != 0
+			? (Coefficient / CoefficientTotal * AllotedSize)
+			: InDesiredSizes[Index];
+	}
+}
 
 void SGridPanel::OnArrangeChildren( const FGeometry& AllottedGeometry, FArrangedChildren& ArrangedChildren ) const
 {
@@ -130,59 +162,13 @@ void SGridPanel::OnArrangeChildren( const FGeometry& AllottedGeometry, FArranged
 		FinalRows[FinalRows.Num()-1] = 0.0f;
 	}
 	
-	FVector2D FlexSpace = AllottedGeometry.Size;
-	const int32 ColFillCoeffsLength = ColFillCoefficients.Num();
-	for(int32 ColIndex=0; ColIndex < Columns.Num(); ++ColIndex)
-	{
-		// Compute the total space available for stretchy columns.
-		if (ColIndex >= ColFillCoeffsLength || ColFillCoefficients[ColIndex] == 0)
-		{
-			FlexSpace.X -= Columns[ColIndex];
-		}
-		else //(ColIndex < ColFillCoeffsLength)
-		{
-			// Compute the denominator for dividing up the stretchy column space
-			ColumnCoeffTotal += ColFillCoefficients[ColIndex];
-		}
-
-	}
-
-	for(int32 ColIndex=0; ColIndex < Columns.Num(); ++ColIndex)
-	{
-		// Figure out how big each column needs to be
-		FinalColumns[ColIndex] = (ColIndex < ColFillCoeffsLength && ColFillCoefficients[ColIndex] != 0)
-			? (ColFillCoefficients[ColIndex] / ColumnCoeffTotal * FlexSpace.X)
-			: Columns[ColIndex];
-	}
-
-	const int32 RowFillCoeffsLength = RowFillCoefficients.Num();
-	for(int32 RowIndex=0; RowIndex < Rows.Num(); ++RowIndex)
-	{
-		// Compute the total space available for stretchy rows.
-		if (RowIndex >= RowFillCoeffsLength || RowFillCoefficients[RowIndex] == 0)
-		{
-			FlexSpace.Y -= Rows[RowIndex];
-		}
-		else //(RowIndex < RowFillCoeffsLength)
-		{
-			// Compute the denominator for dividing up the stretchy row space
-			RowCoeffTotal += RowFillCoefficients[RowIndex];
-		}
-
-	}
-
-	for(int32 RowIndex=0; RowIndex < Rows.Num(); ++RowIndex)
-	{
-		// Compute how big each row needs to be
-		FinalRows[RowIndex] = (RowIndex < RowFillCoeffsLength && RowFillCoefficients[RowIndex] != 0)
-			? (RowFillCoefficients[RowIndex] / RowCoeffTotal * FlexSpace.Y)
-			: Rows[RowIndex];
-	}
+	CalculateStretchedCellSizes(FinalColumns, AllottedGeometry.Size.X, Columns, ColFillCoefficients);
+	CalculateStretchedCellSizes(FinalRows, AllottedGeometry.Size.Y, Rows, RowFillCoefficients);
 	
 	// Build up partial sums for row and column sizes so that we can handle column and row spans conveniently.
 	ComputePartialSums(FinalColumns);
 	ComputePartialSums(FinalRows);
-	
+
 	// ARRANGE PHASE
 	for( int32 SlotIndex=0; SlotIndex < Slots.Num(); ++SlotIndex )
 	{
@@ -265,20 +251,20 @@ FVector2D SGridPanel::GetDesiredSize( const FIntPoint& StartCell, int32 Width, i
 	}
 }
 
-void SGridPanel::SetColumnFill( int32 ColumnId, float Coefficient )
+void SGridPanel::SetColumnFill( int32 ColumnId, const TAttribute<float>& Coefficient )
 {
-	if (ColFillCoefficients.Num() <= ColumnId)
+	while (ColFillCoefficients.Num() <= ColumnId)
 	{
-		ColFillCoefficients.AddZeroed( ColumnId - ColFillCoefficients.Num() + 1 );
+		ColFillCoefficients.Emplace(0);
 	}
 	ColFillCoefficients[ColumnId] = Coefficient;
 }
 
-void SGridPanel::SetRowFill( int32 RowId, float Coefficient )
+void SGridPanel::SetRowFill( int32 RowId, const TAttribute<float>& Coefficient )
 {
 	if (RowFillCoefficients.Num() <= RowId)
 	{
-		RowFillCoefficients.AddZeroed( RowId - RowFillCoefficients.Num() + 1 );
+		RowFillCoefficients.Emplace(0);
 	}
 	RowFillCoefficients[RowId] = Coefficient;
 }
diff --git a/Engine/Source/Runtime/Slate/Public/Widgets/Layout/SGridPanel.h b/Engine/Source/Runtime/Slate/Public/Widgets/Layout/SGridPanel.h
--- a/Engine/Source/Runtime/Slate/Public/Widgets/Layout/SGridPanel.h
+++ b/Engine/Source/Runtime/Slate/Public/Widgets/Layout/SGridPanel.h
@@ -147,32 +147,32 @@ public:
 		SLATE_SUPPORTS_SLOT( FSlot )
 
 		/** Specify a column to stretch instead of sizing to content. */
-		FArguments& FillColumn( int32 ColumnId, float Coefficient )
+		FArguments& FillColumn( int32 ColumnId, const TAttribute<float>& Coefficient )
 		{
-			if (ColFillCoefficients.Num() <= ColumnId)
+			while (ColFillCoefficients.Num() <= ColumnId)
 			{
-				ColFillCoefficients.AddZeroed( ColumnId - ColFillCoefficients.Num() + 1 );
+				ColFillCoefficients.Emplace( 0 );
 			}
 			ColFillCoefficients[ColumnId] = Coefficient;
 			return Me();
 		}
 
 		/** Specify a row to stretch instead of sizing to content. */
-		FArguments& FillRow( int32 RowId, float Coefficient )
+		FArguments& FillRow( int32 RowId, const TAttribute<float>& Coefficient )
 		{
-			if (RowFillCoefficients.Num() <= RowId)
+			while (RowFillCoefficients.Num() <= RowId)
 			{
-				RowFillCoefficients.AddZeroed( RowId - RowFillCoefficients.Num() + 1 );
+				RowFillCoefficients.Emplace( 0 );
 			}
 			RowFillCoefficients[RowId] = Coefficient;
 			return Me();
 		}
 
 		/** Coefficients for columns that need to stretch instead of size to content */
-		TArray<float> ColFillCoefficients;
+		TArray<TAttribute<float>> ColFillCoefficients;
 		
 		/** Coefficients for rows that need to stretch instead of size to content */
-		TArray<float> RowFillCoefficients;
+		TArray<TAttribute<float>> RowFillCoefficients;
 		
 	SLATE_END_ARGS()
 
@@ -194,10 +194,10 @@ public:
 	FVector2D GetDesiredSize( const FIntPoint& StartCell, int32 Width, int32 Height ) const;
 
 	/** Specify a column to stretch instead of sizing to content. */
-	void SetColumnFill( int32 ColumnId, float Coefficient );
+	void SetColumnFill( int32 ColumnId, const TAttribute<float>& Coefficient );
 
 	/** Specify a row to stretch instead of sizing to content. */
-	void SetRowFill( int32 RowId, float Coefficient );
+	void SetRowFill( int32 RowId, const TAttribute<float>& Coefficient );
 
 public:
 
@@ -264,8 +264,8 @@ private:
 	FVector2D TotalDesiredSizes;
 
 	/** Fill coefficients for the columns */
-	TArray<float> ColFillCoefficients;
+	TArray<TAttribute<float>> ColFillCoefficients;
 
 	/** Fill coefficients for the rows */
-	TArray<float> RowFillCoefficients;
+	TArray<TAttribute<float>> RowFillCoefficients;
 };
