Fixed some places in FSlateApplication that were incorrectly taking a TSharedRef by ref, rather than const-ref

diff --git a/Engine/Source/Runtime/Slate/Private/Framework/Application/SlateApplication.cpp b/Engine/Source/Runtime/Slate/Private/Framework/Application/SlateApplication.cpp
--- a/Engine/Source/Runtime/Slate/Private/Framework/Application/SlateApplication.cpp
+++ b/Engine/Source/Runtime/Slate/Private/Framework/Application/SlateApplication.cpp
@@ -3528,17 +3528,17 @@ void FSlateApplication::GotoLineInSource(const FString& FileName, int32 LineNumb
 	}
 }
 
-void FSlateApplication::ForceRedrawWindow(TSharedRef<SWindow>& InWindowToDraw)
+void FSlateApplication::ForceRedrawWindow(const TSharedRef<SWindow>& InWindowToDraw)
 {
 	PrivateDrawWindows( InWindowToDraw );
 }
 
-bool FSlateApplication::TakeScreenshot(TSharedRef<SWidget>& Widget, TArray<FColor>&OutColorData, FIntVector& OutSize)
+bool FSlateApplication::TakeScreenshot(const TSharedRef<SWidget>& Widget, TArray<FColor>&OutColorData, FIntVector& OutSize)
 {
 	return TakeScreenshot(Widget, FIntRect(), OutColorData, OutSize);
 }
 
-bool FSlateApplication::TakeScreenshot(TSharedRef<SWidget>& Widget, const FIntRect& InnerWidgetArea, TArray<FColor>& OutColorData, FIntVector& OutSize)
+bool FSlateApplication::TakeScreenshot(const TSharedRef<SWidget>& Widget, const FIntRect& InnerWidgetArea, TArray<FColor>& OutColorData, FIntVector& OutSize)
 {
 	// We can't screenshot the widget unless there's a valid window handle to draw it in.
 	TSharedPtr<SWindow> WidgetWindow = FSlateApplication::Get().FindWidgetWindow(Widget);
diff --git a/Engine/Source/Runtime/Slate/Public/Framework/Application/SlateApplication.h b/Engine/Source/Runtime/Slate/Public/Framework/Application/SlateApplication.h
--- a/Engine/Source/Runtime/Slate/Public/Framework/Application/SlateApplication.h
+++ b/Engine/Source/Runtime/Slate/Public/Framework/Application/SlateApplication.h
@@ -730,7 +730,7 @@ public:
 	/**
 	 * Forces the window to redraw immediately.
 	 */
-	void ForceRedrawWindow( TSharedRef<SWindow>& InWindowToDraw );
+	void ForceRedrawWindow( const TSharedRef<SWindow>& InWindowToDraw );
 
 	/**
 	 * Takes a screenshot of the widget writing the results into the color buffer provided.  Note that the format is BGRA.
@@ -738,7 +738,7 @@ public:
 	 * 
 	 * @return true if taking the screenshot was successful.
 	 */
-	bool TakeScreenshot(TSharedRef<SWidget>& Widget, TArray<FColor>&OutColorData, FIntVector& OutSize);
+	bool TakeScreenshot(const TSharedRef<SWidget>& Widget, TArray<FColor>&OutColorData, FIntVector& OutSize);
 
 	/**
 	 * Takes a screenshot of the widget writing the results into the color buffer provided, this version allows you to provide 
@@ -746,7 +746,7 @@ public:
 	 *
 	 * @return true if taking the screenshot was successful.
 	 */
-	bool TakeScreenshot(TSharedRef<SWidget>& Widget, const FIntRect& InnerWidgetArea, TArray<FColor>& OutColorData, FIntVector& OutSize);
+	bool TakeScreenshot(const TSharedRef<SWidget>& Widget, const FIntRect& InnerWidgetArea, TArray<FColor>& OutColorData, FIntVector& OutSize);
 
 protected:
 
