Removed appearance override properties from UMG widgets. Style property is now used.
[UE-14234] UMG widgets with overriding appearance properties always clobber the Style

#change Deprecated appearance properties on UWidgets that override the Style
#change Incremented the engine object version number and fixed up deprecated properties in OnLoad
#change Added sync'ing of Style properties in the UMG editor where missing. UMG wysiwyg needs to update in response to styling changes. Required adding of some SetStyle code to SWidgets here and there.

diff --git a/Engine/Source/Runtime/Core/Public/UObject/ObjectVersion.h b/Engine/Source/Runtime/Core/Public/UObject/ObjectVersion.h
--- a/Engine/Source/Runtime/Core/Public/UObject/ObjectVersion.h
+++ b/Engine/Source/Runtime/Core/Public/UObject/ObjectVersion.h
@@ -1,4 +1,4 @@
-// Copyright 1998-2015 Epic Games, Inc. All Rights Reserved.
+﻿// Copyright 1998-2015 Epic Games, Inc. All Rights Reserved.
 
 #pragma once
 
@@ -528,6 +528,9 @@ enum EUnrealEngineObjectUE4Version
 	VER_UE4_TEXTURE_LEGACY_GAMMA,
 	// Added WithSerializer for basic native structures like FVector, FColor etc to improve serialization performance
 	VER_UE4_ADDED_NATIVE_SERIALIZATION_FOR_IMMUTABLE_STRUCTURES,
+	// Deprecated attributes that override the style on UMG widgets
+	VER_UE4_DEPRECATE_UMG_STYLE_OVERRIDES,
+
 	// -----<new versions can be added before this line>-------------------------------------------------
 	// - this needs to be the last line (see note below)
 	VER_UE4_AUTOMATIC_VERSION_PLUS_ONE,
diff --git a/Engine/Source/Runtime/Slate/Private/Widgets/Input/SCheckBox.cpp b/Engine/Source/Runtime/Slate/Private/Widgets/Input/SCheckBox.cpp
--- a/Engine/Source/Runtime/Slate/Private/Widgets/Input/SCheckBox.cpp
+++ b/Engine/Source/Runtime/Slate/Private/Widgets/Input/SCheckBox.cpp
@@ -25,69 +25,17 @@
 	UndeterminedHoveredImage = InArgs._UndeterminedHoveredImage;
 	UndeterminedPressedImage = InArgs._UndeterminedPressedImage;
 	
-	TAttribute<FMargin> Padding = InArgs._Padding.IsSet() ? InArgs._Padding : InArgs._Style->Padding;
-	TAttribute<FSlateColor> ForegroundColor = InArgs._ForegroundColor.IsSet() ? InArgs._ForegroundColor : InArgs._Style->ForegroundColor;
-	TAttribute<FSlateColor> BorderBackgroundColor = InArgs._BorderBackgroundColor.IsSet() ? InArgs._BorderBackgroundColor : InArgs._Style->BorderBackgroundColor;
-	ESlateCheckBoxType::Type CheckBoxType = InArgs._Type.IsSet() ? InArgs._Type.GetValue() : InArgs._Style->CheckBoxType.GetValue();
+	PaddingOverride = InArgs._Padding;
+	ForegroundColorOverride = InArgs._ForegroundColor;
+	BorderBackgroundColorOverride = InArgs._BorderBackgroundColor;
+	CheckBoxTypeOverride = InArgs._Type;
 
+	HorizontalAlignment = InArgs._HAlign;
+
 	bIsPressed = false;
 	bIsFocusable = InArgs._IsFocusable;
 
-	if( CheckBoxType == ESlateCheckBoxType::CheckBox )
-	{
-		// Check boxes use a separate check button to the side of the user's content (often, a text label or icon.)
-		this->ChildSlot
-		[
-			SNew(SHorizontalBox)
-			// Make sure we aren't trying to compute the desired size when the check box is collapsed
-			.Visibility(InArgs._Visibility)
-			+ SHorizontalBox::Slot()
-			.AutoWidth()
-			.VAlign( VAlign_Center )
-			.HAlign( HAlign_Center )
-			[
-				SNew(SImage)
-				.Image( this, &SCheckBox::OnGetCheckImage )
-				.ColorAndOpacity( ForegroundColor )
-			]
-			+ SHorizontalBox::Slot()
-			.AutoWidth()
-			.Padding( Padding )
-			.VAlign( VAlign_Center )
-			[
-				SAssignNew(ContentContainer, SBorder)
-				.BorderImage(FStyleDefaults::GetNoBrush())
-				.Padding(0.0f)
-				[
-					InArgs._Content.Widget
-				]
-			]
-		];
-	}
-	else if( ensure( CheckBoxType == ESlateCheckBoxType::ToggleButton ) )
-	{
+	BuildCheckBox(InArgs._Content.Widget);
 
-		// Toggle buttons have a visual appearance that is similar to a Slate button
-		this->ChildSlot
-		[
-			SAssignNew( ContentContainer, SBorder )
-			// Make sure we aren't trying to compute the desired size when the check box is collapsed
-			.Visibility(InArgs._Visibility)
-			// Bind the border background to our method that gets a slate brush for the current state of the control
-			.BorderImage( this, &SCheckBox::OnGetCheckImage )
-			// Let the user decide the padding amount
-			.Padding( Padding )
-			// Set the user's incoming content as the content of our border
-			.ForegroundColor( ForegroundColor )
-			// Set the color of the border image
-			.BorderBackgroundColor( BorderBackgroundColor )
-			// Content alignment
-			.HAlign( InArgs._HAlign )
-			[
-				InArgs._Content.Widget
-			]
-		];
-	}
-
 	IsCheckboxChecked = InArgs._IsChecked;
 	OnCheckStateChanged = InArgs._OnCheckStateChanged;
@@ -366,6 +314,8 @@
 void SCheckBox::SetStyle(const FCheckBoxStyle* InStyle)
 {
 	Style = InStyle;
+
+	BuildCheckBox(ContentContainer->GetContent());
 }
 
 void SCheckBox::SetUncheckedImage(const FSlateBrush* Brush)
@@ -413,6 +363,91 @@
 	UndeterminedPressedImage = Brush;
 }
 
+void SCheckBox::BuildCheckBox(TSharedRef<SWidget> InContent)
+{
+	if (ContentContainer.IsValid())
+	{
+		ContentContainer->SetContent(SNullWidget::NullWidget);
+	}
+
+	ESlateCheckBoxType::Type CheckBoxType = OnGetCheckBoxType();
+
+	if (CheckBoxType == ESlateCheckBoxType::CheckBox)
+	{
+		// Check boxes use a separate check button to the side of the user's content (often, a text label or icon.)
+		this->ChildSlot
+		[
+			SNew(SHorizontalBox)
+			// Make sure we aren't trying to compute the desired size when the check box is collapsed
+			.Visibility(this->Visibility)
+			+ SHorizontalBox::Slot()
+			.AutoWidth()
+			.VAlign(VAlign_Center)
+			.HAlign(HAlign_Center)
+			[
+				SNew(SImage)
+				.Image(this, &SCheckBox::OnGetCheckImage)
+				.ColorAndOpacity(this, &SCheckBox::OnGetForegroundColor)
+			]
+			+ SHorizontalBox::Slot()
+			.AutoWidth()
+			.Padding(TAttribute<FMargin>(this, &SCheckBox::OnGetPadding))
+			.VAlign(VAlign_Center)
+			[
+				SAssignNew(ContentContainer, SBorder)
+				.BorderImage(FStyleDefaults::GetNoBrush())
+				.Padding(0.0f)
+				[
+					InContent
+				]
+			]
+		];
+	}
+	else if (ensure(CheckBoxType == ESlateCheckBoxType::ToggleButton))
+	{
+		// Toggle buttons have a visual appearance that is similar to a Slate button
+		this->ChildSlot
+		[
+			SAssignNew(ContentContainer, SBorder)
+			// Make sure we aren't trying to compute the desired size when the check box is collapsed
+			.Visibility(this->Visibility)
+			// Bind the border background to our method that gets a slate brush for the current state of the control
+			.BorderImage(this, &SCheckBox::OnGetCheckImage)
+			// Let the user decide the padding amount
+			.Padding(this, &SCheckBox::OnGetPadding)
+			// Set the user's incoming content as the content of our border
+			.ForegroundColor(this, &SCheckBox::OnGetForegroundColor)
+			// Set the color of the border image
+			.BorderBackgroundColor(this, &SCheckBox::OnGetBorderBackgroundColor)
+			// Content alignment
+			.HAlign(HorizontalAlignment)
+			[
+				InContent
+			]
+		];
+	}
+}
+
+FSlateColor SCheckBox::OnGetForegroundColor() const
+{
+	return ForegroundColorOverride.IsSet() ? ForegroundColorOverride.Get() : Style->ForegroundColor;
+}
+
+FMargin SCheckBox::OnGetPadding() const
+{
+	return PaddingOverride.IsSet() ? PaddingOverride.Get() : Style->Padding;
+}
+
+FSlateColor SCheckBox::OnGetBorderBackgroundColor() const
+{
+	return BorderBackgroundColorOverride.IsSet() ? BorderBackgroundColorOverride.Get() : Style->BorderBackgroundColor;
+}
+
+ESlateCheckBoxType::Type SCheckBox::OnGetCheckBoxType() const
+{
+	return CheckBoxTypeOverride.IsSet() ? CheckBoxTypeOverride.GetValue() : Style->CheckBoxType.GetValue();
+}
+
 const FSlateBrush* SCheckBox::GetUncheckedImage() const
 {
 	return UncheckedImage ? UncheckedImage : &Style->UncheckedImage;
diff --git a/Engine/Source/Runtime/Slate/Private/Widgets/Input/SEditableTextBox.cpp b/Engine/Source/Runtime/Slate/Private/Widgets/Input/SEditableTextBox.cpp
--- a/Engine/Source/Runtime/Slate/Private/Widgets/Input/SEditableTextBox.cpp
+++ b/Engine/Source/Runtime/Slate/Private/Widgets/Input/SEditableTextBox.cpp
@@ -13,11 +13,17 @@ void SEditableTextBox::Construct( const FArguments& InArgs )
 	check (InArgs._Style);
 	SetStyle(InArgs._Style);
 
-	TAttribute<FMargin> Padding = InArgs._Padding.IsSet() ? InArgs._Padding : InArgs._Style->Padding;
-	TAttribute<FSlateFontInfo> Font = InArgs._Font.IsSet() ? InArgs._Font : InArgs._Style->Font;
-	TAttribute<FSlateColor> ForegroundColor = InArgs._ForegroundColor.IsSet() ? InArgs._ForegroundColor : InArgs._Style->ForegroundColor;
-	TAttribute<FSlateColor> BackgroundColor = InArgs._BackgroundColor.IsSet() ? InArgs._BackgroundColor : InArgs._Style->BackgroundColor;
-	ReadOnlyForegroundColor = InArgs._ReadOnlyForegroundColor.IsSet() ? InArgs._ReadOnlyForegroundColor : InArgs._Style->ReadOnlyForegroundColor;
+	PaddingOverride = InArgs._Padding;
+	FontOverride = InArgs._Font;
+	ForegroundColorOverride = InArgs._ForegroundColor;
+	BackgroundColorOverride = InArgs._BackgroundColor;
+	ReadOnlyForegroundColorOverride = InArgs._ReadOnlyForegroundColor;
+
+	TAttribute<FMargin> Padding = PaddingOverride.IsSet() ? PaddingOverride : InArgs._Style->Padding;
+	TAttribute<FSlateFontInfo> Font = FontOverride.IsSet() ? FontOverride : InArgs._Style->Font;
+	TAttribute<FSlateColor> ForegroundColor = ForegroundColorOverride.IsSet() ? ForegroundColorOverride : InArgs._Style->ForegroundColor;
+	TAttribute<FSlateColor> BackgroundColor = BackgroundColorOverride.IsSet() ? BackgroundColorOverride : InArgs._Style->BackgroundColor;
+	ReadOnlyForegroundColor = ReadOnlyForegroundColorOverride.IsSet() ? ReadOnlyForegroundColorOverride : InArgs._Style->ReadOnlyForegroundColor;
 
 	SBorder::Construct( SBorder::FArguments()
 		.BorderImage( this, &SEditableTextBox::GetBorderImage )
@@ -32,7 +38,7 @@ void SEditableTextBox::Construct( const FArguments& InArgs )
 			.HAlign(HAlign_Fill)
 			.FillWidth(1)
 			[
-				SNew( SBox )
+				SAssignNew(PaddingBox, SBox)
 				.Padding( Padding )
 				[
 					SAssignNew( EditableText, SEditableText )
@@ -80,6 +86,31 @@ void SEditableTextBox::SetStyle(const FEditableTextBoxStyle* InStyle)
 		Style = Defaults._Style;
 	}
 
+	if (!PaddingOverride.IsSet() && PaddingBox.IsValid())
+	{
+		PaddingBox->SetPadding(Style->Padding);
+	}
+
+	if (!FontOverride.IsSet() && EditableText.IsValid())
+	{
+		EditableText->SetFont(Style->Font);
+	}
+
+	if (!ForegroundColorOverride.IsSet())
+	{
+		SetForegroundColor(Style->ForegroundColor);
+	}
+
+	if (!BackgroundColorOverride.IsSet())
+	{
+		SetBorderBackgroundColor(Style->BackgroundColor);
+	}
+
+	if (!ReadOnlyForegroundColorOverride.IsSet())
+	{
+		ReadOnlyForegroundColor = Style->ReadOnlyForegroundColor;
+	}
+
 	BorderImageNormal = &Style->BackgroundImageNormal;
 	BorderImageHovered = &Style->BackgroundImageHovered;
 	BorderImageFocused = &Style->BackgroundImageFocused;
@@ -211,30 +242,32 @@ void SEditableTextBox::SetIsPassword(TAttribute< bool > InIsPassword)
 
 void SEditableTextBox::SetFont(const TAttribute<FSlateFontInfo>& InFont)
 {
-	TAttribute<FSlateFontInfo> Font = InFont.IsSet() ? InFont : Style->Font;
+	FontOverride = InFont;
 
-	EditableText->SetFont(Font);
+	EditableText->SetFont(FontOverride.IsSet() ? FontOverride : Style->Font);
 }
 
 
 void SEditableTextBox::SetTextBoxForegroundColor(const TAttribute<FSlateColor>& InForegroundColor)
 {
-	TAttribute<FSlateColor> ForegroundColor = InForegroundColor.IsSet() ? InForegroundColor : Style->ForegroundColor;
+	ForegroundColorOverride = InForegroundColor;
 
-	SetForegroundColor(ForegroundColor);
+	SetForegroundColor(ForegroundColorOverride.IsSet() ? ForegroundColorOverride : Style->ForegroundColor);
 }
 
 void SEditableTextBox::SetTextBoxBackgroundColor(const TAttribute<FSlateColor>& InBackgroundColor)
 {
-	TAttribute<FSlateColor> BackgroundColor = InBackgroundColor.IsSet() ? InBackgroundColor : Style->BackgroundColor;
+	BackgroundColorOverride = InBackgroundColor;
 
-	SetBorderBackgroundColor(BackgroundColor);
+	SetBorderBackgroundColor(BackgroundColorOverride.IsSet() ? BackgroundColorOverride : Style->BackgroundColor);
 }
 
 
 void SEditableTextBox::SetReadOnlyForegroundColor(const TAttribute<FSlateColor>& InReadOnlyForegroundColor)
 {
-	ReadOnlyForegroundColor = InReadOnlyForegroundColor.IsSet() ? InReadOnlyForegroundColor : Style->ReadOnlyForegroundColor;
+	ReadOnlyForegroundColorOverride = InReadOnlyForegroundColor;
+
+	ReadOnlyForegroundColor = ReadOnlyForegroundColorOverride.IsSet() ? ReadOnlyForegroundColorOverride : Style->ReadOnlyForegroundColor;
 }
 
 
diff --git a/Engine/Source/Runtime/Slate/Private/Widgets/Input/SMultiLineEditableTextBox.cpp b/Engine/Source/Runtime/Slate/Private/Widgets/Input/SMultiLineEditableTextBox.cpp
--- a/Engine/Source/Runtime/Slate/Private/Widgets/Input/SMultiLineEditableTextBox.cpp
+++ b/Engine/Source/Runtime/Slate/Private/Widgets/Input/SMultiLineEditableTextBox.cpp
@@ -36,6 +36,15 @@ void SMultiLineEditableTextBox::Construct( const FArguments& InArgs )
 	BorderImageHovered = &InArgs._Style->BackgroundImageHovered;
 	BorderImageFocused = &InArgs._Style->BackgroundImageFocused;
 	BorderImageReadOnly = &InArgs._Style->BackgroundImageReadOnly;
+
+	PaddingOverride = InArgs._Padding;
+	HScrollBarPaddingOverride = InArgs._HScrollBarPadding;
+	VScrollBarPaddingOverride = InArgs._VScrollBarPadding;
+	FontOverride = InArgs._Font;
+	ForegroundColorOverride = InArgs._ForegroundColor;
+	BackgroundColorOverride = InArgs._BackgroundColor;
+	ReadOnlyForegroundColorOverride = InArgs._ReadOnlyForegroundColor;
+
 	TAttribute<FMargin> Padding = InArgs._Padding.IsSet() ? InArgs._Padding : InArgs._Style->Padding;
 	TAttribute<FMargin> HScrollBarPadding = InArgs._HScrollBarPadding.IsSet() ? InArgs._HScrollBarPadding : InArgs._Style->HScrollBarPadding;
 	TAttribute<FMargin> VScrollBarPadding = InArgs._VScrollBarPadding.IsSet() ? InArgs._VScrollBarPadding : InArgs._Style->VScrollBarPadding;
@@ -44,7 +53,8 @@ void SMultiLineEditableTextBox::Construct( const FArguments& InArgs )
 	TAttribute<FSlateColor> BackgroundColor = InArgs._BackgroundColor.IsSet() ? InArgs._BackgroundColor : InArgs._Style->BackgroundColor;
 	ReadOnlyForegroundColor = InArgs._ReadOnlyForegroundColor.IsSet() ? InArgs._ReadOnlyForegroundColor : InArgs._Style->ReadOnlyForegroundColor;
 
-	TSharedPtr<SScrollBar> HScrollBar = InArgs._HScrollBar;
+	bHasExternalHScrollBar = InArgs._HScrollBar.IsValid();
+	HScrollBar = InArgs._HScrollBar;
 	if (!HScrollBar.IsValid())
 	{
 		// Create and use our own scrollbar
@@ -55,7 +65,8 @@ void SMultiLineEditableTextBox::Construct( const FArguments& InArgs )
 			.Thickness(FVector2D(5.0f, 5.0f));
 	}
 	
-	TSharedPtr<SScrollBar> VScrollBar = InArgs._VScrollBar;
+	bHasExternalVScrollBar = InArgs._VScrollBar.IsValid();
+	VScrollBar = InArgs._VScrollBar;
 	if (!VScrollBar.IsValid())
 	{
 		// Create and use our own scrollbar
@@ -115,17 +126,23 @@ void SMultiLineEditableTextBox::Construct( const FArguments& InArgs )
 
 				+SVerticalBox::Slot()
 				.AutoHeight()
-				.Padding(HScrollBarPadding)
 				[
-					AsWidgetRef( HScrollBar )
+					SAssignNew(HScrollBarPaddingBox, SBox)
+					.Padding(HScrollBarPadding)
+					[
+						AsWidgetRef( HScrollBar )
+					]
 				]
 			]
 
 			+SHorizontalBox::Slot()
 			.AutoWidth()
-			.Padding(VScrollBarPadding)
 			[
-				AsWidgetRef( VScrollBar )
+				SAssignNew(VScrollBarPaddingBox, SBox)
+				.Padding(VScrollBarPadding)
+				[
+					AsWidgetRef( VScrollBar )
+				]
 			]
 		]
 	);
@@ -143,6 +160,62 @@ void SMultiLineEditableTextBox::Construct( const FArguments& InArgs )
 
 }
 
+void SMultiLineEditableTextBox::SetStyle(const FEditableTextBoxStyle* InStyle)
+{
+	check(InStyle);
+	Style = InStyle;
+
+	if (!PaddingOverride.IsSet())
+	{
+		SetPadding(Style->Padding);
+	}
+
+	if (!HScrollBarPaddingOverride.IsSet() && HScrollBarPaddingBox.IsValid())
+	{
+		HScrollBarPaddingBox->SetPadding(Style->HScrollBarPadding);
+	}
+
+	if (!VScrollBarPaddingOverride.IsSet() && VScrollBarPaddingBox.IsValid())
+	{
+		VScrollBarPaddingBox->SetPadding(Style->VScrollBarPadding);
+	}
+
+	if (!FontOverride.IsSet() && EditableText.IsValid())
+	{
+		EditableText->SetFont(Style->Font);
+	}
+
+	if (!ForegroundColorOverride.IsSet())
+	{
+		SetForegroundColor(Style->ForegroundColor);
+	}
+
+	if (!BackgroundColorOverride.IsSet())
+	{
+		SetBorderBackgroundColor(Style->BackgroundColor);
+	}
+
+	if (!ReadOnlyForegroundColorOverride.IsSet())
+	{
+		ReadOnlyForegroundColor = Style->ReadOnlyForegroundColor;
+	}
+
+	if (!bHasExternalHScrollBar && HScrollBar.IsValid())
+	{
+		HScrollBar->SetStyle(&Style->ScrollBarStyle);
+	}
+
+	if (!bHasExternalVScrollBar && VScrollBar.IsValid())
+	{
+		VScrollBar->SetStyle(&Style->ScrollBarStyle);
+	}
+
+	BorderImageNormal = &InStyle->BackgroundImageNormal;
+	BorderImageHovered = &InStyle->BackgroundImageHovered;
+	BorderImageFocused = &InStyle->BackgroundImageFocused;
+	BorderImageReadOnly = &InStyle->BackgroundImageReadOnly;
+}
+
 void SMultiLineEditableTextBox::SetText( const TAttribute< FText >& InNewText )
 {
 	EditableText->SetText( InNewText );
@@ -155,21 +228,23 @@ void SMultiLineEditableTextBox::SetHintText( const TAttribute< FText >& InHintTe
 
 void SMultiLineEditableTextBox::SetTextBoxForegroundColor(const TAttribute<FSlateColor>& InForegroundColor)
 {
-	TAttribute<FSlateColor> ForegroundColor = InForegroundColor.IsSet() ? InForegroundColor : Style->ForegroundColor;
+	ForegroundColorOverride = InForegroundColor;
 
-	SetForegroundColor(ForegroundColor);
+	SetForegroundColor(ForegroundColorOverride.IsSet() ? ForegroundColorOverride : Style->ForegroundColor);
 }
 
 void SMultiLineEditableTextBox::SetTextBoxBackgroundColor(const TAttribute<FSlateColor>& InBackgroundColor)
 {
-	TAttribute<FSlateColor> BackgroundColor = InBackgroundColor.IsSet() ? InBackgroundColor : Style->BackgroundColor;
+	BackgroundColorOverride = InBackgroundColor;
 
-	SetBorderBackgroundColor(BackgroundColor);
+	SetBorderBackgroundColor(BackgroundColorOverride.IsSet() ? BackgroundColorOverride : Style->BackgroundColor);
 }
 
 void SMultiLineEditableTextBox::SetReadOnlyForegroundColor(const TAttribute<FSlateColor>& InReadOnlyForegroundColor)
 {
-	ReadOnlyForegroundColor = InReadOnlyForegroundColor.IsSet() ? InReadOnlyForegroundColor : Style->ReadOnlyForegroundColor;
+	ReadOnlyForegroundColorOverride = InReadOnlyForegroundColor;
+
+	ReadOnlyForegroundColor = ReadOnlyForegroundColorOverride.IsSet() ? ReadOnlyForegroundColorOverride : Style->ReadOnlyForegroundColor;
 }
 
 void SMultiLineEditableTextBox::SetError( const FText& InError )
diff --git a/Engine/Source/Runtime/Slate/Private/Widgets/Layout/SScrollBar.cpp b/Engine/Source/Runtime/Slate/Private/Widgets/Layout/SScrollBar.cpp
--- a/Engine/Source/Runtime/Slate/Private/Widgets/Layout/SScrollBar.cpp
+++ b/Engine/Source/Runtime/Slate/Private/Widgets/Layout/SScrollBar.cpp
@@ -244,38 +244,15 @@ protected:
  */
 void SScrollBar::Construct(const FArguments& InArgs)
 {
-	check(InArgs._Style);
-
 	OnUserScrolled = InArgs._OnUserScrolled;
 	Orientation = InArgs._Orientation;
 	UserVisibility = InArgs._Visibility;
 
-	NormalThumbImage = &InArgs._Style->NormalThumbImage;
-	HoveredThumbImage = &InArgs._Style->HoveredThumbImage;
-	DraggedThumbImage = &InArgs._Style->DraggedThumbImage;
-
-	EHorizontalAlignment HorizontalAlignment;
-	EVerticalAlignment VerticalAlignment;
-	const FSlateBrush* BackgroundBrush;
-	const FSlateBrush* TopBrush;
-	const FSlateBrush* BottomBrush;
+	check(InArgs._Style);
+	SetStyle(InArgs._Style);
 
-	if(Orientation == Orient_Vertical)
-	{
-		HorizontalAlignment = HAlign_Center;
-		VerticalAlignment = VAlign_Fill;
-		BackgroundBrush = &InArgs._Style->VerticalBackgroundImage;
-		TopBrush = &InArgs._Style->VerticalTopSlotImage;
-		BottomBrush = &InArgs._Style->VerticalBottomSlotImage;
-	}
-	else
-	{
-		HorizontalAlignment = HAlign_Fill;
-		VerticalAlignment = VAlign_Center;
-		BackgroundBrush = &InArgs._Style->HorizontalBackgroundImage;
-		TopBrush = &InArgs._Style->HorizontalTopSlotImage;
-		BottomBrush = &InArgs._Style->HorizontalBottomSlotImage;
-	}
+	EHorizontalAlignment HorizontalAlignment = Orientation == Orient_Vertical ? HAlign_Center : HAlign_Fill;
+	EVerticalAlignment VerticalAlignment = Orientation == Orient_Vertical ? VAlign_Fill : VAlign_Center;
 
 	SBorder::Construct( SBorder::FArguments()
 		.BorderImage(FCoreStyle::Get().GetBrush("NoBorder"))
@@ -538,6 +515,28 @@ bool SScrollBar::IsScrolling() const
 	return bDraggingThumb;
 }
 
+void SScrollBar::SetStyle(const FScrollBarStyle* InStyle)
+{
+	check(InStyle);
+
+	NormalThumbImage = &InStyle->NormalThumbImage;
+	HoveredThumbImage = &InStyle->HoveredThumbImage;
+	DraggedThumbImage = &InStyle->DraggedThumbImage;
+
+	if (Orientation == Orient_Vertical)
+	{
+		BackgroundBrush = &InStyle->VerticalBackgroundImage;
+		TopBrush = &InStyle->VerticalTopSlotImage;
+		BottomBrush = &InStyle->VerticalBottomSlotImage;
+	}
+	else
+	{
+		BackgroundBrush = &InStyle->HorizontalBackgroundImage;
+		TopBrush = &InStyle->HorizontalTopSlotImage;
+		BottomBrush = &InStyle->HorizontalBottomSlotImage;
+	}
+}
+
 void SScrollBar::SetThickness(TAttribute<FVector2D> InThickness)
 {
 	ThicknessSpacer->SetSize(InThickness);
diff --git a/Engine/Source/Runtime/Slate/Private/Widgets/Text/SMultiLineEditableText.cpp b/Engine/Source/Runtime/Slate/Private/Widgets/Text/SMultiLineEditableText.cpp
--- a/Engine/Source/Runtime/Slate/Private/Widgets/Text/SMultiLineEditableText.cpp
+++ b/Engine/Source/Runtime/Slate/Private/Widgets/Text/SMultiLineEditableText.cpp
@@ -417,6 +417,13 @@ void SMultiLineEditableText::SetHintText(const TAttribute< FText >& InHintText)
 	}
 }
 
+void SMultiLineEditableText::SetFont(const TAttribute< FSlateFontInfo >& InNewFont)
+{
+	TextStyle.SetFont(InNewFont.Get());
+
+	// @todo: Chris.Wood - this doesn't update the font (needs the TextLayout updating?)
+}
+
 void SMultiLineEditableText::SetTextFromVirtualKeyboard(const FText& InNewText)
 {
 	// Only set the text if the text attribute doesn't have a getter binding (otherwise it would be blown away).
diff --git a/Engine/Source/Runtime/Slate/Public/Widgets/Input/SCheckBox.h b/Engine/Source/Runtime/Slate/Public/Widgets/Input/SCheckBox.h
--- a/Engine/Source/Runtime/Slate/Public/Widgets/Input/SCheckBox.h
+++ b/Engine/Source/Runtime/Slate/Public/Widgets/Input/SCheckBox.h
@@ -218,6 +218,18 @@ public:
 
 protected:
 
+	/** Rebuilds the checkbox based on the current ESlateCheckBoxType */
+	void BuildCheckBox(TSharedRef<SWidget> InContent);
+
+	/** Attribute getter for the foreground color */
+	FSlateColor OnGetForegroundColor() const;
+	/** Attribute getter for the padding */
+	FMargin OnGetPadding() const;
+	/** Attribute getter for the border background color */
+	FSlateColor OnGetBorderBackgroundColor() const;
+	/** Attribute getter for the checkbox type */
+	ESlateCheckBoxType::Type OnGetCheckBoxType() const;
+
 	/**
 	 * Gets the check image to display for the current state of the check box
 	 * @return	The name of the image to display
@@ -268,6 +280,18 @@ protected:
 	/** Image to use when the checkbox is in an ambiguous state and pressed*/
 	const FSlateBrush* UndeterminedPressedImage;
 
+	/** Overrides padding in the widget style, if set */
+	TAttribute<FMargin> PaddingOverride;
+	/** Overrides foreground color in the widget style, if set */
+	TAttribute<FSlateColor> ForegroundColorOverride;
+	/** Overrides border background color in the widget style, if set */
+	TAttribute<FSlateColor> BorderBackgroundColorOverride;
+	/** Overrides checkbox type in the widget style, if set */
+	TOptional<ESlateCheckBoxType::Type> CheckBoxTypeOverride;
+
+	/** Horiz align setting if in togglebox mode */
+	EHorizontalAlignment HorizontalAlignment;
+
 	/** Sets whether a click should be triggered on mouse down, mouse up, or that both a mouse down and up are required. */
 	EButtonClickMethod::Type ClickMethod;
 
diff --git a/Engine/Source/Runtime/Slate/Public/Widgets/Input/SEditableTextBox.h b/Engine/Source/Runtime/Slate/Public/Widgets/Input/SEditableTextBox.h
--- a/Engine/Source/Runtime/Slate/Public/Widgets/Input/SEditableTextBox.h
+++ b/Engine/Source/Runtime/Slate/Public/Widgets/Input/SEditableTextBox.h
@@ -224,9 +224,27 @@ public:
 protected:
 	const FEditableTextBoxStyle* Style;
 
+	/** Box widget that adds padding around the editable text */
+	TSharedPtr< SBox > PaddingBox;
+
 	/** Editable text widget */
 	TSharedPtr< SEditableText > EditableText;
 
+	/** Padding (overrides style) */
+	TAttribute<FMargin> PaddingOverride;
+
+	/** Font (overrides style) */
+	TAttribute<FSlateFontInfo> FontOverride;
+
+	/** Foreground color (overrides style) */
+	TAttribute<FSlateColor> ForegroundColorOverride;
+
+	/** Background color (overrides style) */
+	TAttribute<FSlateColor> BackgroundColorOverride;
+
+	/** Read-only foreground color (overrides style) */
+	TAttribute<FSlateColor> ReadOnlyForegroundColorOverride;
+
 	/** Read-only foreground color */
 	TAttribute<FSlateColor> ReadOnlyForegroundColor;
 
diff --git a/Engine/Source/Runtime/Slate/Public/Widgets/Input/SMultiLineEditableTextBox.h b/Engine/Source/Runtime/Slate/Public/Widgets/Input/SMultiLineEditableTextBox.h
--- a/Engine/Source/Runtime/Slate/Public/Widgets/Input/SMultiLineEditableTextBox.h
+++ b/Engine/Source/Runtime/Slate/Public/Widgets/Input/SMultiLineEditableTextBox.h
@@ -185,6 +185,9 @@ public:
 		return EditableText->GetPlainText();
 	}
 
+	/** See attribute Style */
+	void SetStyle(const FEditableTextBoxStyle* InStyle);
+
 	/**
 	 * Sets the text string currently being edited 
 	 *
@@ -284,12 +287,51 @@ protected:
 	/** Editable text widget */
 	TSharedPtr< SMultiLineEditableText > EditableText;
 
+	/** Padding (overrides style) */
+	TAttribute<FMargin> PaddingOverride;
+
+	/** Horiz scrollbar padding (overrides style) */
+	TAttribute<FMargin> HScrollBarPaddingOverride;
+
+	/** Vert scrollbar padding (overrides style) */
+	TAttribute<FMargin> VScrollBarPaddingOverride;
+
+	/** Font (overrides style) */
+	TAttribute<FSlateFontInfo> FontOverride;
+
+	/** Foreground color (overrides style) */
+	TAttribute<FSlateColor> ForegroundColorOverride;
+
+	/** Background color (overrides style) */
+	TAttribute<FSlateColor> BackgroundColorOverride;
+
+	/** Read-only foreground color (overrides style) */
+	TAttribute<FSlateColor> ReadOnlyForegroundColorOverride;
+
 	/** Read-only foreground color */
 	TAttribute<FSlateColor> ReadOnlyForegroundColor;
 
 	/** Allows for inserting additional widgets that extend the functionality of the text box */
 	TSharedPtr<SHorizontalBox> Box;
 
+	/** Whether we have an externally supplied horizontal scrollbar or one created internally */
+	bool bHasExternalHScrollBar;
+
+	/** Horiz scrollbar */
+	TSharedPtr<SScrollBar> HScrollBar;
+
+	/** Box around the horiz scrollbar used for adding padding */
+	TSharedPtr<SBox> HScrollBarPaddingBox;
+
+	/** Whether we have an externally supplied vertical scrollbar or one created internally */
+	bool bHasExternalVScrollBar;
+
+	/** Vert scrollbar */
+	TSharedPtr<SScrollBar> VScrollBar;
+
+	/** Box around the vert scrollbar used for adding padding */
+	TSharedPtr<SBox> VScrollBarPaddingBox;
+
 	/** SomeWidget reporting */
 	TSharedPtr<class IErrorReportingWidget> ErrorReporting;
 
diff --git a/Engine/Source/Runtime/Slate/Public/Widgets/Layout/SScrollBar.h b/Engine/Source/Runtime/Slate/Public/Widgets/Layout/SScrollBar.h
--- a/Engine/Source/Runtime/Slate/Public/Widgets/Layout/SScrollBar.h
+++ b/Engine/Source/Runtime/Slate/Public/Widgets/Layout/SScrollBar.h
@@ -101,6 +101,9 @@ public:
 	/** @return True if the user is scrolling by dragging the scroll bar thumb. */
 	bool IsScrolling() const;
 
+	/** See argument Style */
+	void SetStyle(const FScrollBarStyle* InStyle);
+
 	/** See UserVisibility attribute */
 	void SetUserVisibility(TAttribute<EVisibility> InUserVisibility) { UserVisibility = InUserVisibility; }
 
@@ -143,4 +146,10 @@ protected:
 	const FSlateBrush* HoveredThumbImage;
 	/** Image to use when the scrollbar thumb is in its dragged state */
 	const FSlateBrush* DraggedThumbImage;
+	/** Background brush */
+	const FSlateBrush* BackgroundBrush;
+	/** Top brush */
+	const FSlateBrush* TopBrush;
+	/** Bottom brush */
+	const FSlateBrush* BottomBrush;
 };
diff --git a/Engine/Source/Runtime/Slate/Public/Widgets/Text/SMultiLineEditableText.h b/Engine/Source/Runtime/Slate/Public/Widgets/Text/SMultiLineEditableText.h
--- a/Engine/Source/Runtime/Slate/Public/Widgets/Text/SMultiLineEditableText.h
+++ b/Engine/Source/Runtime/Slate/Public/Widgets/Text/SMultiLineEditableText.h
@@ -140,6 +140,13 @@ public:
 	 */
 	void SetHintText(const TAttribute< FText >& InHintText);
 
+	/**
+	* Sets the font used to draw the text
+	*
+	* @param  InNewFont	The new font to use
+	*/
+	void SetFont(const TAttribute< FSlateFontInfo >& InNewFont);
+
 	virtual bool GetIsReadOnly() const override;
 	virtual void ClearSelection() override;
 
diff --git a/Engine/Source/Runtime/UMG/Private/Components/CheckBox.cpp b/Engine/Source/Runtime/UMG/Private/Components/CheckBox.cpp
--- a/Engine/Source/Runtime/UMG/Private/Components/CheckBox.cpp
+++ b/Engine/Source/Runtime/UMG/Private/Components/CheckBox.cpp
@@ -16,9 +16,9 @@ UCheckBox::UCheckBox(const FObjectInitializer& ObjectInitializer)
 	CheckedState = ECheckBoxState::Unchecked;
 
 	HorizontalAlignment = SlateDefaults._HAlign;
-	Padding = SlateDefaults._Padding.Get();
+	Padding_DEPRECATED = SlateDefaults._Padding.Get();
 
-	BorderBackgroundColor = FLinearColor::White;
+	BorderBackgroundColor_DEPRECATED = FLinearColor::White;
 }
 
 void UCheckBox::ReleaseSlateResources(bool bReleaseChildren)
@@ -34,8 +34,6 @@ TSharedRef<SWidget> UCheckBox::RebuildWidget()
 		.OnCheckStateChanged( BIND_UOBJECT_DELEGATE(FOnCheckStateChanged, SlateOnCheckStateChangedCallback) )
 		.Style(&WidgetStyle)
 		.HAlign( HorizontalAlignment )
-		.Padding( Padding )
-		.BorderBackgroundColor( BorderBackgroundColor )
 		;
 
 	if ( GetChildrenCount() > 0 )
@@ -50,6 +48,7 @@ void UCheckBox::SynchronizeProperties()
 {
 	Super::SynchronizeProperties();
 
+	MyCheckbox->SetStyle(&WidgetStyle);
 	MyCheckbox->SetIsChecked( OPTIONAL_BINDING(ECheckBoxState, CheckedState) );
 }
 
@@ -199,6 +198,18 @@ void UCheckBox::PostLoad()
 			UndeterminedPressedImage_DEPRECATED = nullptr;
 		}
 	}
+
+	if (GetLinkerUE4Version() < VER_UE4_DEPRECATE_UMG_STYLE_OVERRIDES)
+	{
+		WidgetStyle.Padding = Padding_DEPRECATED;
+		Padding_DEPRECATED = FMargin(0);
+
+		if (BorderBackgroundColor_DEPRECATED != FLinearColor::White)
+		{
+			WidgetStyle.BorderBackgroundColor = BorderBackgroundColor_DEPRECATED;
+			BorderBackgroundColor_DEPRECATED = FLinearColor::White;
+		}
+	}
 }
 
 #if WITH_EDITOR
diff --git a/Engine/Source/Runtime/UMG/Private/Components/EditableText.cpp b/Engine/Source/Runtime/UMG/Private/Components/EditableText.cpp
--- a/Engine/Source/Runtime/UMG/Private/Components/EditableText.cpp
+++ b/Engine/Source/Runtime/UMG/Private/Components/EditableText.cpp
@@ -14,10 +14,10 @@ UEditableText::UEditableText(const FObjectInitializer& ObjectInitializer)
 	SEditableText::FArguments Defaults;
 	WidgetStyle = *Defaults._Style;
 
-	ColorAndOpacity = FLinearColor::Black;
+	ColorAndOpacity_DEPRECATED = FLinearColor::Black;
 
 	static ConstructorHelpers::FObjectFinder<UFont> RobotoFontObj(TEXT("/Engine/EngineFonts/Roboto"));
-	Font = FSlateFontInfo(RobotoFontObj.Object, 12, FName("Bold"));
+	Font_DEPRECATED = FSlateFontInfo(RobotoFontObj.Object, 12, FName("Bold"));
 
 	// Grab other defaults from slate arguments.
 	IsReadOnly = Defaults._IsReadOnly.Get();
@@ -41,7 +41,6 @@ TSharedRef<SWidget> UEditableText::RebuildWidget()
 {
 	MyEditableText = SNew(SEditableText)
 	.Style(&WidgetStyle)
-	.Font(Font)
 	.MinDesiredWidth(MinimumDesiredWidth)
 	.IsCaretMovedWhenGainFocus(IsCaretMovedWhenGainFocus)
 	.SelectAllTextWhenFocused(SelectAllTextWhenFocused)
@@ -66,7 +65,6 @@ void UEditableText::SynchronizeProperties()
 	MyEditableText->SetHintText(HintTextBinding);
 	MyEditableText->SetIsReadOnly(IsReadOnly);
 	MyEditableText->SetIsPassword(IsPassword);
-	MyEditableText->SetColorAndOpacity(ColorAndOpacity);
 
 	// TODO UMG Complete making all properties settable on SEditableText
 }
@@ -162,6 +160,21 @@ void UEditableText::PostLoad()
 			CaretImage_DEPRECATED = nullptr;
 		}
 	}
+
+	if (GetLinkerUE4Version() < VER_UE4_DEPRECATE_UMG_STYLE_OVERRIDES)
+	{
+		if (Font_DEPRECATED.HasValidFont())
+		{
+			WidgetStyle.Font = Font_DEPRECATED;
+			Font_DEPRECATED = FSlateFontInfo();
+		}
+
+		if (ColorAndOpacity_DEPRECATED != FLinearColor::Black)
+		{
+			WidgetStyle.ColorAndOpacity = ColorAndOpacity_DEPRECATED;
+			ColorAndOpacity_DEPRECATED = FLinearColor::Black;
+		}
+	}
 }
 
 #if WITH_EDITOR
diff --git a/Engine/Source/Runtime/UMG/Private/Components/EditableTextBox.cpp b/Engine/Source/Runtime/UMG/Private/Components/EditableTextBox.cpp
--- a/Engine/Source/Runtime/UMG/Private/Components/EditableTextBox.cpp
+++ b/Engine/Source/Runtime/UMG/Private/Components/EditableTextBox.cpp
@@ -11,19 +11,19 @@
 UEditableTextBox::UEditableTextBox(const FObjectInitializer& ObjectInitializer)
 	: Super(ObjectInitializer)
 {
-	ForegroundColor = FLinearColor::Black;
-	BackgroundColor = FLinearColor::White;
-	ReadOnlyForegroundColor = FLinearColor::Black;
+	ForegroundColor_DEPRECATED = FLinearColor::Black;
+	BackgroundColor_DEPRECATED = FLinearColor::White;
+	ReadOnlyForegroundColor_DEPRECATED = FLinearColor::Black;
 
 	static ConstructorHelpers::FObjectFinder<UFont> RobotoFontObj(TEXT("/Engine/EngineFonts/Roboto"));
-	Font = FSlateFontInfo(RobotoFontObj.Object, 12, FName("Bold"));
+	Font_DEPRECATED = FSlateFontInfo(RobotoFontObj.Object, 12, FName("Bold"));
 
 	// Grab other defaults from slate arguments.
 	SEditableTextBox::FArguments Defaults;
 	IsReadOnly = Defaults._IsReadOnly.Get();
 	IsPassword = Defaults._IsPassword.Get();
 	MinimumDesiredWidth = Defaults._MinDesiredWidth.Get();
-	Padding = Defaults._Padding.Get();
+	Padding_DEPRECATED = Defaults._Padding.Get();
 	IsCaretMovedWhenGainFocus = Defaults._IsCaretMovedWhenGainFocus.Get();
 	SelectAllTextWhenFocused = Defaults._SelectAllTextWhenFocused.Get();
 	RevertTextOnEscape = Defaults._RevertTextOnEscape.Get();
@@ -44,12 +44,7 @@ TSharedRef<SWidget> UEditableTextBox::RebuildWidget()
 {
 	MyEditableTextBlock = SNew(SEditableTextBox)
 		.Style(&WidgetStyle)
-		.Font(Font)
-		.ForegroundColor(ForegroundColor)
-		.BackgroundColor(BackgroundColor)
-		.ReadOnlyForegroundColor(ReadOnlyForegroundColor)
 		.MinDesiredWidth(MinimumDesiredWidth)
-		.Padding(Padding)
 		.IsCaretMovedWhenGainFocus(IsCaretMovedWhenGainFocus)
 		.SelectAllTextWhenFocused(SelectAllTextWhenFocused)
 		.RevertTextOnEscape(RevertTextOnEscape)
@@ -69,16 +64,12 @@ void UEditableTextBox::SynchronizeProperties()
 	TAttribute<FText> TextBinding = OPTIONAL_BINDING(FText, Text);
 	TAttribute<FText> HintTextBinding = OPTIONAL_BINDING(FText, HintText);
 
+	MyEditableTextBlock->SetStyle(&WidgetStyle);
 	MyEditableTextBlock->SetText(TextBinding);
 	MyEditableTextBlock->SetHintText(HintTextBinding);
-	MyEditableTextBlock->SetFont(Font);
-	MyEditableTextBlock->SetTextBoxForegroundColor(ForegroundColor);
-	MyEditableTextBlock->SetTextBoxBackgroundColor(BackgroundColor);
-	MyEditableTextBlock->SetReadOnlyForegroundColor(ReadOnlyForegroundColor);
 	MyEditableTextBlock->SetIsReadOnly(IsReadOnly);
 	MyEditableTextBlock->SetIsPassword(IsPassword);
 	MyEditableTextBlock->SetMinimumDesiredWidth(MinimumDesiredWidth);
-	MyEditableTextBlock->SetPadding(Padding);
 	MyEditableTextBlock->SetIsCaretMovedWhenGainFocus(IsCaretMovedWhenGainFocus);
 	MyEditableTextBlock->SetSelectAllTextWhenFocused(SelectAllTextWhenFocused);
 	MyEditableTextBlock->SetRevertTextOnEscape(RevertTextOnEscape);
@@ -150,6 +141,36 @@ void UEditableTextBox::PostLoad()
 			Style_DEPRECATED = nullptr;
 		}
 	}
+
+	if (GetLinkerUE4Version() < VER_UE4_DEPRECATE_UMG_STYLE_OVERRIDES)
+	{
+		if (Font_DEPRECATED.HasValidFont())
+		{
+			WidgetStyle.Font = Font_DEPRECATED;
+			Font_DEPRECATED = FSlateFontInfo();
+		}
+
+		WidgetStyle.Padding = Padding_DEPRECATED;
+		Padding_DEPRECATED = FMargin(0);
+
+		if (ForegroundColor_DEPRECATED != FLinearColor::Black)
+		{
+			WidgetStyle.ForegroundColor = ForegroundColor_DEPRECATED;
+			ForegroundColor_DEPRECATED = FLinearColor::Black;
+		}
+
+		if (BackgroundColor_DEPRECATED != FLinearColor::White)
+		{
+			WidgetStyle.BackgroundColor = BackgroundColor_DEPRECATED;
+			BackgroundColor_DEPRECATED = FLinearColor::White;
+		}
+
+		if (ReadOnlyForegroundColor_DEPRECATED != FLinearColor::Black)
+		{
+			WidgetStyle.ReadOnlyForegroundColor = ReadOnlyForegroundColor_DEPRECATED;
+			ReadOnlyForegroundColor_DEPRECATED = FLinearColor::Black;
+		}
+	}
 }
 
 #if WITH_EDITOR
diff --git a/Engine/Source/Runtime/UMG/Private/Components/MultiLineEditableText.cpp b/Engine/Source/Runtime/UMG/Private/Components/MultiLineEditableText.cpp
--- a/Engine/Source/Runtime/UMG/Private/Components/MultiLineEditableText.cpp
+++ b/Engine/Source/Runtime/UMG/Private/Components/MultiLineEditableText.cpp
@@ -16,7 +16,7 @@ UMultiLineEditableText::UMultiLineEditableText(const FObjectInitializer& ObjectI
 	bAutoWrapText = true;
 	
 	static ConstructorHelpers::FObjectFinder<UFont> RobotoFontObj(TEXT("/Engine/EngineFonts/Roboto"));
-	Font = FSlateFontInfo(RobotoFontObj.Object, 12, FName("Bold"));
+	Font_DEPRECATED = FSlateFontInfo(RobotoFontObj.Object, 12, FName("Bold"));
 }
 
 void UMultiLineEditableText::ReleaseSlateResources(bool bReleaseChildren)
@@ -30,7 +30,6 @@ TSharedRef<SWidget> UMultiLineEditableText::RebuildWidget()
 {
 	MyMultiLineEditableText = SNew(SMultiLineEditableText)
 	.TextStyle(&WidgetStyle)
-	.Font(Font)
 	.Justification(Justification)
 	.WrapTextAt( WrapTextAt )
 	.AutoWrapText( bAutoWrapText )
@@ -93,6 +92,18 @@ void UMultiLineEditableText::HandleOnTextCommitted(const FText& InText, ETextCom
 	OnTextCommitted.Broadcast(InText, CommitMethod);
 }
 
+void UMultiLineEditableText::PostLoad()
+{
+	if (GetLinkerUE4Version() < VER_UE4_DEPRECATE_UMG_STYLE_OVERRIDES)
+	{
+		if (Font_DEPRECATED.HasValidFont())
+		{
+			WidgetStyle.Font = Font_DEPRECATED;
+			Font_DEPRECATED = FSlateFontInfo();
+		}
+	}
+}
+
 #if WITH_EDITOR
 
 const FSlateBrush* UMultiLineEditableText::GetEditorIcon()
diff --git a/Engine/Source/Runtime/UMG/Private/Components/MultiLineEditableTextBox.cpp b/Engine/Source/Runtime/UMG/Private/Components/MultiLineEditableTextBox.cpp
--- a/Engine/Source/Runtime/UMG/Private/Components/MultiLineEditableTextBox.cpp
+++ b/Engine/Source/Runtime/UMG/Private/Components/MultiLineEditableTextBox.cpp
@@ -11,12 +11,12 @@
 UMultiLineEditableTextBox::UMultiLineEditableTextBox(const FObjectInitializer& ObjectInitializer)
 	: Super(ObjectInitializer)
 {
-	ForegroundColor = FLinearColor::Black;
-	BackgroundColor = FLinearColor::White;
-	ReadOnlyForegroundColor = FLinearColor::Black;
+	ForegroundColor_DEPRECATED = FLinearColor::Black;
+	BackgroundColor_DEPRECATED = FLinearColor::White;
+	ReadOnlyForegroundColor_DEPRECATED = FLinearColor::Black;
 
 	static ConstructorHelpers::FObjectFinder<UFont> RobotoFontObj(TEXT("/Engine/EngineFonts/Roboto"));
-	Font = FSlateFontInfo(RobotoFontObj.Object, 12, FName("Bold"));
+	Font_DEPRECATED = FSlateFontInfo(RobotoFontObj.Object, 12, FName("Bold"));
 
 	bAutoWrapText = true;
 
@@ -37,11 +37,7 @@ TSharedRef<SWidget> UMultiLineEditableTextBox::RebuildWidget()
 	MyEditableTextBlock = SNew(SMultiLineEditableTextBox)
 		.Style(&WidgetStyle)
 		.TextStyle(&TextStyle)
-		.Font(Font)
 		.Justification(Justification)
-		.ForegroundColor(ForegroundColor)
-		.BackgroundColor(BackgroundColor)
-		.ReadOnlyForegroundColor(ReadOnlyForegroundColor)
 		.AutoWrapText( bAutoWrapText )
 		.WrapTextAt( WrapTextAt )
 //		.MinDesiredWidth(MinimumDesiredWidth)
@@ -62,17 +58,13 @@ void UMultiLineEditableTextBox::SynchronizeProperties()
 {
 	Super::SynchronizeProperties();
 
+	MyEditableTextBlock->SetStyle(&WidgetStyle);
 	MyEditableTextBlock->SetText(Text);
 //	MyEditableTextBlock->SetHintText(HintText);
 //	MyEditableTextBlock->SetIsReadOnly(IsReadOnly);
 //	MyEditableTextBlock->SetIsPassword(IsPassword);
 //	MyEditableTextBlock->SetColorAndOpacity(ColorAndOpacity);
 
-//	MyEditableTextBlock->SetFont(Font);		@todo: can't set the font on an existing multiline text block
-	MyEditableTextBlock->SetTextBoxForegroundColor(ForegroundColor);
-	MyEditableTextBlock->SetTextBoxBackgroundColor(BackgroundColor);
-	MyEditableTextBlock->SetReadOnlyForegroundColor(ReadOnlyForegroundColor);
-
 	// TODO UMG Complete making all properties settable on SMultiLineEditableTextBox
 }
 
@@ -130,6 +122,33 @@ void UMultiLineEditableTextBox::PostLoad()
 			Style_DEPRECATED = nullptr;
 		}
 	}
+
+	if (GetLinkerUE4Version() < VER_UE4_DEPRECATE_UMG_STYLE_OVERRIDES)
+	{
+		if (Font_DEPRECATED.HasValidFont())
+		{
+			WidgetStyle.Font = Font_DEPRECATED;
+			Font_DEPRECATED = FSlateFontInfo();
+		}
+
+		if (ForegroundColor_DEPRECATED != FLinearColor::Black)
+		{
+			WidgetStyle.ForegroundColor = ForegroundColor_DEPRECATED;
+			ForegroundColor_DEPRECATED = FLinearColor::Black;
+		}
+
+		if (BackgroundColor_DEPRECATED != FLinearColor::White)
+		{
+			WidgetStyle.BackgroundColor = BackgroundColor_DEPRECATED;
+			BackgroundColor_DEPRECATED = FLinearColor::White;
+		}
+
+		if (ReadOnlyForegroundColor_DEPRECATED != FLinearColor::Black)
+		{
+			WidgetStyle.ReadOnlyForegroundColor = ReadOnlyForegroundColor_DEPRECATED;
+			ReadOnlyForegroundColor_DEPRECATED = FLinearColor::Black;
+		}
+	}
 }
 
 #if WITH_EDITOR
diff --git a/Engine/Source/Runtime/UMG/Public/Components/CheckBox.h b/Engine/Source/Runtime/UMG/Public/Components/CheckBox.h
--- a/Engine/Source/Runtime/UMG/Public/Components/CheckBox.h
+++ b/Engine/Source/Runtime/UMG/Public/Components/CheckBox.h
@@ -80,12 +80,12 @@ public:
 	TEnumAsByte<EHorizontalAlignment> HorizontalAlignment;
 
 	/** Spacing between the check box image and its content */
-	UPROPERTY(EditAnywhere, BlueprintReadOnly, Category=Appearance)
-	FMargin Padding;
+	UPROPERTY()
+	FMargin Padding_DEPRECATED;
 
 	/** The color of the background border */
-	UPROPERTY(EditAnywhere, BlueprintReadOnly, Category=Appearance)
-	FSlateColor BorderBackgroundColor;
+	UPROPERTY()
+	FSlateColor BorderBackgroundColor_DEPRECATED;
 
 public:
 
diff --git a/Engine/Source/Runtime/UMG/Public/Components/EditableText.h b/Engine/Source/Runtime/UMG/Public/Components/EditableText.h
--- a/Engine/Source/Runtime/UMG/Public/Components/EditableText.h
+++ b/Engine/Source/Runtime/UMG/Public/Components/EditableText.h
@@ -59,12 +59,12 @@ public:
 	USlateBrushAsset* CaretImage_DEPRECATED;
 
 	/** Font color and opacity (overrides Style) */
-	UPROPERTY(EditAnywhere, Category=Appearance)
-	FSlateFontInfo Font;
+	UPROPERTY()
+	FSlateFontInfo Font_DEPRECATED;
 
 	/** Text color and opacity (overrides Style) */
-	UPROPERTY(EditAnywhere, Category=Appearance)
-	FSlateColor ColorAndOpacity;
+	UPROPERTY()
+	FSlateColor ColorAndOpacity_DEPRECATED;
 
 	/** Sets whether this text box can actually be modified interactively by the user */
 	UPROPERTY(EditAnywhere, Category=Appearance)
diff --git a/Engine/Source/Runtime/UMG/Public/Components/EditableTextBox.h b/Engine/Source/Runtime/UMG/Public/Components/EditableTextBox.h
--- a/Engine/Source/Runtime/UMG/Public/Components/EditableTextBox.h
+++ b/Engine/Source/Runtime/UMG/Public/Components/EditableTextBox.h
@@ -48,20 +48,20 @@ public:
 	FGetText HintTextDelegate;
 
 	/** Font color and opacity (overrides Style) */
-	UPROPERTY(EditAnywhere, BlueprintReadOnly, Category=Appearance)
-	FSlateFontInfo Font;
+	UPROPERTY()
+	FSlateFontInfo Font_DEPRECATED;
 
 	/** Text color and opacity (overrides Style) */
-	UPROPERTY(EditAnywhere, BlueprintReadOnly, Category=Appearance)
-	FLinearColor ForegroundColor;
+	UPROPERTY()
+	FLinearColor ForegroundColor_DEPRECATED;
 
 	/** The color of the background/border around the editable text (overrides Style) */
-	UPROPERTY(EditAnywhere, BlueprintReadOnly, Category=Appearance)
-	FLinearColor BackgroundColor;
+	UPROPERTY()
+	FLinearColor BackgroundColor_DEPRECATED;
 
 	/** Text color and opacity when read-only (overrides Style) */
-	UPROPERTY(EditAnywhere, BlueprintReadOnly, Category=Appearance)
-	FLinearColor ReadOnlyForegroundColor;
+	UPROPERTY()
+	FLinearColor ReadOnlyForegroundColor_DEPRECATED;
 
 	/** Sets whether this text box can actually be modified interactively by the user */
 	UPROPERTY(EditAnywhere, BlueprintReadOnly, Category=Appearance)
@@ -76,8 +76,8 @@ public:
 	float MinimumDesiredWidth;
 
 	/** Padding between the box/border and the text widget inside (overrides Style) */
-	UPROPERTY(EditAnywhere, BlueprintReadOnly, Category=Appearance)
-	FMargin Padding;
+	UPROPERTY()
+	FMargin Padding_DEPRECATED;
 
 	/** Workaround as we lose focus when the auto completion closes. */
 	UPROPERTY(EditAnywhere, BlueprintReadOnly, Category=Behavior, AdvancedDisplay)
diff --git a/Engine/Source/Runtime/UMG/Public/Components/MultiLineEditableText.h b/Engine/Source/Runtime/UMG/Public/Components/MultiLineEditableText.h
--- a/Engine/Source/Runtime/UMG/Public/Components/MultiLineEditableText.h
+++ b/Engine/Source/Runtime/UMG/Public/Components/MultiLineEditableText.h
@@ -41,8 +41,8 @@ public:
 	float WrapTextAt;
 
 	/** Font color and opacity (overrides Style) */
-	UPROPERTY(EditAnywhere, BlueprintReadWrite, Category=Appearance)
-	FSlateFontInfo Font;
+	UPROPERTY()
+	FSlateFontInfo Font_DEPRECATED;
 
 	/** Called whenever the text is changed interactively by the user */
 	UPROPERTY(BlueprintAssignable, Category="Widget Event", meta=(DisplayName="OnTextChanged (Multi-Line Editable Text)"))
@@ -70,6 +70,10 @@ public:
 	virtual void ReleaseSlateResources(bool bReleaseChildren) override;
 	// End of UVisual interface
 
+	// Begin UObject interface
+	virtual void PostLoad() override;
+	// End of UObject interface
+
 #if WITH_EDITOR
 	virtual const FSlateBrush* GetEditorIcon() override;
 	virtual const FText GetPaletteCategory() override;
diff --git a/Engine/Source/Runtime/UMG/Public/Components/MultiLineEditableTextBox.h b/Engine/Source/Runtime/UMG/Public/Components/MultiLineEditableTextBox.h
--- a/Engine/Source/Runtime/UMG/Public/Components/MultiLineEditableTextBox.h
+++ b/Engine/Source/Runtime/UMG/Public/Components/MultiLineEditableTextBox.h
@@ -49,20 +49,20 @@ public:
 	float WrapTextAt;
 
 	/** Font color and opacity (overrides Style) */
-	UPROPERTY(EditAnywhere, BlueprintReadOnly, Category=Appearance)
-	FSlateFontInfo Font;
+	UPROPERTY()
+	FSlateFontInfo Font_DEPRECATED;
 
 	/** Text color and opacity (overrides Style) */
-	UPROPERTY(EditAnywhere, BlueprintReadOnly, Category=Appearance)
-	FLinearColor ForegroundColor;
+	UPROPERTY()
+	FLinearColor ForegroundColor_DEPRECATED;
 
 	/** The color of the background/border around the editable text (overrides Style) */
-	UPROPERTY(EditAnywhere, BlueprintReadOnly, Category=Appearance)
-	FLinearColor BackgroundColor;
+	UPROPERTY()
+	FLinearColor BackgroundColor_DEPRECATED;
 
 	/** Text color and opacity when read-only (overrides Style) */
-	UPROPERTY(EditAnywhere, BlueprintReadOnly, Category=Appearance)
-	FLinearColor ReadOnlyForegroundColor;
+	UPROPERTY()
+	FLinearColor ReadOnlyForegroundColor_DEPRECATED;
 
 	/** Called whenever the text is changed interactively by the user */
 	UPROPERTY(BlueprintAssignable, Category="Widget Event", meta=(DisplayName="OnTextChanged (Multi-Line Text Box)"))
