Fixed STableRow::SetContent from removing the extra UI added by STableRow::ConstructChildren

You must now use STableRow::SetRowContent if you want to replace the entire row.

diff --git a/Engine/Source/Runtime/Slate/Public/Widgets/Views/STableRow.h b/Engine/Source/Runtime/Slate/Public/Widgets/Views/STableRow.h
--- a/Engine/Source/Runtime/Slate/Public/Widgets/Views/STableRow.h
+++ b/Engine/Source/Runtime/Slate/Public/Widgets/Views/STableRow.h
@@ -142,21 +142,27 @@ public:
 	virtual void ConstructChildren( ETableViewMode::Type InOwnerTableMode, const TAttribute<FMargin>& InPadding, const TSharedRef<SWidget>& InContent )
 	{
 		this->Content = InContent;
+		InnerContentSlot = nullptr;
 
 		if ( InOwnerTableMode == ETableViewMode::List || InOwnerTableMode == ETableViewMode::Tile )
 		{
 			// -- Row is in a ListView or the user --
+			FSimpleSlot* InnerContentSlotNativePtr = nullptr;
 
 			// We just need to hold on to this row's content.
 			this->ChildSlot
+			.Expose( InnerContentSlotNativePtr )
 			.Padding( InPadding )
 			[
 				InContent
 			];
+
+			InnerContentSlot = InnerContentSlotNativePtr;
 		}
 		else
 		{
 			// -- Row is for TreeView --
+			SHorizontalBox::FSlot* InnerContentSlotNativePtr = nullptr;
 
 			// Rows in a TreeView need an expander button and some indentation
 			this->ChildSlot
@@ -173,11 +179,14 @@ public:
 
 				+ SHorizontalBox::Slot()
 				.FillWidth(1)
+				.Expose( InnerContentSlotNativePtr )
 				.Padding( InPadding )
 				[
 					InContent
 				]
 			];
+
+			InnerContentSlot = InnerContentSlotNativePtr;
 		}
 	}
 
@@ -706,12 +715,30 @@ public:
 		return SharedThis(this);
 	}
 
-	virtual void SetContent(TSharedRef< SWidget > InContent) override
+	/** Set the entire content of this row, replacing any extra UI (such as the expander arrows for tree views) that was added by ConstructChildren */
+	void SetRowContent(TSharedRef< SWidget > InContent)
 	{
 		this->Content = InContent;
+		InnerContentSlot = nullptr;
 		SBorder::SetContent(InContent);
 	}
 
+	/** Set the inner content of this row, preserving any extra UI (such as the expander arrows for tree views) that was added by ConstructChildren */
+	virtual void SetContent(TSharedRef< SWidget > InContent) override
+	{
+		this->Content = InContent;
+
+		if (InnerContentSlot)
+		{
+			InnerContentSlot->AttachWidget(InContent);
+		}
+		else
+		{
+			SBorder::SetContent(InContent);
+		}
+	}
+
+	/** Get the inner content of this row */
 	virtual TSharedPtr<SWidget> GetContent() override
 	{
 		if ( this->Content.IsValid() )
@@ -932,6 +959,9 @@ protected:
 	/** Delegate triggered when a user's drag is dropped in the bounds of this list item */
 	FOnTableRowDrop OnDrop_Handler;
 
+	/** The slot that contains the inner content for this row. If this is set, SetContent populates this slot with the new content rather than replace the content wholesale */
+	FSlotBase* InnerContentSlot;
+
 	/** The widget in the content slot for this row */
 	TWeakPtr<SWidget> Content;
 
