Rolled back a change from CL# 2570717.

All menu builder icons are once again a fixed 16x16 size.

diff --git a/Engine/Source/Runtime/Slate/Private/Framework/MultiBox/SMenuEntryBlock.cpp b/Engine/Source/Runtime/Slate/Private/Framework/MultiBox/SMenuEntryBlock.cpp
--- a/Engine/Source/Runtime/Slate/Private/Framework/MultiBox/SMenuEntryBlock.cpp
+++ b/Engine/Source/Runtime/Slate/Private/Framework/MultiBox/SMenuEntryBlock.cpp
@@ -390,8 +390,6 @@ TSharedRef< SWidget > SMenuEntryBlock::BuildMenuEntryWidget( const FMenuEntryBui
 				SNew( SBox )
 				.WidthOverride( MultiBoxConstants::MenuIconSize )
 				.HeightOverride( MultiBoxConstants::MenuIconSize )
-				.HAlign(HAlign_Center)
-				.VAlign(VAlign_Center)
 				[
 					IconWidget
 				]
