Slate: Breadcrumb not getting updated when removed. Remove slot wasn't picking up the button and delimiter because they were wrapped in SVerticalBox.

diff --git a/Engine/Source/Runtime/Slate/Public/Widgets/Navigation/SBreadcrumbTrail.h b/Engine/Source/Runtime/Slate/Public/Widgets/Navigation/SBreadcrumbTrail.h
--- a/Engine/Source/Runtime/Slate/Public/Widgets/Navigation/SBreadcrumbTrail.h
+++ b/Engine/Source/Runtime/Slate/Public/Widgets/Navigation/SBreadcrumbTrail.h
@@ -16,12 +16,16 @@ private:
 		int32 CrumbID;
 		TSharedRef<SButton> Button;
 		TSharedRef<SMenuAnchor> Delimiter;
+		TSharedRef<SVerticalBox> ButtonBox;
+		TSharedRef<SVerticalBox> DelimiterBox;
 		ItemType CrumbData;
 
-		FCrumbItem(int32 InCrumbID, TSharedRef<SButton> InButton, TSharedRef<SMenuAnchor> InDelimiter, ItemType InCrumbData)
+		FCrumbItem(int32 InCrumbID, TSharedRef<SButton> InButton, TSharedRef<SMenuAnchor> InDelimiter, TSharedRef<SVerticalBox> InButtonBox, TSharedRef<SVerticalBox> InDelimiterBox, ItemType InCrumbData)
 			: CrumbID(InCrumbID)
 			, Button(InButton)
 			, Delimiter(InDelimiter)
+			, ButtonBox(InButtonBox)
+			, DelimiterBox(InDelimiterBox)
 			, CrumbData(InCrumbData)
 		{}
 	};
@@ -119,10 +123,13 @@ public:
 		TSharedPtr<SButton> NewButton;
 		TSharedPtr<SMenuAnchor> NewDelimiter;
 
+		TSharedPtr<SVerticalBox> NewButtonBox;
+		TSharedPtr<SVerticalBox> NewDelimiterBox;
+
 		// Add the crumb button
 		CrumbBox->AddSlot()
 		[
-			SNew(SVerticalBox)
+			SAssignNew(NewButtonBox, SVerticalBox)
 
 			+SVerticalBox::Slot()
 			.FillHeight( 1.0f )
@@ -167,7 +174,7 @@ public:
 
 		CrumbBox->AddSlot()
 		[
-			SNew(SVerticalBox)
+			SAssignNew(NewDelimiterBox, SVerticalBox)
 
 			+SVerticalBox::Slot()
 			.FillHeight(1.0f)
@@ -181,7 +188,7 @@ public:
 		];
 
 		// Push the crumb data
-		new (CrumbList) FCrumbItem(NextValidCrumbID, NewButton.ToSharedRef(), NewDelimiter.ToSharedRef(), NewCrumbData);
+		new (CrumbList) FCrumbItem(NextValidCrumbID, NewButton.ToSharedRef(), NewDelimiter.ToSharedRef(), NewButtonBox.ToSharedRef(), NewDelimiterBox.ToSharedRef(), NewCrumbData);
 
 		// Increment the crumb ID for the next crumb
 		NextValidCrumbID = (NextValidCrumbID + 1) % (INT_MAX - 1);
@@ -200,8 +207,8 @@ public:
 
 		// Remove from the crumb list and box
 		FCrumbItem LastCrumbItem = CrumbList.Pop();
-		CrumbBox->RemoveSlot(LastCrumbItem.Button);
-		CrumbBox->RemoveSlot(LastCrumbItem.Delimiter);
+		CrumbBox->RemoveSlot(LastCrumbItem.ButtonBox);
+		CrumbBox->RemoveSlot(LastCrumbItem.DelimiterBox);
 
 		// Trigger event
 		OnCrumbPopped.ExecuteIfBound(LastCrumbItem.CrumbData);
