Follow up to CL# 2589989 to allow you to override SetRowContent in derived widgets

Some BP menu code was doing this for SetContent, but it now needs to override SetRowContent instead.

diff --git a/Engine/Source/Runtime/Slate/Public/Widgets/Views/STableRow.h b/Engine/Source/Runtime/Slate/Public/Widgets/Views/STableRow.h
--- a/Engine/Source/Runtime/Slate/Public/Widgets/Views/STableRow.h
+++ b/Engine/Source/Runtime/Slate/Public/Widgets/Views/STableRow.h
@@ -716,7 +716,7 @@ public:
 	}
 
 	/** Set the entire content of this row, replacing any extra UI (such as the expander arrows for tree views) that was added by ConstructChildren */
-	void SetRowContent(TSharedRef< SWidget > InContent)
+	virtual void SetRowContent(TSharedRef< SWidget > InContent)
 	{
 		this->Content = InContent;
 		InnerContentSlot = nullptr;
