Virtual joysticks now deal with screen resolution changes properly
#ue4
#ios
#android

diff --git a/Engine/Source/Runtime/Slate/Private/Widgets/Input/SVirtualJoystick.cpp b/Engine/Source/Runtime/Slate/Private/Widgets/Input/SVirtualJoystick.cpp
--- a/Engine/Source/Runtime/Slate/Private/Widgets/Input/SVirtualJoystick.cpp
+++ b/Engine/Source/Runtime/Slate/Private/Widgets/Input/SVirtualJoystick.cpp
@@ -23,7 +23,7 @@ FORCEINLINE float SVirtualJoystick::GetBaseOpacity()
 void SVirtualJoystick::FControlInfo::Reset()
 {
 	// snap the visual center back to normal (for controls that have a center on touch)
-	VisualCenter = Center;
+	VisualCenter = CorrectedCenter;
 }
 
 void SVirtualJoystick::Construct( const FArguments& InArgs )
@@ -40,6 +40,18 @@ void SVirtualJoystick::Construct( const FArguments& InArgs )
 	ActivationDelay = 0.f;
 	CurrentOpacity = InactiveOpacity;
 	StartupDelay = 0.f;
+
+	// listen for displaymetrics changes to reposition controls
+	FSlateApplication::Get().GetPlatformApplication()->OnDisplayMetricsChanged().AddSP(this, &SVirtualJoystick::HandleDisplayMetricsChanged);
+}
+
+void SVirtualJoystick::HandleDisplayMetricsChanged(const FDisplayMetrics& NewDisplayMetric)
+{
+	// Mark all controls to be repositioned on next tick
+	for (int32 ControlIndex = 0; ControlIndex < Controls.Num(); ControlIndex++)
+	{
+		Controls[ControlIndex].bHasBeenPositioned = false;
+	}
 }
 
 void SVirtualJoystick::SetGlobalParameters(float InActiveOpacity, float InInactiveOpacity, float InTimeUntilDeactive, float InTimeUntilReset, float InActivationDelay, bool InbPreventReCenter, float InStartupDelay)
@@ -93,13 +105,13 @@ static int32 ResolveRelativePosition(float Position, float RelativeTo, float Sca
 
 }
 
-static bool PositionIsInside(const SVirtualJoystick::FControlInfo& Control, const FVector2D& Position, const FVector2D& BoxSize)
+static bool PositionIsInside(const FVector2D& Center, const FVector2D& Position, const FVector2D& BoxSize)
 {
 	return
-		Position.X >= Control.Center.X - BoxSize.X * 0.5f &&
-		Position.X <= Control.Center.X + BoxSize.X * 0.5f &&
-		Position.Y >= Control.Center.Y - BoxSize.Y * 0.5f &&
-		Position.Y <= Control.Center.Y + BoxSize.Y * 0.5f;
+		Position.X >= Center.X - BoxSize.X * 0.5f &&
+		Position.X <= Center.X + BoxSize.X * 0.5f &&
+		Position.Y >= Center.Y - BoxSize.Y * 0.5f &&
+		Position.Y <= Center.Y + BoxSize.Y * 0.5f;
 }
 
 int32 SVirtualJoystick::OnPaint( const FPaintArgs& Args, const FGeometry& AllottedGeometry, const FSlateRect& MyClippingRect, FSlateWindowElementList& OutDrawElements, int32 LayerId, const FWidgetStyle& InWidgetStyle, bool bParentEnabled ) const
@@ -121,8 +133,8 @@ int32 SVirtualJoystick::OnPaint( const FPaintArgs& Args, const FGeometry& Allott
 					OutDrawElements,
 					RetLayerId++,
 					AllottedGeometry.ToPaintGeometry(
-					Control.VisualCenter - FVector2D(Control.VisualSize.X * 0.5f, Control.VisualSize.Y * 0.5f),
-					Control.VisualSize),
+					Control.VisualCenter - FVector2D(Control.CorrectedVisualSize.X * 0.5f, Control.CorrectedVisualSize.Y * 0.5f),
+					Control.CorrectedVisualSize),
 					Control.Image2.Get(),
 					MyClippingRect,
 					ESlateDrawEffect::None,
@@ -136,8 +148,8 @@ int32 SVirtualJoystick::OnPaint( const FPaintArgs& Args, const FGeometry& Allott
 					OutDrawElements,
 					RetLayerId++,
 					AllottedGeometry.ToPaintGeometry(
-					Control.VisualCenter + Control.ThumbPosition - FVector2D(Control.ThumbSize.X * 0.5f, Control.ThumbSize.Y * 0.5f),
-					Control.ThumbSize),
+					Control.VisualCenter + Control.ThumbPosition - FVector2D(Control.CorrectedThumbSize.X * 0.5f, Control.CorrectedThumbSize.Y * 0.5f),
+					Control.CorrectedThumbSize),
 					Control.Image1.Get(),
 					MyClippingRect,
 					ESlateDrawEffect::None,
@@ -173,10 +185,10 @@ FReply SVirtualJoystick::OnTouchStarted(const FGeometry& MyGeometry, const FPoin
 		// skip controls already in use
 		if (Control.CapturedPointerIndex == -1)
 		{
-			if (PositionIsInside(Control, LocalCoord, Control.InteractionSize))
+			if (PositionIsInside(Control.CorrectedCenter, LocalCoord, Control.CorrectedInteractionSize))
 			{
 				// Align Joystick inside of Screen
-				AlignBoxIntoScreen(LocalCoord, Control.VisualSize, MyGeometry.Size);
+				AlignBoxIntoScreen(LocalCoord, Control.CorrectedVisualSize, MyGeometry.Size);
 
 				Control.CapturedPointerIndex = Event.GetPointerIndex();
 
@@ -286,8 +298,8 @@ bool SVirtualJoystick::HandleTouch(int32 ControlIndex, const FVector2D& LocalCoo
 		// length along line to ellipse: L = 1.0 / sqrt(((sin(T)/Rx)^2 + (cos(T)/Ry)^2))
 		float CosAngle = FMath::Cos(Angle);
 		float SinAngle = FMath::Sin(Angle);
-		float XTerm = CosAngle / (Control.VisualSize.X * 0.5f);
-		float YTerm = SinAngle / (Control.VisualSize.Y * 0.5f);
+		float XTerm = CosAngle / (Control.CorrectedVisualSize.X * 0.5f);
+		float YTerm = SinAngle / (Control.CorrectedVisualSize.Y * 0.5f);
 		float DistanceToEdge = FMath::InvSqrt(XTerm * XTerm + YTerm * YTerm);
 
 		// only clamp 
@@ -302,7 +314,7 @@ bool SVirtualJoystick::HandleTouch(int32 ControlIndex, const FVector2D& LocalCoo
 	}
 
 	FVector2D AbsoluteThumbPos = Control.ThumbPosition + Controls[ControlIndex].VisualCenter;
-	AlignBoxIntoScreen(AbsoluteThumbPos, Control.ThumbSize, ScreenSize);
+	AlignBoxIntoScreen(AbsoluteThumbPos, Control.CorrectedThumbSize, ScreenSize);
 	Control.ThumbPosition = AbsoluteThumbPos - Controls[ControlIndex].VisualCenter;
 
 	return true;
@@ -352,12 +364,12 @@ void SVirtualJoystick::Tick( const FGeometry& AllottedGeometry, const double InC
 			float ScaleFactor = GetScaleFactor(AllottedGeometry);
 
 			// update all the sizes
-			Control.Center = FVector2D(ResolveRelativePosition(Control.Center.X, AllottedGeometry.Size.X, ScaleFactor), ResolveRelativePosition(Control.Center.Y, AllottedGeometry.Size.Y, ScaleFactor));
-			Control.VisualCenter = Control.Center;
-			Control.VisualSize = FVector2D(ResolveRelativePosition(Control.VisualSize.X, AllottedGeometry.Size.X, ScaleFactor), ResolveRelativePosition(Control.VisualSize.Y, AllottedGeometry.Size.Y, ScaleFactor));
-			Control.InteractionSize = FVector2D(ResolveRelativePosition(Control.InteractionSize.X, AllottedGeometry.Size.X, ScaleFactor), ResolveRelativePosition(Control.InteractionSize.Y, AllottedGeometry.Size.Y, ScaleFactor));
-			Control.ThumbSize = FVector2D(ResolveRelativePosition(Control.ThumbSize.X, AllottedGeometry.Size.X, ScaleFactor), ResolveRelativePosition(Control.ThumbSize.Y, AllottedGeometry.Size.Y, ScaleFactor));
-//			Control.InputScale = Control.InputScale * ScaleFactor;
+			Control.CorrectedCenter = FVector2D(ResolveRelativePosition(Control.Center.X, AllottedGeometry.Size.X, ScaleFactor), ResolveRelativePosition(Control.Center.Y, AllottedGeometry.Size.Y, ScaleFactor));
+			Control.VisualCenter = Control.CorrectedCenter;
+			Control.CorrectedVisualSize = FVector2D(ResolveRelativePosition(Control.VisualSize.X, AllottedGeometry.Size.X, ScaleFactor), ResolveRelativePosition(Control.VisualSize.Y, AllottedGeometry.Size.Y, ScaleFactor));
+			Control.CorrectedInteractionSize = FVector2D(ResolveRelativePosition(Control.InteractionSize.X, AllottedGeometry.Size.X, ScaleFactor), ResolveRelativePosition(Control.InteractionSize.Y, AllottedGeometry.Size.Y, ScaleFactor));
+			Control.CorrectedThumbSize = FVector2D(ResolveRelativePosition(Control.ThumbSize.X, AllottedGeometry.Size.X, ScaleFactor), ResolveRelativePosition(Control.ThumbSize.Y, AllottedGeometry.Size.Y, ScaleFactor));
+			Control.CorrectedInputScale = Control.InputScale; // *ScaleFactor;
 			Control.bHasBeenPositioned = true;
 		}
 
@@ -367,8 +379,8 @@ void SVirtualJoystick::Tick( const FGeometry& AllottedGeometry, const double InC
 
 			// now pass the fake joystick events to the game
 			// Assume that joystick size is all equal
-			float JoysticInputSize = Control.ThumbPosition.Size() * 2.f / Control.VisualSize.X;
-			FVector2D NormalizedOffset = Control.ThumbPosition.GetSafeNormal() * Control.InputScale * JoysticInputSize;
+			float JoystickInputSize = Control.ThumbPosition.Size() * 2.f / Control.CorrectedVisualSize.X;
+			FVector2D NormalizedOffset = Control.ThumbPosition.GetSafeNormal() * Control.CorrectedInputScale * JoystickInputSize;
 			EControllerButtons::Type XAxis = ControlIndex == 0 ? EControllerButtons::LeftAnalogX : EControllerButtons::RightAnalogX;
 			EControllerButtons::Type YAxis = ControlIndex == 0 ? EControllerButtons::LeftAnalogY : EControllerButtons::RightAnalogY;
 
diff --git a/Engine/Source/Runtime/Slate/Public/Widgets/Input/SVirtualJoystick.h b/Engine/Source/Runtime/Slate/Public/Widgets/Input/SVirtualJoystick.h
--- a/Engine/Source/Runtime/Slate/Public/Widgets/Input/SVirtualJoystick.h
+++ b/Engine/Source/Runtime/Slate/Public/Widgets/Input/SVirtualJoystick.h
@@ -49,6 +49,9 @@ public:
 		/** The secondary input (for sticks, this is the vertical/Y input, unused for buttons) */ 
 		FKey AltInputKey;
 
+		/** Positioned center in viewport */
+		FVector2D PositionedCenter;
+
 	private:
 		friend SVirtualJoystick;
 
@@ -65,6 +68,21 @@ public:
 		/** For recentered joysticks, this is the re-center location */
 		FVector2D VisualCenter;
 
+		/** The corrected actual center of the control */
+		FVector2D CorrectedCenter;
+
+		/** The corrected size of a joystick that can be re-centered within InteractionSize area */
+		FVector2D CorrectedVisualSize;
+
+		/** The corrected size of the thumb that can be re-centered within InteractionSize area */
+		FVector2D CorrectedThumbSize;
+
+		/** The corrected size of a the interactable area around Center */
+		FVector2D CorrectedInteractionSize;
+
+		/** The corrected scale for control input */
+		FVector2D CorrectedInputScale;
+
 		/** Which pointer index is interacting with this control right now, or -1 if not interacting */
 		int32 CapturedPointerIndex;
 
@@ -135,6 +153,9 @@ public:
 	virtual bool SupportsKeyboardFocus() const override;
 
 private:
+	/** Callback for handling display metrics changes. */
+	void HandleDisplayMetricsChanged(const FDisplayMetrics& NewDisplayMetric);
+
 	void AlignBoxIntoScreen(FVector2D& Position, const FVector2D& Size, const FVector2D& ScreenSize);
 
 	/**
