Fix completely transparent drop shadows still rendering

diff --git a/Engine/Source/Runtime/Slate/Private/Framework/Text/SlateTextRun.cpp b/Engine/Source/Runtime/Slate/Private/Framework/Text/SlateTextRun.cpp
--- a/Engine/Source/Runtime/Slate/Private/Framework/Text/SlateTextRun.cpp
+++ b/Engine/Source/Runtime/Slate/Private/Framework/Text/SlateTextRun.cpp
@@ -73,7 +73,7 @@ int32 FSlateTextRun::OnPaint( const FPaintArgs& Args, const FTextLayout::FLineVi
 	const FSlateRect ClippingRect = AllottedGeometry.GetClippingRect().IntersectionWith(MyClippingRect);
 	const ESlateDrawEffect::Type DrawEffects = bParentEnabled ? ESlateDrawEffect::None : ESlateDrawEffect::DisabledEffect;
 
-	const bool ShouldDropShadow = Style.ShadowOffset.Size() > 0;
+	const bool ShouldDropShadow = Style.ShadowOffset.Size() > 0 && Style.ShadowColorAndOpacity.A > 0;
 	const FVector2D BlockLocationOffset = Block->GetLocationOffset();
 	const FTextRange BlockRange = Block->GetTextRange();
 	
