Slate/UMG - The SButton no longer attempts to handle MouseUp, if the button wasn't pressed.  Previously pressing outside the button, and then releasing the mouse inside the button would cause the button to attempt to handle it, even though the down never originated on the button, which is never a case we care about.  Now we first check if the bIsPressed flag is true before attempting to handle the mouse up.

diff --git a/Engine/Source/Runtime/Slate/Private/Widgets/Input/SButton.cpp b/Engine/Source/Runtime/Slate/Private/Widgets/Input/SButton.cpp
--- a/Engine/Source/Runtime/Slate/Private/Widgets/Input/SButton.cpp
+++ b/Engine/Source/Runtime/Slate/Private/Widgets/Input/SButton.cpp
@@ -144,7 +144,7 @@ FReply SButton::OnKeyUp(const FGeometry& MyGeometry, const FKeyEvent& InKeyEvent
 
 	if (IsEnabled() && (InKeyEvent.GetKey() == EKeys::Enter || InKeyEvent.GetKey() == EKeys::SpaceBar || InKeyEvent.GetKey() == EKeys::Gamepad_FaceButton_Bottom))
 	{
-		bool bWasPressed = bIsPressed;
+		const bool bWasPressed = bIsPressed;
 
 		Release();
 
@@ -207,7 +207,7 @@ FReply SButton::OnMouseButtonDoubleClick( const FGeometry& InMyGeometry, const F
 FReply SButton::OnMouseButtonUp( const FGeometry& MyGeometry, const FPointerEvent& MouseEvent )
 {
 	FReply Reply = FReply::Unhandled();
-	if (IsEnabled() && (MouseEvent.GetEffectingButton() == EKeys::LeftMouseButton || MouseEvent.IsTouchEvent()))
+	if ( bIsPressed && IsEnabled() && ( MouseEvent.GetEffectingButton() == EKeys::LeftMouseButton || MouseEvent.IsTouchEvent() ) )
 	{
 		Release();
 
@@ -222,7 +222,7 @@ FReply SButton::OnMouseButtonUp( const FGeometry& MyGeometry, const FPointerEven
 			const bool bIsUnderMouse = MyGeometry.IsUnderLocation(MouseEvent.GetScreenSpacePosition());
 			if( bIsUnderMouse )
 			{
-				// If we asked for a precide tap, all we need is for the user to have not moved their pointer very far.
+				// If we asked for a precise tap, all we need is for the user to have not moved their pointer very far.
 				const bool bTriggerForTouchEvent = IsPreciseTapOrClick(MouseEvent);
 
 				// If we were asked to allow the button to be clicked on mouse up, regardless of whether the user
@@ -238,14 +238,14 @@ FReply SButton::OnMouseButtonUp( const FGeometry& MyGeometry, const FPointerEven
 		
 		//If the user of the button didn't handle this click, then the button's
 		//default behavior handles it.
-		if(Reply.IsEventHandled() == false)
+		if ( Reply.IsEventHandled() == false )
 		{
 			Reply = FReply::Handled();
 		}
 
 		//If the user hasn't requested a new mouse captor, then the default
 		//behavior of the button is to release mouse capture.
-		if(Reply.GetMouseCaptor().IsValid() == false)
+		if ( Reply.GetMouseCaptor().IsValid() == false )
 		{
 			Reply.ReleaseMouseCapture();
 		}
