UMG/Slate - The STooltipPresenter now calculates a post transformed, rounded layout position to ensure that tooltips do not start on a half pixel.  This reduces a lot of the inner jittering of a tooltip as you move your mouse across the surface of a widget.

diff --git a/Engine/Source/Runtime/Slate/Private/Widgets/LayerManager/STooltipPresenter.cpp b/Engine/Source/Runtime/Slate/Private/Widgets/LayerManager/STooltipPresenter.cpp
--- a/Engine/Source/Runtime/Slate/Private/Widgets/LayerManager/STooltipPresenter.cpp
+++ b/Engine/Source/Runtime/Slate/Private/Widgets/LayerManager/STooltipPresenter.cpp
@@ -34,11 +34,15 @@ void STooltipPresenter::OnArrangeChildren(const FGeometry& AllottedGeometry, FAr
 	const FSlateRect TooltipPopup(LocalCursorPosition + CursorSize, LocalCursorPosition + CursorSize + ChildSlot.GetWidget()->GetDesiredSize());
 	
 	const FVector2D TooltipPosition = ComputePopupFitInRect(CursorAnchorRect, TooltipPopup, EOrientation::Orient_Vertical, FSlateRect(FVector2D::ZeroVector, AllottedGeometry.GetLocalSize()));
+	
+	// We round the final tooltip position so that our tooltip doesn't begin at a half pixel offset, which avoids the contents of the tooltip
+	// jittering in relation to one another.
+	const FVector2D TooltipPositionRounded = AllottedGeometry.LocalToRoundedLocal(TooltipPosition);
 
 	ArrangedChildren.AddWidget(AllottedGeometry.MakeChild(
 		ChildSlot.GetWidget(),
 		ChildSlot.GetWidget()->GetDesiredSize(),
-		FSlateLayoutTransform(TooltipPosition)
+		FSlateLayoutTransform(TooltipPositionRounded)
 	));
 }
 
diff --git a/Engine/Source/Runtime/SlateCore/Private/Layout/Geometry.cpp b/Engine/Source/Runtime/SlateCore/Private/Layout/Geometry.cpp
--- a/Engine/Source/Runtime/SlateCore/Private/Layout/Geometry.cpp
+++ b/Engine/Source/Runtime/SlateCore/Private/Layout/Geometry.cpp
@@ -186,6 +186,14 @@ FVector2D FGeometry::LocalToAbsolute(FVector2D LocalCoordinate) const
 	return TransformPoint(GetAccumulatedRenderTransform(), LocalCoordinate);
 }
 
+FVector2D FGeometry::LocalToRoundedLocal(FVector2D LocalCoordinate) const
+{
+	const FVector2D AbsoluteCoordinate = LocalToAbsolute(LocalCoordinate);
+	const FVector2D AbsoluteCoordinateRounded = FVector2D(FMath::RoundToInt(AbsoluteCoordinate.X), FMath::RoundToInt(AbsoluteCoordinate.Y));
+
+	return AbsoluteToLocal(AbsoluteCoordinateRounded);
+}
+
 bool FGeometry::IsUnderLocation(const FVector2D& AbsoluteCoordinate) const
 {
 	// this render transform invert is a little expensive. We might consider caching it.
diff --git a/Engine/Source/Runtime/SlateCore/Public/Layout/Geometry.h b/Engine/Source/Runtime/SlateCore/Public/Layout/Geometry.h
--- a/Engine/Source/Runtime/SlateCore/Public/Layout/Geometry.h
+++ b/Engine/Source/Runtime/SlateCore/Public/Layout/Geometry.h
@@ -281,6 +281,14 @@ public:
 	 * @return  Absolute coordinates
 	 */
 	FVector2D LocalToAbsolute( FVector2D LocalCoordinate ) const;
+
+	/**
+	 * Translates the local coordinates into local coordinates that after being transformed into absolute space will be rounded
+	 * to a whole number or approximately a whole number.  This is important for cases where you want to show a popup or a tooltip
+	 * and not have the window start on a half pixel, which can cause the contents to jitter in relation to eachother as the tooltip 
+	 * or popup moves around.
+	 */
+	FVector2D LocalToRoundedLocal( FVector2D LocalCoordinate ) const;
 	
 	/**
 	 * !!! DEPRECATED !!! This legacy function does not account for render transforms.
