Slate/UMG - Moving the SConstraintCanvas out of the UMG Module and into Slate proper.  This will allow usage in more core engine modules that come before UMG.

diff --git a/Engine/Source/Editor/UMGEditor/Private/Extensions/CanvasSlotExtension.h b/Engine/Source/Editor/UMGEditor/Private/Extensions/CanvasSlotExtension.h
--- a/Engine/Source/Editor/UMGEditor/Private/Extensions/CanvasSlotExtension.h
+++ b/Engine/Source/Editor/UMGEditor/Private/Extensions/CanvasSlotExtension.h
@@ -3,7 +3,7 @@
 #pragma once
 
 #include "DesignerExtension.h"
-#include "Slate/Anchors.h"
+#include "Layout/Anchors.h"
 
 class UCanvasPanel;
 class UCanvasPanelSlot;
diff --git a/Engine/Source/Runtime/Slate/Public/Widgets/Layout/Anchors.h b/Engine/Source/Runtime/Slate/Public/Widgets/Layout/Anchors.h
new file mode 100644
--- /dev/null
+++ b/Engine/Source/Runtime/Slate/Public/Widgets/Layout/Anchors.h
@@ -0,0 +1,58 @@
+// Copyright 1998-2015 Epic Games, Inc. All Rights Reserved.
+
+#pragma once
+
+#include "Anchors.generated.h"
+
+/**
+ * Describes how a widget is anchored.
+ */
+USTRUCT(BlueprintType)
+struct FAnchors
+{
+	GENERATED_USTRUCT_BODY()
+
+	/** Holds the minimum anchors, left + top. */
+	UPROPERTY(EditAnywhere, BlueprintReadWrite, Category=Appearance)
+	FVector2D Minimum;
+
+	/** Holds the maximum anchors, right + bottom. */
+	UPROPERTY(EditAnywhere, BlueprintReadWrite, Category=Appearance)
+	FVector2D Maximum;
+
+public:
+
+	/**
+	 * Default constructor.
+	 *
+	 * The default margin size is zero on all four sides..
+	 */
+	FAnchors()
+		: Minimum(0.0f, 0.0f)
+		, Maximum(0.0f, 0.0f)
+	{ }
+
+	/** Construct a Anchors with uniform space on all sides */
+	FAnchors(float UnifromAnchors)
+		: Minimum(UnifromAnchors, UnifromAnchors)
+		, Maximum(UnifromAnchors, UnifromAnchors)
+	{ }
+
+	/** Construct a Anchors where Horizontal describes Left and Right spacing while Vertical describes Top and Bottom spacing */
+	FAnchors(float Horizontal, float Vertical)
+		: Minimum(Horizontal, Vertical)
+		, Maximum(Horizontal, Vertical)
+	{ }
+
+	/** Construct Anchors where the spacing on each side is individually specified. */
+	FAnchors(float MinX, float MinY, float MaxX, float MaxY)
+		: Minimum(MinX, MinY)
+		, Maximum(MaxX, MaxY)
+	{ }
+
+	/** Returns true if the anchors represent a stretch along the vertical axis */
+	bool IsStretchedVertical() const { return Minimum.Y != Maximum.Y; }
+
+	/** Returns true if the anchors represent a stretch along the horizontal axis */
+	bool IsStretchedHorizontal() const { return Minimum.X != Maximum.X; }
+};
diff --git a/Engine/Source/Runtime/Slate/Public/Widgets/Layout/SConstraintCanvas.cpp b/Engine/Source/Runtime/Slate/Public/Widgets/Layout/SConstraintCanvas.cpp
new file mode 100644
--- /dev/null
+++ b/Engine/Source/Runtime/Slate/Public/Widgets/Layout/SConstraintCanvas.cpp
@@ -0,0 +1,221 @@
+// Copyright 1998-2015 Epic Games, Inc. All Rights Reserved.
+
+#include "SlatePrivatePCH.h"
+
+#include "SConstraintCanvas.h"
+#include "LayoutUtils.h"
+
+
+/* SConstraintCanvas interface
+ *****************************************************************************/
+
+SConstraintCanvas::SConstraintCanvas()
+: Children()
+{
+}
+
+void SConstraintCanvas::Construct( const SConstraintCanvas::FArguments& InArgs )
+{
+	const int32 NumSlots = InArgs.Slots.Num();
+	for ( int32 SlotIndex = 0; SlotIndex < NumSlots; ++SlotIndex )
+	{
+		Children.Add( InArgs.Slots[SlotIndex] );
+	}
+}
+
+void SConstraintCanvas::ClearChildren( )
+{
+	Children.Empty();
+}
+
+int32 SConstraintCanvas::RemoveSlot( const TSharedRef<SWidget>& SlotWidget )
+{
+	for (int32 SlotIdx = 0; SlotIdx < Children.Num(); ++SlotIdx)
+	{
+		if (SlotWidget == Children[SlotIdx].GetWidget())
+		{
+			Children.RemoveAt(SlotIdx);
+			return SlotIdx;
+		}
+	}
+
+	return -1;
+}
+
+struct FChildZOrder
+{
+	int32 ChildIndex;
+	float ZOrder;
+};
+
+struct FSortSlotsByZOrder
+{
+	FORCEINLINE bool operator()(const FChildZOrder& A, const FChildZOrder& B) const
+	{
+		return A.ZOrder == B.ZOrder ? A.ChildIndex < B.ChildIndex : A.ZOrder < B.ZOrder;
+	}
+};
+
+/* SWidget overrides
+ *****************************************************************************/
+
+void SConstraintCanvas::OnArrangeChildren( const FGeometry& AllottedGeometry, FArrangedChildren& ArrangedChildren ) const
+{
+	CachedGeometry = AllottedGeometry;
+
+	if (Children.Num() > 0)
+	{
+		// Sort the children based on zorder.
+		TArray< FChildZOrder > SlotOrder;
+		for ( int32 ChildIndex = 0; ChildIndex < Children.Num(); ++ChildIndex )
+		{
+			const SConstraintCanvas::FSlot& CurChild = Children[ChildIndex];
+
+			FChildZOrder Order;
+			Order.ChildIndex = ChildIndex;
+			Order.ZOrder = CurChild.ZOrderAttr.Get();
+			SlotOrder.Add( Order );
+		}
+
+		SlotOrder.Sort(FSortSlotsByZOrder());
+
+		// Arrange the children now in their proper z-order.
+		for (int32 ChildIndex = 0; ChildIndex < Children.Num(); ++ChildIndex)
+		{
+			const FChildZOrder& CurSlot = SlotOrder[ChildIndex];
+			const SConstraintCanvas::FSlot& CurChild = Children[CurSlot.ChildIndex];
+			const TSharedRef<SWidget>& CurWidget = CurChild.GetWidget();
+
+			const EVisibility ChildVisibility = CurWidget->GetVisibility();
+			if ( ArrangedChildren.Accepts(ChildVisibility) )
+			{
+				const FMargin Offset = CurChild.OffsetAttr.Get();
+				const FVector2D Alignment = CurChild.AlignmentAttr.Get();
+				const FAnchors Anchors = CurChild.AnchorsAttr.Get();
+
+				const bool AutoSize = CurChild.AutoSizeAttr.Get();
+
+				const FMargin AnchorPixels =
+					FMargin(Anchors.Minimum.X * AllottedGeometry.Size.X,
+					Anchors.Minimum.Y * AllottedGeometry.Size.Y,
+					Anchors.Maximum.X * AllottedGeometry.Size.X,
+					Anchors.Maximum.Y * AllottedGeometry.Size.Y);
+
+				const bool bIsHorizontalStretch = Anchors.Minimum.X != Anchors.Maximum.X;
+				const bool bIsVerticalStretch = Anchors.Minimum.Y != Anchors.Maximum.Y;
+
+				const FVector2D SlotSize = FVector2D(Offset.Right, Offset.Bottom);
+				const FVector2D WidgetDesiredSize = CurWidget->GetDesiredSize();
+
+				const FVector2D Size = AutoSize ? WidgetDesiredSize : SlotSize;
+
+				// Calculate the offset based on the pivot position.
+				FVector2D AlignmentOffset = Size * Alignment;
+
+				// Calculate the local position based on the anchor and position offset.
+				FVector2D LocalPosition, LocalSize;
+
+				// Calculate the position and size based on the horizontal stretch or non-stretch
+				if ( bIsHorizontalStretch )
+				{
+					LocalPosition.X = AnchorPixels.Left + Offset.Left;
+					LocalSize.X = AnchorPixels.Right - LocalPosition.X - Offset.Right;
+				}
+				else
+				{
+					LocalPosition.X = AnchorPixels.Left + Offset.Left - AlignmentOffset.X;
+					LocalSize.X = Size.X;
+				}
+
+				// Calculate the position and size based on the vertical stretch or non-stretch
+				if ( bIsVerticalStretch )
+				{
+					LocalPosition.Y = AnchorPixels.Top + Offset.Top;
+					LocalSize.Y = AnchorPixels.Bottom - LocalPosition.Y - Offset.Bottom;
+				}
+				else
+				{
+					LocalPosition.Y = AnchorPixels.Top + Offset.Top - AlignmentOffset.Y;
+					LocalSize.Y = Size.Y;
+				}
+
+				// Add the information about this child to the output list (ArrangedChildren)
+				ArrangedChildren.AddWidget(ChildVisibility, AllottedGeometry.MakeChild(
+					// The child widget being arranged
+					CurWidget,
+					// Child's local position (i.e. position within parent)
+					LocalPosition,
+					// Child's size
+					LocalSize
+					));
+			}
+		}
+	}
+}
+
+int32 SConstraintCanvas::OnPaint( const FPaintArgs& Args, const FGeometry& AllottedGeometry, const FSlateRect& MyClippingRect, FSlateWindowElementList& OutDrawElements, int32 LayerId, const FWidgetStyle& InWidgetStyle, bool bParentEnabled ) const
+{
+	FArrangedChildren ArrangedChildren(EVisibility::Visible);
+	ArrangeChildren(AllottedGeometry, ArrangedChildren);
+
+	bool bForwardedEnabled = ShouldBeEnabled(bParentEnabled);
+
+	// Because we paint multiple children, we must track the maximum layer id that they produced in case one of our parents
+	// wants to an overlay for all of its contents.
+	int32 MaxLayerId = LayerId;
+
+	for (int32 ChildIndex = 0; ChildIndex < ArrangedChildren.Num(); ++ChildIndex)
+	{
+		FArrangedWidget& CurWidget = ArrangedChildren[ChildIndex];
+		FSlateRect ChildClipRect = MyClippingRect.IntersectionWith(CurWidget.Geometry.GetClippingRect());
+
+		if ( ChildClipRect.GetSize().IsZero() == false )
+		{
+			const int32 CurWidgetsMaxLayerId = CurWidget.Widget->Paint(Args.WithNewParent(this), CurWidget.Geometry, ChildClipRect, OutDrawElements, MaxLayerId + 1, InWidgetStyle, bForwardedEnabled);
+
+			MaxLayerId = FMath::Max(MaxLayerId, CurWidgetsMaxLayerId);
+		}
+	}
+
+	return MaxLayerId;
+}
+
+FVector2D SConstraintCanvas::ComputeDesiredSize( float ) const
+{
+	FVector2D DesiredSize(0,0);
+
+	// Arrange the children now in their proper z-order.
+	for ( int32 ChildIndex = 0; ChildIndex < Children.Num(); ++ChildIndex )
+	{
+		const SConstraintCanvas::FSlot& CurChild = Children[ChildIndex];
+		const TSharedRef<SWidget>& Widget = CurChild.GetWidget();
+
+		// As long as the widgets are not collapsed, they should contribute to the desired size.
+		if ( Widget->GetVisibility() != EVisibility::Collapsed )
+		{
+			const FMargin Offset = CurChild.OffsetAttr.Get();
+			const FVector2D Alignment = CurChild.AlignmentAttr.Get();
+			const FAnchors Anchors = CurChild.AnchorsAttr.Get();
+
+			const FVector2D SlotSize = FVector2D(Offset.Right, Offset.Bottom);
+			const FVector2D WidgetDesiredSize = Widget->GetDesiredSize();
+
+			const bool AutoSize = CurChild.AutoSizeAttr.Get();
+
+			const FVector2D Size = AutoSize ? WidgetDesiredSize : SlotSize;
+
+			const bool bIsDockedHorizontally = ( Anchors.Minimum.X == Anchors.Maximum.X ) && ( Anchors.Minimum.X == 0 || Anchors.Minimum.X == 1 );
+			const bool bIsDockedVertically = ( Anchors.Minimum.Y == Anchors.Maximum.Y ) && ( Anchors.Minimum.Y == 0 || Anchors.Minimum.Y == 1 );
+
+			DesiredSize.X = FMath::Max(DesiredSize.X, Size.X + ( bIsDockedHorizontally ? FMath::Abs(Offset.Left) : 0.0f ));
+			DesiredSize.Y = FMath::Max(DesiredSize.Y, Size.Y + ( bIsDockedVertically ? FMath::Abs(Offset.Top) : 0.0f ));
+		}
+	}
+
+	return DesiredSize;
+}
+
+FChildren* SConstraintCanvas::GetChildren()
+{
+	return &Children;
+}
diff --git a/Engine/Source/Runtime/Slate/Public/Widgets/Layout/SConstraintCanvas.h b/Engine/Source/Runtime/Slate/Public/Widgets/Layout/SConstraintCanvas.h
new file mode 100644
--- /dev/null
+++ b/Engine/Source/Runtime/Slate/Public/Widgets/Layout/SConstraintCanvas.h
@@ -0,0 +1,153 @@
+// Copyright 1998-2015 Epic Games, Inc. All Rights Reserved.
+
+#pragma once
+
+#include "Anchors.h"
+
+/**
+ * ConstraintCanvas is a layout widget that allows you to arbitrary position and size child widgets in a 
+ * relative coordinate space.  Additionally it permits anchoring widgets.
+ */
+class SLATE_API SConstraintCanvas : public SPanel
+{
+public:
+
+	/**
+	 * ConstraintCanvas slots allow child widgets to be positioned and sized
+	 */
+	class FSlot : public TSlotBase<FSlot>
+	{
+	public:		
+		FSlot& Offset( const TAttribute<FMargin>& InOffset )
+		{
+			OffsetAttr = InOffset;
+			return *this;
+		}
+
+		FSlot& Anchors( const TAttribute<FAnchors>& InAnchors )
+		{
+			AnchorsAttr = InAnchors;
+			return *this;
+		}
+
+		FSlot& Alignment(const TAttribute<FVector2D>& InAlignment)
+		{
+			AlignmentAttr = InAlignment;
+			return *this;
+		}
+
+		FSlot& AutoSize(const TAttribute<bool>& InAutoSize)
+		{
+			AutoSizeAttr = InAutoSize;
+			return *this;
+		}
+
+		FSlot& ZOrder(const TAttribute<float>& InZOrder)
+		{
+			ZOrderAttr = InZOrder;
+			return *this;
+		}
+
+		FSlot& Expose( FSlot*& OutVarToInit )
+		{
+			OutVarToInit = this;
+			return *this;
+		}
+
+		/** Offset */
+		TAttribute<FMargin> OffsetAttr;
+
+		/** Anchors */
+		TAttribute<FAnchors> AnchorsAttr;
+
+		/** Size */
+		TAttribute<FVector2D> AlignmentAttr;
+
+		/** Auto-Size */
+		TAttribute<bool> AutoSizeAttr;
+
+		/** Z-Order */
+		TAttribute<float> ZOrderAttr;
+
+		/** Default values for a slot. */
+		FSlot()
+			: TSlotBase<FSlot>()
+			, OffsetAttr( FMargin( 0, 0, 1, 1 ) )
+			, AnchorsAttr( FAnchors( 0.0f, 0.0f ) )
+			, AlignmentAttr( FVector2D( 0.5f, 0.5f ) )
+			, AutoSizeAttr( false )
+			, ZOrderAttr( 0 )
+		{ }
+	};
+
+	SLATE_BEGIN_ARGS( SConstraintCanvas )
+		{
+			_Visibility = EVisibility::SelfHitTestInvisible;
+		}
+
+		SLATE_SUPPORTS_SLOT( SConstraintCanvas::FSlot )
+
+	SLATE_END_ARGS()
+
+	SConstraintCanvas();
+
+	/**
+	 * Construct this widget
+	 *
+	 * @param	InArgs	The declaration data for this widget
+	 */
+	void Construct( const FArguments& InArgs );
+
+	static FSlot& Slot()
+	{
+		return *(new FSlot());
+	}
+
+	/**
+	 * Adds a content slot.
+	 *
+	 * @return The added slot.
+	 */
+	FSlot& AddSlot()
+	{
+		SConstraintCanvas::FSlot& NewSlot = *(new FSlot());
+		this->Children.Add( &NewSlot );
+		return NewSlot;
+	}
+
+	/**
+	 * Removes a particular content slot.
+	 *
+	 * @param SlotWidget The widget in the slot to remove.
+	 */
+	int32 RemoveSlot( const TSharedRef<SWidget>& SlotWidget );
+
+	/**
+	 * Removes all slots from the panel.
+	 */
+	void ClearChildren();
+
+	/** Get the cached geometry of the canvas. */
+	FGeometry GetCachedGeometry() const { return CachedGeometry; }
+
+public:
+
+	// Begin SWidget overrides
+
+	virtual void OnArrangeChildren( const FGeometry& AllottedGeometry, FArrangedChildren& ArrangedChildren ) const override;
+	
+	virtual int32 OnPaint( const FPaintArgs& Args, const FGeometry& AllottedGeometry, const FSlateRect& MyClippingRect, FSlateWindowElementList& OutDrawElements, int32 LayerId, const FWidgetStyle& InWidgetStyle, bool bParentEnabled ) const override;
+
+	virtual FVector2D ComputeDesiredSize(float) const override;
+
+	virtual FChildren* GetChildren() override;
+
+	// End SWidget overrides
+
+protected:
+
+	/** The ConstraintCanvas widget's children. */
+	TPanelChildren< FSlot > Children;
+
+	mutable FGeometry CachedGeometry;
+};
diff --git a/Engine/Source/Runtime/UMG/Private/Slate/SConstraintCanvas.cpp b/Engine/Source/Runtime/UMG/Private/Slate/SConstraintCanvas.cpp
deleted file mode 100644
diff --git a/Engine/Source/Runtime/UMG/Public/Blueprint/UserWidget.h b/Engine/Source/Runtime/UMG/Public/Blueprint/UserWidget.h
--- a/Engine/Source/Runtime/UMG/Public/Blueprint/UserWidget.h
+++ b/Engine/Source/Runtime/UMG/Public/Blueprint/UserWidget.h
@@ -1,13 +1,13 @@
 // Copyright 1998-2015 Epic Games, Inc. All Rights Reserved.
 
 #pragma once
+
 #include "Components/SlateWrapperTypes.h"
 #include "Components/Widget.h"
 #include "Geometry.h"
 #include "Engine/GameInstance.h"
-
+#include "Layout/Anchors.h"
 #include "NamedSlotInterface.h"
-#include "Slate/Anchors.h"
 #include "Engine/LocalPlayer.h"
 #include "Logging/MessageLog.h"
 
@@ -20,7 +20,6 @@ static FWidgetStyle NullStyle;
 
 class UWidget;
 class UWidgetAnimation;
-struct FAnchors;
 struct FLocalPlayerContext;
 
 /**
diff --git a/Engine/Source/Runtime/UMG/Public/Components/CanvasPanelSlot.h b/Engine/Source/Runtime/UMG/Public/Components/CanvasPanelSlot.h
--- a/Engine/Source/Runtime/UMG/Public/Components/CanvasPanelSlot.h
+++ b/Engine/Source/Runtime/UMG/Public/Components/CanvasPanelSlot.h
@@ -3,8 +3,9 @@
 #pragma once
 
 #include "Anchors.h"
+#include "SConstraintCanvas.h"
 #include "Components/PanelSlot.h"
-#include "Slate/SConstraintCanvas.h"
+
 #include "CanvasPanelSlot.generated.h"
 
 USTRUCT(BlueprintType)
diff --git a/Engine/Source/Runtime/UMG/Public/Slate/Anchors.h b/Engine/Source/Runtime/UMG/Public/Slate/Anchors.h
deleted file mode 100644
diff --git a/Engine/Source/Runtime/UMG/Public/Slate/SConstraintCanvas.h b/Engine/Source/Runtime/UMG/Public/Slate/SConstraintCanvas.h
deleted file mode 100644
