Fixed some margin issues with multiline editable text

Some follow up from CL# 2585066. Adding a new line wasn't applying the margin, and the margin wasn't being applied to the lines themselves, causing the highlights (such as the caret and selection marker) to be drawn in the wrong place.

diff --git a/Engine/Source/Runtime/Slate/Private/Framework/Text/TextLayout.cpp b/Engine/Source/Runtime/Slate/Private/Framework/Text/TextLayout.cpp
--- a/Engine/Source/Runtime/Slate/Private/Framework/Text/TextLayout.cpp
+++ b/Engine/Source/Runtime/Slate/Private/Framework/Text/TextLayout.cpp
@@ -336,14 +336,13 @@ void FTextLayout::MarginLayout()
 
 	// Adjust the lines to be offset
 	FVector2D OffsetAdjustment = FVector2D(Margin.Left, Margin.Top) * Scale;
-	for ( int32 LineViewIndex = 0; LineViewIndex < LineViews.Num(); LineViewIndex++ )
+	for (FLineView& LineView : LineViews)
 	{
-		FLineView& LineView = LineViews[LineViewIndex];
+		LineView.Offset += OffsetAdjustment;
 
-		for ( int32 BlockIndex = 0; BlockIndex < LineView.Blocks.Num(); BlockIndex++ )
+		for (const TSharedRef< ILayoutBlock >& Block : LineView.Blocks)
 		{
-			const TSharedRef< ILayoutBlock >& Block = LineView.Blocks[BlockIndex];
-			Block->SetLocationOffset(Block->GetLocationOffset() + OffsetAdjustment);
+			Block->SetLocationOffset( Block->GetLocationOffset() + OffsetAdjustment );
 		}
 	}
 }
@@ -1449,6 +1448,23 @@ void FTextLayout::AddLine( const TSharedRef< FString >& Text, const TArray< TSha
 		TArray<TSharedRef<ILayoutBlock>> SoftLine;
 		FlowLineLayout(LineModelIndex, GetWrappingDrawWidth(), SoftLine);
 
+		// Apply the current margin to the newly added line
+		if (LineViews.Num() > 0)
+		{
+			const FVector2D MarginOffsetAdjustment = FVector2D(Margin.Left, Margin.Top) * Scale;
+
+			FLineView& LastLineView = LineViews.Last();
+			if (LastLineView.ModelIndex == LineModelIndex)
+			{
+				LastLineView.Offset += MarginOffsetAdjustment;
+			}
+
+			for (const TSharedRef< ILayoutBlock >& Block : SoftLine)
+			{
+				Block->SetLocationOffset( Block->GetLocationOffset() + MarginOffsetAdjustment );
+			}
+		}
+
 		// We need to re-justify all lines, as the new line view(s) added by this line model may have affected everything
 		JustifyLayout();
 
diff --git a/Engine/Source/Runtime/Slate/Private/Widgets/Text/SMultiLineEditableText.cpp b/Engine/Source/Runtime/Slate/Private/Widgets/Text/SMultiLineEditableText.cpp
--- a/Engine/Source/Runtime/Slate/Private/Widgets/Text/SMultiLineEditableText.cpp
+++ b/Engine/Source/Runtime/Slate/Private/Widgets/Text/SMultiLineEditableText.cpp
@@ -2598,7 +2598,7 @@ void SMultiLineEditableText::CacheDesiredSize(float LayoutScaleMultiplier)
 
 	TextLayout->SetScale( LayoutScaleMultiplier );
 	TextLayout->SetWrappingWidth( WrappingWidth );
-	TextLayout->SetMargin( Margin.Get() * TextLayout->GetScale() );
+	TextLayout->SetMargin( Margin.Get() );
 	TextLayout->SetLineHeightPercentage( LineHeightPercentage.Get() );
 	TextLayout->SetJustification( Justification.Get() );
 	TextLayout->SetVisibleRegion( CachedSize, ScrollOffset * TextLayout->GetScale() );
