Made clearing the text selection when a ML editable text loses focus optional

UE-17599 - Rich Text for tutorials is not working in Main

Previously SMultiLineEditableText never cleared its selection when it lost focus, however recent changes to make it behave more like SEditableText had undone that behavior. This caused issues with some of our text editors, as they rely on the selection being preserved when focus is lost so that they can style the selected text.

I've added an option to control whether focus loss should clear the current selection, and have turned it off for our text editors.

diff --git a/Engine/Source/Editor/IntroTutorials/Private/STutorialEditableText.cpp b/Engine/Source/Editor/IntroTutorials/Private/STutorialEditableText.cpp
--- a/Engine/Source/Editor/IntroTutorials/Private/STutorialEditableText.cpp
+++ b/Engine/Source/Editor/IntroTutorials/Private/STutorialEditableText.cpp
@@ -59,6 +59,7 @@ void STutorialEditableText::Construct(const FArguments& InArgs)
 			.OnTextCommitted(this, &STutorialEditableText::HandleRichEditableTextCommitted)
 			.OnCursorMoved(this, &STutorialEditableText::HandleRichEditableTextCursorMoved)
 			.Marshaller(RichTextMarshaller)
+			.ClearTextSelectionOnFocusLoss(false)
 			.AutoWrapText(true)
 			.Margin(4)
 			.LineHeightPercentage(1.1f)
diff --git a/Engine/Source/Runtime/AppFramework/Private/Widgets/Testing/STestSuite.cpp b/Engine/Source/Runtime/AppFramework/Private/Widgets/Testing/STestSuite.cpp
--- a/Engine/Source/Runtime/AppFramework/Private/Widgets/Testing/STestSuite.cpp
+++ b/Engine/Source/Runtime/AppFramework/Private/Widgets/Testing/STestSuite.cpp
@@ -2744,6 +2744,7 @@ public:
 									.OnTextCommitted(this, &SRichTextEditTest::HandleRichEditableTextCommitted)
 									.OnCursorMoved(this, &SRichTextEditTest::HandleRichEditableTextCursorMoved)
 									.Marshaller(RichTextMarshaller)
+									.ClearTextSelectionOnFocusLoss(false)
 									.AutoWrapText(true)
 									.Margin(4)
 									.LineHeightPercentage(1.1f)
diff --git a/Engine/Source/Runtime/Slate/Private/Widgets/Input/SMultiLineEditableTextBox.cpp b/Engine/Source/Runtime/Slate/Private/Widgets/Input/SMultiLineEditableTextBox.cpp
--- a/Engine/Source/Runtime/Slate/Private/Widgets/Input/SMultiLineEditableTextBox.cpp
+++ b/Engine/Source/Runtime/Slate/Private/Widgets/Input/SMultiLineEditableTextBox.cpp
@@ -111,6 +111,7 @@ void SMultiLineEditableTextBox::Construct( const FArguments& InArgs )
 					.Justification(InArgs._Justification)
 					.RevertTextOnEscape(InArgs._RevertTextOnEscape)
 					.SelectAllTextWhenFocused(InArgs._SelectAllTextWhenFocused)
+					.ClearTextSelectionOnFocusLoss(InArgs._ClearTextSelectionOnFocusLoss)
 					.ClearKeyboardFocusOnCommit(InArgs._ClearKeyboardFocusOnCommit)
 					.LineHeightPercentage(InArgs._LineHeightPercentage)
 					.Margin(InArgs._Margin)
diff --git a/Engine/Source/Runtime/Slate/Private/Widgets/Text/SMultiLineEditableText.cpp b/Engine/Source/Runtime/Slate/Private/Widgets/Text/SMultiLineEditableText.cpp
--- a/Engine/Source/Runtime/Slate/Private/Widgets/Text/SMultiLineEditableText.cpp
+++ b/Engine/Source/Runtime/Slate/Private/Widgets/Text/SMultiLineEditableText.cpp
@@ -273,6 +273,7 @@ void SMultiLineEditableText::Construct( const FArguments& InArgs )
 	OnTextCommitted = InArgs._OnTextCommitted;
 	OnCursorMoved = InArgs._OnCursorMoved;
 	bSelectAllTextWhenFocused = InArgs._SelectAllTextWhenFocused;
+	bClearTextSelectionOnFocusLoss = InArgs._ClearTextSelectionOnFocusLoss;
 	bClearKeyboardFocusOnCommit = InArgs._ClearKeyboardFocusOnCommit;
 	OnContextMenuOpening = InArgs._OnContextMenuOpening;
 	bRevertTextOnEscape = InArgs._RevertTextOnEscape;
@@ -613,7 +614,7 @@ void SMultiLineEditableText::OnFocusLost( const FFocusEvent& InFocusEvent )
 		}
 
 		// Clear selection unless activating a new window (otherwise can't copy and past on right click)
-		if (InFocusEvent.GetCause() != EFocusCause::WindowActivate)
+		if (bClearTextSelectionOnFocusLoss.Get() && InFocusEvent.GetCause() != EFocusCause::WindowActivate)
 		{
 			ClearSelection();
 		}
diff --git a/Engine/Source/Runtime/Slate/Public/Widgets/Input/SMultiLineEditableTextBox.h b/Engine/Source/Runtime/Slate/Public/Widgets/Input/SMultiLineEditableTextBox.h
--- a/Engine/Source/Runtime/Slate/Public/Widgets/Input/SMultiLineEditableTextBox.h
+++ b/Engine/Source/Runtime/Slate/Public/Widgets/Input/SMultiLineEditableTextBox.h
@@ -27,6 +27,7 @@ public:
 		, _IsPassword( false )
 		, _IsCaretMovedWhenGainFocus ( true )
 		, _SelectAllTextWhenFocused( false )
+		, _ClearTextSelectionOnFocusLoss( true )
 		, _RevertTextOnEscape( false )
 		, _ClearKeyboardFocusOnCommit( true )
 		, _AlwaysShowScrollbars( false )
@@ -84,6 +85,9 @@ public:
 		/** Whether to select all text when the user clicks to give focus on the widget */
 		SLATE_ATTRIBUTE( bool, SelectAllTextWhenFocused )
 
+		/** Whether to clear text selection when focus is lost */
+		SLATE_ATTRIBUTE( bool, ClearTextSelectionOnFocusLoss )
+
 		/** Whether to allow the user to back out of changes when they press the escape key */
 		SLATE_ATTRIBUTE( bool, RevertTextOnEscape )
 
diff --git a/Engine/Source/Runtime/Slate/Public/Widgets/Text/SMultiLineEditableText.h b/Engine/Source/Runtime/Slate/Public/Widgets/Text/SMultiLineEditableText.h
--- a/Engine/Source/Runtime/Slate/Public/Widgets/Text/SMultiLineEditableText.h
+++ b/Engine/Source/Runtime/Slate/Public/Widgets/Text/SMultiLineEditableText.h
@@ -34,6 +34,7 @@ public:
 		, _OnTextChanged()
 		, _OnTextCommitted()
 		, _SelectAllTextWhenFocused(false)
+		, _ClearTextSelectionOnFocusLoss(true)
 		, _RevertTextOnEscape(false)
 		, _ClearKeyboardFocusOnCommit(true)
 		, _OnCursorMoved()
@@ -90,6 +91,9 @@ public:
 		/** Whether to select all text when the user clicks to give focus on the widget */
 		SLATE_ATTRIBUTE(bool, SelectAllTextWhenFocused)
 
+		/** Whether to clear text selection when focus is lost */
+		SLATE_ATTRIBUTE( bool, ClearTextSelectionOnFocusLoss )
+
 		/** Whether to allow the user to back out of changes when they press the escape key */
 		SLATE_ATTRIBUTE(bool, RevertTextOnEscape)
 
@@ -676,6 +680,9 @@ private:
 	/** Whether to select all text when the user clicks to give focus on the widget */
 	TAttribute< bool > bSelectAllTextWhenFocused;
 
+	/** Whether to clear text selection when focus is lost */
+	TAttribute< bool > bClearTextSelectionOnFocusLoss;
+
 	/** True if any changes should be reverted if we recieve an escape key */
 	TAttribute< bool > bRevertTextOnEscape;
 
