Fix so that menus set the widget to focus upon activation to the actual content widget, rather than the wrapper.

diff --git a/Engine/Source/Runtime/Slate/Private/Framework/Application/MenuStack.cpp b/Engine/Source/Runtime/Slate/Private/Framework/Application/MenuStack.cpp
--- a/Engine/Source/Runtime/Slate/Private/Framework/Application/MenuStack.cpp
+++ b/Engine/Source/Runtime/Slate/Private/Framework/Application/MenuStack.cpp
@@ -292,6 +292,10 @@ FMenuStack::FPrePushResults FMenuStack::PrePush(const FPrePushArgs& InArgs)
 {
 	FPrePushResults OutResults;
 	OutResults.bFocusImmediately = InArgs.bFocusImmediately;
+	if (InArgs.bFocusImmediately)
+	{
+		OutResults.WidgetToFocus = InArgs.Content;
+	}
 
 	// Only enable window position/size transitions if we're running at a decent frame rate
 	OutResults.bAllowAnimations = FSlateApplication::Get().AreMenuAnimationsEnabled() && FSlateApplication::Get().IsRunningAtTargetFrameRate();
@@ -416,10 +420,10 @@ TSharedRef<FMenuBase> FMenuStack::PushNewWindow(TSharedPtr<IMenu> InParentMenu,
 
 	PendingNewWindow = NewMenuWindow;
 
-	if (InPrePushResults.bFocusImmediately)
+	if (InPrePushResults.bFocusImmediately && InPrePushResults.WidgetToFocus.IsValid())
 	{
-		// Focus the content rather than just the window
-		NewMenuWindow->SetWidgetToFocusOnActivate(InPrePushResults.WrappedContent);
+		// Focus the unwrapped content rather than just the window
+		NewMenuWindow->SetWidgetToFocusOnActivate(InPrePushResults.WidgetToFocus);
 	}
 
 	TSharedRef<FMenuInWindow> Menu = MakeShareable(new FMenuInWindow(NewMenuWindow, InPrePushResults.WrappedContent.ToSharedRef()));
@@ -466,9 +470,9 @@ TSharedRef<FMenuBase> FMenuStack::PushPopup(TSharedPtr<IMenu> InParentMenu, cons
 	// Add it to a slot on the menus panel widget
 	HostWindowPopupPanel->PushMenu(Menu, InPrePushResults.AnimFinalLocation);
 
-	if (InPrePushResults.bFocusImmediately)
+	if (InPrePushResults.bFocusImmediately && InPrePushResults.WidgetToFocus.IsValid())
 	{
-		FSlateApplication::Get().SetKeyboardFocus(InPrePushResults.WrappedContent, EFocusCause::SetDirectly);
+		FSlateApplication::Get().SetKeyboardFocus(InPrePushResults.WidgetToFocus, EFocusCause::SetDirectly);
 	}
 
 	return Menu;
diff --git a/Engine/Source/Runtime/Slate/Public/Framework/Application/MenuStack.h b/Engine/Source/Runtime/Slate/Public/Framework/Application/MenuStack.h
--- a/Engine/Source/Runtime/Slate/Public/Framework/Application/MenuStack.h
+++ b/Engine/Source/Runtime/Slate/Public/Framework/Application/MenuStack.h
@@ -310,6 +310,7 @@ private:
 	struct FPrePushResults
 	{
 		TSharedPtr<SWidget> WrappedContent;
+		TSharedPtr<SWidget> WidgetToFocus;
 		FVector2D AnimStartLocation;
 		FVector2D AnimFinalLocation;
 		bool bAnchorSetsMinWidth;
