Slate/UMG - Making some improvements to the way the SScaleBox handles fill justification so that it actually fills the entire contents of the scalebox, and handles additional scale being applied to it.

diff --git a/Engine/Source/Runtime/Slate/Private/Widgets/Layout/SScaleBox.cpp b/Engine/Source/Runtime/Slate/Private/Widgets/Layout/SScaleBox.cpp
--- a/Engine/Source/Runtime/Slate/Private/Widgets/Layout/SScaleBox.cpp
+++ b/Engine/Source/Runtime/Slate/Private/Widgets/Layout/SScaleBox.cpp
@@ -31,12 +31,13 @@ void SScaleBox::OnArrangeChildren( const FGeometry& AllottedGeometry, FArrangedC
 	const EVisibility ChildVisibility = ChildSlot.GetWidget()->GetVisibility();
 	if ( ArrangedChildren.Accepts(ChildVisibility) )
 	{
+		const FVector2D AreaSize = AllottedGeometry.GetLocalSize();
 		FVector2D DesiredSize = ChildSlot.GetWidget()->GetDesiredSize();
 
 		float FinalScale = 1;
 
-		EStretch::Type CurrentStretch = Stretch.Get();
-		EStretchDirection::Type CurrentStretchDirection = StretchDirection.Get();
+		const EStretch::Type CurrentStretch = Stretch.Get();
+		const EStretchDirection::Type CurrentStretchDirection = StretchDirection.Get();
 
 		if ( DesiredSize.X != 0 && DesiredSize.Y != 0 )
 		{
@@ -45,13 +46,13 @@ void SScaleBox::OnArrangeChildren( const FGeometry& AllottedGeometry, FArrangedC
 			case EStretch::None:
 				break;
 			case EStretch::Fill:
-				DesiredSize = AllottedGeometry.Size;
+				DesiredSize = AreaSize;
 				break;
 			case EStretch::ScaleToFit:
-				FinalScale = FMath::Min(AllottedGeometry.Size.X / DesiredSize.X, AllottedGeometry.Size.Y / DesiredSize.Y);
+				FinalScale = FMath::Min(AreaSize.X / DesiredSize.X, AreaSize.Y / DesiredSize.Y);
 				break;
 			case EStretch::ScaleToFill:
-				FinalScale = FMath::Max(AllottedGeometry.Size.X / DesiredSize.X, AllottedGeometry.Size.Y / DesiredSize.Y);
+				FinalScale = FMath::Max(AreaSize.X / DesiredSize.X, AreaSize.Y / DesiredSize.Y);
 				break;
 			case EStretch::UserSpecified:
 				FinalScale = UserSpecifiedScale.Get(1.0f);
@@ -73,13 +74,28 @@ void SScaleBox::OnArrangeChildren( const FGeometry& AllottedGeometry, FArrangedC
 
 		FVector2D FinalOffset(0, 0);
 
+		// If we're just filling, there's no scale applied, we're just filling the area.
 		if ( CurrentStretch != EStretch::Fill )
 		{
 			const FMargin SlotPadding(ChildSlot.SlotPadding.Get());
-			AlignmentArrangeResult XResult = AlignChild<Orient_Horizontal>(AllottedGeometry.Size.X, ChildSlot, SlotPadding, FinalScale, false);
-			AlignmentArrangeResult YResult = AlignChild<Orient_Vertical>(AllottedGeometry.Size.Y, ChildSlot, SlotPadding, FinalScale, false);
+			AlignmentArrangeResult XResult = AlignChild<Orient_Horizontal>(AreaSize.X, ChildSlot, SlotPadding, FinalScale, false);
+			AlignmentArrangeResult YResult = AlignChild<Orient_Vertical>(AreaSize.Y, ChildSlot, SlotPadding, FinalScale, false);
 
-			FinalOffset = FVector2D(XResult.Offset, YResult.Offset) * ( 1.0f / FinalScale );
+			FinalOffset = FVector2D(XResult.Offset, YResult.Offset) / FinalScale;
+
+			// If the layout horizontally is fill, then we need the desired size to be the whole size of the widget, 
+			// but scale the inverse of the scale we're applying.
+			if ( ChildSlot.HAlignment == HAlign_Fill )
+			{
+				DesiredSize.X = AreaSize.X / FinalScale;
+			}
+
+			// If the layout vertically is fill, then we need the desired size to be the whole size of the widget, 
+			// but scale the inverse of the scale we're applying.
+			if ( ChildSlot.VAlignment == VAlign_Fill )
+			{
+				DesiredSize.Y = AreaSize.Y / FinalScale;
+			}
 		}
 
 		ArrangedChildren.AddWidget(ChildVisibility, AllottedGeometry.MakeChild(
@@ -91,7 +107,6 @@ void SScaleBox::OnArrangeChildren( const FGeometry& AllottedGeometry, FArrangedC
 	}
 }
 
-
 int32 SScaleBox::OnPaint( const FPaintArgs& Args, const FGeometry& AllottedGeometry, const FSlateRect& MyClippingRect, FSlateWindowElementList& OutDrawElements, int32 LayerId, const FWidgetStyle& InWidgetStyle, bool bParentEnabled ) const
 {
 	return SCompoundWidget::OnPaint(Args, AllottedGeometry, MyClippingRect, OutDrawElements, LayerId, InWidgetStyle, bParentEnabled);
