Platforms using a virtual keyboard no longer attempt to show it when focus is received by a read-only text box.
[UE-13398] - UMG- Text Box bug

diff --git a/Engine/Source/Runtime/Slate/Private/Widgets/Input/SEditableText.cpp b/Engine/Source/Runtime/Slate/Private/Widgets/Input/SEditableText.cpp
--- a/Engine/Source/Runtime/Slate/Private/Widgets/Input/SEditableText.cpp
+++ b/Engine/Source/Runtime/Slate/Private/Widgets/Input/SEditableText.cpp
@@ -75,6 +75,7 @@ void SEditableText::Construct( const FArguments& InArgs )
 	SelectAllTextOnCommit = InArgs._SelectAllTextOnCommit;
 	VirtualKeyboardType = IsPassword.Get() ? EKeyboardType::Keyboard_Password : InArgs._VirtualKeyboardType;
 	OnKeyDownHandler = InArgs._OnKeyDownHandler;
+	bShowingVirtualKeyboard = false;
 
 	// Map UI commands to delegates which are called when the command should be executed
 	UICommandList->MapAction( FGenericCommands::Get().Undo,
@@ -1416,8 +1417,12 @@ FReply SEditableText::OnFocusReceived( const FGeometry& MyGeometry, const FFocus
 		FSlateApplication& SlateApplication = FSlateApplication::Get();
 		if (FPlatformMisc::GetRequiresVirtualKeyboard())
 		{
-			// @TODO: Create ITextInputMethodSystem derivations for mobile
-			SlateApplication.ShowVirtualKeyboard(true, InFocusEvent.GetUser(), SharedThis(this));
+			if (!GetIsReadOnly())
+			{
+				// @TODO: Create ITextInputMethodSystem derivations for mobile
+				SlateApplication.ShowVirtualKeyboard(true, InFocusEvent.GetUser(), SharedThis(this));
+				bShowingVirtualKeyboard = true;
+			}
 		}
 		else
 		{
@@ -1444,7 +1449,11 @@ void SEditableText::OnFocusLost( const FFocusEvent& InFocusEvent )
 		FSlateApplication& SlateApplication = FSlateApplication::Get();
 		if (FPlatformMisc::GetRequiresVirtualKeyboard())
 		{
-			SlateApplication.ShowVirtualKeyboard(false, InFocusEvent.GetUser());
+			if (bShowingVirtualKeyboard)
+			{
+				SlateApplication.ShowVirtualKeyboard(false, InFocusEvent.GetUser());
+				bShowingVirtualKeyboard = false;
+			}
 		}
 		else
 		{
diff --git a/Engine/Source/Runtime/Slate/Public/Widgets/Input/SEditableText.h b/Engine/Source/Runtime/Slate/Public/Widgets/Input/SEditableText.h
--- a/Engine/Source/Runtime/Slate/Public/Widgets/Input/SEditableText.h
+++ b/Engine/Source/Runtime/Slate/Public/Widgets/Input/SEditableText.h
@@ -665,4 +665,7 @@ private:
 
 	/** Callback delegate to have first chance handling of the OnKeyDown event */
 	FOnKeyDown OnKeyDownHandler;
+
+	/** Keep track of whether we showed the virtual keyboard when focus was received */
+	bool bShowingVirtualKeyboard;
 };
