Fixed some context menus not opening properly on mac

This addresses UE-16702 - Viewport buttons are difficult to click on Mac

diff --git a/Engine/Source/Runtime/Slate/Private/Framework/Application/SlateApplication.cpp b/Engine/Source/Runtime/Slate/Private/Framework/Application/SlateApplication.cpp
--- a/Engine/Source/Runtime/Slate/Private/Framework/Application/SlateApplication.cpp
+++ b/Engine/Source/Runtime/Slate/Private/Framework/Application/SlateApplication.cpp
@@ -4166,7 +4166,7 @@ bool FSlateApplication::ProcessMouseButtonDownEvent( const TSharedPtr< FGenericW
 				}
 
 #if PLATFORM_MAC
-				if (WidgetsUnderCursor.TopLevelWindow.IsValid())
+				if (bNeedToActivateWindow && WidgetsUnderCursor.TopLevelWindow.IsValid())
 				{
 					// Clicking on a context menu should not activate anything
 					// @todo: This needs to be updated when we have window type in SWindow and we no longer have to guess if WidgetsUnderCursor.TopLevelWindow is a menu
