Slate/Core - The TAttribute no longer allows you to directly compare another TAttrbiute for equality.  This was being abused and hiding several problems in the codebase because it was just doing an implicit call to Get() == Get().  Which is not a proper comparision for equality between two attributes, that's only comparing their values.  They've been deprecated, and corrected to only allow comarision against literal values instead of comparision attrbiutes, which was causing implict construction of attrbiutes when comparing against a literal.  There's  a new IdenticalTo() function on TAttribute which can actually compare two attributes to see if they are really the same.

diff --git a/Engine/Source/Runtime/Core/Public/Misc/Attribute.h b/Engine/Source/Runtime/Core/Public/Misc/Attribute.h
--- a/Engine/Source/Runtime/Core/Public/Misc/Attribute.h
+++ b/Engine/Source/Runtime/Core/Public/Misc/Attribute.h
@@ -286,16 +286,41 @@ public:
 	}
 
 	/**
-	 * Equality operator
+	 * Is this attribute identical to another TAttribute.
 	 *
 	 * @param InOther The other attribute to compare with.
-	 * @return true if the attributes are equal, false otherwise.
+	 * @return true if the attributes are identical, false otherwise.
 	 */
-	bool operator==( const TAttribute& InOther ) const
+	bool IdenticalTo(const TAttribute& InOther) const
 	{
-		return InOther.Get() == Get();
+		const bool bIsBound = IsBound();
+
+		if ( bIsBound == InOther.IsBound() )
+		{
+			if ( bIsBound )
+			{
+				return Getter.GetHandle() == InOther.Getter.GetHandle();
+			}
+			else
+			{
+				return Value == InOther.Value;
+			}
+		}
+
+		return false;
 	}
 
+	/**
+	 * Equal operator.
+	 *
+	 * @param InOther The other attribute to compare with.
+	 * @return true if the attributes are not equal, false otherwise.
+	 */
+	DEPRECATED(4.9, "Comparing an attribute directly to a literal value is deprecated.  Please call Get() on the attribute before comparing it." )
+	bool operator==( const ObjectType& InOther ) const
+	{
+		return Get() == InOther;
+	}
 
 	/**
 	 * Not equal operator.
@@ -303,9 +328,10 @@ public:
 	 * @param InOther The other attribute to compare with.
 	 * @return true if the attributes are not equal, false otherwise.
 	 */
-	bool operator!=( const TAttribute& InOther ) const
+	DEPRECATED(4.9, "Comparing an attribute directly to a literal value is deprecated.  Please call Get() on the attribute before comparing it.")
+	bool operator!=( const ObjectType& InOther ) const
 	{
-		return InOther.Get() != Get();
+		return Get() != InOther;
 	}
 	
 private:
diff --git a/Engine/Source/Runtime/Slate/Private/Widgets/Input/SCheckBox.cpp b/Engine/Source/Runtime/Slate/Private/Widgets/Input/SCheckBox.cpp
--- a/Engine/Source/Runtime/Slate/Private/Widgets/Input/SCheckBox.cpp
+++ b/Engine/Source/Runtime/Slate/Private/Widgets/Input/SCheckBox.cpp
@@ -65,12 +65,13 @@ FReply SCheckBox::OnKeyUp( const FGeometry& MyGeometry, const FKeyEvent& InKeyEv
 	if ( InKeyEvent.GetKey() == EKeys::SpaceBar )
 	{
 		ToggleCheckedState();
-		const TAttribute<ECheckBoxState>& State = IsCheckboxChecked.Get();
-		if(State == ECheckBoxState::Checked)
+
+		const ECheckBoxState State = IsCheckboxChecked.Get();
+		if ( State == ECheckBoxState::Checked )
 		{
 			PlayCheckedSound();
 		}
-		else if(State == ECheckBoxState::Unchecked)
+		else if ( State == ECheckBoxState::Unchecked )
 		{
 			PlayUncheckedSound();
 		}
@@ -98,7 +99,7 @@ FReply SCheckBox::OnMouseButtonDown( const FGeometry& MyGeometry, const FPointer
 		if( ClickMethod == EButtonClickMethod::MouseDown )
 		{
 			ToggleCheckedState();
-			const TAttribute<ECheckBoxState>& State = IsCheckboxChecked.Get();
+			const ECheckBoxState State = IsCheckboxChecked.Get();
 			if(State == ECheckBoxState::Checked)
 			{
 				PlayCheckedSound();
@@ -180,7 +181,7 @@ FReply SCheckBox::OnMouseButtonUp( const FGeometry& MyGeometry, const FPointerEv
 				if( ClickMethod == EButtonClickMethod::MouseUp || HasMouseCapture() )
 				{
 					ToggleCheckedState();
-					const TAttribute<ECheckBoxState>& State = IsCheckboxChecked.Get();
+					const ECheckBoxState State = IsCheckboxChecked.Get();
 					if(State == ECheckBoxState::Checked)
 					{
 						PlayCheckedSound();
@@ -259,7 +260,7 @@ ECheckBoxState SCheckBox::GetCheckedState() const
  */
 void SCheckBox::ToggleCheckedState()
 {
-	const TAttribute<ECheckBoxState>& State = IsCheckboxChecked.Get();
+	const ECheckBoxState State = IsCheckboxChecked.Get();
 
 	// If the current check box state is checked OR undetermined we set the check box to checked.
 	if( State == ECheckBoxState::Checked || State == ECheckBoxState::Undetermined )
diff --git a/Engine/Source/Runtime/Slate/Private/Widgets/Layout/SSplitter.cpp b/Engine/Source/Runtime/Slate/Private/Widgets/Layout/SSplitter.cpp
--- a/Engine/Source/Runtime/Slate/Private/Widgets/Layout/SSplitter.cpp
+++ b/Engine/Source/Runtime/Slate/Private/Widgets/Layout/SSplitter.cpp
@@ -97,7 +97,7 @@ TArray<FLayoutGeometry> SSplitter::ArrangeChildrenForLayout(const FGeometry& All
 			{
 				++NumNonCollapsedChildren;
 
-				if ( Children[ChildIndex].SizingRule == SSplitter::SizeToContent )
+				if ( Children[ChildIndex].SizingRule.Get() == SSplitter::SizeToContent )
 				{
 					NonResizeableSpace += Children[ChildIndex].GetWidget()->GetDesiredSize()[AxisIndex];
 				}
@@ -121,7 +121,7 @@ TArray<FLayoutGeometry> SSplitter::ArrangeChildrenForLayout(const FGeometry& All
 	{
 		const FSlot& CurSlot = Children[ChildIndex];
 
-		const float ChildSpace = ( CurSlot.SizingRule == SSplitter::SizeToContent )
+		const float ChildSpace = ( CurSlot.SizingRule.Get() == SSplitter::SizeToContent )
 			? CurSlot.GetWidget()->GetDesiredSize()[AxisIndex]
 		: ResizeableSpace * CurSlot.SizeValue.Get() / CoefficientTotal;
 
diff --git a/Engine/Source/Runtime/Slate/Public/Widgets/Views/SListView.h b/Engine/Source/Runtime/Slate/Public/Widgets/Views/SListView.h
--- a/Engine/Source/Runtime/Slate/Public/Widgets/Views/SListView.h
+++ b/Engine/Source/Runtime/Slate/Public/Widgets/Views/SListView.h
@@ -288,7 +288,7 @@ public:
 					ItemType SelectorItemDereference( TListTypeTraits<ItemType>::NullableItemTypeConvertToItemType( SelectorItem ) );
 
 					// Deselect.
-					if( InKeyEvent.IsControlDown() || this->SelectionMode == ESelectionMode::SingleToggle )
+					if( InKeyEvent.IsControlDown() || SelectionMode.Get() == ESelectionMode::SingleToggle )
 					{
 						this->Private_SetItemSelection( SelectorItemDereference, !( this->Private_IsItemSelected( SelectorItemDereference ) ), true );
 						this->Private_SignalSelectionChanged( ESelectInfo::OnKeyPress );
@@ -320,7 +320,7 @@ public:
 					}
 				}
 				// Select all items
-				else if ( (!InKeyEvent.IsShiftDown() && !InKeyEvent.IsAltDown() && InKeyEvent.IsControlDown() && InKeyEvent.GetKey() == EKeys::A) && this->SelectionMode == ESelectionMode::Multi )
+				else if ( (!InKeyEvent.IsShiftDown() && !InKeyEvent.IsAltDown() && InKeyEvent.IsControlDown() && InKeyEvent.GetKey() == EKeys::A) && SelectionMode.Get() == ESelectionMode::Multi )
 				{
 					this->Private_ClearSelection();
 
@@ -344,7 +344,7 @@ public:
 	virtual FReply OnMouseButtonDown( const FGeometry& MyGeometry, const FPointerEvent& MouseEvent ) override
 	{
 		if ( bClearSelectionOnClick
-			&& SelectionMode != ESelectionMode::None
+			&& SelectionMode.Get() != ESelectionMode::None
 			&& MouseEvent.GetEffectingButton() == EKeys::LeftMouseButton
 			&& !MouseEvent.IsControlDown()
 			&& !MouseEvent.IsShiftDown()
@@ -367,7 +367,7 @@ public:
 	virtual FReply OnMouseButtonUp( const FGeometry& MyGeometry, const FPointerEvent& MouseEvent ) override
 	{
 		if ( bClearSelectionOnClick
-			&& SelectionMode != ESelectionMode::None
+			&& SelectionMode.Get() != ESelectionMode::None
 			&& MouseEvent.GetEffectingButton() == EKeys::RightMouseButton
 			&& !MouseEvent.IsControlDown()
 			&& !MouseEvent.IsShiftDown()
@@ -529,7 +529,7 @@ public:
 
 	virtual void Private_SetItemSelection( ItemType TheItem, bool bShouldBeSelected, bool bWasUserDirected = false ) override
 	{
-		if ( SelectionMode == ESelectionMode::None )
+		if ( SelectionMode.Get() == ESelectionMode::None )
 		{
 			return;
 		}
@@ -562,7 +562,7 @@ public:
 
 	virtual void Private_SelectRangeFromCurrentTo( ItemType InRangeSelectionEnd ) override
 	{
-		if ( SelectionMode == ESelectionMode::None )
+		if ( SelectionMode.Get() == ESelectionMode::None )
 		{
 			return;
 		}
@@ -595,7 +595,7 @@ public:
 
 	virtual void Private_SignalSelectionChanged(ESelectInfo::Type SelectInfo) override
 	{
-		if ( SelectionMode == ESelectionMode::None )
+		if ( SelectionMode.Get() == ESelectionMode::None )
 		{
 			return;
 		}
@@ -969,7 +969,7 @@ public:
 	 */
 	bool IsItemSelected( const ItemType& InItem ) const
 	{
-		if ( SelectionMode == ESelectionMode::None )
+		if ( SelectionMode.Get() == ESelectionMode::None )
 		{
 			return false;
 		}
@@ -986,7 +986,7 @@ public:
 	 */
 	void SetItemSelection( const ItemType& InItem, bool bSelected, ESelectInfo::Type SelectInfo = ESelectInfo::Direct )
 	{
-		if ( SelectionMode == ESelectionMode::None )
+		if ( SelectionMode.Get() == ESelectionMode::None )
 		{
 			return;
 		}
@@ -1000,7 +1000,7 @@ public:
 	 */
 	void ClearSelection()
 	{
-		if ( SelectionMode == ESelectionMode::None )
+		if ( SelectionMode.Get() == ESelectionMode::None )
 		{
 			return;
 		}
@@ -1319,12 +1319,14 @@ protected:
 	 */
 	virtual void KeyboardSelect(const ItemType& ItemToSelect, const FKeyEvent& InKeyEvent, bool bCausedByNavigation=false)
 	{
-		if ( SelectionMode != ESelectionMode::None )
+		const ESelectionMode::Type CurrentSelectionMode = SelectionMode.Get();
+
+		if ( CurrentSelectionMode != ESelectionMode::None )
 		{
 			// Must be set before signaling selection changes because sometimes new items will be selected that need to stomp this value
 			SelectorItem = ItemToSelect;
 
-			if ( SelectionMode == ESelectionMode::Multi && ( InKeyEvent.IsShiftDown() || InKeyEvent.IsControlDown() ) )
+			if ( CurrentSelectionMode == ESelectionMode::Multi && ( InKeyEvent.IsShiftDown() || InKeyEvent.IsControlDown() ) )
 			{
 				// Range select.
 				if ( InKeyEvent.IsShiftDown() )
