Merging a fix for OS X 10.11 for an old tooltip window briefly showing before we update the contents and position when spawning a new one.

diff --git a/Engine/Source/Runtime/Slate/Private/Framework/Application/SlateApplication.cpp b/Engine/Source/Runtime/Slate/Private/Framework/Application/SlateApplication.cpp
--- a/Engine/Source/Runtime/Slate/Private/Framework/Application/SlateApplication.cpp
+++ b/Engine/Source/Runtime/Slate/Private/Framework/Application/SlateApplication.cpp
@@ -2851,6 +2851,7 @@ void FSlateApplication::SpawnToolTip( const TSharedRef<IToolTip>& InToolTip, con
 			NewToolTipWindow->MoveWindowTo( DesiredToolTipLocation );
 		}
 
+		NewToolTipWindow->SetOpacity(0.0f);
 
 		// Show the window
 		NewToolTipWindow->ShowWindow();
