Fixed shadow variable warning.

diff --git a/Engine/Source/Runtime/Slate/Private/Widgets/Input/SSubMenuHandler.cpp b/Engine/Source/Runtime/Slate/Private/Widgets/Input/SSubMenuHandler.cpp
--- a/Engine/Source/Runtime/Slate/Private/Widgets/Input/SSubMenuHandler.cpp
+++ b/Engine/Source/Runtime/Slate/Private/Widgets/Input/SSubMenuHandler.cpp
@@ -62,7 +62,6 @@ void SSubMenuHandler::OnMouseEnter( const FGeometry& MyGeometry, const FPointerE
 
 	if (MenuAnchor.IsValid())
 	{
-		TSharedPtr< SMenuOwner > PinnedMenuOwnerWidget( MenuOwnerWidget.Pin() );
 		check( PinnedMenuOwnerWidget.IsValid() );
 
 		// Do we have a different pull-down menu open?
