Slate - The STableViewBase now properly handles touch events when under editor conditions and sets up the callback to handle the inertial scrolling.

diff --git a/Engine/Source/Runtime/Slate/Private/Widgets/Views/STableViewBase.cpp b/Engine/Source/Runtime/Slate/Private/Widgets/Views/STableViewBase.cpp
--- a/Engine/Source/Runtime/Slate/Private/Widgets/Views/STableViewBase.cpp
+++ b/Engine/Source/Runtime/Slate/Private/Widgets/Views/STableViewBase.cpp
@@ -533,6 +533,13 @@ FReply STableViewBase::OnTouchMoved( const FGeometry& MyGeometry, const FPointer
 
 		if (AmountScrolledWhileRightMouseDown > FSlateApplication::Get().GetDragTriggerDistance())
 		{
+			// Make sure the active timer is registered to update the inertial scroll
+			if ( !bIsScrollingActiveTimerRegistered )
+			{
+				bIsScrollingActiveTimerRegistered = true;
+				RegisterActiveTimer(0.f, FWidgetActiveTimerDelegate::CreateSP(this, &STableViewBase::UpdateInertialScroll));
+			}
+
 			const float AmountScrolled = this->ScrollBy( MyGeometry, -ScrollByAmount, EAllowOverscroll::Yes );
 
 			// The user has moved the list some amount; they are probably
