Slate - Fixing the widget reflector.

diff --git a/Engine/Source/Runtime/Slate/Private/Framework/Application/SlateApplication.cpp b/Engine/Source/Runtime/Slate/Private/Framework/Application/SlateApplication.cpp
--- a/Engine/Source/Runtime/Slate/Private/Framework/Application/SlateApplication.cpp
+++ b/Engine/Source/Runtime/Slate/Private/Framework/Application/SlateApplication.cpp
@@ -13,6 +13,8 @@
 #include "ToolboxModule.h"
 #include "TabCommands.h"
 
+#define SLATE_HAS_WIDGET_REFLECTOR !UE_BUILD_SHIPPING || PLATFORM_DESKTOP
+
 extern SLATECORE_API TOptional<FShortRect> GSlateScissorRect;
 
 class FEventRouter
@@ -5523,3 +5525,5 @@ void FSlateApplication::SetWidgetReflector(const TSharedRef<IWidgetReflector>& W
 
 	WidgetReflectorPtr = WidgetReflector;
 }
+
+#undef SLATE_HAS_WIDGET_REFLECTOR
diff --git a/Engine/Source/Runtime/Slate/Public/Slate.h b/Engine/Source/Runtime/Slate/Public/Slate.h
--- a/Engine/Source/Runtime/Slate/Public/Slate.h
+++ b/Engine/Source/Runtime/Slate/Public/Slate.h
@@ -6,4 +6,3 @@ EMIT_DEPRECATED_WARNING_MESSAGE("including Slate.h is deprecated. Please include
 #include "SlateBasics.h"
 #include "SlateExtras.h"
 
-#define SLATE_HAS_WIDGET_REFLECTOR !UE_BUILD_SHIPPING || PLATFORM_DESKTOP
