#slate Submenu arrows and check marks now properly invert when highlighted.

diff --git a/Engine/Source/Runtime/Slate/Private/Framework/MultiBox/SMenuEntryBlock.cpp b/Engine/Source/Runtime/Slate/Private/Framework/MultiBox/SMenuEntryBlock.cpp
--- a/Engine/Source/Runtime/Slate/Private/Framework/MultiBox/SMenuEntryBlock.cpp
+++ b/Engine/Source/Runtime/Slate/Private/Framework/MultiBox/SMenuEntryBlock.cpp
@@ -437,6 +437,7 @@ TSharedRef< SWidget > SMenuEntryBlock::BuildMenuEntryWidget( const FMenuEntryBui
 				// For check style menus, use an image instead of a CheckBox because it can't really be checked.
 				UserInterfaceType == EUserInterfaceActionType::Check
 					? StaticCastSharedRef<SWidget>(SNew(SImage)
+						.ColorAndOpacity(FSlateColor::UseForeground())
 						.Image(this, &SMenuEntryBlock::GetCheckBoxImageBrushFromStyle, &StyleSet->GetWidgetStyle<FCheckBoxStyle>(CheckBoxStyle)))
 					: StaticCastSharedRef<SWidget>(SNew(SCheckBox)
 						.ForegroundColor( CheckBoxForegroundColor )
@@ -719,6 +720,7 @@ TSharedRef< SWidget> SMenuEntryBlock::BuildSubMenuWidget( const FMenuEntryBuildP
 					.Padding(FMargin(7,0,0,0))
 					[
 						SNew( SImage )
+						.ColorAndOpacity(FSlateColor::UseForeground())
 						.Image( StyleSet->GetBrush( StyleName, ".SubMenuIndicator" ) )
 					]
 				]
