Adding an optimization to the TextLayout::RemoveLine function which prevents a dirty of the layout.

diff --git a/Engine/Source/Runtime/Slate/Private/Framework/Text/TextLayout.cpp b/Engine/Source/Runtime/Slate/Private/Framework/Text/TextLayout.cpp
--- a/Engine/Source/Runtime/Slate/Private/Framework/Text/TextLayout.cpp
+++ b/Engine/Source/Runtime/Slate/Private/Framework/Text/TextLayout.cpp
@@ -1419,7 +1419,49 @@ bool FTextLayout::RemoveLine(int32 LineIndex)
 
 	LineModels.RemoveAt(LineIndex);
 
-	DirtyFlags |= EDirtyState::Layout;
+	// If our layout is clean, then we can remove this line immediately (and efficiently)
+	// If our layout is dirty, then we might as well wait as the next UpdateLayout call will remove it
+	if (!(DirtyFlags & EDirtyState::Layout))
+	{
+		//Lots of room for additional optimization
+		float OffsetAdjustment = 0;
+
+		for (int32 ViewIndex = 0; ViewIndex < LineViews.Num(); ViewIndex++)
+		{
+			FLineView& LineView = LineViews[ViewIndex];
+
+			if (LineView.ModelIndex == LineIndex)
+			{
+				if (ViewIndex - 1 <= 0)
+				{
+					OffsetAdjustment += LineView.Offset.Y;
+				}
+				else
+				{
+					//Since the offsets are not relative to other lines, if we aren't removing the top line then
+					//we don't aggregate the offset from any previous removals as we'd be double counting.
+					OffsetAdjustment = (LineView.Offset.Y - LineViews[ViewIndex - 1].Offset.Y);
+				}
+
+				LineViews.RemoveAt(ViewIndex);
+				--ViewIndex;
+			}
+			else if (LineView.ModelIndex > LineIndex)
+			{
+				//We've removed a line model so update the LineView indices
+				--LineView.ModelIndex;
+				LineView.Offset.Y -= OffsetAdjustment;
+
+				for (const TSharedRef< ILayoutBlock >& Block : LineView.Blocks)
+				{
+					FVector2D BlockOffset = Block->GetLocationOffset();
+					BlockOffset.Y -= OffsetAdjustment;
+					Block->SetLocationOffset(BlockOffset);
+				}
+			}
+		}
+	}
+
 	return true;
 }
 
