Fixed issue where popups from graph nodes would intermittently flicker between two scales. Also fixed the issue where the popup was being created with dimensions according to the parent widget scale instead of its own scale.

SMenuAnchor::Tick was accessing deprecated members of FGeometry, and transforming the size of the CreateWindow popup by the accumulated transform, which was wrong, since it will always use the ApplicationScale.

diff --git a/Engine/Source/Runtime/Slate/Private/Widgets/Input/SMenuAnchor.cpp b/Engine/Source/Runtime/Slate/Private/Widgets/Input/SMenuAnchor.cpp
--- a/Engine/Source/Runtime/Slate/Private/Widgets/Input/SMenuAnchor.cpp
+++ b/Engine/Source/Runtime/Slate/Private/Widgets/Input/SMenuAnchor.cpp
@@ -110,8 +110,9 @@ void SMenuAnchor::Tick( const FGeometry& AllottedGeometry, const double InCurren
 		// Figure out where our attached pop-up window should be placed.
 		const FVector2D PopupContentDesiredSize = PopupWindow->GetContent()->GetDesiredSize();
 		FGeometry PopupGeometry = ComputeMenuPlacement( AllottedGeometry, PopupContentDesiredSize, Placement.Get() );
-		const FVector2D NewPosition = PopupGeometry.AbsolutePosition;
-		const FVector2D NewSize = PopupGeometry.GetDrawSize( );
+		const FVector2D NewPosition = PopupGeometry.LocalToAbsolute(FVector2D::ZeroVector);
+		// For the CreateWindow case, don't transform the size; it will always use the ApplicationScale
+		const FVector2D NewSize = PopupGeometry.GetLocalSize();
 
 		const FSlateRect NewShape = FSlateRect( NewPosition.X, NewPosition.Y, NewPosition.X + NewSize.X, NewPosition.Y + NewSize.Y );
 
@@ -309,7 +310,7 @@ void SMenuAnchor::SetIsOpen( bool InIsOpen, const bool bFocusMenu )
 				if (MyWidgetPath.IsValid())
 				{
 					const FGeometry& MyGeometry = MyWidgetPath.Widgets.Last().Geometry;
-					const float LayoutScaleMultiplier = MyGeometry.GetAccumulatedLayoutTransform().GetScale();
+					const float LayoutScaleMultiplier =  MyGeometry.GetAccumulatedLayoutTransform().GetScale();
 
 					SlatePrepass(LayoutScaleMultiplier);
 
