Slate/UMG - Continuing to work on the Invalidation Panel, correctly offsetting volatile widgets to render from the true last hit testable index that their potentinally cached parent used.  Also introducing a way to cache transforms relative to the parent so that widgets that fly around the screen but otherwise don't change shape...etc can be cached.

diff --git a/Engine/Source/Runtime/AppFramework/Private/Widgets/Testing/STestSuite.cpp b/Engine/Source/Runtime/AppFramework/Private/Widgets/Testing/STestSuite.cpp
--- a/Engine/Source/Runtime/AppFramework/Private/Widgets/Testing/STestSuite.cpp
+++ b/Engine/Source/Runtime/AppFramework/Private/Widgets/Testing/STestSuite.cpp
@@ -3919,6 +3919,23 @@ class SInvalidationTest : public SCompoundWidget
 					]
 				]
 			]
+
+			+ SVerticalBox::Slot()
+			.AutoHeight()
+			[
+				SNew(SInvalidationPanel)
+				.CacheRelativeTransforms(true)
+				[
+					SNew(SVerticalBox)
+
+					+ SVerticalBox::Slot()
+					.AutoHeight()
+					[
+						SNew(STextBlock)
+						.Text(LOCTEXT("RelativeCacheMessage", "I cache relative transforms.  Moving the window has no effect on me, woo!"))
+					]
+				]
+			]
 		];
 	}
 
diff --git a/Engine/Source/Runtime/Slate/Private/Widgets/SInvalidationPanel.cpp b/Engine/Source/Runtime/Slate/Private/Widgets/SInvalidationPanel.cpp
--- a/Engine/Source/Runtime/Slate/Private/Widgets/SInvalidationPanel.cpp
+++ b/Engine/Source/Runtime/Slate/Private/Widgets/SInvalidationPanel.cpp
@@ -43,6 +43,10 @@ void SInvalidationPanel::Construct( const FArguments& InArgs )
 	bCanCache = true;
 	RootCacheNode = nullptr;
 	LastUsedCachedNodeIndex = 0;
+	LastLayerId = 0;
+	LastHitTestIndex = 0;
+
+	bCacheRelativeTransforms = InArgs._CacheRelativeTransforms;
 }
 
 SInvalidationPanel::~SInvalidationPanel()
@@ -73,11 +77,30 @@ void SInvalidationPanel::Tick( const FGeometry& AllottedGeometry, const double I
 {
 	if ( GetCanCache() )
 	{
+		const bool bWasCachingNeeded = bNeedsCaching;
+
 		if ( bNeedsCaching == false )
 		{
-			// If the container we're in has changed in any way we need to invalidate for sure.
-			if ( AllottedGeometry.GetAccumulatedLayoutTransform() != LastAllottedGeometry.GetAccumulatedLayoutTransform() ||
-				AllottedGeometry.GetAccumulatedRenderTransform() != LastAllottedGeometry.GetAccumulatedRenderTransform() )
+			if ( bCacheRelativeTransforms )
+			{
+				// If the container we're in has changed in either scale or the rotation matrix has changed, 
+				if ( AllottedGeometry.GetAccumulatedLayoutTransform().GetScale() != LastAllottedGeometry.GetAccumulatedLayoutTransform().GetScale() ||
+					 AllottedGeometry.GetAccumulatedRenderTransform().GetMatrix() != LastAllottedGeometry.GetAccumulatedRenderTransform().GetMatrix() )
+				{
+					InvalidateCache();
+				}
+			}
+			else
+			{
+				// If the container we're in has changed in any way we need to invalidate for sure.
+				if ( AllottedGeometry.GetAccumulatedLayoutTransform() != LastAllottedGeometry.GetAccumulatedLayoutTransform() ||
+					AllottedGeometry.GetAccumulatedRenderTransform() != LastAllottedGeometry.GetAccumulatedRenderTransform() )
+				{
+					InvalidateCache();
+				}
+			}
+
+			if ( AllottedGeometry.GetLocalSize() != LastAllottedGeometry.GetLocalSize() )
 			{
 				InvalidateCache();
 			}
@@ -85,9 +108,11 @@ void SInvalidationPanel::Tick( const FGeometry& AllottedGeometry, const double I
 
 		LastAllottedGeometry = AllottedGeometry;
 
+		// TODO We may be double pre-passing here, if the invalidation happened at the end of last frame,
+		// we'll have already done one pre-pass before getting here.
 		if ( bNeedsCaching )
 		{
-			SlatePrepass(LastAllottedGeometry.Scale);
+			SlatePrepass(AllottedGeometry.Scale);
 			CachePrepass(SharedThis(this));
 		}
 	}
@@ -137,15 +162,14 @@ FCachedWidgetNode* SInvalidationPanel::CreateCacheNode() const
 
 int32 SInvalidationPanel::OnPaint( const FPaintArgs& Args, const FGeometry& AllottedGeometry, const FSlateRect& MyClippingRect, FSlateWindowElementList& OutDrawElements, int32 LayerId, const FWidgetStyle& InWidgetStyle, bool bParentEnabled ) const
 {
-	//TODO If the LayerId changes, or other properties, we should probably invalidate, but
-	// we can't invalidate here, because we haven't done a pre-pass yet, need to do it next frame.
-
 	if ( GetCanCache() )
 	{
 		const bool bWasCachingNeeded = bNeedsCaching;
 
 		if ( bNeedsCaching )
 		{
+			SInvalidationPanel* MutableThis = const_cast<SInvalidationPanel*>( this );
+
 			// Always set the caching flag to false first, during the paint / tick pass we may change something
 			// to volatile and need to re-cache.
 			bNeedsCaching = false;
@@ -165,11 +189,11 @@ int32 SInvalidationPanel::OnPaint( const FPaintArgs& Args, const FGeometry& Allo
 			LastUsedCachedNodeIndex = 0;
 
 			RootCacheNode = CreateCacheNode();
-			RootCacheNode->Initialize(Args, SharedThis(const_cast<SInvalidationPanel*>( this )), AllottedGeometry, MyClippingRect);
+			RootCacheNode->Initialize(Args, SharedThis(MutableThis), AllottedGeometry, MyClippingRect);
 
 			//TODO: When SWidget::Paint is called don't drag self if volatile, and we're doing a cache pass.
 			CachedMaxChildLayer = SCompoundWidget::OnPaint(
-				Args.EnableCaching(SharedThis(const_cast<SInvalidationPanel*>( this )), RootCacheNode, true, false),
+				Args.EnableCaching(SharedThis(MutableThis), RootCacheNode, true, false),
 				AllottedGeometry,
 				MyClippingRect,
 				*CachedWindowElements.Get(),
@@ -177,12 +201,45 @@ int32 SInvalidationPanel::OnPaint( const FPaintArgs& Args, const FGeometry& Allo
 				InWidgetStyle,
 				bParentEnabled);
 
+			if ( bCacheRelativeTransforms )
+			{
+				CachedAbsolutePosition = AllottedGeometry.Position;
+			}
+
+			LastLayerId = LayerId;
+			LastHitTestIndex = Args.GetLastHitTestIndex();
+
 			bIsInvalidating = false;
 		}
 
-		RootCacheNode->RecordHittestGeometry(Args.GetGrid(), Args.GetLastHitTestIndex());
+		// The hit test grid is actually populated during the initial cache phase, so don't bother
+		// recording the hit test geometry on the same frame that we regenerate the cache.
+		if ( bWasCachingNeeded == false )
+		{
+			RootCacheNode->RecordHittestGeometry(Args.GetGrid(), Args.GetLastHitTestIndex());
+		}
 
-		OutDrawElements.AppendDrawElements(CachedWindowElements->GetDrawElements());
+		if ( bCacheRelativeTransforms )
+		{
+			FVector2D DeltaPosition = AllottedGeometry.Position - CachedAbsolutePosition;
+
+			const TArray<FSlateDrawElement>& CachedElements = CachedWindowElements->GetDrawElements();
+			const int32 CachedElementCount = CachedElements.Num();
+			for ( int32 Index = 0; Index < CachedElementCount; Index++ )
+			{
+				const FSlateDrawElement& LocalElement = CachedElements[Index];
+				FSlateDrawElement AbsElement = LocalElement;
+
+				AbsElement.SetPosition(LocalElement.GetPosition() + DeltaPosition);
+				AbsElement.SetClippingRect(LocalElement.GetClippingRect().OffsetBy(DeltaPosition));
+
+				OutDrawElements.AddItem(AbsElement);
+			}
+		}
+		else
+		{
+			OutDrawElements.AppendDrawElements(CachedWindowElements->GetDrawElements());
+		}
 
 		int32 OutMaxChildLayer = CachedMaxChildLayer;
 
@@ -212,7 +269,7 @@ int32 SInvalidationPanel::OnPaint( const FPaintArgs& Args, const FGeometry& Allo
 					MyClippingRect,
 					ESlateDrawEffect::None,
 					DebugTint
-					);
+				);
 			}
 
 			// Draw a yellow outline around any volatile elements.
@@ -278,11 +335,6 @@ void SInvalidationPanel::SetContent(const TSharedRef< SWidget >& InContent)
 	];
 }
 
-void SInvalidationPanel::InvalidateCache()
-{
-	bNeedsCaching = true;
-}
-
 bool SInvalidationPanel::ComputeVolatility() const
 {
 	return true;
diff --git a/Engine/Source/Runtime/Slate/Public/Widgets/SInvalidationPanel.h b/Engine/Source/Runtime/Slate/Public/Widgets/SInvalidationPanel.h
--- a/Engine/Source/Runtime/Slate/Public/Widgets/SInvalidationPanel.h
+++ b/Engine/Source/Runtime/Slate/Public/Widgets/SInvalidationPanel.h
@@ -10,8 +10,10 @@ public:
 	SLATE_BEGIN_ARGS( SInvalidationPanel )
 	{
 		_Visibility = EVisibility::SelfHitTestInvisible;
+		_CacheRelativeTransforms = false;
 	}
 		SLATE_DEFAULT_SLOT( FArguments, Content )
+		SLATE_ARGUMENT( bool, CacheRelativeTransforms )
 	SLATE_END_ARGS()
 
 	void Construct( const FArguments& InArgs );
@@ -20,7 +22,7 @@ public:
 	bool GetCanCache() const;
 	void SetCanCache(bool InCanCache);
 
-	void InvalidateCache();
+	FORCEINLINE void InvalidateCache() { bNeedsCaching = true; }
 
 	// ILayoutCache overrides
 	virtual void InvalidateWidget(SWidget* InvalidateWidget) override;
@@ -55,11 +57,16 @@ private:
 
 	mutable FCachedWidgetNode* RootCacheNode;
 	mutable TSharedPtr< FSlateWindowElementList > CachedWindowElements;
+	mutable FVector2D CachedAbsolutePosition;
 	mutable TArray< FCachedWidgetNode* > NodePool;
 	mutable int32 LastUsedCachedNodeIndex;
+	mutable int32 LastLayerId;
+	mutable int32 LastHitTestIndex;
 
 	mutable int32 CachedMaxChildLayer;
 	mutable bool bNeedsCaching;
 	mutable bool bIsInvalidating;
 	bool bCanCache;
+
+	bool bCacheRelativeTransforms;
 };
diff --git a/Engine/Source/Runtime/SlateCore/Private/Layout/WidgetCaching.cpp b/Engine/Source/Runtime/SlateCore/Private/Layout/WidgetCaching.cpp
--- a/Engine/Source/Runtime/SlateCore/Private/Layout/WidgetCaching.cpp
+++ b/Engine/Source/Runtime/SlateCore/Private/Layout/WidgetCaching.cpp
@@ -14,6 +14,7 @@ void FCachedWidgetNode::Initialize(const FPaintArgs& Args, TSharedRef<SWidget> I
 	ClippingRect = InClippingRect;
 	WindowOffset = Args.GetWindowToDesktopTransform();
 	RecordedVisibility = Args.GetLastRecordedVisibility();
+	LastRecordedHittestIndex = 0;
 
 	if ( RecordedVisibility.AreChildrenHitTestVisible() )
 	{
@@ -27,12 +28,16 @@ void FCachedWidgetNode::RecordHittestGeometry(FHittestGrid& Grid, int32 LastHitt
 {
 	if ( RecordedVisibility.AreChildrenHitTestVisible() )
 	{
-		const int32 RecordedHittestIndex = Grid.InsertWidget(LastHittestIndex, RecordedVisibility, FArrangedWidget(Widget.Pin()->AsShared(), Geometry), WindowOffset, ClippingRect);
-
-		int32 ChildCount = Children.Num();
-		for ( int32 i = 0; i < ChildCount; i++ )
+		TSharedPtr<SWidget> SafeWidget = Widget.Pin();
+		if ( SafeWidget.IsValid() )
 		{
-			Children[i]->RecordHittestGeometry(Grid, RecordedHittestIndex);
+			LastRecordedHittestIndex = Grid.InsertWidget(LastHittestIndex, RecordedVisibility, FArrangedWidget(SafeWidget.ToSharedRef(), Geometry), WindowOffset, ClippingRect);
+
+			const int32 ChildCount = Children.Num();
+			for ( int32 i = 0; i < ChildCount; i++ )
+			{
+				Children[i]->RecordHittestGeometry(Grid, LastRecordedHittestIndex);
+			}
 		}
 	}
 }
diff --git a/Engine/Source/Runtime/SlateCore/Private/Rendering/DrawElements.cpp b/Engine/Source/Runtime/SlateCore/Private/Rendering/DrawElements.cpp
--- a/Engine/Source/Runtime/SlateCore/Private/Rendering/DrawElements.cpp
+++ b/Engine/Source/Runtime/SlateCore/Private/Rendering/DrawElements.cpp
@@ -363,7 +363,7 @@ int32 FSlateWindowElementList::PaintDeferred(int32 LayerId)
 
 FSlateWindowElementList::FVolatilePaint::FVolatilePaint(const TSharedRef<const SWidget>& InWidgetToPaint, const FPaintArgs& InArgs, const FGeometry InAllottedGeometry, const FSlateRect InMyClippingRect, int32 InLayerId, const FWidgetStyle& InWidgetStyle, bool InParentEnabled)
 	: WidgetToPaintPtr(InWidgetToPaint)
-	, Args(InArgs.EnableCaching(InArgs.GetLayoutCache(), nullptr, false, true))
+	, Args(InArgs.EnableCaching(InArgs.GetLayoutCache(), InArgs.GetParentCacheNode(), false, true))
 	, AllottedGeometry(InAllottedGeometry)
 	, MyClippingRect(InMyClippingRect)
 	, LayerId(InLayerId)
diff --git a/Engine/Source/Runtime/SlateCore/Private/Types/PaintArgs.cpp b/Engine/Source/Runtime/SlateCore/Private/Types/PaintArgs.cpp
--- a/Engine/Source/Runtime/SlateCore/Private/Types/PaintArgs.cpp
+++ b/Engine/Source/Runtime/SlateCore/Private/Types/PaintArgs.cpp
@@ -13,6 +13,7 @@ FPaintArgs::FPaintArgs( const TSharedRef<SWidget>& Parent, FHittestGrid& InHitte
 , CurrentTime(InCurrentTime)
 , DeltaTime(InDeltaTime)
 , bIsCaching(false)
+, bIsVolatilityPass(false)
 , ParentCacheNode(nullptr)
 {
 }
@@ -50,18 +51,24 @@ FPaintArgs FPaintArgs::RecordHittestGeometry(const SWidget* Widget, const FGeome
 		if ( bIsCaching )
 		{
 			FCachedWidgetNode* CacheNode = LayoutCache->CreateCacheNode();
-			CacheNode->Initialize(UpdatedArgs, const_cast<SWidget*>( Widget )->AsShared(), WidgetGeometry, InClippingRect);
+			CacheNode->Initialize(*this, const_cast<SWidget*>( Widget )->AsShared(), WidgetGeometry, InClippingRect);
 			UpdatedArgs.ParentCacheNode->Children.Add(CacheNode);
 
 			UpdatedArgs.ParentCacheNode = CacheNode;
 		}
-		else
+
+		int32 RealLastHitTestIndex = LastHittestIndex;
+		if ( bIsVolatilityPass && ParentCacheNode )
 		{
-			const EVisibility RecordedVisibility = Widget->GetVisibility();
-			const int32 RecordedHittestIndex = Grid.InsertWidget(LastHittestIndex, RecordedVisibility, FArrangedWidget(const_cast<SWidget*>( Widget )->AsShared(), WidgetGeometry), WindowOffset, InClippingRect);
-			UpdatedArgs.LastHittestIndex = RecordedHittestIndex;
-			UpdatedArgs.LastRecordedVisibility = RecordedVisibility;
+			RealLastHitTestIndex = ParentCacheNode->LastRecordedHittestIndex;
+			UpdatedArgs.ParentCacheNode = nullptr;
 		}
+
+		// When rendering volatile widgets, their parent widgets who have been cached 
+		const EVisibility RecordedVisibility = Widget->GetVisibility();
+		const int32 RecordedHittestIndex = Grid.InsertWidget(RealLastHitTestIndex, RecordedVisibility, FArrangedWidget(const_cast<SWidget*>( Widget )->AsShared(), WidgetGeometry), WindowOffset, InClippingRect);
+		UpdatedArgs.LastHittestIndex = RecordedHittestIndex;
+		UpdatedArgs.LastRecordedVisibility = RecordedVisibility;
 	}
 	else
 	{
diff --git a/Engine/Source/Runtime/SlateCore/Public/Layout/WidgetCaching.h b/Engine/Source/Runtime/SlateCore/Public/Layout/WidgetCaching.h
--- a/Engine/Source/Runtime/SlateCore/Public/Layout/WidgetCaching.h
+++ b/Engine/Source/Runtime/SlateCore/Public/Layout/WidgetCaching.h
@@ -24,6 +24,7 @@ public:
 	FVector2D WindowOffset;
 
 	EVisibility RecordedVisibility;
+	int32 LastRecordedHittestIndex;
 
 	// TODO Do we want to cache information about the elements drawn by each widget?
 	//TArray< FSlateWindowElementList > Elements;
diff --git a/Engine/Source/Runtime/SlateCore/Public/Rendering/DrawElements.h b/Engine/Source/Runtime/SlateCore/Public/Rendering/DrawElements.h
--- a/Engine/Source/Runtime/SlateCore/Public/Rendering/DrawElements.h
+++ b/Engine/Source/Runtime/SlateCore/Public/Rendering/DrawElements.h
@@ -317,10 +317,12 @@ public:
 	EElementType GetElementType() const { return ElementType; }
 	uint32 GetLayer() const { return Layer; }
 	const FSlateRenderTransform& GetRenderTransform() const { return RenderTransform; }
-	const FVector2D& GetPosition() const { return Position; }
+	SLATECORE_API const FVector2D& GetPosition() const { return Position; }
+	SLATECORE_API void SetPosition(const FVector2D& InPosition) { Position = Position; }
 	const FVector2D& GetLocalSize() const { return LocalSize; }
 	float GetScale() const { return Scale; }
-	const FSlateRect& GetClippingRect() const { return ClippingRect; }
+	SLATECORE_API const FSlateRect& GetClippingRect() const { return ClippingRect; }
+	SLATECORE_API void SetClippingRect(const FSlateRect& InClippingRect) { ClippingRect = InClippingRect; }
 	const FSlateDataPayload& GetDataPayload() const { return DataPayload; }
 	uint32 GetDrawEffects() const { return DrawEffects; }
 	const TOptional<FShortRect>& GetScissorRect() const { return ScissorRect; }
diff --git a/Engine/Source/Runtime/SlateCore/Public/Types/PaintArgs.h b/Engine/Source/Runtime/SlateCore/Public/Types/PaintArgs.h
--- a/Engine/Source/Runtime/SlateCore/Public/Types/PaintArgs.h
+++ b/Engine/Source/Runtime/SlateCore/Public/Types/PaintArgs.h
@@ -36,6 +36,7 @@ public:
 	bool IsCaching() const { return bIsCaching; }
 	bool IsVolatilityPass() const { return bIsVolatilityPass; }
 	TSharedPtr<ILayoutCache> GetLayoutCache() const { return LayoutCache; }
+	FCachedWidgetNode* GetParentCacheNode() const { return ParentCacheNode; }
 
 private:
 	const TSharedRef<SWidget>& ParentPtr;
diff --git a/Engine/Source/Runtime/UMG/Private/Components/InvalidationBox.cpp b/Engine/Source/Runtime/UMG/Private/Components/InvalidationBox.cpp
--- a/Engine/Source/Runtime/UMG/Private/Components/InvalidationBox.cpp
+++ b/Engine/Source/Runtime/UMG/Private/Components/InvalidationBox.cpp
@@ -12,7 +12,7 @@
 UInvalidationBox::UInvalidationBox(const FObjectInitializer& ObjectInitializer)
 	: Super(ObjectInitializer)
 {
-	Visibility = ESlateVisibility::HitTestInvisible;
+	Visibility = ESlateVisibility::SelfHitTestInvisible;
 }
 
 void UInvalidationBox::ReleaseSlateResources(bool bReleaseChildren)
@@ -24,7 +24,9 @@ void UInvalidationBox::ReleaseSlateResources(bool bReleaseChildren)
 
 TSharedRef<SWidget> UInvalidationBox::RebuildWidget()
 {
-	MyInvalidationPanel = SNew(SInvalidationPanel);
+	MyInvalidationPanel = 
+		SNew(SInvalidationPanel)
+		.CacheRelativeTransforms(CacheRelativeTransforms);
 
 	if ( IsDesignTime() )
 	{
diff --git a/Engine/Source/Runtime/UMG/Public/Components/InvalidationBox.h b/Engine/Source/Runtime/UMG/Public/Components/InvalidationBox.h
--- a/Engine/Source/Runtime/UMG/Public/Components/InvalidationBox.h
+++ b/Engine/Source/Runtime/UMG/Public/Components/InvalidationBox.h
@@ -7,7 +7,7 @@
 /**
  * Invalidate
  * ● Single Child
- * ● Popup
+ * ● Caching / Performance
  */
 UCLASS()
 class UMG_API UInvalidationBox : public UContentWidget
@@ -16,6 +16,14 @@ class UMG_API UInvalidationBox : public UContentWidget
 
 public:
 
+	/**
+	 * Caches the locations for child draw elements relative to the invalidation box,
+	 * this adds extra overhead to drawing them every frame.  However, in cases where
+	 * the position of the invalidation boxes changes every frame this can be a big savings.
+	 */
+	UPROPERTY(EditAnywhere, BlueprintReadOnly, Category="Caching")
+	bool CacheRelativeTransforms;
+
 	/**
 	 * 
 	 */
