Menus from menu anchors no longer auto-dismiss when a parent (in the menu stack) receives focus.
[UE-18364] - 4.8 Regression: Sub-menus close when their menu button is clicked

diff --git a/Engine/Source/Runtime/Slate/Private/Framework/Application/Menu.cpp b/Engine/Source/Runtime/Slate/Private/Framework/Application/Menu.cpp
--- a/Engine/Source/Runtime/Slate/Private/Framework/Application/Menu.cpp
+++ b/Engine/Source/Runtime/Slate/Private/Framework/Application/Menu.cpp
@@ -3,14 +3,15 @@
 #include "SlatePrivatePCH.h"
 #include "Menu.h"
 
-FMenuBase::FMenuBase(TSharedRef<SWidget> InContent)
+FMenuBase::FMenuBase(TSharedRef<SWidget> InContent, const bool bIsCollapsedByParent)
 	: Content(InContent)
 	, bDismissing(false)
+	, bIsCollapsedByParent(bIsCollapsedByParent)
 {
 }
 
-FMenuInWindow::FMenuInWindow(TSharedRef<SWindow> InWindow, TSharedRef<SWidget> InContent)
-	: FMenuBase(InContent)
+FMenuInWindow::FMenuInWindow(TSharedRef<SWindow> InWindow, TSharedRef<SWidget> InContent, const bool bIsCollapsedByParent)
+	: FMenuBase(InContent, bIsCollapsedByParent)
 	, Window(InWindow)
 {
 }
@@ -39,8 +40,8 @@ void FMenuInWindow::Dismiss()
 }
 
 
-FMenuInPopup::FMenuInPopup(TSharedRef<SWidget> InContent)
-	: FMenuBase(InContent)
+FMenuInPopup::FMenuInPopup(TSharedRef<SWidget> InContent, const bool bIsCollapsedByParent)
+	: FMenuBase(InContent, bIsCollapsedByParent)
 {
 }
 
@@ -60,8 +61,8 @@ void FMenuInPopup::Dismiss()
 }
 
 
-FMenuInHostWidget::FMenuInHostWidget(TSharedRef<IMenuHost> InHost, const TSharedRef<SWidget>& InContent)
-	: FMenuBase(InContent)
+FMenuInHostWidget::FMenuInHostWidget(TSharedRef<IMenuHost> InHost, const TSharedRef<SWidget>& InContent, const bool bIsCollapsedByParent)
+	: FMenuBase(InContent, bIsCollapsedByParent)
 	, MenuHost(InHost)
 {
 }
diff --git a/Engine/Source/Runtime/Slate/Private/Framework/Application/Menu.h b/Engine/Source/Runtime/Slate/Private/Framework/Application/Menu.h
--- a/Engine/Source/Runtime/Slate/Private/Framework/Application/Menu.h
+++ b/Engine/Source/Runtime/Slate/Private/Framework/Application/Menu.h
@@ -12,13 +12,15 @@ class FMenuBase : public IMenu, public TSharedFromThis<FMenuBase>
 public:
 	virtual FOnMenuDismissed& GetOnMenuDismissed() override { return OnMenuDismissed; }
 	virtual TSharedPtr<SWidget> GetContent() const override { return Content; }
+	bool IsCollapsedByParent() const { return bIsCollapsedByParent; }
 
 protected:
-	FMenuBase(TSharedRef<SWidget> InContent);
+	FMenuBase(TSharedRef<SWidget> InContent, const bool bIsCollapsedByParent);
 
 	FOnMenuDismissed OnMenuDismissed;
 	TSharedRef<SWidget> Content;
 	bool bDismissing;
+	bool bIsCollapsedByParent;
 };
 
 /**
@@ -27,7 +29,7 @@ protected:
 class FMenuInWindow : public FMenuBase
 {
 public:
-	FMenuInWindow(TSharedRef<SWindow> InWindow, TSharedRef<SWidget> InContent);
+	FMenuInWindow(TSharedRef<SWindow> InWindow, TSharedRef<SWidget> InContent, const bool bIsCollapsedByParent);
 	virtual ~FMenuInWindow() {}
 
 	virtual EPopupMethod GetPopupMethod() const override { return EPopupMethod::CreateNewWindow; }
@@ -45,7 +47,7 @@ private:
 class FMenuInPopup : public FMenuBase
 {
 public:
-	FMenuInPopup(TSharedRef<SWidget> InContent);
+	FMenuInPopup(TSharedRef<SWidget> InContent, const bool bIsCollapsedByParent);
 	virtual ~FMenuInPopup() {}
 
 	virtual EPopupMethod GetPopupMethod() const { return EPopupMethod::UseCurrentWindow; }
@@ -60,7 +62,7 @@ public:
 class FMenuInHostWidget : public FMenuBase
 {
 public:
-	FMenuInHostWidget(TSharedRef<IMenuHost> InHost, const TSharedRef<SWidget>& InContent);
+	FMenuInHostWidget(TSharedRef<IMenuHost> InHost, const TSharedRef<SWidget>& InContent, const bool bIsCollapsedByParent);
 	virtual ~FMenuInHostWidget() {}
 
 	virtual EPopupMethod GetPopupMethod() const { return EPopupMethod::UseCurrentWindow; }
diff --git a/Engine/Source/Runtime/Slate/Private/Framework/Application/MenuStack.cpp b/Engine/Source/Runtime/Slate/Private/Framework/Application/MenuStack.cpp
--- a/Engine/Source/Runtime/Slate/Private/Framework/Application/MenuStack.cpp
+++ b/Engine/Source/Runtime/Slate/Private/Framework/Application/MenuStack.cpp
@@ -177,7 +177,7 @@ TSharedRef<SWindow> FMenuStack::PushMenu( const TSharedRef<SWindow>& ParentWindo
 	return NewMenuWindow;
 }
 
-TSharedRef<IMenu> FMenuStack::Push(const FWidgetPath& InOwnerPath, const TSharedRef<SWidget>& InContent, const FVector2D& SummonLocation, const FPopupTransitionEffect& TransitionEffect, const bool bFocusImmediately, const FVector2D& SummonLocationSize, TOptional<EPopupMethod> InMethod)
+TSharedRef<IMenu> FMenuStack::Push(const FWidgetPath& InOwnerPath, const TSharedRef<SWidget>& InContent, const FVector2D& SummonLocation, const FPopupTransitionEffect& TransitionEffect, const bool bFocusImmediately, const FVector2D& SummonLocationSize, TOptional<EPopupMethod> InMethod, const bool bIsCollapsedByParent)
 {
 	FSlateRect Anchor(SummonLocation, SummonLocation + SummonLocationSize);
 	TSharedPtr<IMenu> ParentMenu;
@@ -200,20 +200,20 @@ TSharedRef<IMenu> FMenuStack::Push(const FWidgetPath& InOwnerPath, const TShared
 		SetHostWindow(InOwnerPath.GetWindow());
 	}
 
-	return PushInternal(ParentMenu, InContent, Anchor, TransitionEffect, bFocusImmediately);
+	return PushInternal(ParentMenu, InContent, Anchor, TransitionEffect, bFocusImmediately, bIsCollapsedByParent);
 }
 
-TSharedRef<IMenu> FMenuStack::Push(const TSharedPtr<IMenu>& InParentMenu, const TSharedRef<SWidget>& InContent, const FVector2D& SummonLocation, const FPopupTransitionEffect& TransitionEffect, const bool bFocusImmediately, const FVector2D& SummonLocationSize)
+TSharedRef<IMenu> FMenuStack::Push(const TSharedPtr<IMenu>& InParentMenu, const TSharedRef<SWidget>& InContent, const FVector2D& SummonLocation, const FPopupTransitionEffect& TransitionEffect, const bool bFocusImmediately, const FVector2D& SummonLocationSize, const bool bIsCollapsedByParent)
 {
 	check(Stack.Contains(InParentMenu));
 	check(HostWindow.IsValid());
 
 	FSlateRect Anchor(SummonLocation, SummonLocation + SummonLocationSize);
 
-	return PushInternal(InParentMenu, InContent, Anchor, TransitionEffect, bFocusImmediately);
+	return PushInternal(InParentMenu, InContent, Anchor, TransitionEffect, bFocusImmediately, bIsCollapsedByParent);
 }
 
-TSharedRef<IMenu> FMenuStack::PushHosted(const FWidgetPath& InOwnerPath, const TSharedRef<IMenuHost>& InMenuHost, const TSharedRef<SWidget>& InContent, TSharedPtr<SWidget>& OutWrappedContent, const FPopupTransitionEffect& TransitionEffect)
+TSharedRef<IMenu> FMenuStack::PushHosted(const FWidgetPath& InOwnerPath, const TSharedRef<IMenuHost>& InMenuHost, const TSharedRef<SWidget>& InContent, TSharedPtr<SWidget>& OutWrappedContent, const FPopupTransitionEffect& TransitionEffect, const bool bIsCollapsedByParent)
 {
 	TSharedPtr<IMenu> ParentMenu;
 
@@ -238,13 +238,13 @@ TSharedRef<IMenu> FMenuStack::PushHosted(const FWidgetPath& InOwnerPath, const T
 	return PushHosted(ParentMenu, InMenuHost, InContent, OutWrappedContent, TransitionEffect);
 }
 
-TSharedRef<IMenu> FMenuStack::PushHosted(const TSharedPtr<IMenu>& InParentMenu, const TSharedRef<IMenuHost>& InMenuHost, const TSharedRef<SWidget>& InContent, TSharedPtr<SWidget>& OutWrappedContent, const FPopupTransitionEffect& TransitionEffect)
+TSharedRef<IMenu> FMenuStack::PushHosted(const TSharedPtr<IMenu>& InParentMenu, const TSharedRef<IMenuHost>& InMenuHost, const TSharedRef<SWidget>& InContent, TSharedPtr<SWidget>& OutWrappedContent, const FPopupTransitionEffect& TransitionEffect, const bool bIsCollapsedByParent)
 {
 	check(HostWindow.IsValid());
 
 	// Create a FMenuInHostWidget
 	TSharedRef<SWidget> WrappedContent = WrapContent(InContent);
-	TSharedRef<FMenuInHostWidget> OutMenu = MakeShareable(new FMenuInHostWidget(InMenuHost, WrappedContent));
+	TSharedRef<FMenuInHostWidget> OutMenu = MakeShareable(new FMenuInHostWidget(InMenuHost, WrappedContent, bIsCollapsedByParent));
 	PendingNewMenu = OutMenu;
 
 	// Set the returned content - this must be drawn by the hosting widget until the menu gets dismissed and calls IMenuHost::OnMenuDismissed on its host
@@ -260,13 +260,14 @@ TSharedRef<IMenu> FMenuStack::PushHosted(const TSharedPtr<IMenu>& InParentMenu,
 	return OutMenu;
 }
 
-TSharedRef<IMenu> FMenuStack::PushInternal(const TSharedPtr<IMenu>& InParentMenu, const TSharedRef<SWidget>& InContent, FSlateRect Anchor, const FPopupTransitionEffect& TransitionEffect, const bool bFocusImmediately)
+TSharedRef<IMenu> FMenuStack::PushInternal(const TSharedPtr<IMenu>& InParentMenu, const TSharedRef<SWidget>& InContent, FSlateRect Anchor, const FPopupTransitionEffect& TransitionEffect, const bool bFocusImmediately, const bool bIsCollapsedByParent)
 {
 	FPrePushArgs PrePushArgs;
 	PrePushArgs.Content = InContent;
 	PrePushArgs.Anchor = Anchor;
 	PrePushArgs.TransitionEffect = TransitionEffect;
 	PrePushArgs.bFocusImmediately = bFocusImmediately;
+	PrePushArgs.bIsCollapsedByParent = bIsCollapsedByParent;
 
 	// Pre-push stage
 	//   Determines correct layout
@@ -291,6 +292,7 @@ TSharedRef<IMenu> FMenuStack::PushInternal(const TSharedPtr<IMenu>& InParentMenu
 FMenuStack::FPrePushResults FMenuStack::PrePush(const FPrePushArgs& InArgs)
 {
 	FPrePushResults OutResults;
+	OutResults.bIsCollapsedByParent = InArgs.bIsCollapsedByParent;
 	OutResults.bFocusImmediately = InArgs.bFocusImmediately;
 	if (InArgs.bFocusImmediately)
 	{
@@ -426,7 +428,7 @@ TSharedRef<FMenuBase> FMenuStack::PushNewWindow(TSharedPtr<IMenu> InParentMenu,
 		NewMenuWindow->SetWidgetToFocusOnActivate(InPrePushResults.WidgetToFocus);
 	}
 
-	TSharedRef<FMenuInWindow> Menu = MakeShareable(new FMenuInWindow(NewMenuWindow, InPrePushResults.WrappedContent.ToSharedRef()));
+	TSharedRef<FMenuInWindow> Menu = MakeShareable(new FMenuInWindow(NewMenuWindow, InPrePushResults.WrappedContent.ToSharedRef(), InPrePushResults.bIsCollapsedByParent));
 	PendingNewMenu = Menu;
 
 	TSharedPtr<SWindow> ParentWindow;
@@ -461,7 +463,7 @@ TSharedRef<FMenuBase> FMenuStack::PushPopup(TSharedPtr<IMenu> InParentMenu, cons
 
 	// Create a FMenuInPopup
 	check(InPrePushResults.WrappedContent.IsValid());
-	TSharedRef<FMenuInPopup> Menu = MakeShareable(new FMenuInPopup(InPrePushResults.WrappedContent.ToSharedRef()));
+	TSharedRef<FMenuInPopup> Menu = MakeShareable(new FMenuInPopup(InPrePushResults.WrappedContent.ToSharedRef(), InPrePushResults.bIsCollapsedByParent));
 	PendingNewMenu = Menu;
 
 	// Register to get callback when it's dismissed - to fixup stack
@@ -630,9 +632,13 @@ void FMenuStack::OnMenuContentLostFocus(const FWidgetPath& InFocussedPath)
 			int32 FocussedIndex = Stack.IndexOfByKey(FocussedMenu);
 			check(FocussedIndex != INDEX_NONE);
 
-			if (Stack.Num() > FocussedIndex + 1)
+			for (int32 DismissIndex = FocussedIndex + 1; DismissIndex < Stack.Num(); DismissIndex++)
 			{
-				DismissFrom(Stack[FocussedIndex + 1]);
+				if (Stack[DismissIndex]->IsCollapsedByParent())
+				{
+					DismissFrom(Stack[DismissIndex]);
+					break;
+				}
 			}
 		}
 		else
@@ -716,9 +722,13 @@ void FMenuStack::OnWindowActivated( TSharedRef<SWindow> ActivatedWindow )
 			int32 ActivatedIndex = Stack.IndexOfByKey(ActivatedMenu);
 			check(ActivatedIndex != INDEX_NONE);
 
-			if (Stack.Num() > ActivatedIndex + 1)
+			for (int32 DismissIndex = ActivatedIndex + 1; DismissIndex < Stack.Num(); DismissIndex++)
 			{
-				DismissFrom(Stack[ActivatedIndex + 1]);
+				if (Stack[DismissIndex]->IsCollapsedByParent())
+				{
+					DismissFrom(Stack[DismissIndex]);
+					break;
+				}
 			}
 		}
 		else
diff --git a/Engine/Source/Runtime/Slate/Private/Framework/Application/SlateApplication.cpp b/Engine/Source/Runtime/Slate/Private/Framework/Application/SlateApplication.cpp
--- a/Engine/Source/Runtime/Slate/Private/Framework/Application/SlateApplication.cpp
+++ b/Engine/Source/Runtime/Slate/Private/Framework/Application/SlateApplication.cpp
@@ -1793,50 +1793,50 @@ TSharedRef<SWindow> FSlateApplication::PushMenu( const TSharedRef<SWidget>& InPa
 	return Window.ToSharedRef();
 }
 
-TSharedPtr<IMenu> FSlateApplication::PushMenu(const TSharedRef<SWidget>& InParentWidget, const FWidgetPath& InOwnerPath, const TSharedRef<SWidget>& InContent, const FVector2D& SummonLocation, const FPopupTransitionEffect& TransitionEffect, const bool bFocusImmediately, const FVector2D& SummonLocationSize, TOptional<EPopupMethod> Method)
+TSharedPtr<IMenu> FSlateApplication::PushMenu(const TSharedRef<SWidget>& InParentWidget, const FWidgetPath& InOwnerPath, const TSharedRef<SWidget>& InContent, const FVector2D& SummonLocation, const FPopupTransitionEffect& TransitionEffect, const bool bFocusImmediately, const FVector2D& SummonLocationSize, TOptional<EPopupMethod> Method, const bool bIsCollapsedByParent)
 {
 	// Caller supplied a valid path? Pass it to the menu stack.
 	if (InOwnerPath.IsValid())
 	{
-		return MenuStack.Push(InOwnerPath, InContent, SummonLocation, TransitionEffect, bFocusImmediately, SummonLocationSize, Method);
+		return MenuStack.Push(InOwnerPath, InContent, SummonLocation, TransitionEffect, bFocusImmediately, SummonLocationSize, Method, bIsCollapsedByParent);
 	}
 
 	// If the caller doesn't specify a valid event path we'll generate one from InParentWidget
 	FWidgetPath WidgetPath;
 	if (GeneratePathToWidgetUnchecked(InParentWidget, WidgetPath))
 	{
-		return MenuStack.Push(WidgetPath, InContent, SummonLocation, TransitionEffect, bFocusImmediately, SummonLocationSize, Method);
+		return MenuStack.Push(WidgetPath, InContent, SummonLocation, TransitionEffect, bFocusImmediately, SummonLocationSize, Method, bIsCollapsedByParent);
 	}
 
 	return TSharedPtr<IMenu>();
 }
 
-TSharedPtr<IMenu> FSlateApplication::PushMenu(const TSharedPtr<IMenu>& InParentMenu, const TSharedRef<SWidget>& InContent, const FVector2D& SummonLocation, const FPopupTransitionEffect& TransitionEffect, const bool bFocusImmediately, const FVector2D& SummonLocationSize)
+TSharedPtr<IMenu> FSlateApplication::PushMenu(const TSharedPtr<IMenu>& InParentMenu, const TSharedRef<SWidget>& InContent, const FVector2D& SummonLocation, const FPopupTransitionEffect& TransitionEffect, const bool bFocusImmediately, const FVector2D& SummonLocationSize, const bool bIsCollapsedByParent)
 {
-	return MenuStack.Push(InParentMenu, InContent, SummonLocation, TransitionEffect, bFocusImmediately, SummonLocationSize);
+	return MenuStack.Push(InParentMenu, InContent, SummonLocation, TransitionEffect, bFocusImmediately, SummonLocationSize, bIsCollapsedByParent);
 }
 
-TSharedPtr<IMenu> FSlateApplication::PushHostedMenu(const TSharedRef<SWidget>& InParentWidget, const FWidgetPath& InOwnerPath, const TSharedRef<IMenuHost>& InMenuHost, const TSharedRef<SWidget>& InContent, TSharedPtr<SWidget>& OutWrappedContent, const FPopupTransitionEffect& TransitionEffect)
+TSharedPtr<IMenu> FSlateApplication::PushHostedMenu(const TSharedRef<SWidget>& InParentWidget, const FWidgetPath& InOwnerPath, const TSharedRef<IMenuHost>& InMenuHost, const TSharedRef<SWidget>& InContent, TSharedPtr<SWidget>& OutWrappedContent, const FPopupTransitionEffect& TransitionEffect, const bool bIsCollapsedByParent)
 {
 	// Caller supplied a valid path? Pass it to the menu stack.
 	if (InOwnerPath.IsValid())
 	{
-		return MenuStack.PushHosted(InOwnerPath, InMenuHost, InContent, OutWrappedContent, TransitionEffect);
+		return MenuStack.PushHosted(InOwnerPath, InMenuHost, InContent, OutWrappedContent, TransitionEffect, bIsCollapsedByParent);
 	}
 
 	// If the caller doesn't specify a valid event path we'll generate one from InParentWidget
 	FWidgetPath WidgetPath;
 	if (GeneratePathToWidgetUnchecked(InParentWidget, WidgetPath))
 	{
-		return MenuStack.PushHosted(WidgetPath, InMenuHost, InContent, OutWrappedContent, TransitionEffect);
+		return MenuStack.PushHosted(WidgetPath, InMenuHost, InContent, OutWrappedContent, TransitionEffect, bIsCollapsedByParent);
 	}
 
 	return TSharedPtr<IMenu>();
 }
 
-TSharedPtr<IMenu> FSlateApplication::PushHostedMenu(const TSharedPtr<IMenu>& InParentMenu, const TSharedRef<IMenuHost>& InMenuHost, const TSharedRef<SWidget>& InContent, TSharedPtr<SWidget>& OutWrappedContent, const FPopupTransitionEffect& TransitionEffect)
+TSharedPtr<IMenu> FSlateApplication::PushHostedMenu(const TSharedPtr<IMenu>& InParentMenu, const TSharedRef<IMenuHost>& InMenuHost, const TSharedRef<SWidget>& InContent, TSharedPtr<SWidget>& OutWrappedContent, const FPopupTransitionEffect& TransitionEffect, const bool bIsCollapsedByParent)
 {
-	return MenuStack.PushHosted(InParentMenu, InMenuHost, InContent, OutWrappedContent, TransitionEffect);
+	return MenuStack.PushHosted(InParentMenu, InMenuHost, InContent, OutWrappedContent, TransitionEffect, bIsCollapsedByParent);
 }
 
 bool FSlateApplication::HasOpenSubMenus( TSharedRef<SWindow> Window ) const
diff --git a/Engine/Source/Runtime/Slate/Private/Widgets/Input/SMenuAnchor.cpp b/Engine/Source/Runtime/Slate/Private/Widgets/Input/SMenuAnchor.cpp
--- a/Engine/Source/Runtime/Slate/Private/Widgets/Input/SMenuAnchor.cpp
+++ b/Engine/Source/Runtime/Slate/Private/Widgets/Input/SMenuAnchor.cpp
@@ -349,7 +349,8 @@ void SMenuAnchor::SetIsOpen( bool InIsOpen, const bool bFocusMenu )
 					if (MethodInUse == EPopupMethod::CreateNewWindow)
 					{
 						// Open the pop-up
-						TSharedPtr<IMenu> NewMenu = FSlateApplication::Get().PushMenu(AsShared(), MyWidgetPath, MenuContentRef, NewPosition, TransitionEffect, bFocusMenu, SummonLocationSize, MethodInUse);
+						const bool bIsCollapsedByParent = false;	// don't auto-close child menus when the parent gets focus
+						TSharedPtr<IMenu> NewMenu = FSlateApplication::Get().PushMenu(AsShared(), MyWidgetPath, MenuContentRef, NewPosition, TransitionEffect, bFocusMenu, SummonLocationSize, MethodInUse, bIsCollapsedByParent);
 
 						PopupMenuPtr = NewMenu;
 						check(NewMenu.IsValid() && NewMenu->GetOwnedWindow().IsValid());
@@ -371,8 +372,9 @@ void SMenuAnchor::SetIsOpen( bool InIsOpen, const bool bFocusMenu )
 
 						TSharedRef<SMenuAnchor> SharedThis = StaticCastSharedRef<SMenuAnchor>(AsShared());
 
+						const bool bIsCollapsedByParent = false;	// don't auto-close child menus when the parent gets focus
 						TSharedPtr<IMenu> NewMenu = FSlateApplication::Get().PushHostedMenu(
-							SharedThis, MyWidgetPath, SharedThis, MenuContentRef, WrappedContent, TransitionEffect);
+							SharedThis, MyWidgetPath, SharedThis, MenuContentRef, WrappedContent, TransitionEffect, bIsCollapsedByParent);
 
 						PopupMenuPtr = NewMenu;
 						check(NewMenu.IsValid());
diff --git a/Engine/Source/Runtime/Slate/Public/Framework/Application/MenuStack.h b/Engine/Source/Runtime/Slate/Public/Framework/Application/MenuStack.h
--- a/Engine/Source/Runtime/Slate/Public/Framework/Application/MenuStack.h
+++ b/Engine/Source/Runtime/Slate/Public/Framework/Application/MenuStack.h
@@ -86,8 +86,9 @@ public:
 	 * @param bFocusImmediately		Should the popup steal focus when shown?
 	 * @param SummonLocationSize	An optional size around the summon location which describes an area in which the menu may not appear
 	 * @param InMethod				An optional popup method that will override the default method for the widgets in InOwnerPath
+	 * @param bIsCollapsedByParent	Is this menu collapsed when a parent menu receives focus/activation? If false, only focus/activation outside the entire stack will auto collapse it.
 	 */
-	TSharedRef<IMenu> Push(const FWidgetPath& InOwnerPath, const TSharedRef<SWidget>& InContent, const FVector2D& SummonLocation, const FPopupTransitionEffect& TransitionEffect, const bool bFocusImmediately = true, const FVector2D& SummonLocationSize = FVector2D::ZeroVector, TOptional<EPopupMethod> InMethod = TOptional<EPopupMethod>());
+	TSharedRef<IMenu> Push(const FWidgetPath& InOwnerPath, const TSharedRef<SWidget>& InContent, const FVector2D& SummonLocation, const FPopupTransitionEffect& TransitionEffect, const bool bFocusImmediately = true, const FVector2D& SummonLocationSize = FVector2D::ZeroVector, TOptional<EPopupMethod> InMethod = TOptional<EPopupMethod>(), const bool bIsCollapsedByParent = true);
 	
 	/**
 	 * Pushes a new child menu onto the stack.
@@ -99,8 +100,9 @@ public:
 	 * @param TransitionEffect		Animation to use when the popup appears
 	 * @param bFocusImmediately		Should the popup steal focus when shown?
 	 * @param SummonLocationSize	An optional size around the summon location which describes an area in which the menu may not appear
+	 * @param bIsCollapsedByParent	Is this menu collapsed when a parent menu receives focus/activation? If false, only focus/activation outside the entire stack will auto collapse it.
 	 */
-	TSharedRef<IMenu> Push(const TSharedPtr<IMenu>& InParentMenu, const TSharedRef<SWidget>& InContent, const FVector2D& SummonLocation, const FPopupTransitionEffect& TransitionEffect, const bool bFocusImmediately = true, const FVector2D& SummonLocationSize = FVector2D::ZeroVector);
+	TSharedRef<IMenu> Push(const TSharedPtr<IMenu>& InParentMenu, const TSharedRef<SWidget>& InContent, const FVector2D& SummonLocation, const FPopupTransitionEffect& TransitionEffect, const bool bFocusImmediately = true, const FVector2D& SummonLocationSize = FVector2D::ZeroVector, const bool bIsCollapsedByParent = true);
 
 	/**
 	 * Pushes a new menu onto the stack that is drawn by an external host widget.
@@ -112,8 +114,9 @@ public:
 	 * @param InContent				The menu's content
 	 * @param OutWrappedContent		Returns the InContent wrapped with widgets needed by the menu stack system. This is what should be drawn by the host after this call.
 	 * @param TransitionEffect		Animation to use when the popup appears
+	 * @param bIsCollapsedByParent	Is this menu collapsed when a parent menu receives focus/activation? If false, only focus/activation outside the entire stack will auto collapse it.
 	 */
-	TSharedRef<IMenu> PushHosted(const FWidgetPath& InOwnerPath, const TSharedRef<IMenuHost>& InMenuHost, const TSharedRef<SWidget>& InContent, TSharedPtr<SWidget>& OutWrappedContent, const FPopupTransitionEffect& TransitionEffect);
+	TSharedRef<IMenu> PushHosted(const FWidgetPath& InOwnerPath, const TSharedRef<IMenuHost>& InMenuHost, const TSharedRef<SWidget>& InContent, TSharedPtr<SWidget>& OutWrappedContent, const FPopupTransitionEffect& TransitionEffect, const bool bIsCollapsedByParent = true);
 	
 	/**
 	 * Pushes a new child menu onto the stack that is drawn by an external host widget.
@@ -124,8 +127,9 @@ public:
 	 * @param InContent				The menu's content
 	 * @param OutWrappedContent		Returns the InContent wrapped with widgets needed by the menu stack system. This is what should be drawn by the host after this call.
 	 * @param TransitionEffect		Animation to use when the popup appears
+	 * @param bIsCollapsedByParent	Is this menu collapsed when a parent menu receives focus/activation? If false, only focus/activation outside the entire stack will auto collapse it.
 	 */
-	TSharedRef<IMenu> PushHosted(const TSharedPtr<IMenu>& InParentMenu, const TSharedRef<IMenuHost>& InMenuHost, const TSharedRef<SWidget>& InContent, TSharedPtr<SWidget>& OutWrappedContent, const FPopupTransitionEffect& TransitionEffect);
+	TSharedRef<IMenu> PushHosted(const TSharedPtr<IMenu>& InParentMenu, const TSharedRef<IMenuHost>& InMenuHost, const TSharedRef<SWidget>& InContent, TSharedPtr<SWidget>& OutWrappedContent, const FPopupTransitionEffect& TransitionEffect, const bool bIsCollapsedByParent = true);
 
 	/**
 	 * [Deprecated] Dismisses the menu stack up to a certain level (by default removes all levels in the stack)
@@ -304,6 +308,7 @@ private:
 		FSlateRect Anchor;
 		FPopupTransitionEffect TransitionEffect;
 		bool bFocusImmediately;
+		bool bIsCollapsedByParent;
 	};
 
 	/** Contains all the options returned from the pre-push stage of the menu creation process */
@@ -317,6 +322,7 @@ private:
 		FVector2D ExpectedSize;
 		bool bAllowAnimations;
 		bool bFocusImmediately;
+		bool bIsCollapsedByParent;
 	};
 
 	/**
@@ -338,10 +344,11 @@ private:
 	 * @param Anchor				The anchor location for the menu (rect around the parent widget that is summoning the menu)
 	 * @param TransitionEffect		Animation to use when the popup appears
 	 * @param bFocusImmediately		Should the popup steal focus when shown?
+	 * @param bIsCollapsedByParent	Is this menu collapsed when a parent menu receives focus/activation? If false, only focus/activation outside the entire stack will auto collapse it.
 	 *
 	 * @return			The newly created menu.
 	 */
-	TSharedRef<IMenu> PushInternal(const TSharedPtr<IMenu>& InParentMenu, const TSharedRef<SWidget>& InContent, FSlateRect Anchor, const FPopupTransitionEffect& TransitionEffect, const bool bFocusImmediately);
+	TSharedRef<IMenu> PushInternal(const TSharedPtr<IMenu>& InParentMenu, const TSharedRef<SWidget>& InContent, FSlateRect Anchor, const FPopupTransitionEffect& TransitionEffect, const bool bFocusImmediately, const bool bIsCollapsedByParent);
 
 	/**
 	 * This is the actual menu object creation method for FMenuInWindow menus.
diff --git a/Engine/Source/Runtime/Slate/Public/Framework/Application/SlateApplication.h b/Engine/Source/Runtime/Slate/Public/Framework/Application/SlateApplication.h
--- a/Engine/Source/Runtime/Slate/Public/Framework/Application/SlateApplication.h
+++ b/Engine/Source/Runtime/Slate/Public/Framework/Application/SlateApplication.h
@@ -277,8 +277,9 @@ public:
 	 * @param bFocusImmediately		Should the popup steal focus when shown?
 	 * @param SummonLocationSize	An optional rect which describes an area in which the menu may not appear
 	 * @param Method				An optional popup method override. If not set, the widgets in the InOwnerPath will be queried for this.
+	 * @param bIsCollapsedByParent	Is this menu collapsed when a parent menu receives focus/activation? If false, only focus/activation outside the entire stack will auto collapse it.
 	 */
-	TSharedPtr<IMenu> PushMenu(const TSharedRef<SWidget>& InParentWidget, const FWidgetPath& InOwnerPath, const TSharedRef<SWidget>& InContent, const FVector2D& SummonLocation, const FPopupTransitionEffect& TransitionEffect, const bool bFocusImmediately = true, const FVector2D& SummonLocationSize = FVector2D::ZeroVector, TOptional<EPopupMethod> Method = TOptional<EPopupMethod>());
+	TSharedPtr<IMenu> PushMenu(const TSharedRef<SWidget>& InParentWidget, const FWidgetPath& InOwnerPath, const TSharedRef<SWidget>& InContent, const FVector2D& SummonLocation, const FPopupTransitionEffect& TransitionEffect, const bool bFocusImmediately = true, const FVector2D& SummonLocationSize = FVector2D::ZeroVector, TOptional<EPopupMethod> Method = TOptional<EPopupMethod>(), const bool bIsCollapsedByParent = true);
 
 	/**
 	 * Creates a new Menu and adds it to the menu stack under the specified parent menu.
@@ -290,8 +291,9 @@ public:
 	 * @param TransitionEffect		Animation to use when the popup appears
 	 * @param bFocusImmediately		Should the popup steal focus when shown?
 	 * @param SummonLocationSize	An optional rect which describes an area in which the menu may not appear
+	 * @param bIsCollapsedByParent	Is this menu collapsed when a parent menu receives focus/activation? If false, only focus/activation outside the entire stack will auto collapse it.
 	 */
-	TSharedPtr<IMenu> PushMenu(const TSharedPtr<IMenu>& InParentMenu, const TSharedRef<SWidget>& InContent, const FVector2D& SummonLocation, const FPopupTransitionEffect& TransitionEffect, const bool bFocusImmediately = true, const FVector2D& SummonLocationSize = FVector2D::ZeroVector);
+	TSharedPtr<IMenu> PushMenu(const TSharedPtr<IMenu>& InParentMenu, const TSharedRef<SWidget>& InContent, const FVector2D& SummonLocation, const FPopupTransitionEffect& TransitionEffect, const bool bFocusImmediately = true, const FVector2D& SummonLocationSize = FVector2D::ZeroVector, const bool bIsCollapsedByParent = true);
 
 	/**
 	 * Creates a new hosted Menu and adds it to the menu stack.
@@ -303,8 +305,9 @@ public:
 	 * @param InContent				The content to be placed inside the new menu
 	 * @param OutWrappedContent		Returns the InContent wrapped with widgets needed by the menu stack system. This is what should be drawn by the host after this call.
 	 * @param TransitionEffect		Animation to use when the popup appears
+	 * @param bIsCollapsedByParent	Is this menu collapsed when a parent menu receives focus/activation? If false, only focus/activation outside the entire stack will auto collapse it.
 	 */
-	TSharedPtr<IMenu> PushHostedMenu(const TSharedRef<SWidget>& InParentWidget, const FWidgetPath& InOwnerPath, const TSharedRef<IMenuHost>& InMenuHost, const TSharedRef<SWidget>& InContent, TSharedPtr<SWidget>& OutWrappedContent, const FPopupTransitionEffect& TransitionEffect);
+	TSharedPtr<IMenu> PushHostedMenu(const TSharedRef<SWidget>& InParentWidget, const FWidgetPath& InOwnerPath, const TSharedRef<IMenuHost>& InMenuHost, const TSharedRef<SWidget>& InContent, TSharedPtr<SWidget>& OutWrappedContent, const FPopupTransitionEffect& TransitionEffect, const bool bIsCollapsedByParent = true);
 	
 	/**
 	 * Creates a new hosted child Menu and adds it to the menu stack under the specified parent menu.
@@ -315,8 +318,9 @@ public:
 	 * @param InContent				The menu's content
 	 * @param OutWrappedContent		Returns the InContent wrapped with widgets needed by the menu stack system. This is what should be drawn by the host after this call.
 	 * @param TransitionEffect		Animation to use when the popup appears
+	 * @param bIsCollapsedByParent	Is this menu collapsed when a parent menu receives focus/activation? If false, only focus/activation outside the entire stack will auto collapse it.
 	 */	
-	TSharedPtr<IMenu> PushHostedMenu(const TSharedPtr<IMenu>& InParentMenu, const TSharedRef<IMenuHost>& InMenuHost, const TSharedRef<SWidget>& InContent, TSharedPtr<SWidget>& OutWrappedContent, const FPopupTransitionEffect& TransitionEffect);
+	TSharedPtr<IMenu> PushHostedMenu(const TSharedPtr<IMenu>& InParentMenu, const TSharedRef<IMenuHost>& InMenuHost, const TSharedRef<SWidget>& InContent, TSharedPtr<SWidget>& OutWrappedContent, const FPopupTransitionEffect& TransitionEffect, const bool bIsCollapsedByParent = true);
 
 	/** @return Returns whether the window has child menus. */
 	DEPRECATED(4.9, "HasOpenSubMenus() taking a window is deprecated. Use HasOpenSubMenus() taking an IMenu as a parameter.")
