Fix shadow variable warnings

diff --git a/Engine/Source/Runtime/Slate/Private/Framework/Application/Menu.cpp b/Engine/Source/Runtime/Slate/Private/Framework/Application/Menu.cpp
--- a/Engine/Source/Runtime/Slate/Private/Framework/Application/Menu.cpp
+++ b/Engine/Source/Runtime/Slate/Private/Framework/Application/Menu.cpp
@@ -3,10 +3,10 @@
 #include "SlatePrivatePCH.h"
 #include "Menu.h"
 
-FMenuBase::FMenuBase(TSharedRef<SWidget> InContent, const bool bIsCollapsedByParent)
+FMenuBase::FMenuBase(TSharedRef<SWidget> InContent, const bool bCollapsedByParent)
 	: Content(InContent)
 	, bDismissing(false)
-	, bIsCollapsedByParent(bIsCollapsedByParent)
+	, bIsCollapsedByParent(bCollapsedByParent)
 {
 }
 
diff --git a/Engine/Source/Runtime/Slate/Private/Framework/Application/Menu.h b/Engine/Source/Runtime/Slate/Private/Framework/Application/Menu.h
--- a/Engine/Source/Runtime/Slate/Private/Framework/Application/Menu.h
+++ b/Engine/Source/Runtime/Slate/Private/Framework/Application/Menu.h
@@ -15,7 +15,7 @@ public:
 	bool IsCollapsedByParent() const { return bIsCollapsedByParent; }
 
 protected:
-	FMenuBase(TSharedRef<SWidget> InContent, const bool bIsCollapsedByParent);
+	FMenuBase(TSharedRef<SWidget> InContent, const bool bCollapsedByParent);
 
 	FOnMenuDismissed OnMenuDismissed;
 	TSharedRef<SWidget> Content;
