UMG - Adding a few extra scaling modes to the scale box, allowing it to Scale to Fit X, and Scale to Fix Y.

diff --git a/Engine/Source/Runtime/Slate/Private/Widgets/Layout/SScaleBox.cpp b/Engine/Source/Runtime/Slate/Private/Widgets/Layout/SScaleBox.cpp
--- a/Engine/Source/Runtime/Slate/Private/Widgets/Layout/SScaleBox.cpp
+++ b/Engine/Source/Runtime/Slate/Private/Widgets/Layout/SScaleBox.cpp
@@ -51,6 +51,12 @@ void SScaleBox::OnArrangeChildren( const FGeometry& AllottedGeometry, FArrangedC
 			case EStretch::ScaleToFit:
 				FinalScale = FMath::Min(AreaSize.X / SlotWidgetDesiredSize.X, AreaSize.Y / SlotWidgetDesiredSize.Y);
 				break;
+			case EStretch::ScaleToFitX:
+				FinalScale = AreaSize.X / SlotWidgetDesiredSize.X;
+				break;
+			case EStretch::ScaleToFitY:
+				FinalScale = AreaSize.Y / SlotWidgetDesiredSize.Y;
+				break;
 			case EStretch::ScaleToFill:
 				FinalScale = FMath::Max(AreaSize.X / SlotWidgetDesiredSize.X, AreaSize.Y / SlotWidgetDesiredSize.Y);
 				break;
diff --git a/Engine/Source/Runtime/Slate/Public/Widgets/Layout/SScaleBox.h b/Engine/Source/Runtime/Slate/Public/Widgets/Layout/SScaleBox.h
--- a/Engine/Source/Runtime/Slate/Public/Widgets/Layout/SScaleBox.h
+++ b/Engine/Source/Runtime/Slate/Public/Widgets/Layout/SScaleBox.h
@@ -27,9 +27,27 @@ namespace EStretch
 		None,
 		/** Scales the content non-uniformly filling the entire space of the area. */
 		Fill,
-		/** Scales the content uniformly (preserving aspect ratio) until it can no longer scale the content without clipping it. */
+		/**
+		 * Scales the content uniformly (preserving aspect ratio) 
+		 * until it can no longer scale the content without clipping it.
+		 */
 		ScaleToFit,
-		/** Scales the content uniformly (preserving aspect ratio), until all sides meet or exceed the size of the area.  Will result in clipping longer sides. */
+		/**
+		 * Scales the content uniformly (preserving aspect ratio) 
+		 * until it can no longer scale the content without clipping it along the x-axis, 
+		 * the y-axis can/will be clipped.
+		 */
+		ScaleToFitX,
+		/**
+		 * Scales the content uniformly (preserving aspect ratio) 
+		 * until it can no longer scale the content without clipping it along the y-axis, 
+		 * the x-axis can/will be clipped.
+		 */
+		ScaleToFitY,
+		/**
+		 * Scales the content uniformly (preserving aspect ratio), until all sides meet 
+		 * or exceed the size of the area.  Will result in clipping the longer side.
+		 */
 		ScaleToFill,
 		/** Scales the content by the scale specified by the user. */
 		UserSpecified
