Fixed problem with menus not closing when clicking in another window on Mac
[UE-18450] - Cannot close the Content Browser Context Menu on Mac

diff --git a/Engine/Source/Runtime/Slate/Private/Framework/Application/MenuStack.cpp b/Engine/Source/Runtime/Slate/Private/Framework/Application/MenuStack.cpp
--- a/Engine/Source/Runtime/Slate/Private/Framework/Application/MenuStack.cpp
+++ b/Engine/Source/Runtime/Slate/Private/Framework/Application/MenuStack.cpp
@@ -619,9 +619,11 @@ void FMenuStack::OnMenuDestroyed(TSharedRef<IMenu> InMenu)
 
 void FMenuStack::OnMenuContentLostFocus(const FWidgetPath& InFocussedPath)
 {
-	// Window activation messages will make menus collapse when in CreateNewWindow mode
 	// In UseCurrentWindow mode we must look for focus moving from menus
-	if (ActiveMethod.Get(EPopupMethod::CreateNewWindow) == EPopupMethod::UseCurrentWindow && HasMenus() && !PendingNewMenu.IsValid())
+	// Window activation messages will make menus collapse when in CreateNewWindow mode
+	// However, we cannot rely on window activation messages because they might not be generated on Mac.
+	// So, always do this focus/collapse code, even in CreateNewWindow mode.
+	if (HasMenus() && !PendingNewMenu.IsValid())
 	{
 		// If focus is switching determine which of our menus it is in, if any
 		TSharedPtr<IMenu> FocussedMenu = FindMenuInWidgetPath(InFocussedPath);
