Editor: Added responsive crushing of the XYZ labels and margin in property editor vector/rotator rows, allowing them to be editable at much lower widths

diff --git a/Engine/Source/Editor/DetailCustomizations/Private/ComponentTransformDetails.cpp b/Engine/Source/Editor/DetailCustomizations/Private/ComponentTransformDetails.cpp
--- a/Engine/Source/Editor/DetailCustomizations/Private/ComponentTransformDetails.cpp
+++ b/Engine/Source/Editor/DetailCustomizations/Private/ComponentTransformDetails.cpp
@@ -386,6 +386,7 @@ void FComponentTransformDetails::GenerateChildContent( IDetailChildrenBuilder& C
 				.Y( this, &FComponentTransformDetails::GetLocationY )
 				.Z( this, &FComponentTransformDetails::GetLocationZ )
 				.bColorAxisLabels( true )
+				.AllowResponsiveLayout( true )
 				.IsEnabled( this, &FComponentTransformDetails::GetIsEnabled )
 				.OnXCommitted( this, &FComponentTransformDetails::OnSetLocation, 0 )
 				.OnYCommitted( this, &FComponentTransformDetails::OnSetLocation, 1 )
@@ -451,6 +452,7 @@ void FComponentTransformDetails::GenerateChildContent( IDetailChildrenBuilder& C
 				.Roll( this, &FComponentTransformDetails::GetRotationX )
 				.Pitch( this, &FComponentTransformDetails::GetRotationY )
 				.Yaw( this, &FComponentTransformDetails::GetRotationZ )
+				.AllowResponsiveLayout( true )
 				.bColorAxisLabels( true )
 				.IsEnabled( this, &FComponentTransformDetails::GetIsEnabled )
 				.OnBeginSliderMovement( this, &FComponentTransformDetails::OnBeginRotatonSlider )
@@ -516,6 +518,7 @@ void FComponentTransformDetails::GenerateChildContent( IDetailChildrenBuilder& C
 				.Y( this, &FComponentTransformDetails::GetScaleY )
 				.Z( this, &FComponentTransformDetails::GetScaleZ )
 				.bColorAxisLabels( true )
+				.AllowResponsiveLayout( true )
 				.IsEnabled( this, &FComponentTransformDetails::GetIsEnabled )
 				.OnXCommitted( this, &FComponentTransformDetails::OnSetScale, 0 )
 				.OnYCommitted( this, &FComponentTransformDetails::OnSetScale, 1 )
diff --git a/Engine/Source/Editor/LandscapeEditor/Private/LandscapeEditorDetailCustomization_NewLandscape.cpp b/Engine/Source/Editor/LandscapeEditor/Private/LandscapeEditorDetailCustomization_NewLandscape.cpp
--- a/Engine/Source/Editor/LandscapeEditor/Private/LandscapeEditorDetailCustomization_NewLandscape.cpp
+++ b/Engine/Source/Editor/LandscapeEditor/Private/LandscapeEditorDetailCustomization_NewLandscape.cpp
@@ -223,6 +223,7 @@ void FLandscapeEditorDetailCustomization_NewLandscape::CustomizeDetails(IDetailL
 	[
 		SNew(SRotatorInputBox)
 		.bColorAxisLabels(true)
+		.AllowResponsiveLayout(true)
 		.Font(DetailBuilder.GetDetailFont())
 		.Roll_Static(&GetOptionalPropertyValue<float>, PropertyHandle_Rotation_Roll)
 		.Pitch_Static(&GetOptionalPropertyValue<float>, PropertyHandle_Rotation_Pitch)
diff --git a/Engine/Source/Editor/LevelEditor/Private/DlgDeltaTransform.cpp b/Engine/Source/Editor/LevelEditor/Private/DlgDeltaTransform.cpp
--- a/Engine/Source/Editor/LevelEditor/Private/DlgDeltaTransform.cpp
+++ b/Engine/Source/Editor/LevelEditor/Private/DlgDeltaTransform.cpp
@@ -59,6 +59,7 @@ public:
 							.Y( this, &SDlgDeltaTransform::GetDeltaY )
 							.Z( this, &SDlgDeltaTransform::GetDeltaZ )
 							.bColorAxisLabels( true )
+							.AllowResponsiveLayout( true )
 							.OnXCommitted( this, &SDlgDeltaTransform::OnSetDelta, 0 )
 							.OnYCommitted( this, &SDlgDeltaTransform::OnSetDelta, 1 )
 							.OnZCommitted( this, &SDlgDeltaTransform::OnSetDelta, 2 )
diff --git a/Engine/Source/Editor/StaticMeshEditor/Private/StaticMeshEditorTools.cpp b/Engine/Source/Editor/StaticMeshEditor/Private/StaticMeshEditorTools.cpp
--- a/Engine/Source/Editor/StaticMeshEditor/Private/StaticMeshEditorTools.cpp
+++ b/Engine/Source/Editor/StaticMeshEditor/Private/StaticMeshEditorTools.cpp
@@ -515,6 +515,7 @@ void FMeshBuildSettingsLayout::GenerateChildContent( IDetailChildrenBuilder& Chi
 			.Y(this, &FMeshBuildSettingsLayout::GetBuildScaleY)
 			.Z(this, &FMeshBuildSettingsLayout::GetBuildScaleZ)
 			.bColorAxisLabels(false)
+			.AllowResponsiveLayout(true)
 			.OnXCommitted(this, &FMeshBuildSettingsLayout::OnBuildScaleXChanged)
 			.OnYCommitted(this, &FMeshBuildSettingsLayout::OnBuildScaleYChanged)
 			.OnZCommitted(this, &FMeshBuildSettingsLayout::OnBuildScaleZChanged)
diff --git a/Engine/Source/Editor/WorldBrowser/Private/StreamingLevels/StreamingLevelCustomization.cpp b/Engine/Source/Editor/WorldBrowser/Private/StreamingLevels/StreamingLevelCustomization.cpp
--- a/Engine/Source/Editor/WorldBrowser/Private/StreamingLevels/StreamingLevelCustomization.cpp
+++ b/Engine/Source/Editor/WorldBrowser/Private/StreamingLevels/StreamingLevelCustomization.cpp
@@ -58,6 +58,7 @@ void FStreamingLevelCustomization::CustomizeDetails(IDetailLayoutBuilder& Detail
 			.IsEnabled(this, &FStreamingLevelCustomization::LevelEditTextTransformAllowed)
 			.Font(FEditorStyle::GetFontStyle(TEXT("PropertyWindow.NormalFont")))
 			.bColorAxisLabels(true)
+			.AllowResponsiveLayout(true)
 			.X(this, &FStreamingLevelCustomization::OnGetLevelPosition, 0)
 			.Y(this, &FStreamingLevelCustomization::OnGetLevelPosition, 1)
 			.Z(this, &FStreamingLevelCustomization::OnGetLevelPosition, 2)
diff --git a/Engine/Source/Editor/WorldBrowser/Private/Tiles/STiledLandscapeImportDlg.cpp b/Engine/Source/Editor/WorldBrowser/Private/Tiles/STiledLandscapeImportDlg.cpp
--- a/Engine/Source/Editor/WorldBrowser/Private/Tiles/STiledLandscapeImportDlg.cpp
+++ b/Engine/Source/Editor/WorldBrowser/Private/Tiles/STiledLandscapeImportDlg.cpp
@@ -168,6 +168,7 @@ void STiledLandcapeImportDlg::Construct(const FArguments& InArgs, TSharedPtr<SWi
 				[
 					SNew( SVectorInputBox )
 					.bColorAxisLabels( true )
+					.AllowResponsiveLayout( true )
 					.X( this, &STiledLandcapeImportDlg::GetScaleX )
 					.Y( this, &STiledLandcapeImportDlg::GetScaleY )
 					.Z( this, &STiledLandcapeImportDlg::GetScaleZ )
diff --git a/Engine/Source/Runtime/Slate/Private/Widgets/Input/SRotatorInputBox.cpp b/Engine/Source/Runtime/Slate/Private/Widgets/Input/SRotatorInputBox.cpp
--- a/Engine/Source/Runtime/Slate/Private/Widgets/Input/SRotatorInputBox.cpp
+++ b/Engine/Source/Runtime/Slate/Private/Widgets/Input/SRotatorInputBox.cpp
@@ -3,25 +3,23 @@
 #include "SlatePrivatePCH.h"
 #include "SRotatorInputBox.h"
 #include "SNumericEntryBox.h"
+#include "SWidgetSwitcher.h"
 
+extern TAutoConsoleVariable<float> CVarCrushThem;
+extern TAutoConsoleVariable<float> CVarStopCrushWhenAbove;
+extern TAutoConsoleVariable<float> CVarStartCrushWhenBelow;
 
 #define LOCTEXT_NAMESPACE "SRotatorInputBox"
 
-
-/**
- * Construct this widget
- *
- * @param	InArgs	The declaration data for this widget
- */
 void SRotatorInputBox::Construct( const SRotatorInputBox::FArguments& InArgs )
 {
-	FLinearColor LabelColorX;
-	FLinearColor LabelColorY;
-	FLinearColor LabelColorZ;
-	LabelColorX = InArgs._bColorAxisLabels ? SNumericEntryBox<float>::RedLabelBackgroundColor : FLinearColor(0.0f,0.0f,0.0f,.5f);
-	LabelColorY = InArgs._bColorAxisLabels ? SNumericEntryBox<float>::GreenLabelBackgroundColor : FLinearColor(0.0f,0.0f,0.0f,.5f);
-	LabelColorZ = InArgs._bColorAxisLabels ? SNumericEntryBox<float>::BlueLabelBackgroundColor : FLinearColor(0.0f,0.0f,0.0f,.5f);
+	bCanBeCrushed = InArgs._AllowResponsiveLayout;
+
+	const FLinearColor LabelColorX = InArgs._bColorAxisLabels ? SNumericEntryBox<float>::RedLabelBackgroundColor : FLinearColor(0.0f, 0.0f, 0.0f, .5f);
+	const FLinearColor LabelColorY = InArgs._bColorAxisLabels ? SNumericEntryBox<float>::GreenLabelBackgroundColor : FLinearColor(0.0f, 0.0f, 0.0f, .5f);
+	const FLinearColor LabelColorZ = InArgs._bColorAxisLabels ? SNumericEntryBox<float>::BlueLabelBackgroundColor : FLinearColor(0.0f, 0.0f, 0.0f, .5f);
 
+	
 	this->ChildSlot
 	[
 		SNew(SHorizontalBox)
@@ -37,7 +35,7 @@ void SRotatorInputBox::Construct( const SRotatorInputBox::FArguments& InArgs )
 			.LabelPadding(0)
 			.Label()
 			[
-				SNumericEntryBox<float>::BuildLabel( LOCTEXT("Roll_Label", "X"), FLinearColor::White, SNumericEntryBox<float>::RedLabelBackgroundColor )
+				BuildDecoratorLabel(LabelColorX, FLinearColor::White, LOCTEXT("Roll_Label", "X"))
 			]
 			.Font( InArgs._Font )
 			.Value( InArgs._Roll )
@@ -61,7 +59,7 @@ void SRotatorInputBox::Construct( const SRotatorInputBox::FArguments& InArgs )
 			.LabelPadding(0)
 			.Label()
 			[
-				SNumericEntryBox<float>::BuildLabel( LOCTEXT("Pitch_Label", "Y"), FLinearColor::White, SNumericEntryBox<float>::GreenLabelBackgroundColor )
+				BuildDecoratorLabel(LabelColorY, FLinearColor::White, LOCTEXT("Pitch_Label", "Y"))
 			]
 			.Font( InArgs._Font )
 			.Value( InArgs._Pitch )
@@ -85,7 +83,7 @@ void SRotatorInputBox::Construct( const SRotatorInputBox::FArguments& InArgs )
 			.LabelPadding(0)
 			.Label()
 			[
-				SNumericEntryBox<float>::BuildLabel( LOCTEXT("Yaw_Label", "Z"), FLinearColor::White, SNumericEntryBox<float>::BlueLabelBackgroundColor )
+				BuildDecoratorLabel(LabelColorZ, FLinearColor::White, LOCTEXT("Yaw_Label", "Z"))
 			]
 			.Font( InArgs._Font )
 			.Value( InArgs._Yaw )
@@ -101,5 +99,73 @@ void SRotatorInputBox::Construct( const SRotatorInputBox::FArguments& InArgs )
 
 }
 
+TSharedRef<SWidget> SRotatorInputBox::BuildDecoratorLabel(FLinearColor BackgroundColor, FLinearColor ForegroundColor, FText Label)
+{
+	TSharedRef<SWidget> LabelWidget = SNumericEntryBox<float>::BuildLabel(Label, ForegroundColor, BackgroundColor);
+
+	TSharedRef<SWidget> ResultWidget = LabelWidget;
+
+	if (bCanBeCrushed)
+	{
+		ResultWidget =
+			SNew(SWidgetSwitcher)
+			.WidgetIndex(this, &SRotatorInputBox::GetLabelActiveSlot)
+			+SWidgetSwitcher::Slot()
+			[
+				LabelWidget
+			]
+			+SWidgetSwitcher::Slot()
+			[
+				SNew(SBorder)
+				.BorderImage(FCoreStyle::Get().GetBrush("NumericEntrySpinBox.NarrowDecorator"))
+				.BorderBackgroundColor(BackgroundColor)
+				.ForegroundColor(ForegroundColor)
+				.VAlign(VAlign_Center)
+				.HAlign(HAlign_Left)
+				.Padding(FMargin(5.0f, 0.0f, 0.0f, 0.0f))
+			];
+	}
+
+	return ResultWidget;
+}
+
+int32 SRotatorInputBox::GetLabelActiveSlot() const
+{
+	return bIsBeingCrushed ? 1 : 0;
+}
+
+FMargin SRotatorInputBox::GetTextMargin() const
+{
+	return bIsBeingCrushed ? FMargin(1.0f, 2.0f) : FMargin(4.0f, 2.0f);
+}
+
+void SRotatorInputBox::OnArrangeChildren(const FGeometry& AllottedGeometry, FArrangedChildren& ArrangedChildren) const
+{
+	bool bFoop = bCanBeCrushed && (CVarCrushThem.GetValueOnAnyThread() > 0.0f);
+
+	if (bFoop)
+	{
+		const float AlottedWidth = AllottedGeometry.Size.X;
+
+		const float CrushBelow = CVarStartCrushWhenBelow.GetValueOnAnyThread();
+		const float StopCrushing = CVarStopCrushWhenAbove.GetValueOnAnyThread();
+
+		if (bIsBeingCrushed)
+		{
+			bIsBeingCrushed = AlottedWidth < StopCrushing;
+		}
+		else
+		{
+			bIsBeingCrushed = AlottedWidth < CrushBelow;
+		}
+	}
+	else
+	{
+		bIsBeingCrushed = false;
+	}
+
+	SCompoundWidget::OnArrangeChildren(AllottedGeometry, ArrangedChildren);
+}
+
 #undef LOCTEXT_NAMESPACE
 
diff --git a/Engine/Source/Runtime/Slate/Private/Widgets/Input/SVectorInputBox.cpp b/Engine/Source/Runtime/Slate/Private/Widgets/Input/SVectorInputBox.cpp
--- a/Engine/Source/Runtime/Slate/Private/Widgets/Input/SVectorInputBox.cpp
+++ b/Engine/Source/Runtime/Slate/Private/Widgets/Input/SVectorInputBox.cpp
@@ -3,13 +3,18 @@
 #include "SlatePrivatePCH.h"
 #include "SVectorInputBox.h"
 #include "SNumericEntryBox.h"
-
+#include "SWidgetSwitcher.h"
 
 #define LOCTEXT_NAMESPACE "SVectorInputBox"
 
+TAutoConsoleVariable<float> CVarCrushThem(TEXT("Slate.AllowNumericLabelCrush"), 1.0f, TEXT("Should we crush the vector input box?."));
+TAutoConsoleVariable<float> CVarStopCrushWhenAbove(TEXT("Slate.NumericLabelWidthCrushStop"), 200.0f, TEXT("Stop crushing when the width is above."));
+TAutoConsoleVariable<float> CVarStartCrushWhenBelow(TEXT("Slate.NumericLabelWidthCrushStart"), 190.0f, TEXT("Start crushing when the width is below."));
 
 void SVectorInputBox::Construct( const FArguments& InArgs )
 {
+	bCanBeCrushed = InArgs._AllowResponsiveLayout;
+
 	TSharedRef<SHorizontalBox> HorizontalBox = SNew(SHorizontalBox);
 
 	ChildSlot
@@ -24,8 +29,13 @@ void SVectorInputBox::Construct( const FArguments& InArgs )
 
 void SVectorInputBox::ConstructX( const FArguments& InArgs, TSharedRef<SHorizontalBox> HorizontalBox )
 {	
-	FLinearColor LabelColor;
-	LabelColor = InArgs._bColorAxisLabels ?  SNumericEntryBox<float>::RedLabelBackgroundColor : FLinearColor(0.0f,0.0f,0.0f,.5f);
+	const FLinearColor LabelColor = InArgs._bColorAxisLabels ?  SNumericEntryBox<float>::RedLabelBackgroundColor : FLinearColor(0.0f,0.0f,0.0f,.5f);
+	TSharedRef<SWidget> LabelWidget = BuildDecoratorLabel(LabelColor, FLinearColor::White, LOCTEXT("X_Label", "X"));
+	TAttribute<FMargin> MarginAttribute;
+	if (bCanBeCrushed)
+	{
+		MarginAttribute = TAttribute<FMargin>::Create(TAttribute<FMargin>::FGetter::CreateSP(this, &SVectorInputBox::GetTextMargin));
+	}
 
 	HorizontalBox->AddSlot()
 	.VAlign( VAlign_Center )
@@ -40,21 +50,26 @@ void SVectorInputBox::ConstructX( const FArguments& InArgs, TSharedRef<SHorizont
 		.ToolTipText( LOCTEXT("X_ToolTip", "X Value") )
 		.UndeterminedString( LOCTEXT("MultipleValues", "Multiple Values") )
 		.LabelPadding(0)
+		.OverrideTextMargin(MarginAttribute)
 		.ContextMenuExtender( InArgs._ContextMenuExtenderX )
 		.TypeInterface(InArgs._TypeInterface)
 		.Label()
 		[
-			SNumericEntryBox<float>::BuildLabel( LOCTEXT("X_Label", "X"), FLinearColor::White, LabelColor )
+			LabelWidget
 		]
 	];
 	
 }
 
-
 void SVectorInputBox::ConstructY( const FArguments& InArgs, TSharedRef<SHorizontalBox> HorizontalBox )
 {
-	FLinearColor LabelColor;
-	LabelColor = InArgs._bColorAxisLabels ?  SNumericEntryBox<float>::GreenLabelBackgroundColor : FLinearColor(0.0f,0.0f,0.0f,.5f);
+	const FLinearColor LabelColor = InArgs._bColorAxisLabels ?  SNumericEntryBox<float>::GreenLabelBackgroundColor : FLinearColor(0.0f,0.0f,0.0f,.5f);
+	TSharedRef<SWidget> LabelWidget = BuildDecoratorLabel(LabelColor, FLinearColor::White, LOCTEXT("Y_Label", "Y"));
+	TAttribute<FMargin> MarginAttribute;
+	if (bCanBeCrushed)
+	{
+		MarginAttribute = TAttribute<FMargin>::Create(TAttribute<FMargin>::FGetter::CreateSP(this, &SVectorInputBox::GetTextMargin));
+	}
 
 	HorizontalBox->AddSlot()
 	.VAlign( VAlign_Center )
@@ -69,11 +84,12 @@ void SVectorInputBox::ConstructY( const FArguments& InArgs, TSharedRef<SHorizont
 		.ToolTipText( LOCTEXT("Y_ToolTip", "Y Value") )
 		.UndeterminedString( LOCTEXT("MultipleValues", "Multiple Values") )
 		.LabelPadding(0)
-		.ContextMenuExtender( InArgs._ContextMenuExtenderY )
+		.OverrideTextMargin(MarginAttribute)
+		.ContextMenuExtender(InArgs._ContextMenuExtenderY)
 		.TypeInterface(InArgs._TypeInterface)
 		.Label()
 		[
-			SNumericEntryBox<float>::BuildLabel( LOCTEXT("Y_Label", "Y"), FLinearColor::White, LabelColor )
+			LabelWidget
 		]
 	];
 
@@ -81,8 +97,13 @@ void SVectorInputBox::ConstructY( const FArguments& InArgs, TSharedRef<SHorizont
 
 void SVectorInputBox::ConstructZ( const FArguments& InArgs, TSharedRef<SHorizontalBox> HorizontalBox )
 {
-	FLinearColor LabelColor;
-	LabelColor = InArgs._bColorAxisLabels ?  SNumericEntryBox<float>::BlueLabelBackgroundColor : FLinearColor(0.0f,0.0f,0.0f,.5f);
+	const FLinearColor LabelColor = InArgs._bColorAxisLabels ?  SNumericEntryBox<float>::BlueLabelBackgroundColor : FLinearColor(0.0f,0.0f,0.0f,.5f);
+	TSharedRef<SWidget> LabelWidget = BuildDecoratorLabel(LabelColor, FLinearColor::White, LOCTEXT("Z_Label", "Z"));
+	TAttribute<FMargin> MarginAttribute;
+	if (bCanBeCrushed)
+	{
+		MarginAttribute = TAttribute<FMargin>::Create(TAttribute<FMargin>::FGetter::CreateSP(this, &SVectorInputBox::GetTextMargin));
+	}
 
 	HorizontalBox->AddSlot()
 	.VAlign( VAlign_Center )
@@ -97,13 +118,82 @@ void SVectorInputBox::ConstructZ( const FArguments& InArgs, TSharedRef<SHorizont
 		.ToolTipText( LOCTEXT("Z_ToolTip", "Z Value") )
 		.UndeterminedString( LOCTEXT("MultipleValues", "Multiple Values") )
 		.LabelPadding(0)
-		.ContextMenuExtender( InArgs._ContextMenuExtenderZ )
+		.OverrideTextMargin(MarginAttribute)
+		.ContextMenuExtender(InArgs._ContextMenuExtenderZ)
 		.TypeInterface(InArgs._TypeInterface)
 		.Label()
 		[
-			SNumericEntryBox<float>::BuildLabel( LOCTEXT("Z_Label", "Z"), FLinearColor::White, LabelColor )
+			LabelWidget
 		]
 	];
 }
 
+TSharedRef<SWidget> SVectorInputBox::BuildDecoratorLabel(FLinearColor BackgroundColor, FLinearColor ForegroundColor, FText Label)
+{
+	TSharedRef<SWidget> LabelWidget = SNumericEntryBox<float>::BuildLabel(Label, ForegroundColor, BackgroundColor);
+
+	TSharedRef<SWidget> ResultWidget = LabelWidget;
+	
+	if (bCanBeCrushed)
+	{
+		ResultWidget =
+			SNew(SWidgetSwitcher)
+			.WidgetIndex(this, &SVectorInputBox::GetLabelActiveSlot)
+			+SWidgetSwitcher::Slot()
+			[
+				LabelWidget
+			]
+			+SWidgetSwitcher::Slot()
+			[
+				SNew(SBorder)
+				.BorderImage(FCoreStyle::Get().GetBrush("NumericEntrySpinBox.NarrowDecorator"))
+				.BorderBackgroundColor(BackgroundColor)
+				.ForegroundColor(ForegroundColor)
+				.VAlign(VAlign_Center)
+				.HAlign(HAlign_Left)
+				.Padding(FMargin(5.0f, 0.0f, 0.0f, 0.0f))
+			];
+	}
+
+	return ResultWidget;
+}
+
+int32 SVectorInputBox::GetLabelActiveSlot() const
+{
+	return bIsBeingCrushed ? 1 : 0;
+}
+
+FMargin SVectorInputBox::GetTextMargin() const
+{
+	return bIsBeingCrushed ? FMargin(1.0f, 2.0f) : FMargin(4.0f, 2.0f);
+}
+
+void SVectorInputBox::OnArrangeChildren(const FGeometry& AllottedGeometry, FArrangedChildren& ArrangedChildren) const
+{
+	bool bFoop = bCanBeCrushed && (CVarCrushThem.GetValueOnAnyThread() > 0.0f);
+
+	if (bFoop)
+	{
+		const float AlottedWidth = AllottedGeometry.Size.X;
+
+		const float CrushBelow = CVarStartCrushWhenBelow.GetValueOnAnyThread();
+		const float StopCrushing = CVarStopCrushWhenAbove.GetValueOnAnyThread();
+
+		if (bIsBeingCrushed)
+		{
+			bIsBeingCrushed = AlottedWidth < StopCrushing;
+		}
+		else
+		{
+			bIsBeingCrushed = AlottedWidth < CrushBelow;
+		}
+	}
+	else
+	{
+		bIsBeingCrushed = false;
+	}
+
+	SCompoundWidget::OnArrangeChildren(AllottedGeometry, ArrangedChildren);
+}
+
 #undef LOCTEXT_NAMESPACE
diff --git a/Engine/Source/Runtime/Slate/Public/Widgets/Input/SNumericEntryBox.h b/Engine/Source/Runtime/Slate/Public/Widgets/Input/SNumericEntryBox.h
--- a/Engine/Source/Runtime/Slate/Public/Widgets/Input/SNumericEntryBox.h
+++ b/Engine/Source/Runtime/Slate/Public/Widgets/Input/SNumericEntryBox.h
@@ -76,6 +76,8 @@ public:
 		SLATE_ATTRIBUTE( float, SliderExponent )
 		/** The minimum desired width for the value portion of the control. */
 		SLATE_ATTRIBUTE( float, MinDesiredValueWidth )
+		/** The text margin to use if overridden. */
+		SLATE_ATTRIBUTE( FMargin, OverrideTextMargin )
 		/** Called whenever the text is changed interactively by the user */
 		SLATE_EVENT( FOnValueChanged, OnValueChanged )
 		/** Called whenever the text is committed.  This happens when the user presses enter or the text box loses focus. */
@@ -108,8 +110,8 @@ public:
 		BorderImageNormal = &InArgs._EditableTextBoxStyle->BackgroundImageNormal;
 		BorderImageHovered = &InArgs._EditableTextBoxStyle->BackgroundImageHovered;
 		BorderImageFocused = &InArgs._EditableTextBoxStyle->BackgroundImageFocused;
-		const FMargin& TextMargin = InArgs._EditableTextBoxStyle->Padding;
-		
+		TAttribute<FMargin> TextMargin = InArgs._OverrideTextMargin.IsSet() ? InArgs._OverrideTextMargin : InArgs._EditableTextBoxStyle->Padding;
+
 		Interface = InArgs._TypeInterface.IsValid() ? InArgs._TypeInterface : MakeShareable( new TDefaultNumericTypeInterface<NumericType> );
 
 		TSharedPtr<SWidget> FinalWidget;
diff --git a/Engine/Source/Runtime/Slate/Public/Widgets/Input/SRotatorInputBox.h b/Engine/Source/Runtime/Slate/Public/Widgets/Input/SRotatorInputBox.h
--- a/Engine/Source/Runtime/Slate/Public/Widgets/Input/SRotatorInputBox.h
+++ b/Engine/Source/Runtime/Slate/Public/Widgets/Input/SRotatorInputBox.h
@@ -10,10 +10,12 @@ class SLATE_API SRotatorInputBox : public SCompoundWidget
 {
 public:
 	SLATE_BEGIN_ARGS( SRotatorInputBox )
-		: _bColorAxisLabels( false )
+		: _bColorAxisLabels(false)
 		, _Font( FCoreStyle::Get().GetFontStyle("NormalFont") )
+		, _AllowResponsiveLayout(false)
 		, _AllowSpin(true)
 		{}
+
 		/** Roll component of the rotator */
 		SLATE_ATTRIBUTE( TOptional<float>, Roll )
 
@@ -26,6 +28,9 @@ public:
 		/** Should the axis labels be colored */
 		SLATE_ARGUMENT( bool, bColorAxisLabels )		
 
+			/** Allow responsive layout to crush the label and margins when there is not a lot of room */
+		SLATE_ARGUMENT(bool, AllowResponsiveLayout)
+
 		/** Font to use for the text in this box */
 		SLATE_ATTRIBUTE( FSlateFontInfo, Font )
 
@@ -67,4 +72,26 @@ public:
 	 * @param	InArgs	The declaration data for this widget
 	 */
 	void Construct( const FArguments& InArgs );
+
+
+	// SWidget interface
+	virtual void OnArrangeChildren(const FGeometry& AllottedGeometry, FArrangedChildren& ArrangedChildren) const override;
+	// End of SWidget interface
+
+private:
+	/** Are we allowed to be crushed? */
+	bool bCanBeCrushed;
+
+	/** Are we currently being crushed? */
+	mutable bool bIsBeingCrushed;
+
+private:
+	/** Returns the index of the label widget to use (crushed or uncrushed) */
+	int32 GetLabelActiveSlot() const;
+
+	/** Returns the desired text margin for the label */
+	FMargin GetTextMargin() const;
+
+	/** Creates a decorator label (potentially adding a switcher widget if this is cruhsable) */
+	TSharedRef<SWidget> BuildDecoratorLabel(FLinearColor BackgroundColor, FLinearColor ForegroundColor, FText Label);
 };
diff --git a/Engine/Source/Runtime/Slate/Public/Widgets/Input/SSpinBox.h b/Engine/Source/Runtime/Slate/Public/Widgets/Input/SSpinBox.h
--- a/Engine/Source/Runtime/Slate/Public/Widgets/Input/SSpinBox.h
+++ b/Engine/Source/Runtime/Slate/Public/Widgets/Input/SSpinBox.h
@@ -57,7 +57,7 @@ public:
 		/** Font used to display text in the slider */
 		SLATE_ATTRIBUTE( FSlateFontInfo, Font )
 		/** Padding to add around this widget and its internal widgets */
-		SLATE_ARGUMENT( FMargin, ContentPadding )
+		SLATE_ATTRIBUTE( FMargin, ContentPadding )
 		/** Called when the value is changed by slider or typing */
 		SLATE_EVENT( FOnValueChanged, OnValueChanged )
 		/** Called when the value is committed (by pressing enter) */
diff --git a/Engine/Source/Runtime/Slate/Public/Widgets/Input/SVectorInputBox.h b/Engine/Source/Runtime/Slate/Public/Widgets/Input/SVectorInputBox.h
--- a/Engine/Source/Runtime/Slate/Public/Widgets/Input/SVectorInputBox.h
+++ b/Engine/Source/Runtime/Slate/Public/Widgets/Input/SVectorInputBox.h
@@ -11,10 +11,11 @@ class SLATE_API SVectorInputBox : public SCompoundWidget
 {
 public:
 	SLATE_BEGIN_ARGS( SVectorInputBox )
-		: _Font( FCoreStyle::Get().GetFontStyle("NormalFont") ),
-		  _bColorAxisLabels( false )
+		: _Font( FCoreStyle::Get().GetFontStyle("NormalFont") )
+		, _bColorAxisLabels( false )
+		, _AllowResponsiveLayout( false )
 		{}
-		
+
 		/** X Component of the vector */
 		SLATE_ATTRIBUTE( TOptional<float>, X )
 
@@ -30,6 +31,9 @@ public:
 		/** Should the axis labels be colored */
 		SLATE_ARGUMENT( bool, bColorAxisLabels )		
 
+		/** Allow responsive layout to crush the label and margins when there is not a lot of room */
+		SLATE_ARGUMENT(bool, AllowResponsiveLayout)
+
 		/** Called when the x value of the vector is changed */
 		SLATE_EVENT( FOnFloatValueChanged, OnXChanged )
 
@@ -68,7 +72,27 @@ public:
 	 * @param	InArgs	The declaration data for this widget
 	 */
 	void Construct( const FArguments& InArgs );
+
+	// SWidget interface
+	virtual void OnArrangeChildren(const FGeometry& AllottedGeometry, FArrangedChildren& ArrangedChildren) const override;
+	// End of SWidget interface
+
 private:
+	/** Are we allowed to be crushed? */
+	bool bCanBeCrushed;
+
+	/** Are we currently being crushed? */
+	mutable bool bIsBeingCrushed;
+
+private:
+	/** Returns the index of the label widget to use (crushed or uncrushed) */
+	int32 GetLabelActiveSlot() const;
+
+	/** Returns the desired text margin for the label */
+	FMargin GetTextMargin() const;
+
+	/** Creates a decorator label (potentially adding a switcher widget if this is cruhsable) */
+	TSharedRef<SWidget> BuildDecoratorLabel(FLinearColor BackgroundColor, FLinearColor ForegroundColor, FText Label);
 
 	/**
 	 * Construct widgets for the X Value
diff --git a/Engine/Source/Runtime/SlateCore/Private/Styling/CoreStyle.cpp b/Engine/Source/Runtime/SlateCore/Private/Styling/CoreStyle.cpp
--- a/Engine/Source/Runtime/SlateCore/Private/Styling/CoreStyle.cpp
+++ b/Engine/Source/Runtime/SlateCore/Private/Styling/CoreStyle.cpp
@@ -610,6 +610,8 @@ TSharedRef<ISlateStyle> FCoreStyle::Create( const FName& InStyleSetName )
 			.SetForegroundColor( InvertedForeground )
 			);
 		Style->Set( "NumericEntrySpinBox.Decorator", new BOX_BRUSH( "Common/TextBoxLabelBorder", FMargin(5.0f/16.0f) ) );
+
+		Style->Set( "NumericEntrySpinBox.NarrowDecorator", new BOX_BRUSH( "Common/TextBoxLabelBorder", FMargin(2.0f/16.0f, 4.0f/16.0f, 2.0f/16.0f, 4.0f/16.0f) ) );
 	}
 
 	// SColorPicker defaults...
