Slate - Fixing several issues involving multitouch.  Data stored invloving widgets under cursor are now bucketed based on user and cursor index so that multiple pointer devices, as well as index can independently hold seperate widget under cursor data.

diff --git a/Engine/Source/Runtime/Slate/Private/Framework/Application/SlateApplication.cpp b/Engine/Source/Runtime/Slate/Private/Framework/Application/SlateApplication.cpp
--- a/Engine/Source/Runtime/Slate/Private/Framework/Application/SlateApplication.cpp
+++ b/Engine/Source/Runtime/Slate/Private/Framework/Application/SlateApplication.cpp
@@ -230,12 +230,12 @@
 
 	static void ProcessReply( FSlateApplication* Application, const FWidgetPath& RoutingPath, const FReply& Reply, const FWidgetPath* WidgetsUnderCursor, const FInputEvent* PointerEvent )
 	{
-		Application->ProcessReply(RoutingPath, Reply, WidgetsUnderCursor, nullptr);
+		Application->ProcessReply(RoutingPath, Reply, WidgetsUnderCursor, nullptr, PointerEvent->GetUserIndex());
 	}
 
 	static void ProcessReply( FSlateApplication* Application, const FWidgetPath& RoutingPath, const FReply& Reply, const FWidgetPath* WidgetsUnderCursor, const FPointerEvent* PointerEvent )
 	{
-		Application->ProcessReply(RoutingPath, Reply, WidgetsUnderCursor, PointerEvent);
+		Application->ProcessReply(RoutingPath, Reply, WidgetsUnderCursor, PointerEvent, PointerEvent->GetUserIndex());
 	}
 
 	template<typename EventType>
@@ -384,17 +384,17 @@
 	return false;
 }
 
-bool FSlateApplication::MouseCaptorHelper::HasCaptureForPointerIndex(uint32 PointerIndex) const
+bool FSlateApplication::MouseCaptorHelper::HasCaptureForPointerIndex(uint32 UserIndex, uint32 PointerIndex) const
 {
-	const FWeakWidgetPath* MouseCaptorWeakPath = PointerIndexToMouseCaptorWeakPathMap.Find(PointerIndex);
+	const FWeakWidgetPath* MouseCaptorWeakPath = PointerIndexToMouseCaptorWeakPathMap.Find( FUserAndPointer(UserIndex,PointerIndex) );
 	return MouseCaptorWeakPath && MouseCaptorWeakPath->IsValid();
 }
 
-TSharedPtr< SWidget > FSlateApplication::MouseCaptorHelper::ToSharedWidget(uint32 PointerIndex) const
+TSharedPtr< SWidget > FSlateApplication::MouseCaptorHelper::ToSharedWidget(uint32 UserIndex, uint32 PointerIndex) const
 {
 	// If the path is valid then get the last widget, this is the current mouse captor
 	TSharedPtr< SWidget > SharedWidgetPtr;
-	const FWeakWidgetPath* MouseCaptorWeakPath = PointerIndexToMouseCaptorWeakPathMap.Find(PointerIndex);
+	const FWeakWidgetPath* MouseCaptorWeakPath = PointerIndexToMouseCaptorWeakPathMap.Find( FUserAndPointer(UserIndex,PointerIndex) );
 	if (MouseCaptorWeakPath && MouseCaptorWeakPath->IsValid() )
 	{
 		TWeakPtr< SWidget > WeakWidgetPtr = MouseCaptorWeakPath->GetLastWidget();
@@ -418,10 +418,10 @@
 	return Widgets;
 }
 
-TSharedPtr< SWidget > FSlateApplication::MouseCaptorHelper::ToSharedWindow(uint32 PointerIndex)
+TSharedPtr< SWidget > FSlateApplication::MouseCaptorHelper::ToSharedWindow(uint32 UserIndex, uint32 PointerIndex)
 {
 	// if the path is valid then we can get the window the current mouse captor belongs to
-	FWidgetPath MouseCaptorPath = ToWidgetPath(PointerIndex);
+	FWidgetPath MouseCaptorPath = ToWidgetPath( UserIndex, PointerIndex );
 	if ( MouseCaptorPath.IsValid() )
 	{
 		return MouseCaptorPath.GetWindow();
@@ -430,11 +430,11 @@
 	return TSharedPtr< SWidget >();
 }
 
-void FSlateApplication::MouseCaptorHelper::SetMouseCaptor(uint32 PointerIndex, const FWidgetPath& EventPath, TSharedPtr< SWidget > Widget)
+void FSlateApplication::MouseCaptorHelper::SetMouseCaptor(uint32 UserIndex, uint32 PointerIndex, const FWidgetPath& EventPath, TSharedPtr< SWidget > Widget)
 {
 	// Caller is trying to set a new mouse captor, so invalidate the current one - when the function finishes
 	// it still may not have a valid captor widget, this is ok
-	InvalidateCaptureForPointer(PointerIndex);
+	InvalidateCaptureForPointer(UserIndex, PointerIndex);
 
 	if ( Widget.IsValid() )
 	{
@@ -468,76 +468,76 @@
 
 		if (MouseCaptorWeakPath.IsValid())
 		{
-			PointerIndexToMouseCaptorWeakPathMap.Add(PointerIndex, MouseCaptorWeakPath);
+			PointerIndexToMouseCaptorWeakPathMap.Add(FUserAndPointer(UserIndex,PointerIndex), MouseCaptorWeakPath);
 		}
 	}
 }
 
 void FSlateApplication::MouseCaptorHelper::InvalidateCaptureForAllPointers()
 {
-	TArray<uint32> PointerIndices;
+	TArray<FUserAndPointer> PointerIndices;
 	PointerIndexToMouseCaptorWeakPathMap.GenerateKeyArray(PointerIndices);
-	for (auto Index : PointerIndices)
+	for (FUserAndPointer UserAndPointer : PointerIndices)
 	{
-		InvalidateCaptureForPointer(Index);
+		InvalidateCaptureForPointer(UserAndPointer.UserIndex, UserAndPointer.PointerIndex);
 	}
 }
 
 FWidgetPath FSlateApplication::MouseCaptorHelper::ToWidgetPath( FWeakWidgetPath::EInterruptedPathHandling::Type InterruptedPathHandling, const FPointerEvent* PointerEvent )
 {
 	FWidgetPath WidgetPath;
-	const FWeakWidgetPath* MouseCaptorWeakPath = PointerIndexToMouseCaptorWeakPathMap.Find(PointerEvent->GetPointerIndex());
+	const FWeakWidgetPath* MouseCaptorWeakPath = PointerIndexToMouseCaptorWeakPathMap.Find(FUserAndPointer(PointerEvent->GetUserIndex(),PointerEvent->GetPointerIndex()));
 	if ( MouseCaptorWeakPath->IsValid() )
 	{
 		if ( MouseCaptorWeakPath->ToWidgetPath( WidgetPath, InterruptedPathHandling, PointerEvent ) == FWeakWidgetPath::EPathResolutionResult::Truncated )
 		{
 			// If the path was truncated then it means this widget is no longer part of the active set,
 			// so we make sure to invalidate its capture
-			InvalidateCaptureForPointer(PointerEvent->GetPointerIndex());
+			InvalidateCaptureForPointer(PointerEvent->GetUserIndex(), PointerEvent->GetPointerIndex());
 		}
 	}
 
 	return WidgetPath;
 }
 
-FWidgetPath FSlateApplication::MouseCaptorHelper::ToWidgetPath(uint32 PointerIndex, FWeakWidgetPath::EInterruptedPathHandling::Type InterruptedPathHandling)
+FWidgetPath FSlateApplication::MouseCaptorHelper::ToWidgetPath(uint32 UserIndex, uint32 PointerIndex, FWeakWidgetPath::EInterruptedPathHandling::Type InterruptedPathHandling)
 {
 	FWidgetPath WidgetPath;
-	const FWeakWidgetPath* MouseCaptorWeakPath = PointerIndexToMouseCaptorWeakPathMap.Find(PointerIndex);
+	const FWeakWidgetPath* MouseCaptorWeakPath = PointerIndexToMouseCaptorWeakPathMap.Find(FUserAndPointer(UserIndex,PointerIndex));
 	if (MouseCaptorWeakPath && MouseCaptorWeakPath->IsValid() )
 	{
 		if ( MouseCaptorWeakPath->ToWidgetPath( WidgetPath, InterruptedPathHandling ) == FWeakWidgetPath::EPathResolutionResult::Truncated )
 		{
 			// If the path was truncated then it means this widget is no longer part of the active set,
 			// so we make sure to invalidate its capture
-			InvalidateCaptureForPointer(PointerIndex);
+			InvalidateCaptureForPointer(UserIndex,PointerIndex);
 		}
 	}
 
 	return WidgetPath;
 }
 
-void FSlateApplication::MouseCaptorHelper::InvalidateCaptureForPointer(uint32 PointerIndex)
+void FSlateApplication::MouseCaptorHelper::InvalidateCaptureForPointer(uint32 UserIndex, uint32 PointerIndex)
 {
-	InformCurrentCaptorOfCaptureLoss(PointerIndex);
-	PointerIndexToMouseCaptorWeakPathMap.Remove(PointerIndex);
+	InformCurrentCaptorOfCaptureLoss(UserIndex, PointerIndex);
+	PointerIndexToMouseCaptorWeakPathMap.Remove( FUserAndPointer(UserIndex, PointerIndex) );
 }
 
 TArray<FWidgetPath> FSlateApplication::MouseCaptorHelper::ToWidgetPaths()
 {
 	TArray<FWidgetPath> WidgetPaths;
-	TArray<uint32> PointerIndices;
+	TArray<FUserAndPointer> PointerIndices;
 	PointerIndexToMouseCaptorWeakPathMap.GenerateKeyArray(PointerIndices);
 	for (auto Index : PointerIndices)
 	{
-		WidgetPaths.Add(ToWidgetPath(Index));
+		WidgetPaths.Add(ToWidgetPath(Index.UserIndex,Index.PointerIndex));
 	}
 	return WidgetPaths;
 }
 
-FWeakWidgetPath FSlateApplication::MouseCaptorHelper::ToWeakPath(uint32 PointerIndex) const
+FWeakWidgetPath FSlateApplication::MouseCaptorHelper::ToWeakPath(uint32 UserIndex, uint32 PointerIndex) const
 {
-	const FWeakWidgetPath* MouseCaptorWeakPath = PointerIndexToMouseCaptorWeakPathMap.Find(PointerIndex);
+	const FWeakWidgetPath* MouseCaptorWeakPath = PointerIndexToMouseCaptorWeakPathMap.Find(FUserAndPointer(UserIndex,PointerIndex));
 	if (MouseCaptorWeakPath)
 	{
 		return *MouseCaptorWeakPath;
@@ -545,10 +545,10 @@
 	return FWeakWidgetPath();
 }
 
-void FSlateApplication::MouseCaptorHelper::InformCurrentCaptorOfCaptureLoss(uint32 PointerIndex) const
+void FSlateApplication::MouseCaptorHelper::InformCurrentCaptorOfCaptureLoss(uint32 UserIndex,uint32 PointerIndex) const
 {
 	// if we have a path to a widget then it is the current mouse captor and needs to know it has lost capture
-	const FWeakWidgetPath* MouseCaptorWeakPath = PointerIndexToMouseCaptorWeakPathMap.Find(PointerIndex);
+	const FWeakWidgetPath* MouseCaptorWeakPath = PointerIndexToMouseCaptorWeakPathMap.Find(FUserAndPointer(UserIndex,PointerIndex));
 	if (MouseCaptorWeakPath && MouseCaptorWeakPath->IsValid() )
 	{
 		TWeakPtr< SWidget > WeakWidgetPtr = MouseCaptorWeakPath->GetLastWidget();
@@ -1132,7 +1132,7 @@
 #endif
 
 	FWidgetPath WidgetsUnderCursor = bVisualizeLayoutUnderCursor
-		? WidgetsUnderCursorLastEvent.ToWidgetPath()
+		? WidgetsUnderCursorLastEvent.FindRef(FUserAndPointer(CursorUserIndex,CursorPointerIndex)).ToWidgetPath()
 		: FWidgetPath();
 
 	if ( !SkipSecondPrepass.GetValueOnGameThread() )
@@ -1233,16 +1233,21 @@
 	}
 	else
 	{
-		for ( const TWeakPtr<SWidget>& WidgetPtr : WidgetsUnderCursorLastEvent.Widgets )
+		for( auto LastWidgetIterator = WidgetsUnderCursorLastEvent.CreateConstIterator(); LastWidgetIterator; ++LastWidgetIterator )
 		{
-			const TSharedPtr<SWidget>& Widget = WidgetPtr.Pin();
-			if (Widget.IsValid())
+			FWeakWidgetPath Path = LastWidgetIterator.Value();
+
+			for ( const TWeakPtr<SWidget>& WidgetPtr : Path.Widgets )
 			{
-				Widget->OnFinishedPointerInput();
+				const TSharedPtr<SWidget>& Widget = WidgetPtr.Pin();
+				if (Widget.IsValid())
+				{
+					Widget->OnFinishedPointerInput();
+				}
+				else
+				{
+					break;
+				}
 			}
-			else
-			{
-				break;
-			}
 		}
 	}
@@ -2479,10 +2484,11 @@
 	// Release mouse capture if requested or if we are starting a drag and drop.
 	// Make sure to only clobber WidgetsUnderCursor if we actually had a mouse capture.
 	uint32 PointerIndex = InMouseEvent != nullptr ? InMouseEvent->GetPointerIndex() : CursorPointerIndex;
-	if (MouseCaptor.HasCaptureForPointerIndex(PointerIndex) && (TheReply.ShouldReleaseMouse() || bStartingDragDrop) )
+
+	if (MouseCaptor.HasCaptureForPointerIndex(UserIndex,PointerIndex) && (TheReply.ShouldReleaseMouse() || bStartingDragDrop) )
 	{
-		WidgetsUnderCursorLastEvent = MouseCaptor.ToWeakPath(PointerIndex);
-		MouseCaptor.InvalidateCaptureForPointer(PointerIndex);
+		WidgetsUnderCursorLastEvent.Add( FUserAndPointer(UserIndex,PointerIndex), MouseCaptor.ToWeakPath(UserIndex,PointerIndex) );
+		MouseCaptor.InvalidateCaptureForPointer(UserIndex,PointerIndex);
 	}
 
 	// Clear focus is requested.
@@ -2533,7 +2539,7 @@
 	// Do not capture the mouse if we are also starting a drag and drop.
 	if( RequestedMouseCaptor.IsValid() && !bStartingDragDrop )
 	{
-		MouseCaptor.SetMouseCaptor(PointerIndex, CurrentEventPath, RequestedMouseCaptor );
+		MouseCaptor.SetMouseCaptor(UserIndex, PointerIndex, CurrentEventPath, RequestedMouseCaptor );
 	}
 	
 	if( CurrentEventPath.IsValid() && ( TheReply.ShouldReleaseMouse() || RequestedMouseCaptor.IsValid() ) )
@@ -2723,7 +2729,7 @@
 				);
 
 			// Query widgets with mouse capture for the cursor
-			if (MouseCaptor.HasCaptureForPointerIndex(CursorPointerIndex))
+			if (MouseCaptor.HasCaptureForPointerIndex(CursorUserIndex, CursorPointerIndex))
 			{
 				FWidgetPath MouseCaptorPath = MouseCaptor.ToWidgetPath( FWeakWidgetPath::EInterruptedPathHandling::Truncate, &CursorEvent);
 				if ( MouseCaptorPath.IsValid() )
@@ -3164,17 +3170,21 @@
 
 void FSlateApplication::CancelDragDrop()
 {
-	FWidgetPath WidgetsToDragLeave = WidgetsUnderCursorLastEvent.ToWidgetPath(FWeakWidgetPath::EInterruptedPathHandling::Truncate);
-	if (WidgetsToDragLeave.IsValid())
+	for( auto LastWidgetIterator = WidgetsUnderCursorLastEvent.CreateConstIterator(); LastWidgetIterator; ++LastWidgetIterator)
 	{
-		const FDragDropEvent DragDropEvent(FPointerEvent(), DragDropContent);
-		for (int32 WidgetIndex = WidgetsToDragLeave.Widgets.Num() - 1; WidgetIndex >= 0; --WidgetIndex)
+		
+		FWidgetPath WidgetsToDragLeave = LastWidgetIterator.Value().ToWidgetPath(FWeakWidgetPath::EInterruptedPathHandling::Truncate);
+		if(WidgetsToDragLeave.IsValid())
 		{
-			WidgetsToDragLeave.Widgets[WidgetIndex].Widget->OnDragLeave(DragDropEvent);
+			const FDragDropEvent DragDropEvent(FPointerEvent(), DragDropContent);
+			for(int32 WidgetIndex = WidgetsToDragLeave.Widgets.Num() - 1; WidgetIndex >= 0; --WidgetIndex)
+			{
+				WidgetsToDragLeave.Widgets[WidgetIndex].Widget->OnDragLeave(DragDropEvent);
+			}
 		}
 	}
 
-	WidgetsUnderCursorLastEvent = FWeakWidgetPath();
+	WidgetsUnderCursorLastEvent.Empty();
 	DragDropContent.Reset();
 }
 
@@ -3610,7 +3620,7 @@
 
 TSharedPtr<SWidget> FSlateApplication::GetMouseCaptorImpl() const
 {
-	return MouseCaptor.ToSharedWidget(CursorPointerIndex);
+	return MouseCaptor.ToSharedWidget(CursorUserIndex,CursorPointerIndex);
 }
 
 bool FSlateApplication::HasAnyMouseCaptor() const
@@ -3659,7 +3669,15 @@
 
 bool FSlateApplication::IsWidgetDirectlyHovered(const TSharedPtr<const SWidget> Widget) const
 {
-	return WidgetsUnderCursorLastEvent.IsValid() && Widget == WidgetsUnderCursorLastEvent.GetLastWidget().Pin();
+	for( auto LastWidgetIterator = WidgetsUnderCursorLastEvent.CreateConstIterator(); LastWidgetIterator; ++LastWidgetIterator )
+	{
+		const FWeakWidgetPath& WeakPath = LastWidgetIterator.Value();
+		if( WeakPath.IsValid() && Widget == WeakPath.GetLastWidget().Pin() )
+		{
+			return true;
+		}
+	}
+	return false;
 }
 
 bool FSlateApplication::ShowUserFocus(const TSharedPtr<const SWidget> Widget) const
@@ -4065,121 +4083,50 @@
 	if ( !IsDragDropping() )
 	{
 		FReply Reply = FReply::Unhandled();
-		if (MouseCaptor.HasCaptureForPointerIndex(MouseEvent.GetPointerIndex()))
+		if (MouseCaptor.HasCaptureForPointerIndex(MouseEvent.GetUserIndex(), MouseEvent.GetPointerIndex()))
 		{
 			FWidgetPath MouseCaptorPath = MouseCaptor.ToWidgetPath( FWeakWidgetPath::EInterruptedPathHandling::Truncate, &MouseEvent );
 			FArrangedWidget& MouseCaptorWidget = MouseCaptorPath.Widgets.Last();
 
 			// Switch worlds widgets in the current path
-			FScopedSwitchWorldHack SwitchWorld( MouseCaptorPath );
+			FScopedSwitchWorldHack SwitchWorld(MouseCaptorPath);
 			bInGame = FApp::IsGame();
 
-			Reply = FEventRouter::Route<FReply>( this, FEventRouter::FToLeafmostPolicy(MouseCaptorPath), MouseEvent, []( const FArrangedWidget& InMouseCaptorWidget, const FPointerEvent& Event )
+			FReply Reply = FEventRouter::Route<FReply>(this, FEventRouter::FToLeafmostPolicy(MouseCaptorPath), MouseEvent, [] (const FArrangedWidget& InMouseCaptorWidget, const FPointerEvent& Event)
 			{
-				return InMouseCaptorWidget.Widget->OnPreviewMouseButtonDown( InMouseCaptorWidget.Geometry, Event );
+				return InMouseCaptorWidget.Widget->OnPreviewMouseButtonDown(InMouseCaptorWidget.Geometry, Event);
 			});
-			
-			if (!Reply.IsEventHandled())
+
+			if ( !Reply.IsEventHandled() )
 			{
-				Reply = FEventRouter::Route<FReply>( this, FEventRouter::FToLeafmostPolicy(MouseCaptorPath), MouseEvent,
-					[this]( const FArrangedWidget& InMouseCaptorWidget, const FPointerEvent& Event )
+				Reply = FEventRouter::Route<FReply>(this, FEventRouter::FToLeafmostPolicy(MouseCaptorPath), MouseEvent,
+					[this] (const FArrangedWidget& InMouseCaptorWidget, const FPointerEvent& Event)
+				{
+					FReply TempReply = FReply::Unhandled();
+					if ( Event.IsTouchEvent() )
 					{
-						FReply TempReply = FReply::Unhandled();
-						if ( Event.IsTouchEvent() )
-						{
-							TempReply = InMouseCaptorWidget.Widget->OnTouchStarted( InMouseCaptorWidget.Geometry, Event );
-						}
-						if ( !Event.IsTouchEvent() || ( !TempReply.IsEventHandled() && this->bTouchFallbackToMouse ) )
-						{
-							TempReply = InMouseCaptorWidget.Widget->OnMouseButtonDown( InMouseCaptorWidget.Geometry, Event );
-						}
-						return TempReply;
-					} );
+						TempReply = InMouseCaptorWidget.Widget->OnTouchStarted(InMouseCaptorWidget.Geometry, Event);
+					}
+					if ( !Event.IsTouchEvent() || ( !TempReply.IsEventHandled() && this->bTouchFallbackToMouse ) )
+					{
+						TempReply = InMouseCaptorWidget.Widget->OnMouseButtonDown(InMouseCaptorWidget.Geometry, Event);
+					}
+					return TempReply;
+				});
 			}
-			LOG_EVENT( EEventLog::MouseButtonDown, Reply );
+			LOG_EVENT(EEventLog::MouseButtonDown, Reply);
 		}
 		else
 		{
 			FWidgetPath WidgetsUnderCursor = LocateWindowUnderMouse( MouseEvent.GetScreenSpacePosition(), GetInteractiveTopLevelWindows() );
 
-#if PLATFORM_MAC
-			NSWindow* ActiveWindow = [NSApp keyWindow];
-			const bool bNeedToActivateWindow = (ActiveWindow == nullptr);
-#else
-			const bool bNeedToActivateWindow = false;
-#endif
 			PopupSupport.SendNotifications( WidgetsUnderCursor );
 
 			// Switch worlds widgets in the current path
-			FScopedSwitchWorldHack SwitchWorld( WidgetsUnderCursor );
+			FScopedSwitchWorldHack SwitchWorld(WidgetsUnderCursor);
 			bInGame = FApp::IsGame();
 
-			const TSharedPtr<SWidget> PreviouslyFocusedWidget = GetKeyboardFocusedWidget();
-
-			Reply = FEventRouter::Route<FReply>(this, FEventRouter::FTunnelPolicy(WidgetsUnderCursor), MouseEvent, [](const FArrangedWidget TargetWidget, const FPointerEvent& Event)
-			{
-				return TargetWidget.Widget->OnPreviewMouseButtonDown(TargetWidget.Geometry, Event);
-			});
-
-			if ( !Reply.IsEventHandled() )
-			{
-				Reply = FEventRouter::Route<FReply>(this, FEventRouter::FBubblePolicy(WidgetsUnderCursor), MouseEvent, [this](const FArrangedWidget TargetWidget, const FPointerEvent& Event)
-				{
-					FReply ThisReply = FReply::Unhandled();
-					if (!ThisReply.IsEventHandled())
-					{
-						if (Event.IsTouchEvent())
-						{
-							ThisReply = TargetWidget.Widget->OnTouchStarted( TargetWidget.Geometry, Event );
-						}
-						if (!Event.IsTouchEvent() || (!ThisReply.IsEventHandled() && this->bTouchFallbackToMouse))
-						{
-							ThisReply = TargetWidget.Widget->OnMouseButtonDown( TargetWidget.Geometry, Event );
-						}
-					}
-					return ThisReply;
-				});
-			}
-			LOG_EVENT( EEventLog::MouseButtonDown, Reply );
-
-			// If none of the widgets requested keyboard focus to be set (or set the keyboard focus explicitly), set it to the leaf-most widget under the mouse.
-			// On Mac we prevent the OS from activating the window on mouse down, so we have full control and can activate only if there's nothing draggable under the mouse cursor.
-			const bool bFocusChangedByEventHandler = PreviouslyFocusedWidget != GetKeyboardFocusedWidget();
-			if ((!Reply.GetUserFocusRecepient().IsValid() || (PLATFORM_MAC && MouseEvent.GetEffectingButton() == EKeys::LeftMouseButton && !DragDetector.DetectDragForWidget.IsValid()))
-				&& (!bFocusChangedByEventHandler || bNeedToActivateWindow))
-			{
-				// The event handler for OnMouseButton down may have altered the widget hierarchy.
-				// Refresh the previously-cached widget path.
-				WidgetsUnderCursor = LocateWindowUnderMouse( MouseEvent.GetScreenSpacePosition(), GetInteractiveTopLevelWindows() );
-
-				bool bFocusCandidateFound = false;
-				for( int32 WidgetIndex = WidgetsUnderCursor.Widgets.Num()-1; !bFocusCandidateFound && WidgetIndex >= 0 ; --WidgetIndex )
-				{
-					FArrangedWidget& CurWidget = WidgetsUnderCursor.Widgets[WidgetIndex];
-					if (CurWidget.Widget->SupportsKeyboardFocus())
-					{
-						bFocusCandidateFound = true;
-						FWidgetPath NewFocusedWidgetPath = WidgetsUnderCursor.GetPathDownTo( CurWidget.Widget );
-						SetKeyboardFocus( NewFocusedWidgetPath, EFocusCause::Mouse );
-					}
-				}
-
-#if PLATFORM_MAC
-				if (bNeedToActivateWindow && WidgetsUnderCursor.TopLevelWindow.IsValid())
-				{
-					// Clicking on a context menu should not activate anything
-					// @todo: This needs to be updated when we have window type in SWindow and we no longer have to guess if WidgetsUnderCursor.TopLevelWindow is a menu
-					const bool bIsContextMenu = !WidgetsUnderCursor.TopLevelWindow->IsRegularWindow() && WidgetsUnderCursor.TopLevelWindow->HasMinimizeBox() && WidgetsUnderCursor.TopLevelWindow->HasMaximizeBox();
-					if (!bIsContextMenu && MouseEvent.GetEffectingButton() == EKeys::LeftMouseButton && !DragDetector.DetectDragForWidget.IsValid() && ActiveWindow == [NSApp keyWindow])
-					{
-						MouseCaptorHelper Captor = MouseCaptor;
-						FPlatformMisc::ActivateApplication();
-						WidgetsUnderCursor.TopLevelWindow->BringToFront(true);
-						MouseCaptor = Captor;
-					}
-				}
-#endif
-			}
+			Reply = RoutePointerDownEvent(WidgetsUnderCursor, MouseEvent);
 		}
 
 		// See if expensive tasks should be throttled.  By default on mouse down expensive tasks are throttled
@@ -4203,97 +4150,93 @@
 	return true;
 }
 
-bool FSlateApplication::OnMouseDoubleClick( const TSharedPtr< FGenericWindow >& PlatformWindow, const EMouseButtons::Type Button )
+FReply FSlateApplication::RoutePointerDownEvent(FWidgetPath& WidgetsUnderPointer, FPointerEvent& PointerEvent)
 {
-	if (bIsFakingTouch || bIsGameFakingTouch)
+	const TSharedPtr<SWidget> PreviouslyFocusedWidget = GetKeyboardFocusedWidget();
+
+	FReply Reply = FEventRouter::Route<FReply>(this, FEventRouter::FTunnelPolicy(WidgetsUnderPointer), PointerEvent, [] (const FArrangedWidget TargetWidget, const FPointerEvent& Event)
 	{
-		bIsFakingTouched = true;
-		return OnTouchStarted(PlatformWindow, PlatformApplication->Cursor->GetPosition(), 0, 0);
+		return TargetWidget.Widget->OnPreviewMouseButtonDown(TargetWidget.Geometry, Event);
+	});
+
+	if ( !Reply.IsEventHandled() )
+	{
+		Reply = FEventRouter::Route<FReply>(this, FEventRouter::FBubblePolicy(WidgetsUnderPointer), PointerEvent, [this] (const FArrangedWidget TargetWidget, const FPointerEvent& Event)
+		{
+			FReply ThisReply = FReply::Unhandled();
+			if ( !ThisReply.IsEventHandled() )
+			{
+				if ( Event.IsTouchEvent() )
+				{
+					ThisReply = TargetWidget.Widget->OnTouchStarted(TargetWidget.Geometry, Event);
+				}
+				if ( !Event.IsTouchEvent() || ( !ThisReply.IsEventHandled() && this->bTouchFallbackToMouse ) )
+				{
+					ThisReply = TargetWidget.Widget->OnMouseButtonDown(TargetWidget.Geometry, Event);
+				}
+			}
+			return ThisReply;
+		});
+	}
+	LOG_EVENT(EEventLog::MouseButtonDown, Reply);
+
+#if PLATFORM_MAC
+		NSWindow* ActiveWindow = [NSApp keyWindow];
+		const bool bNeedToActivateWindow = (ActiveWindow == nullptr);
+#else
+		const bool bNeedToActivateWindow = false;
+#endif
+
+	// If none of the widgets requested keyboard focus to be set (or set the keyboard focus explicitly), set it to the leaf-most widget under the mouse.
+	// On Mac we prevent the OS from activating the window on mouse down, so we have full control and can activate only if there's nothing draggable under the mouse cursor.
+	const bool bFocusChangedByEventHandler = PreviouslyFocusedWidget != GetKeyboardFocusedWidget();
+	if ((!Reply.GetUserFocusRecepient().IsValid() || (PLATFORM_MAC && PointerEvent.GetEffectingButton() == EKeys::LeftMouseButton && !DragDetector.DetectDragForWidget.IsValid()))
+		&& (!bFocusChangedByEventHandler || bNeedToActivateWindow))
+	{
+		// The event handler for OnMouseButton down may have altered the widget hierarchy.
+		// Refresh the previously-cached widget path.
+		WidgetsUnderPointer = LocateWindowUnderMouse(PointerEvent.GetScreenSpacePosition(), GetInteractiveTopLevelWindows());
+
+		bool bFocusCandidateFound = false;
+		for ( int32 WidgetIndex = WidgetsUnderPointer.Widgets.Num() - 1; !bFocusCandidateFound && WidgetIndex >= 0; --WidgetIndex )
+		{
+			FArrangedWidget& CurWidget = WidgetsUnderPointer.Widgets[WidgetIndex];
+			if ( CurWidget.Widget->SupportsKeyboardFocus() )
+			{
+				bFocusCandidateFound = true;
+				FWidgetPath NewFocusedWidgetPath = WidgetsUnderPointer.GetPathDownTo(CurWidget.Widget);
+				SetKeyboardFocus(NewFocusedWidgetPath, EFocusCause::Mouse);
+			}
+		}
+
+#if PLATFORM_MAC
+		TSharedPtr<SWindow> TopLevelWindow = WidgetsUnderPointer.TopLevelWindow;
+		if ( TopLevelWindow.IsValid() )
+		{
+			// Clicking on a context menu should not activate anything
+			// @todo: This needs to be updated when we have window type in SWindow and we no longer have to guess if WidgetsUnderCursor.TopLevelWindow is a menu
+			const bool bIsContextMenu = !TopLevelWindow->IsRegularWindow() && TopLevelWindow->HasMinimizeBox() && TopLevelWindow->HasMaximizeBox();
+			if ( !bIsContextMenu && PointerEvent.GetEffectingButton() == EKeys::LeftMouseButton && !DragDetector.DetectDragForWidget.IsValid() && ActiveWindow == [NSApp keyWindow] )
+			{
+				MouseCaptorHelper Captor = MouseCaptor;
+				FPlatformMisc::ActivateApplication();
+				TopLevelWindow->BringToFront(true);
+				MouseCaptor = Captor;
+			}
+		}
+#endif
 	}
 
-	FKey Key = TranslateMouseButtonToKey( Button );
-
-	FPointerEvent MouseEvent(
-		CursorPointerIndex,
-		GetCursorPos(),
-		GetLastCursorPos(),
-		PressedMouseButtons,
-		Key,
-		0,
-		PlatformApplication->GetModifierKeys()
-		);
-
-	return ProcessMouseButtonDoubleClickEvent( PlatformWindow, MouseEvent );
+	return Reply;
 }
 
-bool FSlateApplication::ProcessMouseButtonDoubleClickEvent( const TSharedPtr< FGenericWindow >& PlatformWindow, FPointerEvent& InMouseEvent )
+
+void FSlateApplication::RoutePointerUpEvent(FWidgetPath& WidgetsUnderPointer, FPointerEvent& PointerEvent)
 {
-	QueueSynthesizedMouseMove();
-	LastUserInteractionTime = this->GetCurrentTime();
-	LastUserInteractionTimeForThrottling = LastUserInteractionTime;
-	
-	PlatformApplication->SetCapture( PlatformWindow );
-	PressedMouseButtons.Add( InMouseEvent.GetEffectingButton() );
-
-	FWidgetPath WidgetsUnderCursor = LocateWindowUnderMouse( InMouseEvent.GetScreenSpacePosition(), GetInteractiveTopLevelWindows() );
-
-	// Switch worlds widgets in the current path
-	FScopedSwitchWorldHack SwitchWorld( WidgetsUnderCursor );
-
-	FReply Reply = FEventRouter::Route<FReply>( this, FEventRouter::FBubblePolicy(WidgetsUnderCursor), InMouseEvent, [](const FArrangedWidget& TargetWidget, const FPointerEvent& Event)
+	if (MouseCaptor.HasCaptureForPointerIndex(PointerEvent.GetUserIndex(), PointerEvent.GetPointerIndex()))
 	{
-		return TargetWidget.Widget->OnMouseButtonDoubleClick( TargetWidget.Geometry, Event );
-	} );
-
-
-	LOG_EVENT( EEventLog::MouseButtonDoubleClick, Reply );
-
-	PointerIndexLastPositionMap.Add(InMouseEvent.GetPointerIndex(), InMouseEvent.GetScreenSpacePosition());
-
-	return Reply.IsEventHandled();
-}
-
-bool FSlateApplication::OnMouseUp( const EMouseButtons::Type Button )
-{
-	// convert to touch event if we are faking it	
-	if (bIsFakingTouch || bIsGameFakingTouch)
-	{
-		bIsFakingTouched = false;
-		return OnTouchEnded(PlatformApplication->Cursor->GetPosition(), 0, 0);
-	}
-
-	FKey Key = TranslateMouseButtonToKey( Button );
-
-	FPointerEvent MouseEvent(
-		CursorPointerIndex,
-		GetCursorPos(),
-		GetLastCursorPos(),
-		PressedMouseButtons,
-		Key,
-		0,
-		PlatformApplication->GetModifierKeys()
-		);
-
-	return ProcessMouseButtonUpEvent( MouseEvent );
-}
-
-bool FSlateApplication::ProcessMouseButtonUpEvent( FPointerEvent& MouseEvent )
-{
-	QueueSynthesizedMouseMove();
-	LastUserInteractionTime = this->GetCurrentTime();
-	LastUserInteractionTimeForThrottling = LastUserInteractionTime;
-	PressedMouseButtons.Remove( MouseEvent.GetEffectingButton() );
-
-	if (DragDetector.DetectDragForWidget.IsValid() && MouseEvent.GetEffectingButton() == DragDetector.DetectDragButton )
-	{
-		// The user has release the button that was supposed to start the drag; stop detecting it.
-		DragDetector = FDragDetector();
-	}
-
-	if (MouseCaptor.HasCaptureForPointerIndex(MouseEvent.GetPointerIndex()))
-	{
 		//FWidgetPath MouseCaptorPath = MouseCaptor.ToWidgetPath(MouseEvent.GetPointerIndex());
-		FWidgetPath MouseCaptorPath = MouseCaptor.ToWidgetPath( FWeakWidgetPath::EInterruptedPathHandling::Truncate, &MouseEvent );
+		FWidgetPath MouseCaptorPath = MouseCaptor.ToWidgetPath( FWeakWidgetPath::EInterruptedPathHandling::Truncate, &PointerEvent );
 		if ( ensureMsgf(MouseCaptorPath.Widgets.Num() > 0, TEXT("A window had a widget with mouse capture. That entire window has been dismissed before the mouse up could be processed.")) )
 		{
 #if PLATFORM_MAC
@@ -4302,25 +4245,27 @@
 			// Switch worlds widgets in the current path
 			FScopedSwitchWorldHack SwitchWorld( MouseCaptorPath );
 
-			FReply Reply = FEventRouter::Route<FReply>( this, FEventRouter::FToLeafmostPolicy(MouseCaptorPath), MouseEvent, [this]( const FArrangedWidget& TargetWidget, const FPointerEvent& Event )
-			{
-				FReply TempReply = FReply::Unhandled();
-				if (Event.IsTouchEvent())
-			{
-					TempReply = TargetWidget.Widget->OnTouchEnded( TargetWidget.Geometry, Event );
-			}
-				if (!Event.IsTouchEvent() || (!TempReply.IsEventHandled() && this->bTouchFallbackToMouse))
-			{
-					TempReply = TargetWidget.Widget->OnMouseButtonUp( TargetWidget.Geometry, Event );
-			}
-				return TempReply;
-			});
+			FReply Reply =
+				FEventRouter::Route<FReply>( this, FEventRouter::FToLeafmostPolicy(MouseCaptorPath), PointerEvent, [this]( const FArrangedWidget& TargetWidget, const FPointerEvent& Event )
+				{
+					FReply TempReply = FReply::Unhandled();
+					if (Event.IsTouchEvent())
+					{
+						TempReply = TargetWidget.Widget->OnTouchEnded( TargetWidget.Geometry, Event );
+					}
 
+					if (!Event.IsTouchEvent() || (!TempReply.IsEventHandled() && this->bTouchFallbackToMouse))
+					{
+						TempReply = TargetWidget.Widget->OnMouseButtonUp( TargetWidget.Geometry, Event );
+					}
+					return TempReply;
+				});
+
 #if PLATFORM_MAC
 			// Activate a window under the mouse if it's inactive and mouse up didn't bring any window to front
 			TSharedPtr<SWindow> ActiveWindow = GetActiveTopLevelWindow();
 			if (MouseEvent.GetEffectingButton() == EKeys::LeftMouseButton && MouseCaptorPath.TopLevelWindow.IsValid() && ActiveWindow != MouseCaptorPath.TopLevelWindow
-				&& ActiveNativeWindow == [NSApp keyWindow] && ![(NSWindow*)MouseCaptorPath.TopLevelWindow->GetNativeWindow()->GetOSWindowHandle() isMiniaturized])
+				&& ActiveNativeWindow == [NSApp keyWindow] && ![(Route)MouseCaptorPath.TopLevelWindow->GetNativeWindow()->GetOSWindowHandle() isMiniaturized])
 			{
 				MouseCaptorPath.TopLevelWindow->BringToFront(true);
 			}
@@ -4330,17 +4275,17 @@
 	}
 	else
 	{
-		FWidgetPath WidgetsUnderCursor = LocateWindowUnderMouse( MouseEvent.GetScreenSpacePosition(), GetInteractiveTopLevelWindows() );
+		FWidgetPath LocalWidgetsUnderCursor = WidgetsUnderPointer.IsValid() ? WidgetsUnderPointer : LocateWindowUnderMouse( PointerEvent.GetScreenSpacePosition(), GetInteractiveTopLevelWindows() );
 
+		// Switch worlds widgets in the current path
+		FScopedSwitchWorldHack SwitchWorld( LocalWidgetsUnderCursor );
+		
 		// Cache the drag drop content and reset the pointer in case OnMouseButtonUpMessage re-enters as a result of OnDrop
 		const bool bIsDragDropping = IsDragDropping();
 		TSharedPtr< FDragDropOperation > LocalDragDropContent = DragDropContent;
 		DragDropContent.Reset();
 
-		// Switch worlds widgets in the current path
-		FScopedSwitchWorldHack SwitchWorld( WidgetsUnderCursor );
-		
-		FReply Reply = FEventRouter::Route<FReply>( this, FEventRouter::FBubblePolicy(WidgetsUnderCursor), MouseEvent, [&](const FArrangedWidget& CurWidget, const FPointerEvent& Event)
+		FReply Reply = FEventRouter::Route<FReply>( this, FEventRouter::FBubblePolicy(LocalWidgetsUnderCursor), PointerEvent, [&](const FArrangedWidget& CurWidget, const FPointerEvent& Event)
 		{
 			FReply TempReply = FReply::Unhandled();
 			if (Event.IsTouchEvent())
@@ -4362,167 +4307,17 @@
 		if ( bIsDragDropping )
 		{
 			// @todo slate : depending on SetEventPath() is not ideal.
-			MouseEvent.SetEventPath( WidgetsUnderCursor );
-			LocalDragDropContent->OnDrop( Reply.IsEventHandled(), MouseEvent );
-			WidgetsUnderCursorLastEvent = FWeakWidgetPath();
+			PointerEvent.SetEventPath( LocalWidgetsUnderCursor );
+			LocalDragDropContent->OnDrop( Reply.IsEventHandled(), PointerEvent );
+
+			WidgetsUnderCursorLastEvent.Remove( FUserAndPointer( PointerEvent.GetUserIndex(), PointerEvent.GetPointerIndex() ) );
 		}
+
 	}
-
-	// If in responsive mode throttle, leave it on mouse up.
-	if( MouseButtonDownResponsivnessThrottle.IsValid() )
-	{
-		FSlateThrottleManager::Get().LeaveResponsiveMode( MouseButtonDownResponsivnessThrottle );
-	}
-
-	if ( PressedMouseButtons.Num() == 0 )
-	{
-		// Release Capture
-		PlatformApplication->SetCapture( nullptr );
-	}
-
-	return true;
 }
 
-bool FSlateApplication::OnMouseWheel( const float Delta )
+bool FSlateApplication::RoutePointerMoveEvent(const FWidgetPath& WidgetsUnderPointer, FPointerEvent& PointerEvent, bool bIsSynthetic)
 {
-	const FVector2D CurrentCursorPosition = GetCursorPos();
-
-	FPointerEvent MouseWheelEvent(
-		CursorPointerIndex,
-		CurrentCursorPosition,
-		CurrentCursorPosition,
-		PressedMouseButtons,
-		EKeys::Invalid,
-		Delta,
-		PlatformApplication->GetModifierKeys()
-		);
-
-	return ProcessMouseWheelOrGestureEvent( MouseWheelEvent, nullptr );
-}
-
-bool FSlateApplication::ProcessMouseWheelOrGestureEvent( FPointerEvent& InWheelEvent, const FPointerEvent* InGestureEvent )
-{
-	QueueSynthesizedMouseMove();
-
-	const bool bShouldProcessEvent = InWheelEvent.GetWheelDelta() != 0 ||
-		(InGestureEvent != nullptr && InGestureEvent->GetGestureDelta()!=FVector2D::ZeroVector);
-	
-	if ( !bShouldProcessEvent )
-	{
-		return false;
-	}
-
-	LastUserInteractionTime = this->GetCurrentTime();
-	
-	// NOTE: We intentionally don't reset LastUserInteractionTimeForThrottling here so that the UI can be responsive while scrolling
-
-	FWidgetPath EventPath = ( MouseCaptor.HasCaptureForPointerIndex(InWheelEvent.GetPointerIndex()) )
-		? MouseCaptor.ToWidgetPath(InWheelEvent.GetPointerIndex())
-		: LocateWindowUnderMouse( InWheelEvent.GetScreenSpacePosition(), GetInteractiveTopLevelWindows() );
-
-	// Switch worlds widgets in the current path
-	FScopedSwitchWorldHack SwitchWorld( EventPath );
-
-	FReply Reply = FEventRouter::Route<FReply>( this, FEventRouter::FBubblePolicy(EventPath), InWheelEvent, [&InGestureEvent](const FArrangedWidget& CurWidget, const FPointerEvent& Event)
-	{
-		FReply TempReply = FReply::Unhandled();
-		// Gesture event gets first shot, if slate doesn't respond to it, we'll try the wheel event.
-		if( InGestureEvent != nullptr )
-		{
-			TempReply = CurWidget.Widget->OnTouchGesture( CurWidget.Geometry, *InGestureEvent );
-		}
-		
-		// Send the mouse wheel event if we haven't already handled the gesture version of this event.
-		if( !TempReply.IsEventHandled() )
-		{
-			TempReply = CurWidget.Widget->OnMouseWheel( CurWidget.Geometry, Event );
-		}
-
-		return TempReply;
-	} );
-
-	LOG_EVENT( bIsGestureEvent ? EEventLog::TouchGesture : EEventLog::MouseWheel, Reply );
-
-	return Reply.IsEventHandled();
-}
-
-bool FSlateApplication::OnMouseMove()
-{
-	// convert to touch event if we are faking it	
-	if (bIsFakingTouched)
-	{
-		return OnTouchMoved(PlatformApplication->Cursor->GetPosition(), 0, 0);
-	}
-	else if (!bIsGameFakingTouch && bIsFakingTouch)
-	{
-		return false;
-	}
-
-	bool Result = true;
-	const FVector2D CurrentCursorPosition = GetCursorPos();
-	const FVector2D LastCursorPosition = GetLastCursorPos();
-	if ( LastCursorPosition != CurrentCursorPosition )
-	{
-		LastMouseMoveTime = GetCurrentTime();
-
-		FPointerEvent MouseEvent(
-			CursorPointerIndex,
-			CurrentCursorPosition,
-			LastCursorPosition,
-			PressedMouseButtons,
-			EKeys::Invalid,
-			0,
-			PlatformApplication->GetModifierKeys()
-			);
-
-		Result = ProcessMouseMoveEvent( MouseEvent );
-	}
-
-	return Result;
-}
-
-bool FSlateApplication::OnRawMouseMove( const int32 X, const int32 Y )
-{
-	if ( X != 0 || Y != 0 )
-	{
-		FPointerEvent MouseEvent(
-			CursorPointerIndex,
-			GetCursorPos(),
-			GetLastCursorPos(),
-			FVector2D( X, Y ), 
-			PressedMouseButtons,
-			PlatformApplication->GetModifierKeys()
-		);
-
-		ProcessMouseMoveEvent(MouseEvent);
-	}
-	
-	return true;
-}
-
-bool FSlateApplication::ProcessMouseMoveEvent( FPointerEvent& MouseEvent, bool bIsSynthetic )
-{
-	if ( !bIsSynthetic )
-	{
-		QueueSynthesizedMouseMove();
-
-		// Detecting a mouse move of zero delta is our way of filtering out synthesized move events
-		const bool AllowSpawningOfToolTips = true;
-		UpdateToolTip( AllowSpawningOfToolTips );
-		
-		// Guard against synthesized mouse moves and only track user interaction if the cursor pos changed
-		LastUserInteractionTime = this->GetCurrentTime();
-	}
-
-	// When the event came from the OS, we are guaranteed to be over a slate window.
-	// Otherwise, we are synthesizing a MouseMove ourselves, and must verify that the
-	// cursor is indeed over a Slate window.
-	const bool bOverSlateWindow = !bIsSynthetic || PlatformApplication->IsCursorDirectlyOverSlateWindow();
-	
-	FWidgetPath WidgetsUnderCursor = bOverSlateWindow
-		? LocateWindowUnderMouse(MouseEvent.GetScreenSpacePosition(), GetInteractiveTopLevelWindows())
-		: FWidgetPath();
-
 	bool bHandled = false;
 
 	FWeakWidgetPath LastWidgetsUnderCursor;
@@ -4531,7 +4326,7 @@
 	bool bDragDetected = false;
 	if( DragDetector.DetectDragForWidget.IsValid() )
 	{	
-		const FVector2D DragDelta = (DragDetector.DetectDragStartLocation - MouseEvent.GetScreenSpacePosition());
+		const FVector2D DragDelta = (DragDetector.DetectDragStartLocation - PointerEvent.GetScreenSpacePosition());
 		bDragDetected = ( DragDelta.Size() > FSlateApplication::Get().GetDragTriggerDistance() );
 		if (bDragDetected)
 		{
@@ -4553,7 +4348,7 @@
 				// Switch worlds widgets in the current path
 				FScopedSwitchWorldHack SwitchWorld( DragDetectPath );
 
-				FReply Reply = FEventRouter::Route<FReply>(this, FEventRouter::FDirectPolicy(DetectDragForMe,DragDetectPath), MouseEvent, []( const FArrangedWidget& InDetectDragForMe, const FPointerEvent& TranslatedMouseEvent )
+				FReply Reply = FEventRouter::Route<FReply>(this, FEventRouter::FDirectPolicy(DetectDragForMe,DragDetectPath), PointerEvent, []( const FArrangedWidget& InDetectDragForMe, const FPointerEvent& TranslatedMouseEvent )
 				{
 					return InDetectDragForMe.Widget->OnDragDetected(InDetectDragForMe.Geometry, TranslatedMouseEvent );
 				});
@@ -4576,13 +4371,13 @@
 	else
 	{
 		// No Drag Detection
-		LastWidgetsUnderCursor = WidgetsUnderCursorLastEvent;
+		LastWidgetsUnderCursor = WidgetsUnderCursorLastEvent.FindRef( FUserAndPointer( PointerEvent.GetUserIndex(), PointerEvent.GetPointerIndex() ) );
 	}
 
 	// Send out mouse leave events
 	// If we are doing a drag and drop, we will send this event instead.
 	{
-		FDragDropEvent DragDropEvent( MouseEvent, DragDropContent );
+		FDragDropEvent DragDropEvent( PointerEvent, DragDropContent );
 		// Switch worlds widgets in the current path
 		FScopedSwitchWorldHack SwitchWorld( LastWidgetsUnderCursor.Window.Pin() );
 
@@ -4599,7 +4394,7 @@
 				const TSharedPtr<SWidget>& SomeWidgetPreviouslyUnderCursor = LastWidgetsUnderCursor.Widgets[WidgetIndex].Pin();
 				if( SomeWidgetPreviouslyUnderCursor.IsValid() )
 				{
-					TOptional<FArrangedWidget> FoundWidget = WidgetsUnderCursor.FindArrangedWidget( SomeWidgetPreviouslyUnderCursor.ToSharedRef() );
+					TOptional<FArrangedWidget> FoundWidget = WidgetsUnderPointer.FindArrangedWidget( SomeWidgetPreviouslyUnderCursor.ToSharedRef() );
 					const bool bWidgetNoLongerUnderMouse = !FoundWidget.IsSet();
 					if ( bWidgetNoLongerUnderMouse )
 					{
@@ -4618,7 +4413,7 @@
 						else
 						{
 							// Note that the event's pointer position is not translated.
-							SomeWidgetPreviouslyUnderCursor->OnMouseLeave( MouseEvent );
+							SomeWidgetPreviouslyUnderCursor->OnMouseLeave( PointerEvent );
 							LOG_EVENT( EEventLog::MouseLeave, SomeWidgetPreviouslyUnderCursor );
 						}			
 					}
@@ -4629,10 +4424,10 @@
 
 
 	FWidgetPath MouseCaptorPath;
-	if (MouseCaptor.HasCaptureForPointerIndex(MouseEvent.GetPointerIndex()))
+	if (MouseCaptor.HasCaptureForPointerIndex(PointerEvent.GetUserIndex(), PointerEvent.GetPointerIndex()))
 	{
 		//MouseCaptorPath = MouseCaptor.ToWidgetPath(MouseEvent.GetPointerIndex(), FWeakWidgetPath::EInterruptedPathHandling::ReturnInvalid);
-		MouseCaptorPath = MouseCaptor.ToWidgetPath(FWeakWidgetPath::EInterruptedPathHandling::ReturnInvalid, &MouseEvent );
+		MouseCaptorPath = MouseCaptor.ToWidgetPath(FWeakWidgetPath::EInterruptedPathHandling::ReturnInvalid, &PointerEvent );
 	}
 
 	if (MouseCaptorPath.IsValid())
@@ -4642,7 +4437,7 @@
 			// Switch worlds widgets in the current path
 			FScopedSwitchWorldHack SwitchWorld( MouseCaptorPath );
 			
-			FReply Reply = FEventRouter::Route<FReply>( this, FEventRouter::FToLeafmostPolicy(MouseCaptorPath), MouseEvent, [this]( const FArrangedWidget& MouseCaptorWidget, const FPointerEvent& Event )
+			FReply Reply = FEventRouter::Route<FReply>( this, FEventRouter::FToLeafmostPolicy(MouseCaptorPath), PointerEvent, [this]( const FArrangedWidget& MouseCaptorWidget, const FPointerEvent& Event )
 			{
 				FReply TempReply = FReply::Unhandled();
 				if (Event.IsTouchEvent())
@@ -4661,13 +4456,13 @@
 	else
 	{	
 		// Switch worlds widgets in the current path
-		FScopedSwitchWorldHack SwitchWorld( WidgetsUnderCursor );
+		FScopedSwitchWorldHack SwitchWorld( WidgetsUnderPointer );
 
 		// Send out mouse enter events.
 		if (IsDragDropping())
 		{
-			FDragDropEvent DragDropEvent( MouseEvent, DragDropContent );
-			FEventRouter::Route<FNoReply>( this, FEventRouter::FBubblePolicy(WidgetsUnderCursor), DragDropEvent, [&LastWidgetsUnderCursor](const FArrangedWidget& WidgetUnderCursor, const FDragDropEvent& InDragDropEvent)
+			FDragDropEvent DragDropEvent( PointerEvent, DragDropContent );
+			FEventRouter::Route<FNoReply>( this, FEventRouter::FBubblePolicy(WidgetsUnderPointer), DragDropEvent, [&LastWidgetsUnderCursor](const FArrangedWidget& WidgetUnderCursor, const FDragDropEvent& InDragDropEvent)
 			{
 				if ( !LastWidgetsUnderCursor.ContainsWidget(WidgetUnderCursor.Widget) )
 				{
@@ -4678,7 +4473,7 @@
 		}
 		else
 		{
-			FEventRouter::Route<FNoReply>( this, FEventRouter::FBubblePolicy(WidgetsUnderCursor), MouseEvent, [&LastWidgetsUnderCursor](const FArrangedWidget& WidgetUnderCursor, const FPointerEvent& PointerEvent)
+			FEventRouter::Route<FNoReply>( this, FEventRouter::FBubblePolicy(WidgetsUnderPointer), PointerEvent, [&LastWidgetsUnderCursor](const FArrangedWidget& WidgetUnderCursor, const FPointerEvent& PointerEvent)
 			{
 				if ( !LastWidgetsUnderCursor.ContainsWidget(WidgetUnderCursor.Widget) )
 				{
@@ -4689,7 +4484,7 @@
 		}
 		
 		// Bubble the MouseMove event.
-		FReply Reply = FEventRouter::Route<FReply>(this, FEventRouter::FBubblePolicy(WidgetsUnderCursor), MouseEvent, [&](const FArrangedWidget& CurWidget, const FPointerEvent& Event)
+		FReply Reply = FEventRouter::Route<FReply>(this, FEventRouter::FBubblePolicy(WidgetsUnderPointer), PointerEvent, [&](const FArrangedWidget& CurWidget, const FPointerEvent& Event)
 		{
 			FReply TempReply = FReply::Unhandled();
 
@@ -4716,15 +4511,15 @@
 	// custom (e.g. update the Drag/Drop preview based on content)
 	if (IsDragDropping())
 	{
-		FDragDropEvent DragDropEvent( MouseEvent, DragDropContent );
-		FScopedSwitchWorldHack SwitchWorld( WidgetsUnderCursor );
+		FDragDropEvent DragDropEvent( PointerEvent, DragDropContent );
+		FScopedSwitchWorldHack SwitchWorld( WidgetsUnderPointer );
 		DragDropContent->OnDragged( DragDropEvent );
 
 		// Update the window we're under for rendering the drag drop operation if
 		// it's a windowless drag drop operation.
-		if ( WidgetsUnderCursor.IsValid() )
+		if ( WidgetsUnderPointer.IsValid() )
 		{
-			DragDropWindowPtr = WidgetsUnderCursor.GetWindow();
+			DragDropWindowPtr = WidgetsUnderPointer.GetWindow();
 		}
 		else
 		{
@@ -4749,12 +4544,262 @@
 		DragDropWindowPtr = nullptr;
 	}
 
-	WidgetsUnderCursorLastEvent = FWeakWidgetPath( WidgetsUnderCursor );
+	WidgetsUnderCursorLastEvent.Add( FUserAndPointer( PointerEvent.GetUserIndex(), PointerEvent.GetPointerIndex() ), FWeakWidgetPath( WidgetsUnderPointer ) );
 
-	PointerIndexLastPositionMap.Add(MouseEvent.GetPointerIndex(), MouseEvent.GetScreenSpacePosition());
+	PointerIndexLastPositionMap.Add(PointerEvent.GetPointerIndex(), PointerEvent.GetScreenSpacePosition());
 	return bHandled;
 }
 
+bool FSlateApplication::OnMouseDoubleClick( const TSharedPtr< FGenericWindow >& PlatformWindow, const EMouseButtons::Type Button )
+{
+	if (bIsFakingTouch || bIsGameFakingTouch)
+	{
+		bIsFakingTouched = true;
+		return OnTouchStarted(PlatformWindow, PlatformApplication->Cursor->GetPosition(), 0, 0);
+	}
+
+	FKey Key = TranslateMouseButtonToKey( Button );
+
+	FPointerEvent MouseEvent(
+		CursorPointerIndex,
+		GetCursorPos(),
+		GetLastCursorPos(),
+		PressedMouseButtons,
+		Key,
+		0,
+		PlatformApplication->GetModifierKeys()
+		);
+
+	return ProcessMouseButtonDoubleClickEvent( PlatformWindow, MouseEvent );
+}
+
+bool FSlateApplication::ProcessMouseButtonDoubleClickEvent( const TSharedPtr< FGenericWindow >& PlatformWindow, FPointerEvent& InMouseEvent )
+{
+	QueueSynthesizedMouseMove();
+	LastUserInteractionTime = this->GetCurrentTime();
+	LastUserInteractionTimeForThrottling = LastUserInteractionTime;
+	
+	PlatformApplication->SetCapture( PlatformWindow );
+	PressedMouseButtons.Add( InMouseEvent.GetEffectingButton() );
+
+	FWidgetPath WidgetsUnderCursor = LocateWindowUnderMouse( InMouseEvent.GetScreenSpacePosition(), GetInteractiveTopLevelWindows() );
+
+	// Switch worlds widgets in the current path
+	FScopedSwitchWorldHack SwitchWorld( WidgetsUnderCursor );
+
+	FReply Reply = FEventRouter::Route<FReply>( this, FEventRouter::FBubblePolicy(WidgetsUnderCursor), InMouseEvent, [](const FArrangedWidget& TargetWidget, const FPointerEvent& Event)
+	{
+		return TargetWidget.Widget->OnMouseButtonDoubleClick( TargetWidget.Geometry, Event );
+	} );
+
+
+	LOG_EVENT( EEventLog::MouseButtonDoubleClick, Reply );
+
+	PointerIndexLastPositionMap.Add(InMouseEvent.GetPointerIndex(), InMouseEvent.GetScreenSpacePosition());
+
+	return Reply.IsEventHandled();
+}
+
+bool FSlateApplication::OnMouseUp( const EMouseButtons::Type Button )
+{
+	// convert to touch event if we are faking it	
+	if (bIsFakingTouch || bIsGameFakingTouch)
+	{
+		bIsFakingTouched = false;
+		return OnTouchEnded(PlatformApplication->Cursor->GetPosition(), 0, 0);
+	}
+
+	FKey Key = TranslateMouseButtonToKey( Button );
+
+	FPointerEvent MouseEvent(
+		CursorPointerIndex,
+		GetCursorPos(),
+		GetLastCursorPos(),
+		PressedMouseButtons,
+		Key,
+		0,
+		PlatformApplication->GetModifierKeys()
+		);
+
+	return ProcessMouseButtonUpEvent( MouseEvent );
+}
+
+bool FSlateApplication::ProcessMouseButtonUpEvent( FPointerEvent& MouseEvent )
+{
+	QueueSynthesizedMouseMove();
+	LastUserInteractionTime = this->GetCurrentTime();
+	LastUserInteractionTimeForThrottling = LastUserInteractionTime;
+	PressedMouseButtons.Remove( MouseEvent.GetEffectingButton() );
+
+	if (DragDetector.DetectDragForWidget.IsValid() && MouseEvent.GetEffectingButton() == DragDetector.DetectDragButton )
+	{
+		// The user has release the button that was supposed to start the drag; stop detecting it.
+		DragDetector = FDragDetector();
+	}
+
+	// An empty widget path is passed in.  As an optimization, one will be generated only if a captured mouse event isnt routed
+	FWidgetPath EmptyPath;
+	RoutePointerUpEvent( EmptyPath, MouseEvent );
+	
+
+	// If in responsive mode throttle, leave it on mouse up.
+	if( MouseButtonDownResponsivnessThrottle.IsValid() )
+	{
+		FSlateThrottleManager::Get().LeaveResponsiveMode( MouseButtonDownResponsivnessThrottle );
+	}
+
+	if ( PressedMouseButtons.Num() == 0 )
+	{
+		// Release Capture
+		PlatformApplication->SetCapture( nullptr );
+	}
+
+	return true;
+}
+
+bool FSlateApplication::OnMouseWheel( const float Delta )
+{
+	const FVector2D CurrentCursorPosition = GetCursorPos();
+
+	FPointerEvent MouseWheelEvent(
+		CursorPointerIndex,
+		CurrentCursorPosition,
+		CurrentCursorPosition,
+		PressedMouseButtons,
+		EKeys::Invalid,
+		Delta,
+		PlatformApplication->GetModifierKeys()
+		);
+
+	return ProcessMouseWheelOrGestureEvent( MouseWheelEvent, nullptr );
+}
+
+bool FSlateApplication::ProcessMouseWheelOrGestureEvent( FPointerEvent& InWheelEvent, const FPointerEvent* InGestureEvent )
+{
+	QueueSynthesizedMouseMove();
+
+	const bool bShouldProcessEvent = InWheelEvent.GetWheelDelta() != 0 ||
+		(InGestureEvent != nullptr && InGestureEvent->GetGestureDelta()!=FVector2D::ZeroVector);
+	
+	if ( !bShouldProcessEvent )
+	{
+		return false;
+	}
+
+	LastUserInteractionTime = this->GetCurrentTime();
+	
+	// NOTE: We intentionally don't reset LastUserInteractionTimeForThrottling here so that the UI can be responsive while scrolling
+
+	FWidgetPath EventPath = ( MouseCaptor.HasCaptureForPointerIndex(InWheelEvent.GetUserIndex(), InWheelEvent.GetPointerIndex()) )
+		? MouseCaptor.ToWidgetPath( InWheelEvent.GetUserIndex(), InWheelEvent.GetPointerIndex() )
+		: LocateWindowUnderMouse( InWheelEvent.GetScreenSpacePosition(), GetInteractiveTopLevelWindows() );
+
+	// Switch worlds widgets in the current path
+	FScopedSwitchWorldHack SwitchWorld( EventPath );
+
+	FReply Reply = FEventRouter::Route<FReply>( this, FEventRouter::FBubblePolicy(EventPath), InWheelEvent, [&InGestureEvent](const FArrangedWidget& CurWidget, const FPointerEvent& Event)
+	{
+		FReply TempReply = FReply::Unhandled();
+		// Gesture event gets first shot, if slate doesn't respond to it, we'll try the wheel event.
+		if( InGestureEvent != nullptr )
+		{
+			TempReply = CurWidget.Widget->OnTouchGesture( CurWidget.Geometry, *InGestureEvent );
+		}
+		
+		// Send the mouse wheel event if we haven't already handled the gesture version of this event.
+		if( !TempReply.IsEventHandled() )
+		{
+			TempReply = CurWidget.Widget->OnMouseWheel( CurWidget.Geometry, Event );
+		}
+
+		return TempReply;
+	} );
+
+	LOG_EVENT( bIsGestureEvent ? EEventLog::TouchGesture : EEventLog::MouseWheel, Reply );
+
+	return Reply.IsEventHandled();
+}
+
+bool FSlateApplication::OnMouseMove()
+{
+	// convert to touch event if we are faking it	
+	if (bIsFakingTouched)
+	{
+		return OnTouchMoved(PlatformApplication->Cursor->GetPosition(), 0, 0);
+	}
+	else if (!bIsGameFakingTouch && bIsFakingTouch)
+	{
+		return false;
+	}
+
+	bool Result = true;
+	const FVector2D CurrentCursorPosition = GetCursorPos();
+	const FVector2D LastCursorPosition = GetLastCursorPos();
+	if ( LastCursorPosition != CurrentCursorPosition )
+	{
+		LastMouseMoveTime = GetCurrentTime();
+
+		FPointerEvent MouseEvent(
+			CursorPointerIndex,
+			CurrentCursorPosition,
+			LastCursorPosition,
+			PressedMouseButtons,
+			EKeys::Invalid,
+			0,
+			PlatformApplication->GetModifierKeys()
+			);
+
+		Result = ProcessMouseMoveEvent( MouseEvent );
+	}
+
+	return Result;
+}
+
+bool FSlateApplication::OnRawMouseMove( const int32 X, const int32 Y )
+{
+	if ( X != 0 || Y != 0 )
+	{
+		FPointerEvent MouseEvent(
+			CursorPointerIndex,
+			GetCursorPos(),
+			GetLastCursorPos(),
+			FVector2D( X, Y ), 
+			PressedMouseButtons,
+			PlatformApplication->GetModifierKeys()
+		);
+
+		ProcessMouseMoveEvent(MouseEvent);
+	}
+	
+	return true;
+}
+
+bool FSlateApplication::ProcessMouseMoveEvent( FPointerEvent& MouseEvent, bool bIsSynthetic )
+{
+	if ( !bIsSynthetic )
+	{
+		QueueSynthesizedMouseMove();
+
+		// Detecting a mouse move of zero delta is our way of filtering out synthesized move events
+		const bool AllowSpawningOfToolTips = true;
+		UpdateToolTip( AllowSpawningOfToolTips );
+		
+		// Guard against synthesized mouse moves and only track user interaction if the cursor pos changed
+		LastUserInteractionTime = this->GetCurrentTime();
+	}
+
+	// When the event came from the OS, we are guaranteed to be over a slate window.
+	// Otherwise, we are synthesizing a MouseMove ourselves, and must verify that the
+	// cursor is indeed over a Slate window.
+	const bool bOverSlateWindow = !bIsSynthetic || PlatformApplication->IsCursorDirectlyOverSlateWindow();
+	
+	FWidgetPath WidgetsUnderCursor = bOverSlateWindow
+		? LocateWindowUnderMouse(MouseEvent.GetScreenSpacePosition(), GetInteractiveTopLevelWindows())
+		: FWidgetPath();
+
+	return RoutePointerMoveEvent( WidgetsUnderCursor, MouseEvent, bIsSynthetic );
+}
+
 bool FSlateApplication::OnCursorSet()
 {
 	QueryCursor();
diff --git a/Engine/Source/Runtime/Slate/Public/Framework/Application/SlateApplication.h b/Engine/Source/Runtime/Slate/Public/Framework/Application/SlateApplication.h
--- a/Engine/Source/Runtime/Slate/Public/Framework/Application/SlateApplication.h
+++ b/Engine/Source/Runtime/Slate/Public/Framework/Application/SlateApplication.h
@@ -1131,6 +1131,35 @@ public:
 	virtual EDropEffect::Type OnDragDrop( const TSharedPtr< FGenericWindow >& Window ) override;
 	virtual bool OnWindowAction( const TSharedRef< FGenericWindow >& PlatformWindow, const EWindowAction::Type InActionType ) override;
 
+private:
+
+	/**
+	 * Directly routes a pointer down event to the widgets in the specified widget path
+	 *
+	 * @param WidgetsUnderPointer	The path of widgets the event is routed to.
+	 * @param PointerEvent		The event data that is is routed to the widget path
+	 * 
+	 * @return The reply returned by the widget that handled the event
+	 */
+	FReply RoutePointerDownEvent(FWidgetPath& WidgetsUnderPointer, FPointerEvent& PointerEvent);
+
+	/**
+	 * Directly routes a pointer up event to the widgets in the specified widget path
+	 *
+	 * @param WidgetsUnderPointer	The path of widgets the event is routed to.
+	 * @param PointerEvent		The event data that is is routed to the widget path
+	 */
+	void RoutePointerUpEvent(FWidgetPath& WidgetsUnderPointer, FPointerEvent& PointerEvent);
+
+	/**
+	 * Directly routes a pointer move event to the widgets in the specified widget path
+	 *
+	 * @param WidgetsUnderPointer	The path of widgets the event is routed to.
+	 * @param PointerEvent		The event data that is is routed to the widget path
+	 * @param bIsSynthetic		Whether or not the move event is synthetic.  Synthetic pointer moves used simulate an event without the pointer actually moving 
+	 */
+	bool RoutePointerMoveEvent( const FWidgetPath& WidgetsUnderPointer, FPointerEvent& PointerEvent, bool bIsSynthetic );
+
 private:
 
 	TSharedRef< FGenericWindow > MakeWindow( TSharedRef<SWindow> InSlateWindow, const bool bShowImmediately );
@@ -1192,6 +1221,29 @@ private:
 	FSlateApplication();
 
 private:
+	/** Represents a single user and pointer index for a device 
+	 *  Used to uniquely track widget state per user and per device
+	 */
+	struct FUserAndPointer
+	{
+		uint32 UserIndex;
+		uint32 PointerIndex;
+		
+		FUserAndPointer( uint32 InUserIndex, uint32 InPointerIndex )
+			: UserIndex( InUserIndex )
+			, PointerIndex( InPointerIndex )
+		{}
+
+		bool operator==( const FUserAndPointer& Other ) const
+		{
+			return UserIndex == Other.UserIndex && PointerIndex == Other.PointerIndex;
+		}
+
+		friend uint32 GetTypeHash( const FUserAndPointer& UserAndPointer )
+		{
+			return UserAndPointer.UserIndex << 16 | UserAndPointer.PointerIndex;
+		}
+	};
 
 	/**
 	 * Creates a mouse move event for the last known cursor position.  This should be called every tick to make
@@ -1244,7 +1296,7 @@ private:
 	FMenuStack MenuStack;
 
 	/** A vertical slice through the tree of widgets on screen; it represents widgets that were under the cursor last time an event was processed */
-	FWeakWidgetPath WidgetsUnderCursorLastEvent;
+	TMap<FUserAndPointer, FWeakWidgetPath> WidgetsUnderCursorLastEvent;
 
 	/**
 	 * A helper class to wrap the weak path functionality. The advantage of using this
@@ -1263,7 +1315,7 @@ private:
 		/**
 		 * Returns whether or not the particular PointerIndex has capture.
 		 */
-		bool HasCaptureForPointerIndex(uint32 PointerIndex) const;
+		bool HasCaptureForPointerIndex(uint32 UserIndex, uint32 PointerIndex) const;
 
 		/**
 		 * Sets a new mouse captor widget for a specific pointer index, invalidating the previous one if any and calling
@@ -1273,13 +1325,13 @@ private:
 		 * @param EventPath		The path to the event.
 		 * @param Widget		The widget that wants to capture the mouse.
 		 */
-		void SetMouseCaptor(uint32 PointerIndex, const FWidgetPath& EventPath, TSharedPtr< SWidget > Widget );
+		void SetMouseCaptor(uint32 UserIndex, uint32 PointerIndex, const FWidgetPath& EventPath, TSharedPtr< SWidget > Widget );
 
 		/** Invalidates all current mouse captors. Calls OnMouseCaptureLost() on the current mouse captor if one exists */
 		void InvalidateCaptureForAllPointers();
 
 		/** Invalidates a specific mouse captor. Calls OnMouseCaptureLost() on the specific mouse captor if one exists */
-		void InvalidateCaptureForPointer(uint32 PointIndex);
+		void InvalidateCaptureForPointer(uint32 UserIndex, uint32 PointIndex);
 
 		/**
 		 * Retrieves a resolved FWidgetPath for a specific pointer index, if possible.
@@ -1289,7 +1341,7 @@ private:
 		 * @param PointerIndex				The index of the pointer which has capture.
 		 * @param InterruptedPathHandling	How to handled incomplete paths. "Truncate" will return a partial path, "ReturnInvalid" will return an empty path.
 		 */
-		FWidgetPath ToWidgetPath(uint32 PointerIndex, FWeakWidgetPath::EInterruptedPathHandling::Type InterruptedPathHandling = FWeakWidgetPath::EInterruptedPathHandling::Truncate );
+		FWidgetPath ToWidgetPath(uint32 UserIndex, uint32 PointerIndex, FWeakWidgetPath::EInterruptedPathHandling::Type InterruptedPathHandling = FWeakWidgetPath::EInterruptedPathHandling::Truncate );
 
 		FWidgetPath ToWidgetPath( FWeakWidgetPath::EInterruptedPathHandling::Type InterruptedPathHandling = FWeakWidgetPath::EInterruptedPathHandling::Truncate, const FPointerEvent* PointerEvent = nullptr );
 
@@ -1299,11 +1351,11 @@ private:
 		TArray<FWidgetPath> ToWidgetPaths();
 
 		/** Retrieves the weak path for a current mouse captor with a specific pointer index */
-		FWeakWidgetPath ToWeakPath(uint32 PointerIndex) const;
+		FWeakWidgetPath ToWeakPath(uint32 UserIndex, uint32 PointerIndex) const;
 		
 
 		/* Walks the weak path and retrieves the widget that is set as the current mouse captor with a specific pointer index */
-		TSharedPtr< SWidget > ToSharedWidget(uint32 PointerIndex) const;
+		TSharedPtr< SWidget > ToSharedWidget(uint32 UserIndex, uint32 PointerIndex) const;
 
 		/*
 		 * Retrieves an array of shared widget pointers for the active mouse captures.
@@ -1311,14 +1363,14 @@ private:
 		TArray<TSharedRef<SWidget>> ToSharedWidgets() const;
 
 		/* Walks the weak path and retrieves the window for the widget belonging to the mouse captor with the specified pointer index */
-		TSharedPtr< SWidget > ToSharedWindow(uint32 PointerIndex);
+		TSharedPtr< SWidget > ToSharedWindow(uint32 UserIndex, uint32 PointerIndex);
 
 	protected:
 		/** Call the OnMouseCaptureLost() handler for the widget captured by the specific pointer index */
-		void InformCurrentCaptorOfCaptureLoss(uint32 PointerIndex) const;
+		void InformCurrentCaptorOfCaptureLoss(uint32 UserIndex, uint32 PointerIndex) const;
 
 		/** A map of pointer indices to weak widget paths for the active mouse captures */
-		TMap<uint32, FWeakWidgetPath> PointerIndexToMouseCaptorWeakPathMap;
+		TMap<FUserAndPointer, FWeakWidgetPath> PointerIndexToMouseCaptorWeakPathMap;
 	};
 	/** The current mouse captor for the application, if any. */
 	MouseCaptorHelper MouseCaptor;
diff --git a/Engine/Source/Runtime/SlateCore/Private/Application/SlateApplicationBase.cpp b/Engine/Source/Runtime/SlateCore/Private/Application/SlateApplicationBase.cpp
--- a/Engine/Source/Runtime/SlateCore/Private/Application/SlateApplicationBase.cpp
+++ b/Engine/Source/Runtime/SlateCore/Private/Application/SlateApplicationBase.cpp
@@ -12,6 +12,7 @@ TSharedPtr<FSlateApplicationBase> FSlateApplicationBase::CurrentBaseApplication
 TSharedPtr<GenericApplication> FSlateApplicationBase::PlatformApplication = nullptr;
 // TODO: Identifier the cursor index in a smarter way.
 const uint32 FSlateApplicationBase::CursorPointerIndex = EKeys::NUM_TOUCH_KEYS - 1;
+const uint32 FSlateApplicationBase::CursorUserIndex = 0;
 
 FWidgetPath FHitTesting::LocateWidgetInWindow(FVector2D ScreenspaceMouseCoordinate, const TSharedRef<SWindow>& Window, bool bIgnoreEnabledStatus) const
 {
diff --git a/Engine/Source/Runtime/SlateCore/Public/Application/SlateApplicationBase.h b/Engine/Source/Runtime/SlateCore/Public/Application/SlateApplicationBase.h
--- a/Engine/Source/Runtime/SlateCore/Public/Application/SlateApplicationBase.h
+++ b/Engine/Source/Runtime/SlateCore/Public/Application/SlateApplicationBase.h
@@ -376,6 +376,7 @@ protected:
 
 public:
 	const static uint32 CursorPointerIndex;
+	const static uint32 CursorUserIndex;
 
 	/**
 	 * Returns the current instance of the application. The application should have been initialized before
