Build Fix: Mac SlateApplication
error: use of undeclared identifier 'MouseEvent'
error: use of undeclared identifier 'Route'
warning: declaration shadows a local variable [-Wshadow]

diff --git a/Engine/Source/Runtime/Slate/Private/Framework/Application/SlateApplication.cpp b/Engine/Source/Runtime/Slate/Private/Framework/Application/SlateApplication.cpp
--- a/Engine/Source/Runtime/Slate/Private/Framework/Application/SlateApplication.cpp
+++ b/Engine/Source/Runtime/Slate/Private/Framework/Application/SlateApplication.cpp
@@ -4092,7 +4092,7 @@ bool FSlateApplication::ProcessMouseButtonDownEvent( const TSharedPtr< FGenericW
 			FScopedSwitchWorldHack SwitchWorld(MouseCaptorPath);
 			bInGame = FApp::IsGame();
 
-			FReply Reply = FEventRouter::Route<FReply>(this, FEventRouter::FToLeafmostPolicy(MouseCaptorPath), MouseEvent, [] (const FArrangedWidget& InMouseCaptorWidget, const FPointerEvent& Event)
+			Reply = FEventRouter::Route<FReply>(this, FEventRouter::FToLeafmostPolicy(MouseCaptorPath), MouseEvent, [] (const FArrangedWidget& InMouseCaptorWidget, const FPointerEvent& Event)
 			{
 				return InMouseCaptorWidget.Widget->OnPreviewMouseButtonDown(InMouseCaptorWidget.Geometry, Event);
 			});
@@ -4235,7 +4235,7 @@ void FSlateApplication::RoutePointerUpEvent(FWidgetPath& WidgetsUnderPointer, FP
 {
 	if (MouseCaptor.HasCaptureForPointerIndex(PointerEvent.GetUserIndex(), PointerEvent.GetPointerIndex()))
 	{
-		//FWidgetPath MouseCaptorPath = MouseCaptor.ToWidgetPath(MouseEvent.GetPointerIndex());
+		//FWidgetPath MouseCaptorPath = MouseCaptor.ToWidgetPath(PointerEvent.GetPointerIndex());
 		FWidgetPath MouseCaptorPath = MouseCaptor.ToWidgetPath( FWeakWidgetPath::EInterruptedPathHandling::Truncate, &PointerEvent );
 		if ( ensureMsgf(MouseCaptorPath.Widgets.Num() > 0, TEXT("A window had a widget with mouse capture. That entire window has been dismissed before the mouse up could be processed.")) )
 		{
@@ -4264,8 +4264,8 @@ void FSlateApplication::RoutePointerUpEvent(FWidgetPath& WidgetsUnderPointer, FP
 #if PLATFORM_MAC
 			// Activate a window under the mouse if it's inactive and mouse up didn't bring any window to front
 			TSharedPtr<SWindow> ActiveWindow = GetActiveTopLevelWindow();
-			if (MouseEvent.GetEffectingButton() == EKeys::LeftMouseButton && MouseCaptorPath.TopLevelWindow.IsValid() && ActiveWindow != MouseCaptorPath.TopLevelWindow
-				&& ActiveNativeWindow == [NSApp keyWindow] && ![(Route)MouseCaptorPath.TopLevelWindow->GetNativeWindow()->GetOSWindowHandle() isMiniaturized])
+			if ( PointerEvent.GetEffectingButton() == EKeys::LeftMouseButton && MouseCaptorPath.TopLevelWindow.IsValid() && ActiveWindow != MouseCaptorPath.TopLevelWindow
+				&& ActiveNativeWindow == [NSApp keyWindow] && ![(NSWindow*)MouseCaptorPath.TopLevelWindow->GetNativeWindow()->GetOSWindowHandle() isMiniaturized] )
 			{
 				MouseCaptorPath.TopLevelWindow->BringToFront(true);
 			}
@@ -4473,11 +4473,11 @@ bool FSlateApplication::RoutePointerMoveEvent(const FWidgetPath& WidgetsUnderPoi
 		}
 		else
 		{
-			FEventRouter::Route<FNoReply>( this, FEventRouter::FBubblePolicy(WidgetsUnderPointer), PointerEvent, [&LastWidgetsUnderCursor](const FArrangedWidget& WidgetUnderCursor, const FPointerEvent& PointerEvent)
+			FEventRouter::Route<FNoReply>(this, FEventRouter::FBubblePolicy(WidgetsUnderPointer), PointerEvent, [&LastWidgetsUnderCursor] (const FArrangedWidget& WidgetUnderCursor, const FPointerEvent& Event)
 			{
 				if ( !LastWidgetsUnderCursor.ContainsWidget(WidgetUnderCursor.Widget) )
 				{
-					WidgetUnderCursor.Widget->OnMouseEnter( WidgetUnderCursor.Geometry, PointerEvent );
+					WidgetUnderCursor.Widget->OnMouseEnter( WidgetUnderCursor.Geometry, Event );
 				}
 				return FNoReply();
 			} );
