Fixed "LMB hold, drag, release" usage style on menus. Also, fixes pop-under when menus can be clicked accidentally after appearing.
[UE-18721] - 4.9 Regression: Slate Menus: Can no longer LMB hold on menu bar, drag down to a menu-entry and release LMB to select
[UE-18606] - Fix pop-under menus in editor causing menus to activate automatically

#change Changed menu entries to use new button class SMenuEntryButton that allows mouse-up clicking without mouse-down in certain circustances.
#fix Blocks mouse-up click handling on menu items that are under the cursor when they are created.

diff --git a/Engine/Source/Runtime/Slate/Private/Framework/MultiBox/SMenuEntryBlock.cpp b/Engine/Source/Runtime/Slate/Private/Framework/MultiBox/SMenuEntryBlock.cpp
--- a/Engine/Source/Runtime/Slate/Private/Framework/MultiBox/SMenuEntryBlock.cpp
+++ b/Engine/Source/Runtime/Slate/Private/Framework/MultiBox/SMenuEntryBlock.cpp
@@ -277,6 +277,106 @@ EVisibility SMenuEntryBlock::GetVisibility() const
 	return DirectActions.IsVisible();
 }
 
+/**
+* A button for a menu entry that has special mouse up handling
+*/
+class SMenuEntryButton : public SButton
+{
+public:
+	SLATE_BEGIN_ARGS(SMenuEntryButton)
+		: _Content()
+		, _ButtonStyle(nullptr)
+		, _ClickMethod(EButtonClickMethod::DownAndUp)
+	{}
+
+	/** Slot for this button's content (optional) */
+	SLATE_DEFAULT_SLOT(FArguments, Content)
+	/** The style to use */
+	SLATE_STYLE_ARGUMENT(FButtonStyle, ButtonStyle)
+	/** Sets the rules to use for determining whether the button was clicked.  This is an advanced setting and generally should be left as the default. */
+	SLATE_ARGUMENT(EButtonClickMethod::Type, ClickMethod)
+	/** Called when the button is clicked */
+	SLATE_EVENT(FOnClicked, OnClicked)
+
+	SLATE_END_ARGS()
+
+	enum class EResponseToMouseUp
+	{
+		Undetermined,
+		Handle,
+		DoNotHandle,
+	};
+
+	void Construct(const FArguments& InArgs)
+	{
+		SButton::FArguments ButtonArgs;
+		ButtonArgs.ButtonStyle(InArgs._ButtonStyle);
+		ButtonArgs.ClickMethod(InArgs._ClickMethod);
+		ButtonArgs.ToolTip(InArgs._ToolTip);
+		ButtonArgs.ContentPadding(FMargin(0, 2));
+		ButtonArgs.ForegroundColor(FSlateColor::UseForeground());
+		ButtonArgs.OnClicked(InArgs._OnClicked)
+		[
+			InArgs._Content.Widget
+		];
+
+		SButton::Construct(ButtonArgs);
+
+		ResponseToMouseUp = EResponseToMouseUp::Undetermined;
+	}
+
+	virtual void Tick(const FGeometry& AllottedGeometry, const double InCurrentTime, const float InDeltaTime) override
+	{
+		// On first tick, check mouse cursor position.
+		if (ResponseToMouseUp == EResponseToMouseUp::Undetermined)
+		{
+			FVector2D CursorPos = FSlateApplication::Get().GetCursorPos();
+
+			if (AllottedGeometry.IsUnderLocation(CursorPos))
+			{
+				// button was created under the mouse
+				ResponseToMouseUp = EResponseToMouseUp::DoNotHandle;
+			}
+			else
+			{
+				// button was NOT created under the mouse
+				ResponseToMouseUp = EResponseToMouseUp::Handle;
+			}
+		}
+
+		SButton::Tick(AllottedGeometry, InCurrentTime, InDeltaTime);
+	}
+
+	virtual FReply OnMouseButtonUp(const FGeometry& MyGeometry, const FPointerEvent& MouseEvent) override
+	{
+		bool bWasPressed = bIsPressed;
+
+		if (ResponseToMouseUp == EResponseToMouseUp::Handle)
+		{
+			bIsPressed = true;
+		}
+
+		FReply Reply = SButton::OnMouseButtonUp(MyGeometry, MouseEvent);
+
+		bIsPressed = bWasPressed;
+
+		return Reply;
+	}
+
+	virtual void OnMouseLeave(const FPointerEvent& MouseEvent) override
+	{
+		if (ResponseToMouseUp == EResponseToMouseUp::DoNotHandle)
+		{
+			ResponseToMouseUp = EResponseToMouseUp::Handle;
+		}
+
+		SButton::OnMouseLeave(MouseEvent);
+	}
+
+private:
+	EResponseToMouseUp ResponseToMouseUp;
+};
+
 TSharedRef< SWidget > SMenuEntryBlock::BuildMenuEntryWidget( const FMenuEntryBuildParams& InBuildParams )
 {
 	const TAttribute<FText>& Label = InBuildParams.Label;
@@ -452,28 +552,18 @@ TSharedRef< SWidget > SMenuEntryBlock::BuildMenuEntryWidget( const FMenuEntryBui
 		];
 
 	// Create a menu item button
-	TSharedPtr<SWidget> MenuEntryWidget = 
-		SNew( SButton )
-
+	TSharedPtr<SWidget> MenuEntryWidget = SNew(SMenuEntryButton)
 		// Use the menu item style for this button
 		.ButtonStyle( StyleSet, ISlateStyle::Join( StyleName, ".Button" ) )
-
 		// Set our click method for this menu item.  It will be different for pull-down/context menus.
 		.ClickMethod( ButtonClickMethod )
-
 		// Pass along the block's tool-tip string
 		.ToolTip( FMultiBoxSettings::ToolTipConstructor.Execute(EntryToolTip, nullptr, UICommand ) )
-
-		.ContentPadding(FMargin(0, 2))
-
-		.ForegroundColor( FSlateColor::UseForeground() )
+		// Bind the button's "on clicked" event to our object's method for this
+		.OnClicked(this, &SMenuEntryBlock::OnMenuItemButtonClicked)
 		[
 			CheckBoxAndButtonContent.ToSharedRef()
-		]
-
-		// Bind the button's "on clicked" event to our object's method for this
-		.OnClicked( this, &SMenuEntryBlock::OnMenuItemButtonClicked );
-
+		];
 
 	return MenuEntryWidget.ToSharedRef();
 }
