SListPanel is now publically accessible
This allows types derived from the various list implementations to inspect the panel itself directly

diff --git a/Engine/Source/Runtime/Slate/Private/Widgets/Views/SListPanel.h b/Engine/Source/Runtime/Slate/Private/Widgets/Views/SListPanel.h
deleted file mode 100644
diff --git a/Engine/Source/Runtime/Slate/Public/Widgets/Views/SListPanel.h b/Engine/Source/Runtime/Slate/Public/Widgets/Views/SListPanel.h
new file mode 100644
--- /dev/null
+++ b/Engine/Source/Runtime/Slate/Public/Widgets/Views/SListPanel.h
@@ -0,0 +1,149 @@
+// Copyright 1998-2015 Epic Games, Inc. All Rights Reserved.
+ 
+#pragma once
+ 
+ 
+/**
+ * A really simple panel that arranges its children in a vertical list with no spacing.
+ * Items in this panel have a uniform height.
+ * Also supports offsetting its items vertically.
+ */
+class SListPanel : public SPanel
+{
+public:
+	/** A ListPanel slot is very simple - it just stores a widget. */
+	class FSlot : public TSlotBase<FSlot>
+	{
+	public:
+		FSlot()
+		: TSlotBase<FSlot>()
+		{}
+	};
+	
+	/** Make a new ListPanel::Slot  */
+	static FSlot& Slot();
+	
+	/** Add a slot to the ListPanel */
+	FSlot& AddSlot(int32 InsertAtIndex = INDEX_NONE);
+	
+	SLATE_BEGIN_ARGS( SListPanel )
+		: _ItemWidth(0)
+		, _ItemHeight(16)
+		, _NumDesiredItems(0)
+		, _ItemAlignment(EListItemAlignment::EvenlyDistributed)
+		{
+			_Visibility = EVisibility::SelfHitTestInvisible;
+		}
+	
+		SLATE_ATTRIBUTE( float, ItemWidth )
+		SLATE_ATTRIBUTE( float, ItemHeight )
+		SLATE_ATTRIBUTE( int32, NumDesiredItems )
+		SLATE_ATTRIBUTE( EListItemAlignment, ItemAlignment )
+
+	SLATE_END_ARGS()
+	
+	SListPanel();
+
+	/**
+	 * Construct the widget
+	 *
+	 * @param InArgs   A declaration from which to construct the widget
+	 */
+	void Construct( const FArguments& InArgs );
+
+public:
+
+	// SWidget interface
+	virtual void OnArrangeChildren( const FGeometry& AllottedGeometry, FArrangedChildren& ArrangedChildren ) const override;
+	virtual FVector2D ComputeDesiredSize(float) const override;
+	virtual FChildren* GetChildren() override;
+	virtual void Tick( const FGeometry& AllottedGeometry, const double InCurrentTime, const float InDeltaTime ) override;
+	// End of SWidget interface
+
+	/** Fraction of the first item that we should offset by to simulate smooth scrolling. */
+	void SmoothScrollOffset(float OffsetInItems);
+
+	/** Set how much we should appear to have scrolled past the beginning/end of the list. */
+	void SetOverscrollAmount( float InOverscrollAmount );
+	
+	/** Remove all the children from this panel */
+	void ClearItems();
+
+	/** @return the uniform desired item width used when arranging children. */
+	float GetDesiredItemWidth() const;
+
+	/** @return the uniform item width used when arranging children. */
+	float GetItemWidth(const FGeometry& AllottedGeometry) const;
+
+	/** @return the uniform item width used when arranging children. */
+	float GetItemWidth(const FGeometry& AllottedGeometry, const EListItemAlignment ListItemAlignment) const;
+
+	/** @return the horizontal padding applied to each tile item */
+	float GetItemPadding(const FGeometry& AllottedGeometry) const;
+
+	/** @return the horizontal padding applied to each tile item */
+	float GetItemPadding(const FGeometry& AllottedGeometry, const EListItemAlignment ListItemAlignment) const;
+
+	/** @return the horizontal padding applied to all the items on a line */
+	float GetLinePadding(const FGeometry& AllottedGeometry, const int32 LineStartIndex) const;
+	
+	/** @return the uniform item height used when arranging children. */
+	float GetItemHeight() const;
+
+	/** Tells the list panel whether items in the list are pending a refresh */
+	void SetRefreshPending( bool IsPendingRefresh );
+
+	/** Returns true if this list panel is pending a refresh, false otherwise */
+	bool IsRefreshPending() const;
+
+	/** See ItemHeight attribute */
+	void SetItemHeight(TAttribute<float> Height);
+
+	/** See ItemWidth attribute */
+	void SetItemWidth(TAttribute<float> Width);
+	
+protected:
+
+	/** @return true if this panel should arrange items horizontally until it runs out of room, then create new rows */
+	bool ShouldArrangeHorizontally() const;
+	
+protected:
+
+	/** The children being arranged by this panel */
+	TPanelChildren<FSlot> Children;
+
+	/** The uniform item width used to arrange the children. Only relevant for tile views. */
+	TAttribute<float> ItemWidth;
+	
+	/** The uniform item height used to arrange the children */
+	TAttribute<float> ItemHeight;
+
+	/** Total number of items that the tree wants to visualize */
+	TAttribute<int32> NumDesiredItems;
+	
+	/** How should be horizontally aligned? Only relevant for tile views. */
+	TAttribute<EListItemAlignment> ItemAlignment;
+
+	/**
+	 * The offset of the view area from the top of the list in item heights.
+	 * Translate to physical units based on first item in list.
+	 */
+	float SmoothScrollOffsetInItems;
+
+	/** Amount scrolled past beginning/end of list in Slate Units. */
+	float OverscrollAmount;
+
+	/** The preferred number of rows that this widget should have */
+	int32 PreferredRowNum;
+
+	/**
+	 * When true, a refresh of the table view control that is using this panel is pending.
+	 * Some of the widgets in this panel are associated with items that may no longer be sound data.
+	 */
+	bool bIsRefreshPending;
+
+private:
+	
+	/** Used to pretend that the panel has no children when asked to cache desired size. */
+	static FNoChildren NoChildren;
+};
