Added option to STreeView for specifying whether to allow overscroll or not

diff --git a/Engine/Source/Runtime/Slate/Public/Widgets/Views/STreeView.h b/Engine/Source/Runtime/Slate/Public/Widgets/Views/STreeView.h
--- a/Engine/Source/Runtime/Slate/Public/Widgets/Views/STreeView.h
+++ b/Engine/Source/Runtime/Slate/Public/Widgets/Views/STreeView.h
@@ -88,6 +88,7 @@ public:
 		, _ClearSelectionOnClick(true)
 		, _ExternalScrollbar()
 		, _ConsumeMouseWheel( EConsumeMouseWheel::WhenScrollingPossible )
+		, _AllowOverscroll(EAllowOverscroll::Yes)
 		{}
 
 		SLATE_EVENT( FOnGenerateRow, OnGenerateRow )
@@ -121,6 +122,8 @@ public:
 		SLATE_ARGUMENT( TSharedPtr<SScrollBar>, ExternalScrollbar )
 
 		SLATE_ARGUMENT( EConsumeMouseWheel, ConsumeMouseWheel );
+		
+		SLATE_ARGUMENT( EAllowOverscroll, AllowOverscroll );
 
 	SLATE_END_ARGS()
 
@@ -146,6 +149,7 @@ public:
 
 		this->bClearSelectionOnClick = InArgs._ClearSelectionOnClick;
 		this->ConsumeMouseWheel = InArgs._ConsumeMouseWheel;
+		this->AllowOverscroll = InArgs._AllowOverscroll;
 
 		// Check for any parameters that the coder forgot to specify.
 		FString ErrorString;
