Re-added menu opening code on Mac that fixed previous bug UE-16702
[UE-18806] - Viewport buttons on Mac are difficult to open

diff --git a/Engine/Source/Runtime/Slate/Private/Framework/Application/SlateApplication.cpp b/Engine/Source/Runtime/Slate/Private/Framework/Application/SlateApplication.cpp
--- a/Engine/Source/Runtime/Slate/Private/Framework/Application/SlateApplication.cpp
+++ b/Engine/Source/Runtime/Slate/Private/Framework/Application/SlateApplication.cpp
@@ -4211,7 +4211,7 @@ FReply FSlateApplication::RoutePointerDownEvent(FWidgetPath& WidgetsUnderPointer
 
 #if PLATFORM_MAC
 		TSharedPtr<SWindow> TopLevelWindow = WidgetsUnderPointer.TopLevelWindow;
-		if ( TopLevelWindow.IsValid() )
+		if ( bNeedToActivateWindow && TopLevelWindow.IsValid() )
 		{
 			// Clicking on a context menu should not activate anything
 			// @todo: This needs to be updated when we have window type in SWindow and we no longer have to guess if WidgetsUnderCursor.TopLevelWindow is a menu
