Intertial scroll methods and cached geometry of STableViewBase are now available to derived types

diff --git a/Engine/Source/Runtime/Slate/Public/Widgets/Views/STableViewBase.h b/Engine/Source/Runtime/Slate/Public/Widgets/Views/STableViewBase.h
--- a/Engine/Source/Runtime/Slate/Public/Widgets/Views/STableViewBase.h
+++ b/Engine/Source/Runtime/Slate/Public/Widgets/Views/STableViewBase.h
@@ -292,7 +292,7 @@ protected:
 	/**	Whether the software cursor should be drawn in the viewport */
 	bool bShowSoftwareCursor;
 
-private:
+protected:
 
 	/** Check whether the current state of the table warrants inertial scroll by the specified amount */
 	bool CanUseInertialScroll( float ScrollAmount ) const;
