Slate: Optionally show the text value associated with the numeric value.

diff --git a/Engine/Source/Runtime/Slate/Public/Widgets/Input/SNumericDropDown.h b/Engine/Source/Runtime/Slate/Public/Widgets/Input/SNumericDropDown.h
--- a/Engine/Source/Runtime/Slate/Public/Widgets/Input/SNumericDropDown.h
+++ b/Engine/Source/Runtime/Slate/Public/Widgets/Input/SNumericDropDown.h
@@ -51,6 +51,7 @@ public:
 		: _DropDownValues()
 		, _LabelText()
 		, _Orientation( EOrientation::Orient_Horizontal )
+		, _bShowNamedValue( false )
 		, _MinDesiredValueWidth( 40 )
 		{}
 
@@ -62,6 +63,8 @@ public:
 		SLATE_ATTRIBUTE( EOrientation, Orientation )
 		/** Controls the minimum width for the text box portion of the control. */
 		SLATE_ATTRIBUTE( float, MinDesiredValueWidth )
+		/** Toggle to show the drop down text value if the value matches the numeric value. */
+		SLATE_ATTRIBUTE( bool, bShowNamedValue )
 		/** The value displayed by the control. */
 		SLATE_ATTRIBUTE( NumericType, Value )
 		/** The callback for when the value changes. */
@@ -74,6 +77,7 @@ public:
 		DropDownValues = InArgs._DropDownValues;
 		LabelText = InArgs._LabelText;
 		Orientation = InArgs._Orientation;
+		bShowNamedValue = InArgs._bShowNamedValue;
 		Value = InArgs._Value;
 		OnValueChanged = InArgs._OnValueChanged;
 
@@ -131,6 +135,16 @@ private:
 
 	FText GetValueText() const
 	{
+		if (bShowNamedValue.Get())
+		{
+			for ( FNamedValue DropDownValue : DropDownValues )
+			{
+				if (FMath::IsNearlyEqual(DropDownValue.GetValue(), Value.Get()))
+				{
+					return DropDownValue.GetName();
+				}
+			}
+		}
 		return FText::AsNumber( Value.Get() );
 	}
 
@@ -166,6 +180,7 @@ private:
 	TArray<FNamedValue> DropDownValues;
 	TAttribute<FText> LabelText;
 	TAttribute<EOrientation> Orientation;
+	TAttribute<bool> bShowNamedValue;
 	TAttribute<NumericType> Value;
 	FOnValueChanged OnValueChanged;
 };
\ No newline at end of file
