SEditableText no longer displays it's own tab focus border. This responsibility should be delegated to the containing box.

Merging using UE4-To-UE4-LauncherDev

diff --git a/Engine/Source/Runtime/Slate/Private/Widgets/Input/SEditableText.cpp b/Engine/Source/Runtime/Slate/Private/Widgets/Input/SEditableText.cpp
--- a/Engine/Source/Runtime/Slate/Private/Widgets/Input/SEditableText.cpp
+++ b/Engine/Source/Runtime/Slate/Private/Widgets/Input/SEditableText.cpp
@@ -1591,6 +1591,12 @@ FCursorReply SEditableText::OnCursorQuery( const FGeometry& MyGeometry, const FP
 }
 
 
+const FSlateBrush* SEditableText::GetFocusBrush() const
+{
+	return nullptr;
+}
+
+
 void SEditableText::SelectText( const int32 InOldCaretPosition )
 {
 	if( InOldCaretPosition != CaretPosition )
diff --git a/Engine/Source/Runtime/Slate/Public/Widgets/Input/SEditableText.h b/Engine/Source/Runtime/Slate/Public/Widgets/Input/SEditableText.h
--- a/Engine/Source/Runtime/Slate/Public/Widgets/Input/SEditableText.h
+++ b/Engine/Source/Runtime/Slate/Public/Widgets/Input/SEditableText.h
@@ -414,6 +414,7 @@ protected:
 	virtual FReply OnMouseMove( const FGeometry& InMyGeometry, const FPointerEvent& InMouseEvent ) override;
 	virtual FReply OnMouseButtonDoubleClick( const FGeometry& InMyGeometry, const FPointerEvent& InMouseEvent ) override;
 	virtual FCursorReply OnCursorQuery( const FGeometry& MyGeometry, const FPointerEvent& CursorEvent ) const override;
+	virtual const FSlateBrush* GetFocusBrush() const;
 	// END SWidget interface
 
 protected:
