[INTEGRATE] merge ue4<-ue4-launcherDev cl 2640436

diff --git a/Engine/Source/Runtime/Slate/Private/Widgets/Layout/SScrollBox.cpp b/Engine/Source/Runtime/Slate/Private/Widgets/Layout/SScrollBox.cpp
--- a/Engine/Source/Runtime/Slate/Private/Widgets/Layout/SScrollBox.cpp
+++ b/Engine/Source/Runtime/Slate/Private/Widgets/Layout/SScrollBox.cpp
@@ -165,6 +165,7 @@ void SScrollBox::Construct( const FArguments& InArgs )
 	ConsumeMouseWheel = InArgs._ConsumeMouseWheel;
 	TickScrollDelta = 0;
 	AllowOverscroll = InArgs._AllowOverscroll;
+	DestinationScrollingWidgetIntoView = EDescendantScrollDestination::IntoView;
 	bAnimateScrollingWidgetIntoView = false;
 
 	if (InArgs._ExternalScrollbar.IsValid())
@@ -368,16 +369,17 @@ void SScrollBox::ScrollToEnd()
 	bScrollToEnd = true;
 }
 
-void SScrollBox::ScrollDescendantIntoView(const TSharedPtr<SWidget>& WidgetToFind, bool InAnimateScroll)
+void SScrollBox::ScrollDescendantIntoView(const TSharedPtr<SWidget>& WidgetToFind, bool InAnimateScroll, EDescendantScrollDestination InDestination)
 {
 	WidgetToScrollIntoView = WidgetToFind;
+	DestinationScrollingWidgetIntoView = InDestination;
 	bAnimateScrollingWidgetIntoView = InAnimateScroll;
 
 	// This will force the active timer system to wakeup for a frame to ensure we tick at least once.
 	RegisterActiveTimer(0.f, FWidgetActiveTimerDelegate::CreateSP(this, &SScrollBox::UpdateInertialScroll));
 }
 
-bool SScrollBox::ScrollDescendantIntoView(const FGeometry& MyGeometry, const TSharedPtr<SWidget>& WidgetToFind, bool InAnimateScroll)
+bool SScrollBox::ScrollDescendantIntoView(const FGeometry& MyGeometry, const TSharedPtr<SWidget>& WidgetToFind, bool InAnimateScroll, EDescendantScrollDestination InDestination)
 {
 	// We need to safely find the one WidgetToFind among our descendants.
 	TSet< TSharedRef<SWidget> > WidgetsToFind;
@@ -394,32 +396,45 @@ bool SScrollBox::ScrollDescendantIntoView(const FGeometry& MyGeometry, const TSh
 		// @todo: This is a workaround because DesiredScrollOffset can exceed the ScrollMax when mouse dragging on the scroll bar and we need it clamped here or the offset is wrong
 		ScrollBy(MyGeometry, 0, EAllowOverscroll::No, false);
 
-		// If the scale is 0, we can't really calculate things properly.
-		if ( MyGeometry.Scale != 0 )
+		float ScrollOffset = 0.0f;
+		if ( InDestination == EDescendantScrollDestination::TopOrLeft )
 		{
 			// Calculate how much we would need to scroll to bring this to the top/left of the scroll box
-			// NOTE: The scrollbox does all offset calculation using desired sizes which doesn't involve scale,
-			//       so we need to remove the scale from positions so that they are in the same unscaled space.
-			const float WidgetStartPosition = GetScrollComponentFromVector(WidgetGeometry->Geometry.AbsolutePosition / MyGeometry.Scale);
-			const float WidgetEndPosition = GetScrollComponentFromVector(WidgetGeometry->Geometry.AbsolutePosition / MyGeometry.Scale + WidgetGeometry->Geometry.GetLocalSize());
-			const float ViewStartPosition = GetScrollComponentFromVector(MyGeometry.AbsolutePosition / MyGeometry.Scale);
-			const float ViewEndPosition = GetScrollComponentFromVector(MyGeometry.AbsolutePosition / MyGeometry.Scale + MyGeometry.GetLocalSize());
-
-			const float ViewDelta = ( ViewEndPosition - ViewStartPosition );
-			const float WidgetDelta = ( WidgetEndPosition - WidgetStartPosition );
-
-			if ( WidgetStartPosition < ViewStartPosition )
-			{
-				const float ScrollOffset = WidgetStartPosition - ViewStartPosition;
-				ScrollBy(MyGeometry, ScrollOffset, EAllowOverscroll::No, InAnimateScroll);
-			}
-			else if ( WidgetEndPosition > ViewEndPosition )
+			const float WidgetPosition = GetScrollComponentFromVector(WidgetGeometry->Geometry.AbsolutePosition);
+			const float MyPosition = GetScrollComponentFromVector(MyGeometry.AbsolutePosition);
+			ScrollOffset = WidgetPosition - MyPosition;
+		}
+		else
+		{
+			// If the scale is 0, we can't really calculate things properly.
+			if ( MyGeometry.Scale != 0 )
 			{
-				const float ScrollOffset = ( WidgetEndPosition - ViewDelta ) - ViewStartPosition;
-				ScrollBy(MyGeometry, ScrollOffset, EAllowOverscroll::No, InAnimateScroll);
+				// Calculate how much we would need to scroll to bring this to the top/left of the scroll box
+				// NOTE: The scrollbox does all offset calculation using desired sizes which doesn't involve scale,
+				//       so we need to remove the scale from positions so that they are in the same unscaled space.
+				const float WidgetStartPosition = GetScrollComponentFromVector(WidgetGeometry->Geometry.AbsolutePosition / MyGeometry.Scale);
+				const float WidgetEndPosition = GetScrollComponentFromVector(WidgetGeometry->Geometry.AbsolutePosition / MyGeometry.Scale + WidgetGeometry->Geometry.GetLocalSize());
+				const float ViewStartPosition = GetScrollComponentFromVector(MyGeometry.AbsolutePosition / MyGeometry.Scale);
+				const float ViewEndPosition = GetScrollComponentFromVector(MyGeometry.AbsolutePosition / MyGeometry.Scale + MyGeometry.GetLocalSize());
+
+				const float ViewDelta = ( ViewEndPosition - ViewStartPosition );
+				const float WidgetDelta = ( WidgetEndPosition - WidgetStartPosition );
+
+				if ( WidgetStartPosition < ViewStartPosition )
+				{
+					ScrollOffset = WidgetStartPosition - ViewStartPosition;
+				}
+				else if ( WidgetEndPosition > ViewEndPosition )
+				{
+					ScrollOffset = ( WidgetEndPosition - ViewDelta ) - ViewStartPosition;
+				}
 			}
 		}
 
+		if ( ScrollOffset != 0.0f )
+		{
+			ScrollBy(MyGeometry, ScrollOffset, EAllowOverscroll::No, InAnimateScroll);
+		}
 		return true;
 	}
 	return false;
@@ -537,7 +552,7 @@ void SScrollBox::Tick( const FGeometry& AllottedGeometry, const double InCurrent
 	// If we needed a widget to be scrolled into view, make that happen.
 	if ( WidgetToScrollIntoView.IsValid() )
 	{
-		ScrollDescendantIntoView(AllottedGeometry, WidgetToScrollIntoView, bAnimateScrollingWidgetIntoView);
+		ScrollDescendantIntoView(AllottedGeometry, WidgetToScrollIntoView, bAnimateScrollingWidgetIntoView, DestinationScrollingWidgetIntoView);
 		WidgetToScrollIntoView.Reset();
 	}
 
@@ -735,11 +750,11 @@ bool SScrollBox::ScrollBy(const FGeometry& AllottedGeometry, float ScrollAmount,
 
 	const float PreviousScrollOffset = DesiredScrollOffset;
 
-	const float ScrollMin = 0.0f;
-	const float ScrollMax = ContentSize - GetScrollComponentFromVector(ScrollPanelGeometry.Size);
-
 	if ( ScrollAmount != 0 )
 	{
+		const float ScrollMin = 0.0f;
+		const float ScrollMax = ContentSize - GetScrollComponentFromVector(ScrollPanelGeometry.Size);
+
 		if ( Overscrolling == EAllowOverscroll::Yes && Overscroll.ShouldApplyOverscroll(DesiredScrollOffset == 0, DesiredScrollOffset == ScrollMax, ScrollAmount) )
 		{
 			Overscroll.ScrollBy(ScrollAmount);
diff --git a/Engine/Source/Runtime/Slate/Public/Widgets/Layout/SScrollBox.h b/Engine/Source/Runtime/Slate/Public/Widgets/Layout/SScrollBox.h
--- a/Engine/Source/Runtime/Slate/Public/Widgets/Layout/SScrollBox.h
+++ b/Engine/Source/Runtime/Slate/Public/Widgets/Layout/SScrollBox.h
@@ -108,13 +108,24 @@ public:
 
 	void ScrollToEnd();
 
+	/** Where to scroll the descendant to */
+	enum EDescendantScrollDestination
+	{
+		/** Scroll the widget into view */
+		IntoView,
+
+		/** Always scroll the widget so it appears at the top/Left of the scrollable area */
+		TopOrLeft,
+	};
+
 	/** 
 	 * Attempt to scroll a widget into view, will safely handle non-descendant widgets 
 	 *
 	 * @param WidgetToFind The widget whose geometry we wish to discover.
 	 * @param InAnimateScroll	Whether or not to animate the scroll
+	 * @param InDestination		Where we want the child widget to stop.
 	 */
-	void ScrollDescendantIntoView(const TSharedPtr<SWidget>& WidgetToFind, bool InAnimateScroll = true);
+	void ScrollDescendantIntoView(const TSharedPtr<SWidget>& WidgetToFind, bool InAnimateScroll = true, EDescendantScrollDestination InDestination = EDescendantScrollDestination::IntoView);
 
 	/** Get the current orientation of the scrollbox. */
 	EOrientation GetOrientation();
@@ -181,7 +192,7 @@ private:
 	float DesiredScrollOffset;
 
 	/** Scrolls or begins scrolling a widget into view, only valid to call when we have layout geometry. */
-	bool ScrollDescendantIntoView(const FGeometry& MyGeometry, const TSharedPtr<SWidget>& WidgetToFind, bool InAnimateScroll);
+	bool ScrollDescendantIntoView(const FGeometry& MyGeometry, const TSharedPtr<SWidget>& WidgetToFind, bool InAnimateScroll = true, EDescendantScrollDestination InDestination = EDescendantScrollDestination::IntoView);
 
 	/**
 	 * Scroll the view by ScrollAmount given its currently AllottedGeometry.
@@ -258,6 +269,9 @@ private:
 	/** The widget to scroll into view next frame. */
 	TSharedPtr<SWidget> WidgetToScrollIntoView;
 
+	/** Where we should scroll the descendant to */
+	EDescendantScrollDestination DestinationScrollingWidgetIntoView;
+
 	/** Whether we should scroll the widget into view over time. */
 	bool bAnimateScrollingWidgetIntoView : 1;
 
