[UEDOC-2213] Fixed typos.

diff --git a/Engine/Source/Runtime/Slate/Public/Widgets/Layout/SScissorRectBox.h b/Engine/Source/Runtime/Slate/Public/Widgets/Layout/SScissorRectBox.h
--- a/Engine/Source/Runtime/Slate/Public/Widgets/Layout/SScissorRectBox.h
+++ b/Engine/Source/Runtime/Slate/Public/Widgets/Layout/SScissorRectBox.h
@@ -3,13 +3,13 @@
 #pragma once
 
 /**
- * A ScirrorRectBox is a widget that clips all its children using a scissor rect that matches it clipping rect
+ * A ScissorRectBox is a widget that clips all its children using a scissor rect that matches it clipping rect
  * in LAYOUT SPACE (the way it is provided in OnPaint).
  * This ensures that any child widgets with render transforms do not "bleed out" of this parent container.
  * Note there is a significant performance impact to using one of these widgets because any widgets inside 
  * this one end up in a separate set of batches using that scissor rect.
  * 
- * Keep in mind that this widget will use the LAYOUT clippping rect to clip children, since hardware clipping
+ * Keep in mind that this widget will use the LAYOUT clipping rect to clip children, since hardware clipping
  * rects do not support non-aligned rectangles. If a render transform is applied to this widget itself, the results
  * will not be as expected.
  * 
