Fixed unsafe active widget index in SWidgetSwitcher that could act as out-of-bounds index into child array.
[UE-19741] - Switcher widget causes a crash when dragging a widget out of it and on to the canvas panel

diff --git a/Engine/Source/Runtime/Slate/Private/Widgets/Layout/SWidgetSwitcher.cpp b/Engine/Source/Runtime/Slate/Private/Widgets/Layout/SWidgetSwitcher.cpp
--- a/Engine/Source/Runtime/Slate/Private/Widgets/Layout/SWidgetSwitcher.cpp
+++ b/Engine/Source/Runtime/Slate/Private/Widgets/Layout/SWidgetSwitcher.cpp
@@ -21,6 +21,12 @@ SWidgetSwitcher::FSlot& SWidgetSwitcher::AddSlot( int32 SlotIndex )
 	else
 	{
 		AllChildren.Insert(NewSlot, SlotIndex);
+
+		// adjust the active WidgetIndex based on this slot change
+		if (!WidgetIndex.IsBound() && WidgetIndex.Get() >= SlotIndex)
+		{
+			WidgetIndex = WidgetIndex.Get() + 1;
+		}
 	}
 
 	return *NewSlot;
@@ -42,10 +48,11 @@ void SWidgetSwitcher::Construct( const FArguments& InArgs )
 
 TSharedPtr<SWidget> SWidgetSwitcher::GetActiveWidget( ) const
 {
-	const int32 ActiveWidgetIndex = WidgetIndex.Get();
-	if (ActiveWidgetIndex >= 0)
+	const FSlot* ActiveSlot = GetActiveSlot();
+
+	if (ActiveSlot != nullptr)
 	{
-		return AllChildren[ActiveWidgetIndex].GetWidget();
+		return ActiveSlot->GetWidget();
 	}
 
 	return nullptr;
@@ -85,6 +92,16 @@ int32 SWidgetSwitcher::RemoveSlot( TSharedRef<SWidget> WidgetToRemove )
 	{
 		if (AllChildren[SlotIndex].GetWidget() == WidgetToRemove)
 		{
+			// adjust the active WidgetIndex based on this slot change
+			if (!WidgetIndex.IsBound())
+			{
+				int32 ActiveWidgetIndex = WidgetIndex.Get();
+				if (ActiveWidgetIndex > 0 && ActiveWidgetIndex >= SlotIndex)
+				{
+					WidgetIndex = ActiveWidgetIndex - 1;
+				}
+			}
+
 			AllChildren.RemoveAt(SlotIndex);
 			return SlotIndex;
 		}
@@ -109,18 +126,40 @@ void SWidgetSwitcher::OnArrangeChildren( const FGeometry& AllottedGeometry, FArr
 
 	if (AllChildren.Num() > 0)
 	{
-		ArrangeSingleChild(AllottedGeometry, ArrangedChildren, AllChildren[WidgetIndex.Get()], ContentScale);
+		const FSlot* ActiveSlotPtr = GetActiveSlot();	// Returns null if unsafe attribute WidgetIndex is out-of-bounds
+		if (ActiveSlotPtr != nullptr)
+		{
+			ArrangeSingleChild(AllottedGeometry, ArrangedChildren, *ActiveSlotPtr, ContentScale);
+		}
 	}
 }
 	
 FVector2D SWidgetSwitcher::ComputeDesiredSize( float ) const
 {
-	return AllChildren.Num() > 0
-		? AllChildren[WidgetIndex.Get()].GetWidget()->GetDesiredSize()
-		: FVector2D::ZeroVector;
+	if (AllChildren.Num() > 0)
+	{
+		const FSlot* ActiveSlotPtr = GetActiveSlot();	// Returns null if unsafe attribute WidgetIndex is out-of-bounds
+		if (ActiveSlotPtr != nullptr)
+		{
+			return ActiveSlotPtr->GetWidget()->GetDesiredSize();
+		}
+	}
+
+	return FVector2D::ZeroVector;
 }
 
 FChildren* SWidgetSwitcher::GetChildren( )
 {
 	return &OneDynamicChild;
 }
+
+const SWidgetSwitcher::FSlot* SWidgetSwitcher::GetActiveSlot() const
+{
+	const int32 ActiveWidgetIndex = WidgetIndex.Get();
+	if (ActiveWidgetIndex >= 0 && ActiveWidgetIndex < AllChildren.Num())
+	{
+		return &AllChildren[ActiveWidgetIndex];
+	}
+
+	return nullptr;
+}
diff --git a/Engine/Source/Runtime/Slate/Public/Widgets/Layout/SWidgetSwitcher.h b/Engine/Source/Runtime/Slate/Public/Widgets/Layout/SWidgetSwitcher.h
--- a/Engine/Source/Runtime/Slate/Public/Widgets/Layout/SWidgetSwitcher.h
+++ b/Engine/Source/Runtime/Slate/Public/Widgets/Layout/SWidgetSwitcher.h
@@ -142,6 +142,8 @@ protected:
 	virtual FVector2D ComputeDesiredSize( float ) const override;
 	virtual FChildren* GetChildren( ) override;
 
+	const FSlot* GetActiveSlot() const;
+
 private:
 
 	/** Required to implement GetChildren() in a way that can dynamically return the currently active child. */
