fix for crash when accessing setting the text from SEditText on iOS
#ios

diff --git a/Engine/Source/Runtime/Slate/Private/Framework/Text/IOS/IOSPlatformTextField.cpp b/Engine/Source/Runtime/Slate/Private/Framework/Text/IOS/IOSPlatformTextField.cpp
--- a/Engine/Source/Runtime/Slate/Private/Framework/Text/IOS/IOSPlatformTextField.cpp
+++ b/Engine/Source/Runtime/Slate/Private/Framework/Text/IOS/IOSPlatformTextField.cpp
@@ -40,6 +40,8 @@ void FIOSPlatformTextField::ShowVirtualKeyboard(bool bShow, int32 UserIndex, TSh
 -(void)show:(TSharedPtr<IVirtualKeyboardEntry>)InTextWidget
 {
 	TextWidget = InTextWidget;
+	TextEntry = FText::FromString(TEXT(""));
+
 #ifdef __IPHONE_8_0
 	if ([UIAlertController class])
 	{
@@ -52,11 +54,13 @@ void FIOSPlatformTextField::ShowVirtualKeyboard(bool bShow, int32 UserIndex, TSh
 											[AlertController dismissViewControllerAnimated : YES completion : nil];
 
 											UITextField* AlertTextField = AlertController.textFields.firstObject;
-											TextWidget->SetTextFromVirtualKeyboard(FText::FromString(AlertTextField.text));
+											TextEntry = FText::FromString(AlertTextField.text);
 
 											FIOSAsyncTask* AsyncTask = [[FIOSAsyncTask alloc] init];
 											AsyncTask.GameThreadCallback = ^ bool(void)
 											{
+												TextWidget->SetTextFromVirtualKeyboard(TextEntry);
+
 												// clear the TextWidget
 												TextWidget = nullptr;
 												return true;
@@ -183,16 +187,18 @@ void FIOSPlatformTextField::ShowVirtualKeyboard(bool bShow, int32 UserIndex, TSh
 
 - (void)alertView:(UIAlertView *)alertView clickedButtonAtIndex:(NSInteger)buttonIndex
 {
-	// index 1 is the OK button
-	if(buttonIndex == 1)
-	{
-		UITextField* AlertTextField = [alertView textFieldAtIndex: 0];
-		TextWidget->SetTextFromVirtualKeyboard(FText::FromString(AlertTextField.text));
-	}
-    
-    FIOSAsyncTask* AsyncTask = [[FIOSAsyncTask alloc] init];
+	UITextField* AlertTextField = [alertView textFieldAtIndex : 0];
+	TextEntry = FText::FromString(AlertTextField.text);
+
+	FIOSAsyncTask* AsyncTask = [[FIOSAsyncTask alloc] init];
     AsyncTask.GameThreadCallback = ^ bool(void)
     {
+		// index 1 is the OK button
+		if (buttonIndex == 1)
+		{
+			TextWidget->SetTextFromVirtualKeyboard(TextEntry);
+		}
+
         // clear the TextWidget
         TextWidget = nullptr;
         return true;
diff --git a/Engine/Source/Runtime/Slate/Public/Framework/Text/IOS/IOSPlatformTextField.h b/Engine/Source/Runtime/Slate/Public/Framework/Text/IOS/IOSPlatformTextField.h
--- a/Engine/Source/Runtime/Slate/Public/Framework/Text/IOS/IOSPlatformTextField.h
+++ b/Engine/Source/Runtime/Slate/Public/Framework/Text/IOS/IOSPlatformTextField.h
@@ -26,6 +26,7 @@ typedef FIOSPlatformTextField FPlatformTextField;
 @interface SlateTextField : NSObject<UIAlertViewDelegate>
 {
 	TSharedPtr<IVirtualKeyboardEntry> TextWidget;
+	FText TextEntry;
 }
 
 -(void)show:(TSharedPtr<IVirtualKeyboardEntry>)InTextWidget;
