Slate/UMG - Disabling the invalidation panel for now, need to add material resource reference reporting first.

diff --git a/Engine/Source/Runtime/Slate/Private/Widgets/SInvalidationPanel.cpp b/Engine/Source/Runtime/Slate/Private/Widgets/SInvalidationPanel.cpp
--- a/Engine/Source/Runtime/Slate/Private/Widgets/SInvalidationPanel.cpp
+++ b/Engine/Source/Runtime/Slate/Private/Widgets/SInvalidationPanel.cpp
@@ -5,7 +5,7 @@
 #include "SInvalidationPanel.h"
 #include "WidgetCaching.h"
 
-DECLARE_CYCLE_STAT(TEXT("Invalidation Time"), STAT_InvalidationTime, STATGROUP_Slate);
+//DECLARE_CYCLE_STAT(TEXT("Invalidation Time"), STAT_InvalidationTime, STATGROUP_Slate);
 DECLARE_DWORD_COUNTER_STAT(TEXT("Num Cached Elements"), STAT_SlateNumCachedElements, STATGROUP_Slate);
 DECLARE_DWORD_COUNTER_STAT(TEXT("Num Invalidated Elements"), STAT_SlateNumInvalidatedElements, STATGROUP_Slate);
 DECLARE_DWORD_COUNTER_STAT(TEXT("Num Volatile Widgets"), STAT_SlateNumVolatileWidgets, STATGROUP_Slate);
@@ -64,6 +64,9 @@ SInvalidationPanel::~SInvalidationPanel()
 
 bool SInvalidationPanel::GetCanCache() const
 {
+	//HACK: Disabling invalidation panel until material resource reporting is done.
+	return false;
+
 #if !(UE_BUILD_SHIPPING || UE_BUILD_TEST)
 	return bCanCache && EnableWidgetCaching.GetValueOnGameThread() == 1;
 #else
@@ -84,7 +87,7 @@ void SInvalidationPanel::Tick( const FGeometry& AllottedGeometry, const double I
 
 	if ( GetCanCache() )
 	{
-		SCOPE_CYCLE_COUNTER(STAT_InvalidationTime);
+		//SCOPE_CYCLE_COUNTER(STAT_InvalidationTime);
 
 		const bool bWasCachingNeeded = bNeedsCaching;
 
@@ -216,7 +219,7 @@ int32 SInvalidationPanel::OnPaint( const FPaintArgs& Args, const FGeometry& Allo
 {
 	if ( GetCanCache() )
 	{
-		SCOPE_CYCLE_COUNTER(STAT_InvalidationTime);
+		//SCOPE_CYCLE_COUNTER(STAT_InvalidationTime);
 
 		//FPlatformMisc::BeginNamedEvent(FColor::Magenta, "Slate::InvalidationPanel::Paint");
 
