UMG / Slate - Some cleanup in a few classes.

diff --git a/Engine/Source/Runtime/Slate/Private/Framework/Application/SlateApplication.cpp b/Engine/Source/Runtime/Slate/Private/Framework/Application/SlateApplication.cpp
--- a/Engine/Source/Runtime/Slate/Private/Framework/Application/SlateApplication.cpp
+++ b/Engine/Source/Runtime/Slate/Private/Framework/Application/SlateApplication.cpp
@@ -1146,9 +1146,9 @@ void FSlateApplication::PrivateDrawWindows( TSharedPtr<SWindow> DrawOnlyThisWind
 
 	if ( !SkipSecondPrepass.GetValueOnGameThread() )
 	{
-		//FPlatformMisc::BeginNamedEvent(FColor::Magenta, "Slate::Prepass");
+		FPlatformMisc::BeginNamedEvent(FColor::Magenta, "Slate::Prepass");
 		DrawPrepass( DrawOnlyThisWindow );
-		//FPlatformMisc::EndNamedEvent();
+		FPlatformMisc::EndNamedEvent();
 	}
 
 	//FPlatformMisc::BeginNamedEvent(FColor::Magenta, "Slate::GetDrawBuffer");
diff --git a/Engine/Source/Runtime/Slate/Public/Widgets/Layout/SConstraintCanvas.cpp b/Engine/Source/Runtime/Slate/Public/Widgets/Layout/SConstraintCanvas.cpp
--- a/Engine/Source/Runtime/Slate/Public/Widgets/Layout/SConstraintCanvas.cpp
+++ b/Engine/Source/Runtime/Slate/Public/Widgets/Layout/SConstraintCanvas.cpp
@@ -197,9 +197,10 @@ FVector2D SConstraintCanvas::ComputeDesiredSize( float ) const
 	{
 		const SConstraintCanvas::FSlot& CurChild = Children[ChildIndex];
 		const TSharedRef<SWidget>& Widget = CurChild.GetWidget();
+		const EVisibility ChildVisibilty = Widget->GetVisibility();
 
 		// As long as the widgets are not collapsed, they should contribute to the desired size.
-		if ( Widget->GetVisibility() != EVisibility::Collapsed )
+		if ( ChildVisibilty != EVisibility::Collapsed )
 		{
 			const FMargin Offset = CurChild.OffsetAttr.Get();
 			const FVector2D Alignment = CurChild.AlignmentAttr.Get();
diff --git a/Engine/Source/Runtime/SlateCore/Private/Input/HittestGrid.cpp b/Engine/Source/Runtime/SlateCore/Private/Input/HittestGrid.cpp
--- a/Engine/Source/Runtime/SlateCore/Private/Input/HittestGrid.cpp
+++ b/Engine/Source/Runtime/SlateCore/Private/Input/HittestGrid.cpp
@@ -228,7 +228,7 @@ void FHittestGrid::ClearGridForNewFrame( const FSlateRect& HittestArea )
 	GridOrigin = HittestArea.GetTopLeft();
 	const FVector2D GridSize = HittestArea.GetSize();
 	NumCells = FIntPoint( FMath::CeilToInt(GridSize.X / CellSize.X), FMath::CeilToInt(GridSize.Y / CellSize.Y) );
-	WidgetsCachedThisFrame->Empty(WidgetsCachedThisFrame->Num());
+	WidgetsCachedThisFrame->Reset();
 
 	const int32 NewTotalCells = NumCells.X * NumCells.Y;
 	if (NewTotalCells != Cells.Num())
diff --git a/Engine/Source/Runtime/UMG/Private/Components/Widget.cpp b/Engine/Source/Runtime/UMG/Private/Components/Widget.cpp
--- a/Engine/Source/Runtime/UMG/Private/Components/Widget.cpp
+++ b/Engine/Source/Runtime/UMG/Private/Components/Widget.cpp
@@ -237,7 +237,8 @@ void UWidget::SetVisibility(ESlateVisibility InVisibility)
 	TSharedPtr<SWidget> SafeWidget = GetCachedWidget();
 	if (SafeWidget.IsValid())
 	{
-		return SafeWidget->SetVisibility(UWidget::ConvertSerializedVisibilityToRuntime(InVisibility));
+		EVisibility SlateVisibility = UWidget::ConvertSerializedVisibilityToRuntime(InVisibility);
+		return SafeWidget->SetVisibility(SlateVisibility);
 	}
 }
 
