UMG/Slate - Re-enabling border drawing.

diff --git a/Engine/Source/Runtime/Slate/Private/Widgets/SInvalidationPanel.cpp b/Engine/Source/Runtime/Slate/Private/Widgets/SInvalidationPanel.cpp
--- a/Engine/Source/Runtime/Slate/Private/Widgets/SInvalidationPanel.cpp
+++ b/Engine/Source/Runtime/Slate/Private/Widgets/SInvalidationPanel.cpp
@@ -386,6 +386,8 @@ int32 SInvalidationPanel::OnPaint( const FPaintArgs& Args, const FGeometry& Allo
 
 			static const FName InvalidationPanelName(TEXT("SInvalidationPanel"));
 
+			const FSlateBrush* VolatileBrush = FCoreStyle::Get().GetBrush(TEXT("FocusRectangle"));
+
 			// Draw a yellow outline around any volatile elements.
 			const TArray<TSharedPtr<FSlateWindowElementList::FVolatilePaint>>& VolatileElements = CachedWindowElements->GetVolatileElements();
 			for ( const TSharedPtr<FSlateWindowElementList::FVolatilePaint>& VolatileElement : VolatileElements )
@@ -404,7 +406,7 @@ int32 SInvalidationPanel::OnPaint( const FPaintArgs& Args, const FGeometry& Allo
 					OutDrawElements,
 					++OutMaxChildLayer,
 					VolatileElement->GetGeometry().ToPaintGeometry(),
-					FCoreStyle::Get().GetBrush(TEXT("FocusRectangle")),
+					VolatileBrush,
 					MyClippingRect,
 					ESlateDrawEffect::None,
 					FLinearColor::Yellow
diff --git a/Engine/Source/Runtime/SlateCore/Private/Rendering/ElementBatcher.cpp b/Engine/Source/Runtime/SlateCore/Private/Rendering/ElementBatcher.cpp
--- a/Engine/Source/Runtime/SlateCore/Private/Rendering/ElementBatcher.cpp
+++ b/Engine/Source/Runtime/SlateCore/Private/Rendering/ElementBatcher.cpp
@@ -204,7 +204,7 @@ void FSlateElementBatcher::AddElements(FSlateDrawLayer& InDrawLayer)
 					AddViewportElement( DrawElement );
 					break;
 				case FSlateDrawElement::ET_Border:
-					//AddBorderElement( DrawElement );
+					AddBorderElement( DrawElement );
 					break;
 				case FSlateDrawElement::ET_Custom:
 					AddCustomElement( DrawElement );
