Merged Change: 2661515
Added some more verbose Slate stats for perf work

diff --git a/Engine/Source/Runtime/Slate/Private/Widgets/Text/STextBlock.cpp b/Engine/Source/Runtime/Slate/Private/Widgets/Text/STextBlock.cpp
--- a/Engine/Source/Runtime/Slate/Private/Widgets/Text/STextBlock.cpp
+++ b/Engine/Source/Runtime/Slate/Private/Widgets/Text/STextBlock.cpp
@@ -4,6 +4,11 @@
 #include "PlainTextLayoutMarshaller.h"
 #include "TextBlockLayout.h"
 
+DECLARE_CYCLE_STAT(TEXT("STextBlock::SetText Time"), Stat_SlateTextBlockSetText, STATGROUP_SlateVerbose)
+DECLARE_CYCLE_STAT(TEXT("STextBlock::OnPaint Time"), Stat_SlateTextBlockOnPaint, STATGROUP_SlateVerbose)
+DECLARE_CYCLE_STAT(TEXT("STextBlock::ComputeDesiredSize"), Stat_SlateTextBlockCDS, STATGROUP_SlateVerbose)
+DECLARE_CYCLE_STAT(TEXT("STextBlock::ComputeVolitility"), Stat_SlateTextBlockCV, STATGROUP_SlateVerbose)
+
 void STextBlock::Construct( const FArguments& InArgs )
 {
 	TextStyle = InArgs._TextStyle;
@@ -67,6 +72,7 @@ const FSlateBrush* STextBlock::GetHighlightShape() const
 
 void STextBlock::SetText( const TAttribute< FString >& InText )
 {
+	SCOPE_CYCLE_COUNTER(Stat_SlateTextBlockSetText);
 	struct Local
 	{
 		static FText PassThroughAttribute( TAttribute< FString > InString )
@@ -82,18 +88,21 @@ void STextBlock::SetText( const TAttribute< FString >& InText )
 
 void STextBlock::SetText( const FString& InText )
 {
+	SCOPE_CYCLE_COUNTER(Stat_SlateTextBlockSetText);
 	BoundText = FText::FromString( InText );
 	Invalidate(EInvalidateWidget::LayoutAndVolatility);
 }
 
 void STextBlock::SetText( const TAttribute< FText >& InText )
 {
+	SCOPE_CYCLE_COUNTER(Stat_SlateTextBlockSetText);
 	BoundText = InText;
 	Invalidate(EInvalidateWidget::LayoutAndVolatility);
 }
 
 void STextBlock::SetText( const FText& InText )
 {
+	SCOPE_CYCLE_COUNTER(Stat_SlateTextBlockSetText);
 	if ( BoundText.IsBound() || BoundText.Get().ToString() != InText.ToString() )
 	{
 		BoundText = InText;
@@ -103,6 +112,8 @@ void STextBlock::SetText( const FText& InText )
 
 int32 STextBlock::OnPaint( const FPaintArgs& Args, const FGeometry& AllottedGeometry, const FSlateRect& MyClippingRect, FSlateWindowElementList& OutDrawElements, int32 LayerId, const FWidgetStyle& InWidgetStyle, bool bParentEnabled ) const
 {
+	SCOPE_CYCLE_COUNTER(Stat_SlateTextBlockOnPaint);
+
 #if WITH_FANCY_TEXT
 
 	// OnPaint will also update the text layout cache if required
@@ -171,6 +182,7 @@ FReply STextBlock::OnMouseButtonDoubleClick( const FGeometry& InMyGeometry, cons
 
 FVector2D STextBlock::ComputeDesiredSize(float LayoutScaleMultiplier) const
 {
+	SCOPE_CYCLE_COUNTER(Stat_SlateTextBlockCDS);
 #if WITH_FANCY_TEXT
 
 	// ComputeDesiredSize will also update the text layout cache if required
@@ -193,6 +205,7 @@ FVector2D STextBlock::ComputeDesiredSize(float LayoutScaleMultiplier) const
 
 bool STextBlock::ComputeVolatility() const
 {
+	SCOPE_CYCLE_COUNTER(Stat_SlateTextBlockCV);
 	return SLeafWidget::ComputeVolatility() || BoundText.IsBound();
 }
 
