UMG/Slate - Moving Element list management cached element lists to slate application so that it can be cleaned up after the renderer has finished using it.

diff --git a/Engine/Source/Runtime/Slate/Private/Widgets/SInvalidationPanel.cpp b/Engine/Source/Runtime/Slate/Private/Widgets/SInvalidationPanel.cpp
--- a/Engine/Source/Runtime/Slate/Private/Widgets/SInvalidationPanel.cpp
+++ b/Engine/Source/Runtime/Slate/Private/Widgets/SInvalidationPanel.cpp
@@ -64,10 +64,9 @@ SInvalidationPanel::~SInvalidationPanel()
 		delete NodePool[i];
 	}
 
-	if( FSlateApplication::IsInitialized() )
+	if ( FSlateApplication::IsInitialized() )
 	{
-		TSharedPtr<FSlateRenderer> Renderer = FSlateApplicationBase::Get().GetRenderer();
-		Renderer->ReleaseCachingResourcesFor(this);
+		FSlateApplication::Get().ReleaseResourcesForLayoutCache(this);
 	}
 }
 
@@ -185,47 +184,6 @@ FCachedWidgetNode* SInvalidationPanel::CreateCacheNode() const
 	return NewNode;
 }
 
-TSharedPtr< FSlateWindowElementList > SInvalidationPanel::GetNextCachedElementList(const TSharedPtr<SWindow>& CurrentWindow) const
-{
-	TSharedPtr< FSlateWindowElementList > NextElementList;
-
-	// Move any inactive element lists in the active pool to the inactive pool.
-	for ( int32 i = 0; i < ActiveCachedElementListPool.Num(); i++ )
-	{
-		if ( ActiveCachedElementListPool[i]->IsCachedRenderDataInUse() == false )
-		{
-			InactiveCachedElementListPool.Add(ActiveCachedElementListPool[i]);
-			ActiveCachedElementListPool.RemoveAtSwap(i, 1, false);
-		}
-	}
-
-	// Remove inactive lists that don't belong to this window.
-	for ( int32 i = 0; i < InactiveCachedElementListPool.Num(); i++ )
-	{
-		if ( InactiveCachedElementListPool[i]->GetWindow() != CurrentWindow )
-		{
-			InactiveCachedElementListPool.RemoveAtSwap(i, 1, false);
-		}
-	}
-
-	// Create a new element list if none are available, or use an existing one.
-	if ( InactiveCachedElementListPool.Num() == 0 )
-	{
-		NextElementList = MakeShareable(new FSlateWindowElementList(CurrentWindow));
-	}
-	else
-	{
-		NextElementList = InactiveCachedElementListPool[0];
-		NextElementList->ResetBuffers();
-
-		InactiveCachedElementListPool.RemoveAtSwap(0, 1, false);
-	}
-
-	ActiveCachedElementListPool.Add(NextElementList);
-
-	return NextElementList;
-}
-
 void SInvalidationPanel::OnGlobalInvalidate()
 {
 	InvalidateCache();
@@ -262,7 +220,7 @@ int32 SInvalidationPanel::OnPaint( const FPaintArgs& Args, const FGeometry& Allo
 			// Mark that we're in the process of invalidating.
 			bIsInvalidating = true;
 
-			CachedWindowElements = GetNextCachedElementList(OutDrawElements.GetWindow());
+			CachedWindowElements = FSlateApplication::Get().GetCachableElementList(OutDrawElements.GetWindow(), this);
 
 			// Reset the render data handle in case it was in use, and we're not overriding it this frame.
 			CachedRenderData.Reset();
diff --git a/Engine/Source/Runtime/Slate/Public/Framework/Application/SlateApplication.h b/Engine/Source/Runtime/Slate/Public/Framework/Application/SlateApplication.h
--- a/Engine/Source/Runtime/Slate/Public/Framework/Application/SlateApplication.h
+++ b/Engine/Source/Runtime/Slate/Public/Framework/Application/SlateApplication.h
@@ -748,6 +748,18 @@ public:
 	 */
 	bool TakeScreenshot(const TSharedRef<SWidget>& Widget, const FIntRect& InnerWidgetArea, TArray<FColor>& OutColorData, FIntVector& OutSize);
 
+	/**
+	 * 
+	 */
+	TSharedPtr< FSlateWindowElementList > GetCachableElementList(const TSharedPtr<SWindow>& CurrentWindow, const ILayoutCache* LayoutCache);
+
+	/**
+	 * Once a layout cache is destroyed it needs to free any resources it was using in a safe way to prevent
+	 * any in-flight rendering from being interrupted by referencing resources that go away.  So when a layout
+	 * cache is destroyed it should call this function so any associated resources can be collected when it's safe.
+	 */
+	void ReleaseResourcesForLayoutCache(const ILayoutCache* LayoutCache);
+
 protected:
 
 	friend class FAnalogCursor;
@@ -1617,4 +1629,22 @@ private:
 	// e.g. On windows the origin (coordinates X=0, Y=0) is the upper left of the primary monitor,
 	// but there could be another monitor on any of the sides.
 	FSlateRect VirtualDesktopRect;
+
+	//
+	// Invalidation Support
+	//
+
+	class FCacheElementPools
+	{
+	public:
+		TSharedPtr< FSlateWindowElementList > GetNextCachableElementList(const TSharedPtr<SWindow>& CurrentWindow );
+		bool IsInUse() const;
+
+	private:
+		TArray< TSharedPtr< FSlateWindowElementList > > ActiveCachedElementListPool;
+		TArray< TSharedPtr< FSlateWindowElementList > > InactiveCachedElementListPool;
+	};
+
+	TMap< const ILayoutCache*, FCacheElementPools* > CachedElementLists;
+	TArray< FCacheElementPools* > ReleasedCachedElementLists;
 };
diff --git a/Engine/Source/Runtime/Slate/Public/Widgets/SInvalidationPanel.h b/Engine/Source/Runtime/Slate/Public/Widgets/SInvalidationPanel.h
--- a/Engine/Source/Runtime/Slate/Public/Widgets/SInvalidationPanel.h
+++ b/Engine/Source/Runtime/Slate/Public/Widgets/SInvalidationPanel.h
@@ -61,8 +61,6 @@ private:
 #endif
 
 	mutable FCachedWidgetNode* RootCacheNode;
-	mutable TArray< TSharedPtr< FSlateWindowElementList > > ActiveCachedElementListPool;
-	mutable TArray< TSharedPtr< FSlateWindowElementList > > InactiveCachedElementListPool;
 	mutable TSharedPtr< FSlateWindowElementList > CachedWindowElements;
 	mutable TSharedPtr<FSlateRenderDataHandle, ESPMode::ThreadSafe> CachedRenderData;
 
diff --git a/Engine/Source/Runtime/SlateCore/Public/Rendering/DrawElements.h b/Engine/Source/Runtime/SlateCore/Public/Rendering/DrawElements.h
--- a/Engine/Source/Runtime/SlateCore/Public/Rendering/DrawElements.h
+++ b/Engine/Source/Runtime/SlateCore/Public/Rendering/DrawElements.h
@@ -631,12 +631,6 @@ public:
 		, NumIndices( InNumIndices )
 	{}
 
-	//// Render Batch Sort Operator
-	//bool operator < ( const FSlateRenderBatch& RenderBatchIn ) const
-	//{
-	//	return Layer < RenderBatchIn.Layer;
-	//}
-
 public:
 	/** The layer we need to sort by when  */
 	const uint32 Layer;
