Merged CL 2664672

Fix crash with cutoff menus in the editor

https://jira.ol.epicgames.net/browse/UE-20058

diff --git a/Engine/Source/Runtime/Slate/Private/Framework/MultiBox/SClippingHorizontalBox.h b/Engine/Source/Runtime/Slate/Private/Framework/MultiBox/SClippingHorizontalBox.h
--- a/Engine/Source/Runtime/Slate/Private/Framework/MultiBox/SClippingHorizontalBox.h
+++ b/Engine/Source/Runtime/Slate/Private/Framework/MultiBox/SClippingHorizontalBox.h
@@ -18,7 +18,10 @@ public:
 		SLATE_ARGUMENT(FName, StyleName)
 	SLATE_END_ARGS()
 
-	FORCENOINLINE SClippingHorizontalBox() { }
+	FORCENOINLINE SClippingHorizontalBox() 
+	{
+		bCanTick = true;
+	}
 
 	/** SWidget interface */
 	virtual void Tick( const FGeometry& AllottedGeometry, const double InCurrentTime, const float InDeltaTime ) override;
