Slate/UMG - Adding missing file to move caching element lists to slate application.  Also fixing a focus bug with user index not being passed to change focus.

diff --git a/Engine/Source/Runtime/Slate/Private/Framework/Application/SlateApplication.cpp b/Engine/Source/Runtime/Slate/Private/Framework/Application/SlateApplication.cpp
--- a/Engine/Source/Runtime/Slate/Private/Framework/Application/SlateApplication.cpp
+++ b/Engine/Source/Runtime/Slate/Private/Framework/Application/SlateApplication.cpp
@@ -1408,6 +1408,19 @@ void FSlateApplication::Tick()
 		QueueSynthesizedMouseMove();
 	}
 
+	// Check if any element lists used for caching need to be released
+	{
+		for ( int32 CacheIndex = 0; CacheIndex < ReleasedCachedElementLists.Num(); CacheIndex++ )
+		{
+			if ( ReleasedCachedElementLists[CacheIndex]->IsInUse() == false )
+			{
+				delete ReleasedCachedElementLists[CacheIndex];
+				ReleasedCachedElementLists.RemoveAtSwap(CacheIndex, 1, false);
+				CacheIndex--;
+			}
+		}
+	}
+
 	// skip tick/draw if we are idle and there are no active timers registered that we need to drive slate for.
 	// This effectively means the slate application is totally idle and we don't need to update the UI.
 	// This relies on Widgets properly registering for Active timer when they need something to happen even
@@ -2382,7 +2395,7 @@ bool FSlateApplication::SetUserFocus(const uint32 InUserIndex, const FWidgetPath
 		FReply Reply = NewFocusedWidget->OnFocusReceived(WidgetToFocus.Geometry, FFocusEvent(InCause, InUserIndex));
 		if (Reply.IsEventHandled())
 		{
-			ProcessReply(InFocusPath, Reply, nullptr, nullptr);
+			ProcessReply(InFocusPath, Reply, nullptr, nullptr, InUserIndex);
 		}
 	}
 
@@ -3619,6 +3632,91 @@ bool FSlateApplication::TakeScreenshot(const TSharedRef<SWidget>& Widget, const
 	return true;
 }
 
+TSharedPtr< FSlateWindowElementList > FSlateApplication::GetCachableElementList(const TSharedPtr<SWindow>& CurrentWindow, const ILayoutCache* LayoutCache)
+{
+	FCacheElementPools* Pools = CachedElementLists.FindRef(LayoutCache);
+	if ( Pools == nullptr )
+	{
+		Pools = new FCacheElementPools();
+		CachedElementLists.Add(LayoutCache, Pools);
+	}
+
+	TSharedPtr< FSlateWindowElementList > NextElementList = Pools->GetNextCachableElementList(CurrentWindow);
+
+	return NextElementList;
+}
+
+TSharedPtr< FSlateWindowElementList > FSlateApplication::FCacheElementPools::GetNextCachableElementList(const TSharedPtr<SWindow>& CurrentWindow)
+{
+	TSharedPtr< FSlateWindowElementList > NextElementList;
+
+	// Move any inactive element lists in the active pool to the inactive pool.
+	for ( int32 i = 0; i < ActiveCachedElementListPool.Num(); i++ )
+	{
+		if ( ActiveCachedElementListPool[i]->IsCachedRenderDataInUse() == false )
+		{
+			InactiveCachedElementListPool.Add(ActiveCachedElementListPool[i]);
+			ActiveCachedElementListPool.RemoveAtSwap(i, 1, false);
+		}
+	}
+
+	// Remove inactive lists that don't belong to this window.
+	for ( int32 i = 0; i < InactiveCachedElementListPool.Num(); i++ )
+	{
+		if ( InactiveCachedElementListPool[i]->GetWindow() != CurrentWindow )
+		{
+			InactiveCachedElementListPool.RemoveAtSwap(i, 1, false);
+		}
+	}
+
+	// Create a new element list if none are available, or use an existing one.
+	if ( InactiveCachedElementListPool.Num() == 0 )
+	{
+		NextElementList = MakeShareable(new FSlateWindowElementList(CurrentWindow));
+	}
+	else
+	{
+		NextElementList = InactiveCachedElementListPool[0];
+		NextElementList->ResetBuffers();
+
+		InactiveCachedElementListPool.RemoveAtSwap(0, 1, false);
+	}
+
+	ActiveCachedElementListPool.Add(NextElementList);
+
+	return NextElementList;
+}
+
+bool FSlateApplication::FCacheElementPools::IsInUse() const
+{
+	bool bInUse = false;
+	for ( TSharedPtr< FSlateWindowElementList > ElementList : InactiveCachedElementListPool )
+	{
+		bInUse |= ElementList->IsCachedRenderDataInUse();
+	}
+
+	for ( TSharedPtr< FSlateWindowElementList > ElementList : ActiveCachedElementListPool )
+	{
+		bInUse |= ElementList->IsCachedRenderDataInUse();
+	}
+
+	return bInUse;
+}
+
+void FSlateApplication::ReleaseResourcesForLayoutCache(const ILayoutCache* LayoutCache)
+{
+	FCacheElementPools* Pools = CachedElementLists.FindRef(LayoutCache);
+	if ( Pools != nullptr )
+	{
+		ReleasedCachedElementLists.Add(Pools);
+	}
+
+	CachedElementLists.Remove(LayoutCache);
+
+	// Release the rendering related resources.
+	Renderer->ReleaseCachingResourcesFor(LayoutCache);
+}
+
 
 /* FSlateApplicationBase interface
  *****************************************************************************/
