Fix compilation issue on a case-sensitive fs.

diff --git a/Engine/Source/Runtime/Slate/Private/Widgets/SInvalidationPanel.cpp b/Engine/Source/Runtime/Slate/Private/Widgets/SInvalidationPanel.cpp
--- a/Engine/Source/Runtime/Slate/Private/Widgets/SInvalidationPanel.cpp
+++ b/Engine/Source/Runtime/Slate/Private/Widgets/SInvalidationPanel.cpp
@@ -4,7 +4,7 @@
 #include "LayoutUtils.h"
 #include "SInvalidationPanel.h"
 #include "WidgetCaching.h"
-#include "ReflectionMetaData.h"
+#include "ReflectionMetadata.h"
 
 //DECLARE_CYCLE_STAT(TEXT("Invalidation Time"), STAT_InvalidationTime, STATGROUP_Slate);
 DECLARE_DWORD_COUNTER_STAT(TEXT("Num Cached Elements"), STAT_SlateNumCachedElements, STATGROUP_Slate);
