Integrating CL 2655118, 2655170 from UE4-Oriont/... to UE4/...

Now using FPopupMethodReply instead of EPopupMethod to correctly drive throttling.

diff --git a/Engine/Source/Runtime/AppFramework/Private/Framework/Testing/SWidgetGallery.cpp b/Engine/Source/Runtime/AppFramework/Private/Framework/Testing/SWidgetGallery.cpp
--- a/Engine/Source/Runtime/AppFramework/Private/Framework/Testing/SWidgetGallery.cpp
+++ b/Engine/Source/Runtime/AppFramework/Private/Framework/Testing/SWidgetGallery.cpp
@@ -1103,9 +1103,9 @@ private:
 		FSlateApplication::Get().DismissAllMenus();
 	}
 
-	virtual TOptional<EPopupMethod> OnQueryPopupMethod() const override
+	virtual FPopupMethodReply OnQueryPopupMethod() const override
 	{
-		return PopupMethod;
+		return FPopupMethodReply::UseMethod(PopupMethod);
 	}
 
 private:
diff --git a/Engine/Source/Runtime/Engine/Classes/Engine/GameViewportClient.h b/Engine/Source/Runtime/Engine/Classes/Engine/GameViewportClient.h
--- a/Engine/Source/Runtime/Engine/Classes/Engine/GameViewportClient.h
+++ b/Engine/Source/Runtime/Engine/Classes/Engine/GameViewportClient.h
@@ -665,7 +665,7 @@ public:
 		return bHideCursorDuringCapture;
 	}
 
-	virtual TOptional<EPopupMethod> OnQueryPopupMethod() const override;
+	virtual FPopupMethodReply OnQueryPopupMethod() const override;
 
 private:
 	/**
diff --git a/Engine/Source/Runtime/Engine/Private/GameViewportClient.cpp b/Engine/Source/Runtime/Engine/Private/GameViewportClient.cpp
--- a/Engine/Source/Runtime/Engine/Private/GameViewportClient.cpp
+++ b/Engine/Source/Runtime/Engine/Private/GameViewportClient.cpp
@@ -2298,9 +2298,10 @@ bool UGameViewportClient::HandleShowCommand( const TCHAR* Cmd, FOutputDevice& Ar
 	return true;
 }
 
-TOptional<EPopupMethod> UGameViewportClient::OnQueryPopupMethod() const
+FPopupMethodReply UGameViewportClient::OnQueryPopupMethod() const
 {
-	return EPopupMethod::UseCurrentWindow;
+	return FPopupMethodReply::UseMethod(EPopupMethod::UseCurrentWindow)
+		.SetShouldThrottle(EShouldThrottle::No);
 }
 
 void UGameViewportClient::ToggleShowVolumes()
diff --git a/Engine/Source/Runtime/Engine/Private/Slate/SceneViewport.cpp b/Engine/Source/Runtime/Engine/Private/Slate/SceneViewport.cpp
--- a/Engine/Source/Runtime/Engine/Private/Slate/SceneViewport.cpp
+++ b/Engine/Source/Runtime/Engine/Private/Slate/SceneViewport.cpp
@@ -697,7 +697,7 @@ FReply FSceneViewport::OnMotionDetected( const FGeometry& MyGeometry, const FMot
 	return CurrentReplyState;
 }
 
-TOptional<EPopupMethod> FSceneViewport::OnQueryPopupMethod() const
+FPopupMethodReply FSceneViewport::OnQueryPopupMethod() const
 {
 	if (ViewportClient != nullptr)
 	{
@@ -705,7 +705,7 @@ TOptional<EPopupMethod> FSceneViewport::OnQueryPopupMethod() const
 	}
 	else
 	{
-		return TOptional<EPopupMethod>();
+		return FPopupMethodReply::Unhandled();
 	}
 }
 
diff --git a/Engine/Source/Runtime/Engine/Public/Slate/SceneViewport.h b/Engine/Source/Runtime/Engine/Public/Slate/SceneViewport.h
--- a/Engine/Source/Runtime/Engine/Public/Slate/SceneViewport.h
+++ b/Engine/Source/Runtime/Engine/Public/Slate/SceneViewport.h
@@ -203,7 +203,7 @@ public:
 	virtual FReply OnTouchEnded( const FGeometry& MyGeometry, const FPointerEvent& InTouchEvent ) override;
 	virtual FReply OnTouchGesture( const FGeometry& MyGeometry, const FPointerEvent& InGestureEvent ) override;
 	virtual FReply OnMotionDetected( const FGeometry& MyGeometry, const FMotionEvent& InMotionEvent ) override;
-	virtual TOptional<EPopupMethod> OnQueryPopupMethod() const override;
+	virtual FPopupMethodReply OnQueryPopupMethod() const override;
 	virtual TOptional<bool> OnQueryShowFocus(const EFocusCause InFocusCause) const override;
 	virtual void OnFinishedPointerInput() override;
 	virtual FReply OnKeyDown( const FGeometry& InGeometry, const FKeyEvent& InKeyEvent ) override;
diff --git a/Engine/Source/Runtime/Engine/Public/UnrealClient.h b/Engine/Source/Runtime/Engine/Public/UnrealClient.h
--- a/Engine/Source/Runtime/Engine/Public/UnrealClient.h
+++ b/Engine/Source/Runtime/Engine/Public/UnrealClient.h
@@ -10,6 +10,7 @@
 #include "HitProxies.h"
 #include "InputCoreTypes.h"
 #include "Engine/EngineBaseTypes.h"
+#include "PopupMethodReply.h"
 
 class FCanvas;
 class FViewportClient;
@@ -17,7 +18,6 @@ class FViewportClient;
 class SWidget;
 class FCursorReply;
 
-enum class EPopupMethod : uint8;
 enum class EFocusCause : uint8;
 
 /**
@@ -908,7 +908,7 @@ public:
 	virtual bool HideCursorDuringCapture() { return false; }
 
 	/** Should we make new windows for popups or create an overlay in the current window. */
-	virtual TOptional<EPopupMethod> OnQueryPopupMethod() const { return TOptional<EPopupMethod>(); }
+	virtual FPopupMethodReply OnQueryPopupMethod() const { return FPopupMethodReply::Unhandled(); }
 };
 
 /** Tracks the viewport client that should process the stat command, can be NULL */
diff --git a/Engine/Source/Runtime/Slate/Private/Framework/Application/MenuStack.cpp b/Engine/Source/Runtime/Slate/Private/Framework/Application/MenuStack.cpp
--- a/Engine/Source/Runtime/Slate/Private/Framework/Application/MenuStack.cpp
+++ b/Engine/Source/Runtime/Slate/Private/Framework/Application/MenuStack.cpp
@@ -194,13 +194,13 @@ TSharedRef<IMenu> FMenuStack::Push(const FWidgetPath& InOwnerPath, const TShared
 		// pushing a new root menu (leave ParentMenu invalid)
 
 		// The active method is determined when a new root menu is pushed
-		ActiveMethod = InMethod.IsSet() ? InMethod : QueryPopupMethod(InOwnerPath);
+		ActiveMethod = InMethod.IsSet() ? FPopupMethodReply::UseMethod(InMethod.GetValue()) : QueryPopupMethod(InOwnerPath);
 
 		// The host window is determined when a new root menu is pushed
 		SetHostWindow(InOwnerPath.GetWindow());
 	}
 
-	return PushInternal(ParentMenu, InContent, Anchor, TransitionEffect, bFocusImmediately, bIsCollapsedByParent);
+	return PushInternal(ParentMenu, InContent, Anchor, TransitionEffect, bFocusImmediately, ActiveMethod.GetShouldThrottle(), bIsCollapsedByParent);
 }
 
 TSharedRef<IMenu> FMenuStack::Push(const TSharedPtr<IMenu>& InParentMenu, const TSharedRef<SWidget>& InContent, const FVector2D& SummonLocation, const FPopupTransitionEffect& TransitionEffect, const bool bFocusImmediately, const FVector2D& SummonLocationSize, const bool bIsCollapsedByParent)
@@ -210,10 +210,10 @@ TSharedRef<IMenu> FMenuStack::Push(const TSharedPtr<IMenu>& InParentMenu, const
 
 	FSlateRect Anchor(SummonLocation, SummonLocation + SummonLocationSize);
 
-	return PushInternal(InParentMenu, InContent, Anchor, TransitionEffect, bFocusImmediately, bIsCollapsedByParent);
+	return PushInternal(InParentMenu, InContent, Anchor, TransitionEffect, bFocusImmediately, EShouldThrottle::Yes, bIsCollapsedByParent);
 }
 
-TSharedRef<IMenu> FMenuStack::PushHosted(const FWidgetPath& InOwnerPath, const TSharedRef<IMenuHost>& InMenuHost, const TSharedRef<SWidget>& InContent, TSharedPtr<SWidget>& OutWrappedContent, const FPopupTransitionEffect& TransitionEffect, const bool bIsCollapsedByParent)
+TSharedRef<IMenu> FMenuStack::PushHosted(const FWidgetPath& InOwnerPath, const TSharedRef<IMenuHost>& InMenuHost, const TSharedRef<SWidget>& InContent, TSharedPtr<SWidget>& OutWrappedContent, const FPopupTransitionEffect& TransitionEffect, EShouldThrottle ShouldThrottle, const bool bIsCollapsedByParent)
 {
 	TSharedPtr<IMenu> ParentMenu;
 
@@ -229,16 +229,16 @@ TSharedRef<IMenu> FMenuStack::PushHosted(const FWidgetPath& InOwnerPath, const T
 		// pushing a new root menu (leave ParentMenu invalid)
 
 		// The active method is determined when a new root menu is pushed and hosted menus are always UseCurrentWindow
-		ActiveMethod = EPopupMethod::UseCurrentWindow;
+		ActiveMethod = FPopupMethodReply::UseMethod(EPopupMethod::UseCurrentWindow);
 
 		// The host window is determined when a new root menu is pushed
 		SetHostWindow(InOwnerPath.GetWindow());
 	}
 
-	return PushHosted(ParentMenu, InMenuHost, InContent, OutWrappedContent, TransitionEffect);
+	return PushHosted(ParentMenu, InMenuHost, InContent, OutWrappedContent, TransitionEffect, ShouldThrottle);
 }
 
-TSharedRef<IMenu> FMenuStack::PushHosted(const TSharedPtr<IMenu>& InParentMenu, const TSharedRef<IMenuHost>& InMenuHost, const TSharedRef<SWidget>& InContent, TSharedPtr<SWidget>& OutWrappedContent, const FPopupTransitionEffect& TransitionEffect, const bool bIsCollapsedByParent)
+TSharedRef<IMenu> FMenuStack::PushHosted(const TSharedPtr<IMenu>& InParentMenu, const TSharedRef<IMenuHost>& InMenuHost, const TSharedRef<SWidget>& InContent, TSharedPtr<SWidget>& OutWrappedContent, const FPopupTransitionEffect& TransitionEffect, EShouldThrottle ShouldThrottle, const bool bIsCollapsedByParent)
 {
 	check(HostWindow.IsValid());
 
@@ -253,14 +253,14 @@ TSharedRef<IMenu> FMenuStack::PushHosted(const TSharedPtr<IMenu>& InParentMenu,
 	// Register to get a callback when it's dismissed - to fixup stack
 	OutMenu->GetOnMenuDismissed().AddRaw(this, &FMenuStack::OnMenuDestroyed);
 
-	PostPush(InParentMenu, OutMenu);
+	PostPush(InParentMenu, OutMenu, ShouldThrottle);
 
 	PendingNewMenu.Reset();
 
 	return OutMenu;
 }
 
-TSharedRef<IMenu> FMenuStack::PushInternal(const TSharedPtr<IMenu>& InParentMenu, const TSharedRef<SWidget>& InContent, FSlateRect Anchor, const FPopupTransitionEffect& TransitionEffect, const bool bFocusImmediately, const bool bIsCollapsedByParent)
+TSharedRef<IMenu> FMenuStack::PushInternal(const TSharedPtr<IMenu>& InParentMenu, const TSharedRef<SWidget>& InContent, FSlateRect Anchor, const FPopupTransitionEffect& TransitionEffect, const bool bFocusImmediately, EShouldThrottle ShouldThrottle, const bool bIsCollapsedByParent)
 {
 	FPrePushArgs PrePushArgs;
 	PrePushArgs.Content = InContent;
@@ -276,13 +276,13 @@ TSharedRef<IMenu> FMenuStack::PushInternal(const TSharedPtr<IMenu>& InParentMenu
 	const FPrePushResults PrePushResults = PrePush(PrePushArgs);
 
 	// Menu object creation stage
-	TSharedRef<FMenuBase> OutMenu = ActiveMethod.Get(EPopupMethod::CreateNewWindow) == EPopupMethod::CreateNewWindow
+	TSharedRef<FMenuBase> OutMenu = ActiveMethod.GetPopupMethod() == EPopupMethod::CreateNewWindow
 		? PushNewWindow(InParentMenu, PrePushResults)
 		: PushPopup(InParentMenu, PrePushResults);
 
 	// Post-push stage
 	//   Updates the stack and content map member variables
-	PostPush(InParentMenu, OutMenu);
+	PostPush(InParentMenu, OutMenu, ShouldThrottle);
 
 	PendingNewMenu.Reset();
 
@@ -304,7 +304,7 @@ FMenuStack::FPrePushResults FMenuStack::PrePush(const FPrePushArgs& InArgs)
 
 	// Calc the max height available on screen for the menu
 	float MaxHeight;
-	if (ActiveMethod.Get(EPopupMethod::CreateNewWindow) == EPopupMethod::CreateNewWindow)
+	if (ActiveMethod.GetPopupMethod() == EPopupMethod::CreateNewWindow)
 	{
 		FSlateRect WorkArea = FSlateApplication::Get().GetWorkArea(InArgs.Anchor);
 		MaxHeight = FMenuStackDefs::MaxMenuScreenHeightFraction * WorkArea.GetSize().Y;
@@ -333,7 +333,7 @@ FMenuStack::FPrePushResults FMenuStack::PrePush(const FPrePushArgs& InArgs)
 	EOrientation Orientation = (InArgs.TransitionEffect.SlideDirection == FPopupTransitionEffect::SubMenu) ? Orient_Horizontal : Orient_Vertical;
 
 	// Calculate the correct position for the menu based on the popup method
-	if (ActiveMethod.Get(EPopupMethod::CreateNewWindow) == EPopupMethod::CreateNewWindow)
+	if (ActiveMethod.GetPopupMethod() == EPopupMethod::CreateNewWindow)
 	{
 		// Places the menu's window in the work area
 		OutResults.AnimStartLocation = OutResults.AnimFinalLocation = FSlateApplication::Get().CalculatePopupWindowPosition(InArgs.Anchor, OutResults.ExpectedSize, Orientation);
@@ -398,7 +398,7 @@ FMenuStack::FPrePushResults FMenuStack::PrePush(const FPrePushArgs& InArgs)
 
 TSharedRef<FMenuBase> FMenuStack::PushNewWindow(TSharedPtr<IMenu> InParentMenu, const FPrePushResults& InPrePushResults)
 {
-	check(ActiveMethod.GetValue() == EPopupMethod::CreateNewWindow);
+	check(ActiveMethod.GetPopupMethod() == EPopupMethod::CreateNewWindow);
 
 	// Start pop-up windows out transparent, then fade them in over time
 	const EWindowTransparency Transparency(InPrePushResults.bAllowAnimations ? EWindowTransparency::PerWindow : EWindowTransparency::None);
@@ -459,7 +459,7 @@ TSharedRef<FMenuBase> FMenuStack::PushNewWindow(TSharedPtr<IMenu> InParentMenu,
 
 TSharedRef<FMenuBase> FMenuStack::PushPopup(TSharedPtr<IMenu> InParentMenu, const FPrePushResults& InPrePushResults)
 {
-	check(ActiveMethod.GetValue() == EPopupMethod::UseCurrentWindow);
+	check(ActiveMethod.GetPopupMethod() == EPopupMethod::UseCurrentWindow);
 
 	// Create a FMenuInPopup
 	check(InPrePushResults.WrappedContent.IsValid());
@@ -480,7 +480,7 @@ TSharedRef<FMenuBase> FMenuStack::PushPopup(TSharedPtr<IMenu> InParentMenu, cons
 	return Menu;
 }
 
-void FMenuStack::PostPush(TSharedPtr<IMenu> InParentMenu, TSharedRef<FMenuBase> InMenu)
+void FMenuStack::PostPush(TSharedPtr<IMenu> InParentMenu, TSharedRef<FMenuBase> InMenu, EShouldThrottle ShouldThrottle )
 {
 	// Determine at which index we insert this new menu in the stack
 	int32 InsertIndex = 0;
@@ -511,24 +511,24 @@ void FMenuStack::PostPush(TSharedPtr<IMenu> InParentMenu, TSharedRef<FMenuBase>
 
 	// When a new menu is pushed, if we are not already in responsive mode for Slate UI, enter it now
 	// to ensure the menu is responsive in low FPS situations
-	if (!ThrottleHandle.IsValid())
+	if (!ThrottleHandle.IsValid() && ShouldThrottle == EShouldThrottle::Yes)
 	{
 		ThrottleHandle = FSlateThrottleManager::Get().EnterResponsiveMode();
 	}
 }
 
-EPopupMethod FMenuStack::QueryPopupMethod(const FWidgetPath& PathToQuery)
+FPopupMethodReply FMenuStack::QueryPopupMethod(const FWidgetPath& PathToQuery)
 {
 	for (int32 WidgetIndex = PathToQuery.Widgets.Num() - 1; WidgetIndex >= 0; --WidgetIndex)
 	{
-		TOptional<EPopupMethod> PopupMethod = PathToQuery.Widgets[WidgetIndex].Widget->OnQueryPopupMethod();
-		if (PopupMethod.IsSet())
+		FPopupMethodReply PopupMethod = PathToQuery.Widgets[WidgetIndex].Widget->OnQueryPopupMethod();
+		if (PopupMethod.IsEventHandled())
 		{
-			return PopupMethod.GetValue();
+			return PopupMethod;
 		}
 	}
 
-	return EPopupMethod::CreateNewWindow;
+	return FPopupMethodReply::UseMethod(EPopupMethod::CreateNewWindow);
 }
 
 void FMenuStack::Dismiss(int32 FirstStackIndexToRemove)
diff --git a/Engine/Source/Runtime/Slate/Private/Framework/Application/SlateApplication.cpp b/Engine/Source/Runtime/Slate/Private/Framework/Application/SlateApplication.cpp
--- a/Engine/Source/Runtime/Slate/Private/Framework/Application/SlateApplication.cpp
+++ b/Engine/Source/Runtime/Slate/Private/Framework/Application/SlateApplication.cpp
@@ -1860,27 +1860,27 @@ TSharedPtr<IMenu> FSlateApplication::PushMenu(const TSharedPtr<IMenu>& InParentM
 	return MenuStack.Push(InParentMenu, InContent, SummonLocation, TransitionEffect, bFocusImmediately, SummonLocationSize, bIsCollapsedByParent);
 }
 
-TSharedPtr<IMenu> FSlateApplication::PushHostedMenu(const TSharedRef<SWidget>& InParentWidget, const FWidgetPath& InOwnerPath, const TSharedRef<IMenuHost>& InMenuHost, const TSharedRef<SWidget>& InContent, TSharedPtr<SWidget>& OutWrappedContent, const FPopupTransitionEffect& TransitionEffect, const bool bIsCollapsedByParent)
+TSharedPtr<IMenu> FSlateApplication::PushHostedMenu(const TSharedRef<SWidget>& InParentWidget, const FWidgetPath& InOwnerPath, const TSharedRef<IMenuHost>& InMenuHost, const TSharedRef<SWidget>& InContent, TSharedPtr<SWidget>& OutWrappedContent, const FPopupTransitionEffect& TransitionEffect, EShouldThrottle ShouldThrottle, const bool bIsCollapsedByParent)
 {
 	// Caller supplied a valid path? Pass it to the menu stack.
 	if (InOwnerPath.IsValid())
 	{
-		return MenuStack.PushHosted(InOwnerPath, InMenuHost, InContent, OutWrappedContent, TransitionEffect, bIsCollapsedByParent);
+		return MenuStack.PushHosted(InOwnerPath, InMenuHost, InContent, OutWrappedContent, TransitionEffect, ShouldThrottle, bIsCollapsedByParent);
 	}
 
 	// If the caller doesn't specify a valid event path we'll generate one from InParentWidget
 	FWidgetPath WidgetPath;
 	if (GeneratePathToWidgetUnchecked(InParentWidget, WidgetPath))
 	{
-		return MenuStack.PushHosted(WidgetPath, InMenuHost, InContent, OutWrappedContent, TransitionEffect, bIsCollapsedByParent);
+		return MenuStack.PushHosted(WidgetPath, InMenuHost, InContent, OutWrappedContent, TransitionEffect, ShouldThrottle, bIsCollapsedByParent);
 	}
 
 	return TSharedPtr<IMenu>();
 }
 
-TSharedPtr<IMenu> FSlateApplication::PushHostedMenu(const TSharedPtr<IMenu>& InParentMenu, const TSharedRef<IMenuHost>& InMenuHost, const TSharedRef<SWidget>& InContent, TSharedPtr<SWidget>& OutWrappedContent, const FPopupTransitionEffect& TransitionEffect, const bool bIsCollapsedByParent)
+TSharedPtr<IMenu> FSlateApplication::PushHostedMenu(const TSharedPtr<IMenu>& InParentMenu, const TSharedRef<IMenuHost>& InMenuHost, const TSharedRef<SWidget>& InContent, TSharedPtr<SWidget>& OutWrappedContent, const FPopupTransitionEffect& TransitionEffect, EShouldThrottle ShouldThrottle, const bool bIsCollapsedByParent)
 {
-	return MenuStack.PushHosted(InParentMenu, InMenuHost, InContent, OutWrappedContent, TransitionEffect, bIsCollapsedByParent);
+	return MenuStack.PushHosted(InParentMenu, InMenuHost, InContent, OutWrappedContent, TransitionEffect, ShouldThrottle, bIsCollapsedByParent);
 }
 
 bool FSlateApplication::HasOpenSubMenus( TSharedRef<SWindow> Window ) const
diff --git a/Engine/Source/Runtime/Slate/Private/Widgets/Input/SMenuAnchor.cpp b/Engine/Source/Runtime/Slate/Private/Widgets/Input/SMenuAnchor.cpp
--- a/Engine/Source/Runtime/Slate/Private/Widgets/Input/SMenuAnchor.cpp
+++ b/Engine/Source/Runtime/Slate/Private/Widgets/Input/SMenuAnchor.cpp
@@ -239,12 +239,12 @@ int32 SMenuAnchor::OnPaint( const FPaintArgs& Args, const FGeometry& AllottedGeo
 
 bool SMenuAnchor::IsOpenAndReusingWindow() const
 {
-	return MethodInUse.IsSet() && MethodInUse.GetValue() == EPopupMethod::UseCurrentWindow;
+	return MethodInUse.IsSet() && MethodInUse.GetPopupMethod() == EPopupMethod::UseCurrentWindow;
 }
 
 bool SMenuAnchor::IsOpenViaCreatedWindow() const
 {
-	return MethodInUse.IsSet() && MethodInUse.GetValue() == EPopupMethod::CreateNewWindow;
+	return MethodInUse.IsSet() && MethodInUse.GetPopupMethod() == EPopupMethod::CreateNewWindow;
 }
 
 void SMenuAnchor::SetContent(TSharedRef<SWidget> InContent)
@@ -262,18 +262,18 @@ void SMenuAnchor::SetMenuContent(TSharedRef<SWidget> InMenuContent)
 	WrappedContent = InMenuContent;	// wrapping, if any will happen when the menu is opened
 }
 
-EPopupMethod QueryPopupMethod(const FWidgetPath& PathToQuery)
+FPopupMethodReply QueryPopupMethod(const FWidgetPath& PathToQuery)
 {
 	for (int32 WidgetIndex = PathToQuery.Widgets.Num() - 1; WidgetIndex >= 0; --WidgetIndex)
 	{
-		TOptional<EPopupMethod> PopupMethod = PathToQuery.Widgets[WidgetIndex].Widget->OnQueryPopupMethod();
-		if (PopupMethod.IsSet())
+		FPopupMethodReply PopupMethod = PathToQuery.Widgets[WidgetIndex].Widget->OnQueryPopupMethod();
+		if (PopupMethod.IsEventHandled())
 		{
-			return PopupMethod.GetValue();
+			return PopupMethod;
 		}
 	}
 
-	return EPopupMethod::CreateNewWindow;
+	return FPopupMethodReply::UseMethod(EPopupMethod::CreateNewWindow);
 }
 
 
@@ -343,14 +343,14 @@ void SMenuAnchor::SetIsOpen( bool InIsOpen, const bool bFocusMenu )
 					}
 
 					MethodInUse = Method.IsSet()
-						? Method.GetValue()
+						? FPopupMethodReply::UseMethod(Method.GetValue())
 						: QueryPopupMethod(MyWidgetPath);
 
-					if (MethodInUse == EPopupMethod::CreateNewWindow)
+					if (MethodInUse.GetPopupMethod() == EPopupMethod::CreateNewWindow)
 					{
 						// Open the pop-up
 						const bool bIsCollapsedByParent = false;	// don't auto-close child menus when the parent gets focus
-						TSharedPtr<IMenu> NewMenu = FSlateApplication::Get().PushMenu(AsShared(), MyWidgetPath, MenuContentRef, NewPosition, TransitionEffect, bFocusMenu, SummonLocationSize, MethodInUse, bIsCollapsedByParent);
+						TSharedPtr<IMenu> NewMenu = FSlateApplication::Get().PushMenu(AsShared(), MyWidgetPath, MenuContentRef, NewPosition, TransitionEffect, bFocusMenu, SummonLocationSize, MethodInUse.GetPopupMethod(), bIsCollapsedByParent);
 
 						PopupMenuPtr = NewMenu;
 						check(NewMenu.IsValid() && NewMenu->GetOwnedWindow().IsValid());
@@ -361,7 +361,7 @@ void SMenuAnchor::SetIsOpen( bool InIsOpen, const bool bFocusMenu )
 					{
 						// We are re-using the current window instead of creating a new one.
 						// The popup will be presented via an overlay service.
-						ensure(MethodInUse == EPopupMethod::UseCurrentWindow);
+						ensure(MethodInUse.GetPopupMethod() == EPopupMethod::UseCurrentWindow);
 						TSharedRef<SWindow> PopupWindow = MyWidgetPath.GetWindow();
 						PopupWindowPtr = PopupWindow;
 
@@ -374,7 +374,7 @@ void SMenuAnchor::SetIsOpen( bool InIsOpen, const bool bFocusMenu )
 
 						const bool bIsCollapsedByParent = false;	// don't auto-close child menus when the parent gets focus
 						TSharedPtr<IMenu> NewMenu = FSlateApplication::Get().PushHostedMenu(
-							SharedThis, MyWidgetPath, SharedThis, MenuContentRef, WrappedContent, TransitionEffect, bIsCollapsedByParent);
+							SharedThis, MyWidgetPath, SharedThis, MenuContentRef, WrappedContent, TransitionEffect, MethodInUse.GetShouldThrottle(), bIsCollapsedByParent);
 
 						PopupMenuPtr = NewMenu;
 						check(NewMenu.IsValid());
@@ -404,7 +404,7 @@ void SMenuAnchor::SetIsOpen( bool InIsOpen, const bool bFocusMenu )
 			else
 			{
 				PopupWindowPtr.Reset();
-				MethodInUse.Reset();
+				MethodInUse = FPopupMethodReply::Unhandled();
 			}
 		}
 	}
@@ -413,7 +413,7 @@ void SMenuAnchor::SetIsOpen( bool InIsOpen, const bool bFocusMenu )
 void SMenuAnchor::OnMenuClosed(TSharedRef<IMenu> InMenu)
 {
 	bDismissedThisTick = true;
-	MethodInUse = TOptional<EPopupMethod>();
+	MethodInUse = FPopupMethodReply::Unhandled();
 	PopupMenuPtr.Reset();
 	PopupWindowPtr.Reset();
 
@@ -492,5 +492,5 @@ SMenuAnchor::~SMenuAnchor()
 	// We no longer have a popup open, so reset all the tracking state associated.
 	PopupMenuPtr.Reset();
 	PopupWindowPtr.Reset();
-	MethodInUse.Reset();
+	MethodInUse = FPopupMethodReply::Unhandled();
 }
diff --git a/Engine/Source/Runtime/Slate/Private/Widgets/SViewport.cpp b/Engine/Source/Runtime/Slate/Private/Widgets/SViewport.cpp
--- a/Engine/Source/Runtime/Slate/Private/Widgets/SViewport.cpp
+++ b/Engine/Source/Runtime/Slate/Private/Widgets/SViewport.cpp
@@ -299,7 +299,7 @@ TOptional<bool> SViewport::OnQueryShowFocus(const EFocusCause InFocusCause) cons
 	return ViewportInterface.IsValid() ? ViewportInterface.Pin()->OnQueryShowFocus(InFocusCause) : TOptional<bool>();
 }
 
-TOptional<EPopupMethod> SViewport::OnQueryPopupMethod() const
+FPopupMethodReply SViewport::OnQueryPopupMethod() const
 {
 	TSharedPtr<ISlateViewport> PinnedInterface = ViewportInterface.Pin();
 	if (PinnedInterface.IsValid())
@@ -308,7 +308,7 @@ TOptional<EPopupMethod> SViewport::OnQueryPopupMethod() const
 	}
 	else
 	{
-		return EPopupMethod::CreateNewWindow;
+		return FPopupMethodReply::UseMethod(EPopupMethod::CreateNewWindow);
 	}	
 }
 
diff --git a/Engine/Source/Runtime/Slate/Public/Framework/Application/MenuStack.h b/Engine/Source/Runtime/Slate/Public/Framework/Application/MenuStack.h
--- a/Engine/Source/Runtime/Slate/Public/Framework/Application/MenuStack.h
+++ b/Engine/Source/Runtime/Slate/Public/Framework/Application/MenuStack.h
@@ -116,7 +116,7 @@ public:
 	 * @param TransitionEffect		Animation to use when the popup appears
 	 * @param bIsCollapsedByParent	Is this menu collapsed when a parent menu receives focus/activation? If false, only focus/activation outside the entire stack will auto collapse it.
 	 */
-	TSharedRef<IMenu> PushHosted(const FWidgetPath& InOwnerPath, const TSharedRef<IMenuHost>& InMenuHost, const TSharedRef<SWidget>& InContent, TSharedPtr<SWidget>& OutWrappedContent, const FPopupTransitionEffect& TransitionEffect, const bool bIsCollapsedByParent = true);
+	TSharedRef<IMenu> PushHosted(const FWidgetPath& InOwnerPath, const TSharedRef<IMenuHost>& InMenuHost, const TSharedRef<SWidget>& InContent, TSharedPtr<SWidget>& OutWrappedContent, const FPopupTransitionEffect& TransitionEffect, EShouldThrottle ShouldThrottle, const bool bIsCollapsedByParent = true);
 	
 	/**
 	 * Pushes a new child menu onto the stack that is drawn by an external host widget.
@@ -129,7 +129,7 @@ public:
 	 * @param TransitionEffect		Animation to use when the popup appears
 	 * @param bIsCollapsedByParent	Is this menu collapsed when a parent menu receives focus/activation? If false, only focus/activation outside the entire stack will auto collapse it.
 	 */
-	TSharedRef<IMenu> PushHosted(const TSharedPtr<IMenu>& InParentMenu, const TSharedRef<IMenuHost>& InMenuHost, const TSharedRef<SWidget>& InContent, TSharedPtr<SWidget>& OutWrappedContent, const FPopupTransitionEffect& TransitionEffect, const bool bIsCollapsedByParent = true);
+	TSharedRef<IMenu> PushHosted(const TSharedPtr<IMenu>& InParentMenu, const TSharedRef<IMenuHost>& InMenuHost, const TSharedRef<SWidget>& InContent, TSharedPtr<SWidget>& OutWrappedContent, const FPopupTransitionEffect& TransitionEffect, EShouldThrottle ShouldThrottle, const bool bIsCollapsedByParent = true);
 
 	/**
 	 * [Deprecated] Dismisses the menu stack up to a certain level (by default removes all levels in the stack)
@@ -260,7 +260,7 @@ private:
 	 *
 	 * @return	The popup method chosen by a widget in the path. Defaults to EPopupMethod::CreateNewWindow if no widgets in the path have a preference.
 	 */
-	EPopupMethod QueryPopupMethod(const FWidgetPath& PathToQuery);
+	FPopupMethodReply QueryPopupMethod(const FWidgetPath& PathToQuery);
 
 	/**
 	 * Called by public Dismiss methods. Dismisses all menus in the stack from FirstStackIndexToRemove and below.
@@ -348,7 +348,7 @@ private:
 	 *
 	 * @return			The newly created menu.
 	 */
-	TSharedRef<IMenu> PushInternal(const TSharedPtr<IMenu>& InParentMenu, const TSharedRef<SWidget>& InContent, FSlateRect Anchor, const FPopupTransitionEffect& TransitionEffect, const bool bFocusImmediately, const bool bIsCollapsedByParent);
+	TSharedRef<IMenu> PushInternal(const TSharedPtr<IMenu>& InParentMenu, const TSharedRef<SWidget>& InContent, FSlateRect Anchor, const FPopupTransitionEffect& TransitionEffect, const bool bFocusImmediately, EShouldThrottle ShouldThrottle, const bool bIsCollapsedByParent);
 
 	/**
 	 * This is the actual menu object creation method for FMenuInWindow menus.
@@ -376,13 +376,14 @@ private:
 	 * This is the post-push stage of menu creation.
 	 * It is responsible for updating the stack with the new menu and removing any that it replaces.
 	 *
-	 * @param InParentMenu			The parent menu for this menu
-	 * @param InMenu				A newly created menu
+	 * @param InParentMenu         The parent menu for this menu
+	 * @param InMenu               A newly created menu
+	 * @param ShouldThrottle       Should pushing this menu enable throttling of the engine for a more responsive UI
 	 */
-	void PostPush(TSharedPtr<IMenu> InParentMenu, TSharedRef<FMenuBase> InMenu);
+	void PostPush(TSharedPtr<IMenu> InParentMenu, TSharedRef<FMenuBase> InMenu, EShouldThrottle ShouldThrottle);
 
 	/** The popup method currently used by the whole stack. It can only use one at a time */
-	TOptional<EPopupMethod> ActiveMethod;
+	FPopupMethodReply ActiveMethod;
 
 	/** The parent window of the root menu in the stack. NOT the actual menu window if it's a CreateNewWindow */
 	TSharedPtr<SWindow> HostWindow;
diff --git a/Engine/Source/Runtime/Slate/Public/Framework/Application/SlateApplication.h b/Engine/Source/Runtime/Slate/Public/Framework/Application/SlateApplication.h
--- a/Engine/Source/Runtime/Slate/Public/Framework/Application/SlateApplication.h
+++ b/Engine/Source/Runtime/Slate/Public/Framework/Application/SlateApplication.h
@@ -305,9 +305,10 @@ public:
 	 * @param InContent				The content to be placed inside the new menu
 	 * @param OutWrappedContent		Returns the InContent wrapped with widgets needed by the menu stack system. This is what should be drawn by the host after this call.
 	 * @param TransitionEffect		Animation to use when the popup appears
+	 * @param ShouldThrottle		Should we throttle engine ticking to maximize the menu responsiveness
 	 * @param bIsCollapsedByParent	Is this menu collapsed when a parent menu receives focus/activation? If false, only focus/activation outside the entire stack will auto collapse it.
 	 */
-	TSharedPtr<IMenu> PushHostedMenu(const TSharedRef<SWidget>& InParentWidget, const FWidgetPath& InOwnerPath, const TSharedRef<IMenuHost>& InMenuHost, const TSharedRef<SWidget>& InContent, TSharedPtr<SWidget>& OutWrappedContent, const FPopupTransitionEffect& TransitionEffect, const bool bIsCollapsedByParent = true);
+	TSharedPtr<IMenu> PushHostedMenu(const TSharedRef<SWidget>& InParentWidget, const FWidgetPath& InOwnerPath, const TSharedRef<IMenuHost>& InMenuHost, const TSharedRef<SWidget>& InContent, TSharedPtr<SWidget>& OutWrappedContent, const FPopupTransitionEffect& TransitionEffect, EShouldThrottle ShouldThrottle, const bool bIsCollapsedByParent = true);
 	
 	/**
 	 * Creates a new hosted child Menu and adds it to the menu stack under the specified parent menu.
@@ -318,9 +319,10 @@ public:
 	 * @param InContent				The menu's content
 	 * @param OutWrappedContent		Returns the InContent wrapped with widgets needed by the menu stack system. This is what should be drawn by the host after this call.
 	 * @param TransitionEffect		Animation to use when the popup appears
+	 * @param ShouldThrottle		Should we throttle engine ticking to maximize the menu responsiveness
 	 * @param bIsCollapsedByParent	Is this menu collapsed when a parent menu receives focus/activation? If false, only focus/activation outside the entire stack will auto collapse it.
 	 */	
-	TSharedPtr<IMenu> PushHostedMenu(const TSharedPtr<IMenu>& InParentMenu, const TSharedRef<IMenuHost>& InMenuHost, const TSharedRef<SWidget>& InContent, TSharedPtr<SWidget>& OutWrappedContent, const FPopupTransitionEffect& TransitionEffect, const bool bIsCollapsedByParent = true);
+	TSharedPtr<IMenu> PushHostedMenu(const TSharedPtr<IMenu>& InParentMenu, const TSharedRef<IMenuHost>& InMenuHost, const TSharedRef<SWidget>& InContent, TSharedPtr<SWidget>& OutWrappedContent, const FPopupTransitionEffect& TransitionEffect, EShouldThrottle ShouldThrottle, const bool bIsCollapsedByParent = true);
 
 	/** @return Returns whether the window has child menus. */
 	DEPRECATED(4.9, "HasOpenSubMenus() taking a window is deprecated. Use HasOpenSubMenus() taking an IMenu as a parameter.")
diff --git a/Engine/Source/Runtime/Slate/Public/Widgets/Input/SMenuAnchor.h b/Engine/Source/Runtime/Slate/Public/Widgets/Input/SMenuAnchor.h
--- a/Engine/Source/Runtime/Slate/Public/Widgets/Input/SMenuAnchor.h
+++ b/Engine/Source/Runtime/Slate/Public/Widgets/Input/SMenuAnchor.h
@@ -3,20 +3,13 @@
 #pragma once
 #include "STextBlock.h"
 #include "IMenu.h"
+#include "PopupMethodReply.h"
 
 /** Notification when popup is opened/closed. */
 DECLARE_DELEGATE_OneParam(FOnIsOpenChanged, bool)
 
 
-enum class EPopupMethod : uint8
-{
-	/** Creating a new window allows us to place popups outside of the window in which the menu anchor resides. */
-	CreateNewWindow,
-	/** Place the popup into the current window. Applications that intend to run in fullscreen cannot create new windows, so they must use this method.. */
-	UseCurrentWindow
-};
 
-;
 
 /**
  * A PopupAnchor summons a Popup relative to its content.
@@ -150,7 +143,7 @@ protected:
 	TOptional<EPopupMethod> Method;
 
 	/** Method currently being used to show the popup. No value if popup is closed. */
-	TOptional<EPopupMethod> MethodInUse;
+	FPopupMethodReply MethodInUse;
 
 	/**
 	 * @todo Slate : Unify geometry so that this is not necessary.
diff --git a/Engine/Source/Runtime/Slate/Public/Widgets/SViewport.h b/Engine/Source/Runtime/Slate/Public/Widgets/SViewport.h
--- a/Engine/Source/Runtime/Slate/Public/Widgets/SViewport.h
+++ b/Engine/Source/Runtime/Slate/Public/Widgets/SViewport.h
@@ -167,7 +167,7 @@ public:
 	virtual void OnFocusLost( const FFocusEvent& InFocusEvent ) override;
 	virtual FReply OnMotionDetected( const FGeometry& MyGeometry, const FMotionEvent& InMotionEvent ) override;
 	virtual TOptional<bool> OnQueryShowFocus( const EFocusCause InFocusCause ) const override;
-	virtual TOptional<EPopupMethod> OnQueryPopupMethod() const override;
+	virtual FPopupMethodReply OnQueryPopupMethod() const override;
 	virtual void OnFinishedPointerInput() override;
 	virtual void OnArrangeChildren( const FGeometry& AllottedGeometry, FArrangedChildren& ArrangedChildren ) const override;
 	virtual TSharedPtr<struct FVirtualPointerPosition> TranslateMouseCoordinateFor3DChild(const TSharedRef<SWidget>& ChildWidget, const FGeometry& MyGeometry, const FVector2D& ScreenSpaceMouseCoordinate, const FVector2D& LastScreenSpaceMouseCoordinate) const override;
diff --git a/Engine/Source/Runtime/SlateCore/Private/Widgets/SWidget.cpp b/Engine/Source/Runtime/SlateCore/Private/Widgets/SWidget.cpp
--- a/Engine/Source/Runtime/SlateCore/Private/Widgets/SWidget.cpp
+++ b/Engine/Source/Runtime/SlateCore/Private/Widgets/SWidget.cpp
@@ -293,9 +293,9 @@ TOptional<bool> SWidget::OnQueryShowFocus(const EFocusCause InFocusCause) const
 	return TOptional<bool>();
 }
 
-TOptional<EPopupMethod> SWidget::OnQueryPopupMethod() const
+FPopupMethodReply SWidget::OnQueryPopupMethod() const
 {
-	return TOptional<EPopupMethod>();
+	return FPopupMethodReply::Unhandled();
 }
 
 TSharedPtr<struct FVirtualPointerPosition> SWidget::TranslateMouseCoordinateFor3DChild(const TSharedRef<SWidget>& ChildWidget, const FGeometry& MyGeometry, const FVector2D& ScreenSpaceMouseCoordinate, const FVector2D& LastScreenSpaceMouseCoordinate) const
diff --git a/Engine/Source/Runtime/SlateCore/Public/Input/PopupMethodReply.h b/Engine/Source/Runtime/SlateCore/Public/Input/PopupMethodReply.h
new file mode 100644
--- /dev/null
+++ b/Engine/Source/Runtime/SlateCore/Public/Input/PopupMethodReply.h
@@ -0,0 +1,80 @@
+// Copyright 1998-2015 Epic Games, Inc. All Rights Reserved.
+
+#pragma once
+
+#include "ReplyBase.h"
+
+enum class EPopupMethod : uint8
+{
+	/** Creating a new window allows us to place popups outside of the window in which the menu anchor resides. */
+	CreateNewWindow,
+	/** Place the popup into the current window. Applications that intend to run in fullscreen cannot create new windows, so they must use this method.. */
+	UseCurrentWindow
+};
+
+enum class EShouldThrottle : uint8
+{
+	No,
+	Yes
+};
+
+/** Reply informs Slate how it should express the popup: by creating a new window or by reusing the existing window. */
+class FPopupMethodReply : public TReplyBase<FPopupMethodReply>
+{
+public:
+	/** Create a reply that signals not having an opinion about the popup method */
+	static FPopupMethodReply Unhandled()
+	{
+		return FPopupMethodReply(false, EPopupMethod::CreateNewWindow);
+	}
+
+	/** Create a reply that specifies how a popup should be handled. */
+	static FPopupMethodReply UseMethod(EPopupMethod WithMethod)
+	{
+		return FPopupMethodReply(true, WithMethod);
+	}
+
+	/** Specify whether we should throttle the engine ticking s.t. the UI is most responsive when this popup is up. */
+	FPopupMethodReply& SetShouldThrottle(EShouldThrottle InShouldThrottle)
+	{
+		ShouldThrottle = InShouldThrottle;
+		return Me();
+	}
+
+	/** Should we throttle the engine? */
+	EShouldThrottle GetShouldThrottle() const
+	{
+		return ShouldThrottle;
+	}
+
+	/** Which method to use for the popup: new window or reuse current window */
+	EPopupMethod GetPopupMethod() const
+	{
+		return PopupMethod;
+	}
+
+	/** Alias for IsEventHandled for situations where this is used as optional */
+	bool IsSet() const
+	{
+		return IsEventHandled();
+	}
+
+
+public:
+	FPopupMethodReply()
+		:TReplyBase<FPopupMethodReply>(false)
+		, PopupMethod(EPopupMethod::CreateNewWindow)
+		, ShouldThrottle(EShouldThrottle::Yes)
+	{}
+
+private:
+	FPopupMethodReply(bool bIsHandled, EPopupMethod InMethod)
+		: TReplyBase<FPopupMethodReply>(bIsHandled)
+		, PopupMethod(InMethod)
+		, ShouldThrottle(EShouldThrottle::Yes)
+	{
+	}
+
+	EPopupMethod PopupMethod;
+	EShouldThrottle ShouldThrottle;
+};
diff --git a/Engine/Source/Runtime/SlateCore/Public/Input/ReplyBase.h b/Engine/Source/Runtime/SlateCore/Public/Input/ReplyBase.h
--- a/Engine/Source/Runtime/SlateCore/Public/Input/ReplyBase.h
+++ b/Engine/Source/Runtime/SlateCore/Public/Input/ReplyBase.h
@@ -2,6 +2,8 @@
 
 #pragma once
 
+class SWidget;
+
 /**
  * Base class for all the ways that a Slate Widget can reply to events.
  * Typical replies include requests to set focus, capture or release the pointer,
diff --git a/Engine/Source/Runtime/SlateCore/Public/Rendering/RenderingCommon.h b/Engine/Source/Runtime/SlateCore/Public/Rendering/RenderingCommon.h
--- a/Engine/Source/Runtime/SlateCore/Public/Rendering/RenderingCommon.h
+++ b/Engine/Source/Runtime/SlateCore/Public/Rendering/RenderingCommon.h
@@ -2,6 +2,8 @@
 
 #pragma once
 
+#include "PopupMethodReply.h"
+
 struct FVector2D;
 class FSlateRect;
 enum class EPopupMethod : uint8;
@@ -483,9 +485,9 @@ public:
 	 * Making windows allows us to have popups that go outside the parent window, but cannot
 	 * be used in fullscreen and do not have per-pixel alpha.
 	 */
-	virtual TOptional<EPopupMethod> OnQueryPopupMethod() const
+	virtual FPopupMethodReply OnQueryPopupMethod() const
 	{
-		return TOptional<EPopupMethod>();
+		return FPopupMethodReply::Unhandled();
 	}
 
 	/**
diff --git a/Engine/Source/Runtime/SlateCore/Public/Widgets/SWidget.h b/Engine/Source/Runtime/SlateCore/Public/Widgets/SWidget.h
--- a/Engine/Source/Runtime/SlateCore/Public/Widgets/SWidget.h
+++ b/Engine/Source/Runtime/SlateCore/Public/Widgets/SWidget.h
@@ -497,7 +497,7 @@ public:
 	 *      EPopupMethod::UserCurrentWindow. This makes all the menu anchors within them
 	 *      use the current window.
 	 */
-	virtual TOptional<EPopupMethod> OnQueryPopupMethod() const;
+	virtual FPopupMethodReply OnQueryPopupMethod() const;
 
 	virtual TSharedPtr<FVirtualPointerPosition> TranslateMouseCoordinateFor3DChild(const TSharedRef<SWidget>& ChildWidget, const FGeometry& MyGeometry, const FVector2D& ScreenSpaceMouseCoordinate, const FVector2D& LastScreenSpaceMouseCoordinate) const;
 	
diff --git a/Engine/Source/Runtime/WebBrowser/Private/SWebBrowser.cpp b/Engine/Source/Runtime/WebBrowser/Private/SWebBrowser.cpp
--- a/Engine/Source/Runtime/WebBrowser/Private/SWebBrowser.cpp
+++ b/Engine/Source/Runtime/WebBrowser/Private/SWebBrowser.cpp
@@ -551,6 +551,14 @@ void SWebBrowser::HandleMenuDismissed(TSharedRef<IMenu>)
 	PopupMenuPtr.Reset();
 }
 
+
+FPopupMethodReply SWebBrowser::OnQueryPopupMethod() const
+{
+	return PopupMenuMethod.IsSet()
+		? FPopupMethodReply::UseMethod(PopupMenuMethod.GetValue())
+		: FPopupMethodReply::Unhandled();
+}
+
 void SWebBrowser::HandleDismissPopup()
 {
 	if (PopupMenuPtr.IsValid())
diff --git a/Engine/Source/Runtime/WebBrowser/Public/SWebBrowser.h b/Engine/Source/Runtime/WebBrowser/Public/SWebBrowser.h
--- a/Engine/Source/Runtime/WebBrowser/Public/SWebBrowser.h
+++ b/Engine/Source/Runtime/WebBrowser/Public/SWebBrowser.h
@@ -279,10 +279,7 @@ private:
 	/** Callback from the popup menu notifiying it has been dismissed */
 	void HandleMenuDismissed(TSharedRef<IMenu>);
 
-	virtual TOptional<EPopupMethod> OnQueryPopupMethod() const override
-	{
-		return PopupMenuMethod;
-	}
+	virtual FPopupMethodReply OnQueryPopupMethod() const override;
 
 private:
 
