Slate: Prevented incremental resizing in SConstraintCanvas::OnArrangeChildren
[INTEGRATE] Change 2667194 by [REDACTED_PERSON][REDACTED_HANDLE] on 2015/08/25 02:10:37

diff --git a/Engine/Source/Runtime/Slate/Public/Widgets/Layout/SConstraintCanvas.cpp b/Engine/Source/Runtime/Slate/Public/Widgets/Layout/SConstraintCanvas.cpp
--- a/Engine/Source/Runtime/Slate/Public/Widgets/Layout/SConstraintCanvas.cpp
+++ b/Engine/Source/Runtime/Slate/Public/Widgets/Layout/SConstraintCanvas.cpp
@@ -74,7 +74,9 @@ void SConstraintCanvas::OnArrangeChildren( const FGeometry& AllottedGeometry, FA
 	if (Children.Num() > 0)
 	{
 		// Sort the children based on zorder.
-		TArray< FChildZOrder > SlotOrder;
+		TArray< FChildZOrder, TInlineAllocator<64> > SlotOrder;
+		SlotOrder.Reserve(Children.Num());
+
 		for ( int32 ChildIndex = 0; ChildIndex < Children.Num(); ++ChildIndex )
 		{
 			const SConstraintCanvas::FSlot& CurChild = Children[ChildIndex];
