Slate: Avoid incremental allocations in FSlateApplication::MouseCaptorHelper::ToSharedWidgets
[INTEGRATE] Change 2667230 by [REDACTED_PERSON][REDACTED_HANDLE] on 2015/08/25 03:45:29

diff --git a/Engine/Source/Runtime/Slate/Private/Framework/Application/SlateApplication.cpp b/Engine/Source/Runtime/Slate/Private/Framework/Application/SlateApplication.cpp
--- a/Engine/Source/Runtime/Slate/Private/Framework/Application/SlateApplication.cpp
+++ b/Engine/Source/Runtime/Slate/Private/Framework/Application/SlateApplication.cpp
@@ -407,7 +407,9 @@ TSharedPtr< SWidget > FSlateApplication::MouseCaptorHelper::ToSharedWidget(uint3
 TArray<TSharedRef<SWidget>> FSlateApplication::MouseCaptorHelper::ToSharedWidgets() const
 {
 	TArray<TSharedRef<SWidget>> Widgets;
-	for (auto IndexPathPair : PointerIndexToMouseCaptorWeakPathMap)
+	Widgets.Empty(PointerIndexToMouseCaptorWeakPathMap.Num());
+
+	for (const auto& IndexPathPair : PointerIndexToMouseCaptorWeakPathMap)
 	{
 		TSharedPtr<SWidget> LastWidget = IndexPathPair.Value.GetLastWidget().Pin();
 		if (LastWidget.IsValid())
