Added a boolean to turn off Core Style usage in SBreadcrumbTrail, added a delegate that can be used to get foreground color, added breaks to for loops to break out when we find the right index.

diff --git a/Engine/Source/Runtime/Slate/Public/Widgets/Navigation/SBreadcrumbTrail.h b/Engine/Source/Runtime/Slate/Public/Widgets/Navigation/SBreadcrumbTrail.h
--- a/Engine/Source/Runtime/Slate/Public/Widgets/Navigation/SBreadcrumbTrail.h
+++ b/Engine/Source/Runtime/Slate/Public/Widgets/Navigation/SBreadcrumbTrail.h
@@ -4,7 +4,7 @@
 
 
 /**
- * A breadcrumb trail. Allows the user to see his currently selected path and navigate upwards.
+ * A breadcrumb trail. Allows the user to see their currently selected path and navigate upwards.
  */
 template <typename ItemType>
 class SBreadcrumbTrail : public SCompoundWidget
@@ -43,6 +43,8 @@ public:
 	/** Callback for getting the menu content to be displayed when clicking on a crumb's delimiter arrow */
 	DECLARE_DELEGATE_RetVal_OneParam( TSharedPtr< SWidget >, FGetCrumbMenuContent, const ItemType& /*CrumbData*/ );
 
+	DECLARE_DELEGATE_RetVal_TwoParams( FSlateColor, FOnGetCrumbColor, int32 /*CrumbId*/, bool /*bInvert*/ );
+
 
 	SLATE_BEGIN_ARGS( SBreadcrumbTrail )
 		: _InvertTextColorOnHover(true)
@@ -53,7 +55,8 @@ public:
 		, _ShowLeadingDelimiter(false)
 		, _PersistentBreadcrumbs(false)
 		, _GetCrumbMenuContent()
-	{}
+		, _UseCoreStyle( true )
+	    {}
 
 		/** When true, will invert the button text color when a crumb button is hovered */
 		SLATE_ARGUMENT( bool, InvertTextColorOnHover )
@@ -87,6 +90,10 @@ public:
 
 		SLATE_EVENT( FGetCrumbMenuContent, GetCrumbMenuContent )
 
+		SLATE_ARGUMENT( bool, UseCoreStyle )
+
+		SLATE_EVENT( FOnGetCrumbColor, OnGetCrumbColor )
+
 	SLATE_END_ARGS()
 
 	/** Constructs this widget with InArgs */
@@ -103,6 +110,9 @@ public:
 		OnCrumbClicked = InArgs._OnCrumbClicked;
 		bHasStaticBreadcrumbs = InArgs._PersistentBreadcrumbs;
 		GetCrumbMenuContentCallback = InArgs._GetCrumbMenuContent;
+		bUseCoreStyle = InArgs._UseCoreStyle;
+		OnGetCrumbColor = InArgs._OnGetCrumbColor;
+
 
 		NextValidCrumbID = 0;
 
@@ -206,7 +216,8 @@ public:
 		check(HasCrumbs());
 
 		// Remove from the crumb list and box
-		FCrumbItem LastCrumbItem = CrumbList.Pop();
+		const FCrumbItem& LastCrumbItem = CrumbList.Pop();
+
 		CrumbBox->RemoveSlot(LastCrumbItem.ButtonBox);
 		CrumbBox->RemoveSlot(LastCrumbItem.DelimiterBox);
 
@@ -333,27 +344,46 @@ private:
 	/** Handler to determine the text color of crumb buttons. Will invert the text color if allowed. */
 	FSlateColor GetButtonForegroundColor(int32 CrumbID) const
 	{
-		if ( bInvertTextColorOnHover )
-		{
-			TSharedPtr<SButton> CrumbButton;
+		TSharedPtr< SButton > CrumbButton;
 
-			for (int32 CrumbListIdx = 0; CrumbListIdx < NumCrumbs(); ++CrumbListIdx)
+		if ( bUseCoreStyle )
+		{
+			if ( bInvertTextColorOnHover )
 			{
-				if (CrumbList[CrumbListIdx].CrumbID == CrumbID)
+				for ( int32 CrumbListIdx = 0; CrumbListIdx < NumCrumbs(); ++CrumbListIdx )
+				{
+					if ( CrumbList[ CrumbListIdx ].CrumbID == CrumbID )
+					{
+						CrumbButton = CrumbList[ CrumbListIdx ].Button;
+						break;
+					}
+				}
+
+				if ( CrumbButton.IsValid() && CrumbButton->IsHovered() )
 				{
-					CrumbButton = CrumbList[CrumbListIdx].Button;
+					static const FName InvertedForegroundName( "InvertedForeground" );
+					return FCoreStyle::Get().GetSlateColor( InvertedForegroundName );
 				}
 			}
 
-			if ( CrumbButton.IsValid() && CrumbButton->IsHovered() )
+			static const FName DefaultForegroundName;
+			return FCoreStyle::Get().GetSlateColor( DefaultForegroundName );
+		}
+		else
+		{
+			int32 CrumbPosition = INDEX_NONE;
+			for ( int32 CrumbListIdx = 0; CrumbListIdx < NumCrumbs(); ++CrumbListIdx )
 			{
-				static const FName InvertedForegroundName("InvertedForeground");
-				return FCoreStyle::Get().GetSlateColor(InvertedForegroundName);
+				if ( CrumbList[ CrumbListIdx ].CrumbID == CrumbID )
+				{
+					CrumbButton = CrumbList[ CrumbListIdx ].Button;
+					CrumbPosition = CrumbListIdx;
+					break;
+				}
 			}
-		}
 
-		static const FName DefaultForegroundName;
-		return FCoreStyle::Get().GetSlateColor(DefaultForegroundName);
+			return OnGetCrumbColor.Execute( CrumbPosition, ( CrumbButton.IsValid() && CrumbButton->IsHovered() ) );
+		}
 	}
 
 	/** Handler for when a crumb is clicked. Will pop crumbs down to the selected one. */
@@ -379,6 +409,7 @@ private:
 				if (CrumbList[CrumbListIdx].CrumbID == CrumbID)
 				{
 					CrumbIdx = CrumbListIdx;
+					break;
 				}
 			}
 
@@ -450,4 +481,8 @@ private:
 
 	/** If true, don't dynamically remove items when clicking */
 	bool bHasStaticBreadcrumbs;
+
+	bool bUseCoreStyle;
+
+	FOnGetCrumbColor OnGetCrumbColor;
 };
