Removed explicit boolean for using core style from SBreadcrumbTrail

diff --git a/Engine/Source/Runtime/Slate/Public/Widgets/Navigation/SBreadcrumbTrail.h b/Engine/Source/Runtime/Slate/Public/Widgets/Navigation/SBreadcrumbTrail.h
--- a/Engine/Source/Runtime/Slate/Public/Widgets/Navigation/SBreadcrumbTrail.h
+++ b/Engine/Source/Runtime/Slate/Public/Widgets/Navigation/SBreadcrumbTrail.h
@@ -55,7 +55,7 @@ public:
 		, _ShowLeadingDelimiter(false)
 		, _PersistentBreadcrumbs(false)
 		, _GetCrumbMenuContent()
-		, _UseCoreStyle( true )
+		, _OnGetCrumbColor()
 	    {}
 
 		/** When true, will invert the button text color when a crumb button is hovered */
@@ -90,8 +90,6 @@ public:
 
 		SLATE_EVENT( FGetCrumbMenuContent, GetCrumbMenuContent )
 
-		SLATE_ARGUMENT( bool, UseCoreStyle )
-
 		SLATE_EVENT( FOnGetCrumbColor, OnGetCrumbColor )
 
 	SLATE_END_ARGS()
@@ -110,7 +108,6 @@ public:
 		OnCrumbClicked = InArgs._OnCrumbClicked;
 		bHasStaticBreadcrumbs = InArgs._PersistentBreadcrumbs;
 		GetCrumbMenuContentCallback = InArgs._GetCrumbMenuContent;
-		bUseCoreStyle = InArgs._UseCoreStyle;
 		OnGetCrumbColor = InArgs._OnGetCrumbColor;
 
 
@@ -346,7 +343,7 @@ private:
 	{
 		TSharedPtr< SButton > CrumbButton;
 
-		if ( bUseCoreStyle )
+		if ( !OnGetCrumbColor.IsBound() )
 		{
 			if ( bInvertTextColorOnHover )
 			{
@@ -482,7 +479,5 @@ private:
 	/** If true, don't dynamically remove items when clicking */
 	bool bHasStaticBreadcrumbs;
 
-	bool bUseCoreStyle;
-
 	FOnGetCrumbColor OnGetCrumbColor;
 };
