Added an option for anchoring the center of a child widget to the center of a menu anchor.

diff --git a/Engine/Source/Runtime/Slate/Private/Widgets/Input/SMenuAnchor.cpp b/Engine/Source/Runtime/Slate/Private/Widgets/Input/SMenuAnchor.cpp
--- a/Engine/Source/Runtime/Slate/Private/Widgets/Input/SMenuAnchor.cpp
+++ b/Engine/Source/Runtime/Slate/Private/Widgets/Input/SMenuAnchor.cpp
@@ -32,6 +32,9 @@ static FVector2D GetMenuOffsetForPlacement(const FGeometry& AllottedGeometry, EM
 		case MenuPlacement_MenuLeft:
 			return FVector2D(-PopupSizeLocalSpace.X, 0.0f);
 			break;
+		case MenuPlacement_Center:
+			return FVector2D( -( ( PopupSizeLocalSpace.X / 2 ) - ( AllottedGeometry.GetLocalSize().X / 2 ) ), -( ( PopupSizeLocalSpace.Y / 2 ) - ( AllottedGeometry.GetLocalSize().Y / 2 ) ) );
+			break;
 		default:
 			ensureMsgf( false, TEXT("Unhandled placement mode: %d"), PlacementMode );
 			return FVector2D::ZeroVector;
diff --git a/Engine/Source/Runtime/SlateCore/Public/Types/SlateEnums.h b/Engine/Source/Runtime/SlateCore/Public/Types/SlateEnums.h
--- a/Engine/Source/Runtime/SlateCore/Public/Types/SlateEnums.h
+++ b/Engine/Source/Runtime/SlateCore/Public/Types/SlateEnums.h
@@ -174,6 +174,9 @@ enum EMenuPlacement
 
 	/** Place the menu to the left of the anchor */
 	MenuPlacement_MenuLeft UMETA(DisplayName="Left"),
+
+	/** Place the menu's center on top of the menu anchor's center point */
+	MenuPlacement_Center UMETA(DisplayName = "Center"),
 };
 
 
