New UI popup (SMultiBoxWidget) searching system.

-Adds searching widget to display current search and other search feedback
-Hides/shows searched widgets
-Adds highlighting to searches
-Backspace/delete for reducing/removing the search.
-Added support for unique widget elements

UE-809

diff --git a/Engine/Source/Editor/LandscapeEditor/Private/LandscapeEditorDetailWidgets.cpp b/Engine/Source/Editor/LandscapeEditor/Private/LandscapeEditorDetailWidgets.cpp
--- a/Engine/Source/Editor/LandscapeEditor/Private/LandscapeEditorDetailWidgets.cpp
+++ b/Engine/Source/Editor/LandscapeEditor/Private/LandscapeEditorDetailWidgets.cpp
@@ -147,6 +147,10 @@ void SToolSelector::BuildMultiBlockWidget(const ISlateStyle* StyleSet, const FNa
 
 	Label = ToolSelectorButtonBlock->Label;
 	SmallText = ToolSelectorButtonBlock->SmallText;
+
+	// Add this widget to the search list of the multibox
+	if (MultiBlock->GetSearchable())
+		OwnerMultiBoxWidget.Pin()->AddSearchElement(this->AsWidget(), Label.Get());
 	
 	const FString MetaTag = FString::Printf(TEXT("LandscapeToolButton.%s"), Label.IsSet() == true ? *Label.Get().ToString() : TEXT("NoLabel"));
 
diff --git a/Engine/Source/Runtime/Slate/Private/Framework/MultiBox/MultiBox.cpp b/Engine/Source/Runtime/Slate/Private/Framework/MultiBox/MultiBox.cpp
--- a/Engine/Source/Runtime/Slate/Private/Framework/MultiBox/MultiBox.cpp
+++ b/Engine/Source/Runtime/Slate/Private/Framework/MultiBox/MultiBox.cpp
@@ -148,6 +148,15 @@ TSharedRef< IMultiBlockBaseWidget > FMultiBlock::MakeWidget( TSharedRef< SMultiB
 	return NewMultiBlockWidget;
 }
 
+void FMultiBlock::SetSearchable( bool InSearchable )
+{
+	bSearchable = InSearchable;
+}
+bool FMultiBlock::GetSearchable() const
+{
+	return bSearchable;
+}
+
 /**
  * Constructor
  *
@@ -630,6 +639,9 @@ void SMultiBoxWidget::BuildMultiBoxWidget()
 	bool bSectionContainsIcons = false;
 	int32 NextMenuSeparator = INDEX_NONE;
 
+	// Assign and add the search widget at the top
+	SearchTextWidget = MultiBox->SearchTextWidget;
+
 	for( int32 Index = 0; Index < Blocks.Num(); Index++ )
 	{
 		// If we've passed the last menu separator, scan for the next one (the end of the list is also considered a menu separator for the purposes of this index)
@@ -709,7 +721,7 @@ void SMultiBoxWidget::BuildMultiBoxWidget()
 				}
 			}
 		}
-	
+
 
 		if( DragPreview.IsValid() && DragPreview.InsertIndex == Index )
 		{
@@ -1040,14 +1052,164 @@ FReply SMultiBoxWidget::OnFocusReceived( const FGeometry& MyGeometry, const FFoc
 
 FReply SMultiBoxWidget::OnKeyDown( const FGeometry& MyGeometry, const FKeyEvent& KeyEvent )
 {
-	SCompoundWidget::OnKeyDown(MyGeometry, KeyEvent);
+	SCompoundWidget::OnKeyDown( MyGeometry, KeyEvent );
 
+	if ( KeyEvent.GetKey() == EKeys::BackSpace )
+	{
+		ReduceSearch();
+	}
+	else if( KeyEvent.GetKey() == EKeys::Delete )
+	{
+		ResetSearch();
+	}
 	// allow use of up and down keys to transfer focus/hover state
-	if(KeyEvent.GetKey() == EKeys::Up || KeyEvent.GetKey() == EKeys::Down)
+	else if( KeyEvent.GetKey() == EKeys::Up || KeyEvent.GetKey() == EKeys::Down )
 	{
-		return FocusNextWidget(EUINavigation::Next);
+		return FocusNextWidget( EUINavigation::Next );
 	}
 
 	return FReply::Unhandled();
 }
 
+FReply SMultiBoxWidget::OnKeyChar( const FGeometry& MyGeometry, const FCharacterEvent& InCharacterEvent )
+{
+	FReply Reply = FReply::Unhandled();
+	
+	// Check for special characters
+	const TCHAR Character = InCharacterEvent.GetCharacter();
+	TypeChar(Character);
+	Reply = FReply::Handled();
+
+	return Reply;
+}
+
+void SMultiBoxWidget::TypeChar(const int32 InChar)
+{
+	// Certain characters are not allowed
+	bool bIsCharAllowed = true;
+	{
+		if ( InChar == TEXT('\t') )
+		{
+			bIsCharAllowed = true;
+		}
+		else if ( InChar <= 0x1F )
+		{
+			bIsCharAllowed = false;
+		}
+	}
+
+	if ( bIsCharAllowed )
+	{
+		UpdateSearch( InChar );
+	}
+}
+
+void SMultiBoxWidget::UpdateSearch( const int32 CharToAdd )
+{
+	const FString& OldSearchText = SearchText.ToString();
+	SearchText = FText::FromString( OldSearchText + CharToAdd );
+
+	if ( SearchTextWidget.IsValid() && SearchBlockWidget.IsValid())
+	{
+		SearchTextWidget->SetText( FText::Format( FText::FromString("Searching: {0}"), SearchText ) );
+		SearchBlockWidget->SetVisibility(EVisibility::Visible);
+	}
+
+	FilterMultiBoxEntries();
+}
+
+void SMultiBoxWidget::ResetSearch()
+{
+	// Empty search name
+	SearchText = FText::GetEmpty();
+
+	for( auto It = SearchElements.CreateConstIterator(); It; ++It )
+	{
+		It.Key()->SetVisibility( EVisibility::Visible );
+	}
+
+	if (SearchTextWidget.IsValid() && SearchBlockWidget.IsValid())
+	{
+		SearchTextWidget->SetText( FText::FromString("No Search") );
+		SearchBlockWidget->SetVisibility( EVisibility::Collapsed );
+	}
+}
+
+void SMultiBoxWidget::ReduceSearch()
+{
+	if (SearchText.ToString().Len() <= 1)
+	{
+		ResetSearch();
+	}
+	else
+	{
+		SearchText = FText::FromString(SearchText.ToString().LeftChop(1));
+
+		if (SearchTextWidget.IsValid() && SearchBlockWidget.IsValid())
+		{
+			SearchTextWidget->SetText(FText::Format(FText::FromString("Searching: {0}"), SearchText));
+			SearchBlockWidget->SetVisibility(EVisibility::Visible);
+		}
+
+		FilterMultiBoxEntries();
+	}
+}
+
+void SMultiBoxWidget::FilterMultiBoxEntries()
+{
+	bool NoSearchedItems = true;
+	for(auto It = SearchElements.CreateConstIterator(); It; ++It)
+	{
+		// Non-searched elements should not be rendered while searching
+		if(It.Value().IsEmpty())
+		{
+			if(SearchText.IsEmpty())
+			{
+				It.Key()->SetVisibility( EVisibility::Visible );
+			}
+			else
+			{
+				It.Key()->SetVisibility( EVisibility::Collapsed );
+			}
+		}
+		else
+		{
+			// Compare widget text to the current search text
+			if( It.Value().ToString().Contains( SearchText.ToString() ) )
+			{
+				It.Key()->SetVisibility( EVisibility::Visible );
+				NoSearchedItems = false;
+			}
+			else
+			{
+				It.Key()->SetVisibility( EVisibility::Collapsed );
+			}
+		}
+	}
+
+	if (NoSearchedItems && SearchTextWidget.IsValid() && SearchBlockWidget.IsValid() )
+	{
+		SearchTextWidget->SetText( FText::Format( FText::FromString("No Results: {0}"), SearchText ) );
+		SearchBlockWidget->SetVisibility( EVisibility::Visible );
+	}
+}
+
+FText SMultiBoxWidget::GetSearchText() const
+{
+	return SearchText;
+}
+
+TSharedPtr<STextBlock> SMultiBoxWidget::GetSearchTextWidget()
+{
+	return SearchTextWidget;
+}
+
+void SMultiBoxWidget::SetSearchBlockWidget(TSharedPtr<SWidget> InWidget)
+{
+	SearchBlockWidget = InWidget;
+}
+
+void SMultiBoxWidget::AddSearchElement( TSharedPtr<SWidget> BlockWidget, FText BlockDisplayText )
+{
+	SearchElements.Add( BlockWidget, BlockDisplayText );
+}
\ No newline at end of file
diff --git a/Engine/Source/Runtime/Slate/Private/Framework/MultiBox/MultiBoxBuilder.cpp b/Engine/Source/Runtime/Slate/Private/Framework/MultiBox/MultiBoxBuilder.cpp
--- a/Engine/Source/Runtime/Slate/Private/Framework/MultiBox/MultiBoxBuilder.cpp
+++ b/Engine/Source/Runtime/Slate/Private/Framework/MultiBox/MultiBoxBuilder.cpp
@@ -180,7 +180,7 @@ void FMenuBuilder::BeginSection( FName InExtensionHook, const TAttribute< FText
 	{
 		ApplySectionBeginning();
 	}
-	
+
 	ApplyHook(InExtensionHook, EExtensionHook::First);
 }
 
@@ -270,14 +270,25 @@ void FMenuBuilder::AddWrapperSubMenu( const FText& InMenuLabel, const FText& InT
 	MultiBox->AddMultiBlock( NewMenuEntryBlock );
 }
 
-void FMenuBuilder::AddWidget( TSharedRef<SWidget> InWidget, const FText& Label, bool bNoIndent )
+void FMenuBuilder::AddWidget( TSharedRef<SWidget> InWidget, const FText& Label, bool bNoIndent, bool bSearchable )
 {
 	ApplySectionBeginning();
 
-	TSharedRef< FWidgetBlock > NewWidgetBlock( new FWidgetBlock( InWidget, Label, bNoIndent ) );
+	TSharedRef< FWidgetBlock > NewWidgetBlock(new FWidgetBlock( InWidget, Label, bNoIndent ));
+	NewWidgetBlock->SetSearchable( bSearchable );
+
 	MultiBox->AddMultiBlock( NewWidgetBlock );
 }
 
+void FMenuBuilder::AddSearchWidget()
+{
+	MultiBox->SearchTextWidget = SNew(STextBlock)
+		.Visibility( EVisibility::Visible )
+		.Text( FText::FromString("Search Start") );
+
+	AddWidget( MultiBox->SearchTextWidget.ToSharedRef(), FText::GetEmpty(), false, false );
+}
+
 void FMenuBuilder::ApplyHook(FName InExtensionHook, EExtensionHook::Position HookPosition)
 {
 	// this is a virtual function to get a properly typed "this" pointer
@@ -292,7 +303,8 @@ void FMenuBuilder::ApplySectionBeginning()
 {
 	if (bSectionNeedsToBeApplied)
 	{
-		if( MultiBox->GetBlocks().Num() > 0 || FMultiBoxSettings::DisplayMultiboxHooks.Get() )
+		// Do not count search block, which starts as invisible
+		if( MultiBox->GetBlocks().Num() > 1 || FMultiBoxSettings::DisplayMultiboxHooks.Get() )
 		{
 			MultiBox->AddMultiBlock( MakeShareable( new FMenuSeparatorBlock(CurrentSectionExtensionHook) ) );
 		}
diff --git a/Engine/Source/Runtime/Slate/Private/Framework/MultiBox/MultiBoxCustomization.cpp b/Engine/Source/Runtime/Slate/Private/Framework/MultiBox/MultiBoxCustomization.cpp
--- a/Engine/Source/Runtime/Slate/Private/Framework/MultiBox/MultiBoxCustomization.cpp
+++ b/Engine/Source/Runtime/Slate/Private/Framework/MultiBox/MultiBoxCustomization.cpp
@@ -25,6 +25,10 @@ void SCustomToolbarPreviewWidget::BuildMultiBlockWidget(const ISlateStyle* Style
 			Content.ToSharedRef()
 		]
 	];
+
+	// Add this widget to the search list of the multibox and hide it
+	if (MultiBlock->GetSearchable())
+		OwnerMultiBoxWidget.Pin()->AddSearchElement(this->AsWidget(), FText::GetEmpty());
 }
 
 void SMultiBlockDragHandle::Construct( const FArguments& InArgs, TSharedRef<SMultiBoxWidget> InBaseWidget, TSharedRef<const FMultiBlock> InBlock, FName CustomizationName )
diff --git a/Engine/Source/Runtime/Slate/Private/Framework/MultiBox/SButtonRowBlock.cpp b/Engine/Source/Runtime/Slate/Private/Framework/MultiBox/SButtonRowBlock.cpp
--- a/Engine/Source/Runtime/Slate/Private/Framework/MultiBox/SButtonRowBlock.cpp
+++ b/Engine/Source/Runtime/Slate/Private/Framework/MultiBox/SButtonRowBlock.cpp
@@ -83,6 +83,10 @@ void SButtonRowBlock::BuildMultiBlockWidget(const ISlateStyle* StyleSet, const F
 		ActualToolTip = ButtonRowBlock->ToolTipOverride;
 	}
 
+	// Add this widget to the search list of the multibox
+	if (MultiBlock->GetSearchable())
+		OwnerMultiBoxWidget.Pin()->AddSearchElement(this->AsWidget(), ActualLabel.Get());
+
 	// Allow the block to override the button icon, too
 	const FSlateIcon& ActualIcon = !ButtonRowBlock->IconOverride.IsSet() && ButtonRowBlock->GetAction().IsValid() ? ButtonRowBlock->GetAction()->GetIcon() : ButtonRowBlock->IconOverride;
 
diff --git a/Engine/Source/Runtime/Slate/Private/Framework/MultiBox/SEditableTextBlock.cpp b/Engine/Source/Runtime/Slate/Private/Framework/MultiBox/SEditableTextBlock.cpp
--- a/Engine/Source/Runtime/Slate/Private/Framework/MultiBox/SEditableTextBlock.cpp
+++ b/Engine/Source/Runtime/Slate/Private/Framework/MultiBox/SEditableTextBlock.cpp
@@ -53,6 +53,10 @@ void SEditableTextBlock::BuildMultiBlockWidget( const ISlateStyle* StyleSet, con
 	const TAttribute<FText>& TextBlockToolTip = !EditableTextBlock->ToolTipOverride.IsBound() && EditableTextBlock->ToolTipOverride.Get().IsEmpty() && UICommand.IsValid() ? UICommand->GetDescription() : EditableTextBlock->ToolTipOverride;
 	const bool bHasLabel = !Label.Get().IsEmpty();
 
+	// Add this widget to the search list of the multibox
+	if (MultiBlock->GetSearchable())
+		OwnerMultiBoxWidget.Pin()->AddSearchElement(this->AsWidget(), Label.Get());
+
 	// See if the action is valid and if so we will use the actions icon if we dont override it later
 	const FSlateIcon ActionIcon = UICommand.IsValid() ? UICommand->GetIcon() : FSlateIcon();
 
diff --git a/Engine/Source/Runtime/Slate/Private/Framework/MultiBox/SGroupMarkerBlock.cpp b/Engine/Source/Runtime/Slate/Private/Framework/MultiBox/SGroupMarkerBlock.cpp
--- a/Engine/Source/Runtime/Slate/Private/Framework/MultiBox/SGroupMarkerBlock.cpp
+++ b/Engine/Source/Runtime/Slate/Private/Framework/MultiBox/SGroupMarkerBlock.cpp
@@ -50,5 +50,8 @@ void SGroupMarkerBlock::BuildMultiBlockWidget(const ISlateStyle* StyleSet, const
 	[
 		SNullWidget::NullWidget
 	];
-}
 
+	// Add this widget to the search list of the multibox and hide it
+	if (MultiBlock->GetSearchable())
+		OwnerMultiBoxWidget.Pin()->AddSearchElement(this->AsWidget(), FText::GetEmpty());
+}
\ No newline at end of file
diff --git a/Engine/Source/Runtime/Slate/Private/Framework/MultiBox/SHeadingBlock.cpp b/Engine/Source/Runtime/Slate/Private/Framework/MultiBox/SHeadingBlock.cpp
--- a/Engine/Source/Runtime/Slate/Private/Framework/MultiBox/SHeadingBlock.cpp
+++ b/Engine/Source/Runtime/Slate/Private/Framework/MultiBox/SHeadingBlock.cpp
@@ -46,6 +46,10 @@ void SHeadingBlock::BuildMultiBlockWidget(const ISlateStyle* StyleSet, const FNa
 {
 	TSharedRef< const FHeadingBlock > HeadingBlock = StaticCastSharedRef< const FHeadingBlock >( MultiBlock.ToSharedRef() );
 
+	// Add this widget to the search list of the multibox
+	if (MultiBlock->GetSearchable())
+		OwnerMultiBoxWidget.Pin()->AddSearchElement(this->AsWidget(), FText::GetEmpty());
+
 	ChildSlot
 		.Padding( 2.0f )
 		[
@@ -53,5 +57,4 @@ void SHeadingBlock::BuildMultiBlockWidget(const ISlateStyle* StyleSet, const FNa
 				.Text( HeadingBlock->HeadingText )
 				.TextStyle( StyleSet, ISlateStyle::Join( StyleName, ".Heading" ) )
 		];
-}
-
+}
\ No newline at end of file
diff --git a/Engine/Source/Runtime/Slate/Private/Framework/MultiBox/SMenuEntryBlock.cpp b/Engine/Source/Runtime/Slate/Private/Framework/MultiBox/SMenuEntryBlock.cpp
--- a/Engine/Source/Runtime/Slate/Private/Framework/MultiBox/SMenuEntryBlock.cpp
+++ b/Engine/Source/Runtime/Slate/Private/Framework/MultiBox/SMenuEntryBlock.cpp
@@ -178,7 +178,7 @@ void SMenuEntryBlock::Construct( const FArguments& InArgs )
 }
 
 
-TSharedRef< SWidget>  SMenuEntryBlock::BuildMenuBarWidget( const FMenuEntryBuildParams& InBuildParams )
+TSharedRef< SWidget> SMenuEntryBlock::BuildMenuBarWidget( const FMenuEntryBuildParams& InBuildParams )
 {
 	const TAttribute<FText>& Label = InBuildParams.Label;
 	const TAttribute<FText>& EntryToolTip = InBuildParams.ToolTip;
@@ -236,6 +236,7 @@ TSharedRef< SWidget>  SMenuEntryBlock::BuildMenuBarWidget( const FMenuEntryBuild
 					SNew( STextBlock )
 					.TextStyle( StyleSet, ISlateStyle::Join( StyleName, ".Label" ) )
 					.Text( Label )
+					.HighlightText( OwnerMultiBoxWidget.Pin().Get(), &SMultiBoxWidget::GetSearchText )
 				]
 
 				// Bind the button's "on clicked" event to our object's method for this
@@ -250,6 +251,28 @@ TSharedRef< SWidget>  SMenuEntryBlock::BuildMenuBarWidget( const FMenuEntryBuild
 }
 
 
+TSharedRef<SWidget> SMenuEntryBlock::FindTextBlockWidget( TSharedRef<SWidget> Content )
+{
+	if (Content->GetType() == FName(TEXT("STextBlock")))
+	{
+		return Content;
+	}
+
+	FChildren* Children = Content->GetChildren();
+	int32 NumChildren = Children->Num();
+
+	for( int32 Index = 0; Index < NumChildren; ++Index )
+	{
+		TSharedRef<SWidget> Found = FindTextBlockWidget( Children->GetChildAt( Index ));
+		if (Found != SNullWidget::NullWidget)
+		{
+			return Found;
+		}
+	}
+	return SNullWidget::NullWidget;
+}
+
+
 FText SMenuEntryBlock::GetFilteredToolTipText( TAttribute<FText> ToolTipText ) const
 {
 	// If we're part of a menu bar that has a currently open menu, then we suppress our own tool-tip
@@ -470,6 +493,8 @@ TSharedRef< SWidget > SMenuEntryBlock::BuildMenuEntryWidget( const FMenuEntryBui
 		CheckBoxForegroundColor = TAttribute<FSlateColor>::Create( TAttribute<FSlateColor>::FGetter::CreateRaw( this, &SMenuEntryBlock::TintOnHover ) );
 	}
 
+	// If there is custom menu widget, set it
+	// If there isn't, create it
 	TSharedPtr< SWidget > ButtonContent = MenuEntryBlock->EntryWidget;
 	if ( !ButtonContent.IsValid() )
 	{
@@ -503,6 +528,7 @@ TSharedRef< SWidget > SMenuEntryBlock::BuildMenuEntryWidget( const FMenuEntryBui
 			SNew( STextBlock )
 			.TextStyle( StyleSet, ISlateStyle::Join( StyleName, ".Label" ) )
 			.Text( Label )
+			.HighlightText( OwnerMultiBoxWidget.Pin().Get(), &SMultiBoxWidget::GetSearchText )
 		]
 		+ SHorizontalBox::Slot()
 		.AutoWidth()
@@ -776,6 +802,7 @@ TSharedRef< SWidget> SMenuEntryBlock::BuildSubMenuWidget( const FMenuEntryBuildP
 			SNew( STextBlock )
 			.TextStyle( StyleSet, ISlateStyle::Join( StyleName, ".Label" ) )
 			.Text( Label )
+			.HighlightText( OwnerMultiBoxWidget.Pin().Get(), &SMultiBoxWidget::GetSearchText )
 		];
 	}
 
@@ -834,10 +861,12 @@ void SMenuEntryBlock::BuildMultiBlockWidget(const ISlateStyle* StyleSet, const F
 	TSharedPtr< const FMultiBox > MultiBox = OwnerMultiBoxWidget.Pin()->GetMultiBox();
 	TSharedPtr< const FMenuEntryBlock > MenuEntryBlock = StaticCastSharedRef< const FMenuEntryBlock >( MultiBlock.ToSharedRef() );
 	BuildParams.MultiBox = MultiBox;
+	
 	BuildParams.MenuEntryBlock = MenuEntryBlock;
 	BuildParams.UICommand = BuildParams.MenuEntryBlock->GetAction();
 	BuildParams.StyleSet = StyleSet;
 	BuildParams.StyleName = StyleName;
+
 	if (MenuEntryBlock->LabelOverride.IsSet())
 	{
 		BuildParams.Label = MenuEntryBlock->LabelOverride;
@@ -847,6 +876,34 @@ void SMenuEntryBlock::BuildMultiBlockWidget(const ISlateStyle* StyleSet, const F
 		BuildParams.Label = BuildParams.UICommand.IsValid() ? BuildParams.UICommand->GetLabel() : FText::GetEmpty();
 	}
 
+	// Add this widget to the search list of the multibox
+	// If there is a widget already assigned (created early) ensure that it's STextBlock is set up for searching
+	TSharedPtr< SWidget > ButtonContent = MenuEntryBlock->EntryWidget;
+	if( ButtonContent.IsValid() )
+	{
+		TSharedRef<SWidget> TextBlock = FindTextBlockWidget(ButtonContent.ToSharedRef());
+		if (TextBlock != SNullWidget::NullWidget )
+		{
+			TSharedRef<STextBlock> TheTextBlock = StaticCastSharedRef<STextBlock>(TextBlock);
+
+			// Bind the search text to the widgets text to highlight
+			TAttribute<FText> HighlightText;
+			HighlightText.Bind(OwnerMultiBoxWidget.Pin().Get(), &SMultiBoxWidget::GetSearchText);
+			TheTextBlock->SetHighlightText(HighlightText);
+
+			OwnerMultiBoxWidget.Pin()->AddSearchElement( this->AsWidget(), TheTextBlock.Get().GetText() );
+		}
+		else
+		{
+			//UE_LOG(LogTemp, Warning, TEXT("Custom widget block will not be searched, could not find it's displayed STextBlock"));
+		}
+	}
+	else
+	{
+		if (MultiBlock->GetSearchable() && !BuildParams.Label.Get().IsEmpty())
+			OwnerMultiBoxWidget.Pin()->AddSearchElement( this->AsWidget(), BuildParams.Label.Get() );
+	}
+
 	// Tool tips are optional so if the tool tip override is empty and there is no UI command just use the empty tool tip.
 	if (MenuEntryBlock->ToolTipOverride.IsSet())
 	{
diff --git a/Engine/Source/Runtime/Slate/Private/Framework/MultiBox/SMenuEntryBlock.h b/Engine/Source/Runtime/Slate/Private/Framework/MultiBox/SMenuEntryBlock.h
--- a/Engine/Source/Runtime/Slate/Private/Framework/MultiBox/SMenuEntryBlock.h
+++ b/Engine/Source/Runtime/Slate/Private/Framework/MultiBox/SMenuEntryBlock.h
@@ -253,6 +253,14 @@ protected:
 	 * @return	The widget content for the new menu
 	 */
 	TSharedRef< SWidget > BuildMenuBarWidget( const FMenuEntryBuildParams& InBuildParams );
+	
+	/**
+	* Finds the STextBlock that gets displayed in the UI
+	*
+	* @param Content	Widget to check for an STextBlock
+	* @return	The STextBlock widget found
+	*/
+	TSharedRef<SWidget> FindTextBlockWidget(TSharedRef<SWidget> Content);
 
 	/**
 	 * Called to create content for a menu entry inside a pull-down, context, or sub-menu
diff --git a/Engine/Source/Runtime/Slate/Private/Framework/MultiBox/SMenuSeparatorBlock.cpp b/Engine/Source/Runtime/Slate/Private/Framework/MultiBox/SMenuSeparatorBlock.cpp
--- a/Engine/Source/Runtime/Slate/Private/Framework/MultiBox/SMenuSeparatorBlock.cpp
+++ b/Engine/Source/Runtime/Slate/Private/Framework/MultiBox/SMenuSeparatorBlock.cpp
@@ -60,6 +60,8 @@ void SMenuSeparatorBlock::BuildMultiBlockWidget(const ISlateStyle* StyleSet, con
 						.BorderImage( StyleSet->GetBrush( StyleName, ".Separator" ) )
 				]
 	];
-}
-
 
+	// Add this widget to the search list of the multibox and hide it
+	if (MultiBlock->GetSearchable())
+		OwnerMultiBoxWidget.Pin()->AddSearchElement(this->AsWidget(), FText::GetEmpty());
+}
diff --git a/Engine/Source/Runtime/Slate/Private/Framework/MultiBox/SToolBarButtonBlock.cpp b/Engine/Source/Runtime/Slate/Private/Framework/MultiBox/SToolBarButtonBlock.cpp
--- a/Engine/Source/Runtime/Slate/Private/Framework/MultiBox/SToolBarButtonBlock.cpp
+++ b/Engine/Source/Runtime/Slate/Private/Framework/MultiBox/SToolBarButtonBlock.cpp
@@ -139,6 +139,10 @@ void SToolBarButtonBlock::BuildMultiBlockWidget(const ISlateStyle* StyleSet, con
 		ActualLabel = ToolBarButtonBlock->GetAction()->GetLabel();
 	}
 
+	// Add this widget to the search list of the multibox
+	if (MultiBlock->GetSearchable())
+		OwnerMultiBoxWidget.Pin()->AddSearchElement(this->AsWidget(), ActualLabel.Get());
+
 	TAttribute<FText> ActualToolTip;
 	if (ToolBarButtonBlock->ToolTipOverride.IsSet())
 	{
diff --git a/Engine/Source/Runtime/Slate/Private/Framework/MultiBox/SToolBarComboButtonBlock.cpp b/Engine/Source/Runtime/Slate/Private/Framework/MultiBox/SToolBarComboButtonBlock.cpp
--- a/Engine/Source/Runtime/Slate/Private/Framework/MultiBox/SToolBarComboButtonBlock.cpp
+++ b/Engine/Source/Runtime/Slate/Private/Framework/MultiBox/SToolBarComboButtonBlock.cpp
@@ -110,6 +110,10 @@ void SToolBarComboButtonBlock::BuildMultiBlockWidget(const ISlateStyle* StyleSet
 		Label = ToolBarComboButtonBlock->Label;
 	}
 
+	// Add this widget to the search list of the multibox
+	if (MultiBlock->GetSearchable())
+		OwnerMultiBoxWidget.Pin()->AddSearchElement(this->AsWidget(), Label.Get());
+
 	// Setup the string for the metatag
 	FName TagName;
 	if (ToolBarComboButtonBlock->GetTutorialHighlightName() == NAME_None)
diff --git a/Engine/Source/Runtime/Slate/Private/Framework/MultiBox/SToolBarSeparatorBlock.cpp b/Engine/Source/Runtime/Slate/Private/Framework/MultiBox/SToolBarSeparatorBlock.cpp
--- a/Engine/Source/Runtime/Slate/Private/Framework/MultiBox/SToolBarSeparatorBlock.cpp
+++ b/Engine/Source/Runtime/Slate/Private/Framework/MultiBox/SToolBarSeparatorBlock.cpp
@@ -62,6 +62,8 @@ void SToolBarSeparatorBlock::BuildMultiBlockWidget(const ISlateStyle* StyleSet,
 				.SeparatorImage( StyleSet->GetBrush( ISlateStyle::Join( StyleName, ".Separator" ) ) )
 		]
 	];
-}
-
 
+	// Add this widget to the search list of the multibox and hide it
+	if (MultiBlock->GetSearchable())
+		OwnerMultiBoxWidget.Pin()->AddSearchElement(this->AsWidget(), FText::GetEmpty());
+}
\ No newline at end of file
diff --git a/Engine/Source/Runtime/Slate/Private/Framework/MultiBox/SWidgetBlock.cpp b/Engine/Source/Runtime/Slate/Private/Framework/MultiBox/SWidgetBlock.cpp
--- a/Engine/Source/Runtime/Slate/Private/Framework/MultiBox/SWidgetBlock.cpp
+++ b/Engine/Source/Runtime/Slate/Private/Framework/MultiBox/SWidgetBlock.cpp
@@ -59,6 +59,17 @@ void SWidgetBlock::BuildMultiBlockWidget(const ISlateStyle* StyleSet, const FNam
 	bool bHasLabel = !WidgetBlock->Label.IsEmpty();
 	FMargin Padding = WidgetBlock->bNoIndent ? StyleSet->GetMargin( StyleName, ".Block.Padding" ) : StyleSet->GetMargin( StyleName, ".Block.IndentedPadding" );
 
+	// Add this widget to the search list of the multibox
+	if (MultiBlock->GetSearchable())
+		OwnerMultiBoxWidget.Pin()->AddSearchElement(this->AsWidget(), WidgetBlock->Label);
+
+	// This widget holds the search text, set it as the search block widget
+	if (OwnerMultiBoxWidget.Pin()->GetSearchTextWidget() == WidgetBlock->ContentWidget)
+	{
+		OwnerMultiBoxWidget.Pin()->SetSearchBlockWidget(this->AsWidget());
+		this->AsWidget()->SetVisibility(EVisibility::Collapsed);
+	}
+
 	ChildSlot
 	.Padding( Padding )	// Large left margin mimics the indent of normal menu items when bNoIndent is false
 	[
diff --git a/Engine/Source/Runtime/Slate/Private/Widgets/Text/STextBlock.cpp b/Engine/Source/Runtime/Slate/Private/Widgets/Text/STextBlock.cpp
--- a/Engine/Source/Runtime/Slate/Private/Widgets/Text/STextBlock.cpp
+++ b/Engine/Source/Runtime/Slate/Private/Widgets/Text/STextBlock.cpp
@@ -126,6 +126,11 @@ void STextBlock::SetText( const FText& InText )
 	Invalidate(EInvalidateWidget::LayoutAndVolatility);
 }
 
+void STextBlock::SetHighlightText(TAttribute<FText> InText)
+{
+	HighlightText = InText;
+}
+
 int32 STextBlock::OnPaint( const FPaintArgs& Args, const FGeometry& AllottedGeometry, const FSlateRect& MyClippingRect, FSlateWindowElementList& OutDrawElements, int32 LayerId, const FWidgetStyle& InWidgetStyle, bool bParentEnabled ) const
 {
 	SCOPE_CYCLE_COUNTER(Stat_SlateTextBlockOnPaint);
diff --git a/Engine/Source/Runtime/Slate/Public/Framework/MultiBox/MultiBox.h b/Engine/Source/Runtime/Slate/Public/Framework/MultiBox/MultiBox.h
--- a/Engine/Source/Runtime/Slate/Public/Framework/MultiBox/MultiBox.h
+++ b/Engine/Source/Runtime/Slate/Public/Framework/MultiBox/MultiBox.h
@@ -131,6 +131,19 @@ public:
 		return Type;
 	}
 
+	/**
+	* Sets the searchable state of this block
+	*
+	* @param	bSearchable		The searchable state to set
+	*/
+	void SetSearchable(bool bSearchable);
+	/**
+	* Gets the searchable state of this block
+	*
+	* @return	Whether this block is searchable
+	*/
+	bool GetSearchable() const;
+
 private:
 	/**
 	 * Allocates a widget for this type of MultiBlock.  Override this in derived classes.
@@ -164,6 +177,9 @@ private:
 
 	/** Name to identify a widget for tutorials */
 	FName TutorialHighlightName;
+
+	/** Whether this block can be searched */
+	bool bSearchable = true;
 };
 
 
@@ -273,6 +289,10 @@ public:
 	TSharedPtr<const FMultiBlock> FindBlockFromCommand( TSharedPtr<const FUICommandInfo> Command ) const;
 
 	bool IsInEditMode() const { return FMultiBoxSettings::IsInToolbarEditMode() && IsCustomizable(); }
+
+	/* The search widget to be displayed at the top of the multibox */
+	TSharedPtr<STextBlock> SearchTextWidget;
+
 private:
 	
 	/**
@@ -512,6 +532,58 @@ public:
 	virtual bool SupportsKeyboardFocus() const override;
 	virtual FReply OnFocusReceived( const FGeometry& MyGeometry, const FFocusEvent& InFocusEvent ) override;
 	virtual FReply OnKeyDown( const FGeometry& MyGeometry, const FKeyEvent& KeyEvent ) override;
+	virtual FReply OnKeyChar( const FGeometry& MyGeometry, const FCharacterEvent& InCharacterEvent ) override;
+
+	void TypeChar(const int32 InChar);
+
+	/**
+	* Changes visibility of widgets in the multibox
+	*
+	* @param	InSearchText	The search text to check against
+	*/
+	void UpdateSearch(const int32);
+
+	/**
+	* Resets the search to be empty
+	*/
+	void ResetSearch();
+
+	void ReduceSearch();
+
+	/**
+	* Changes visibility of widgets in the multibox
+	*/
+	void FilterMultiBoxEntries();
+
+	/**
+	* Get the text to search by
+	*
+	* @return	The text currently being searched
+	*/
+	FText GetSearchText() const;
+
+	/**
+	* Get the SSearchText widget holding the search text
+	*
+	* @return	The widget to get
+	*/
+	TSharedPtr<STextBlock> GetSearchTextWidget(); 
+
+	/**
+	* Get the block widget holding the search text
+	*
+	* @param	BlockWidget		The widget to set
+	*/
+	void SetSearchBlockWidget(TSharedPtr<SWidget>);
+
+	/**
+	* Adds a widget to SearchElement map, making the widget searchable
+	*
+	* @param	BlockWidget			The widget to add
+	* @param	BlockDisplayText	The display text of the widget to search by
+	*/
+	void AddSearchElement( TSharedPtr<SWidget>, FText );
+
 private:
 	/** Adds a block Widget to this widget */
 	void AddBlockWidget( const FMultiBlock& Block, TSharedPtr<SHorizontalBox> HorizontalBox, TSharedPtr<SVerticalBox> VerticalBox, EMultiBlockLocation::Type InLocation, bool bSectionContainsIcons );
@@ -533,6 +605,7 @@ private:
 	
 	/** Called when a user clicks the delete button on a block */
 	FReply OnDeleteBlockClicked( TWeakPtr<const FMultiBlock> BlockWeakPtr );
+
 private:
 	/** A preview of a block being dragged */
 	struct FDraggedMultiBlockPreview
@@ -571,4 +644,16 @@ private:
 
 	/** A preview of a block being dragged inside this box */
 	FDraggedMultiBlockPreview DragPreview;
+
+	/* The multibox widgets that should be search, linked with their display text */
+	TMap<TSharedPtr<SWidget>, FText > SearchElements;
+
+	/* The search widget to be displayed at the top of the multibox */
+	TSharedPtr<STextBlock> SearchTextWidget;
+
+	/* The search widget to be displayed at the top of the multibox */
+	TSharedPtr<SWidget> SearchBlockWidget;
+
+	/* The text being searched for */
+	FText SearchText = FText::GetEmpty();
 };
diff --git a/Engine/Source/Runtime/Slate/Public/Framework/MultiBox/MultiBoxBuilder.h b/Engine/Source/Runtime/Slate/Public/Framework/MultiBox/MultiBoxBuilder.h
--- a/Engine/Source/Runtime/Slate/Public/Framework/MultiBox/MultiBoxBuilder.h
+++ b/Engine/Source/Runtime/Slate/Public/Framework/MultiBox/MultiBoxBuilder.h
@@ -55,7 +55,7 @@ public:
 	 */
 	TSharedRef< class SWidget > MakeWidget();
 	
-	
+
 	/** 
 	 * Get the multi-box being built.
 	 *
@@ -200,6 +200,7 @@ public:
 		: FBaseMenuBuilder( EMultiBoxType::Menu, bInShouldCloseWindowAfterMenuSelection, InCommandList, bCloseSelfOnly, InExtender, InStyleSet )
 		, bSectionNeedsToBeApplied(false)
 	{
+		AddSearchWidget();
 	}
 
 
@@ -245,9 +246,15 @@ public:
 	 * @param	InWidget			The widget that should be shown in the menu
 	 * @param	InLabel				Optional label text to be added to the left of the content
 	 * @param	bNoIndent			If true, removes the padding from the left of the widget that lines it up with other menu items (default == false)
+	 * @param	bSearchable			If true, widget will be searchable (default == true)
 	 */
-	void AddWidget( TSharedRef<SWidget> InWidget, const FText& Label, bool bNoIndent = false );
-	
+	void AddWidget( TSharedRef<SWidget> InWidget, const FText& Label, bool bNoIndent = false, bool bSearchable = true );
+
+	/**
+	* Adds the widget the multibox will use for searching
+	*/
+	void AddSearchWidget();
+
 protected:
 	/** FMultiBoxBuilder interface */
 	virtual void ApplyHook(FName InExtensionHook, EExtensionHook::Position HookPosition) override;
diff --git a/Engine/Source/Runtime/Slate/Public/Widgets/Text/STextBlock.h b/Engine/Source/Runtime/Slate/Public/Widgets/Text/STextBlock.h
--- a/Engine/Source/Runtime/Slate/Public/Widgets/Text/STextBlock.h
+++ b/Engine/Source/Runtime/Slate/Public/Widgets/Text/STextBlock.h
@@ -133,6 +133,13 @@ public:
 	void SetText( const TAttribute< FText >& InText );
 	void SetText( const FText& InText );
 
+	/**
+	* Sets the highlight text for this text block 
+	*
+	* @param	InText	The new text to highlight
+	*/
+	void SetHighlightText(TAttribute<FText> InText);
+
 	/**
 	 * Sets the font used to draw the text
 	 *
