Fix to SearchBox widget so that clicking the clear button doesn't also close the parent menu.

diff --git a/Engine/Source/Runtime/Slate/Private/Widgets/Input/SSearchBox.cpp b/Engine/Source/Runtime/Slate/Private/Widgets/Input/SSearchBox.cpp
--- a/Engine/Source/Runtime/Slate/Private/Widgets/Input/SSearchBox.cpp
+++ b/Engine/Source/Runtime/Slate/Private/Widgets/Input/SSearchBox.cpp
@@ -107,9 +107,9 @@ void SSearchBox::Construct( const FArguments& InArgs )
 		.VAlign(VAlign_Center)
 		.OnClicked( this, &SSearchBox::OnClearSearch )
 		.ForegroundColor( FSlateColor::UseForeground() )
-		// We want clicking on this to focus the text box instead.
+		// Allow the button to steal focus so that the search text will be automatically committed. Afterwards focus will be returned to the text box.
 		// If the user is keyboard-centric, they'll "ctrl+a, delete" to clear the search
-		.IsFocusable(false)
+		.IsFocusable(true)
 		[
 			SNew(SImage)
 			.Image( &InArgs._Style->ClearImage )
@@ -180,11 +180,13 @@ FReply SSearchBox::OnClickedSearch(SSearchBox::SearchDirection Direction)
 
 FReply SSearchBox::OnClearSearch()
 {
-	// We clear the focus to commit any unset values as the search box is typically used for filtering
-	// and the widget could get immediately destroyed before committing its value.
-	FSlateApplication::Get().ClearKeyboardFocus( EFocusCause::Cleared );
-
+	// When we get here, the button will already have stolen focus, thus committing any unset values in the search box.
+	// This will have allowed any widgets which depend on its state to update themselves prior to the search box being cleared,
+	// which happens now. This is important as the act of clearing the search text may also destroy those widgets (for example,
+	// if the search box is being used as a filter).
 	this->SetText( FText::GetEmpty() );
+
+	// Finally set focus back to the editable text
 	return FReply::Handled().SetUserFocus(EditableText.ToSharedRef(), EFocusCause::SetDirectly);
 }
 
