Fixed UI searching variables for Xbox One.

UE-809

diff --git a/Engine/Source/Runtime/Slate/Private/Framework/MultiBox/MultiBoxBuilder.cpp b/Engine/Source/Runtime/Slate/Private/Framework/MultiBox/MultiBoxBuilder.cpp
--- a/Engine/Source/Runtime/Slate/Private/Framework/MultiBox/MultiBoxBuilder.cpp
+++ b/Engine/Source/Runtime/Slate/Private/Framework/MultiBox/MultiBoxBuilder.cpp
@@ -408,7 +408,7 @@ void FToolBarBuilder::AddComboButton( const FUIAction& InAction, const FOnGetCon
 	MultiBox->AddMultiBlock( NewToolBarComboButtonBlock );
 }
 
-void FToolBarBuilder::AddWidget( TSharedRef<SWidget> InWidget, FName InTutorialHighlightName )
+void FToolBarBuilder::AddWidget( TSharedRef<SWidget> InWidget, FName InTutorialHighlightName, bool bSearchable )
 {
 	ApplySectionBeginning();
 
@@ -425,6 +425,7 @@ void FToolBarBuilder::AddWidget( TSharedRef<SWidget> InWidget, FName InTutorialH
 	
 	TSharedRef< FWidgetBlock > NewWidgetBlock( new FWidgetBlock( InWidget, FText::GetEmpty(), true ) );
 	MultiBox->AddMultiBlock( NewWidgetBlock );
+	NewWidgetBlock->SetSearchable(bSearchable);
 }
 
 void FToolBarBuilder::AddSeparator(FName InExtensionHook)
diff --git a/Engine/Source/Runtime/Slate/Public/Framework/MultiBox/MultiBox.h b/Engine/Source/Runtime/Slate/Public/Framework/MultiBox/MultiBox.h
--- a/Engine/Source/Runtime/Slate/Public/Framework/MultiBox/MultiBox.h
+++ b/Engine/Source/Runtime/Slate/Public/Framework/MultiBox/MultiBox.h
@@ -179,7 +179,7 @@ private:
 	FName TutorialHighlightName;
 
 	/** Whether this block can be searched */
-	bool bSearchable = true;
+	bool bSearchable;
 };
 
 
@@ -655,5 +655,5 @@ private:
 	TSharedPtr<SWidget> SearchBlockWidget;
 
 	/* The text being searched for */
-	FText SearchText = FText::GetEmpty();
+	FText SearchText;
 };
diff --git a/Engine/Source/Runtime/Slate/Public/Framework/MultiBox/MultiBoxBuilder.h b/Engine/Source/Runtime/Slate/Public/Framework/MultiBox/MultiBoxBuilder.h
--- a/Engine/Source/Runtime/Slate/Public/Framework/MultiBox/MultiBoxBuilder.h
+++ b/Engine/Source/Runtime/Slate/Public/Framework/MultiBox/MultiBoxBuilder.h
@@ -397,8 +397,9 @@ public:
 	 * 
 	 * @param	InWidget				The widget that should be shown in the toolbar
 	 * @param	InTutorialHighlightName	Name to identify this widget and highlight during tutorials
+	 * @param	bSearchable			If true, widget will be searchable (default == true)
 	 */
-	void AddWidget( TSharedRef<SWidget> InWidget, FName InTutorialHighlightName = NAME_None );
+	void AddWidget(TSharedRef<SWidget> InWidget, FName InTutorialHighlightName = NAME_None, bool bSearchable = true);
 	
 	/**
 	 * Adds a toolbar separator
