Fixed several SMultiboxwidgets being searchable that shouldn't be. Toolbarwidgets set to be non-searchable by default.

diff --git a/Engine/Source/Runtime/Slate/Private/Framework/MultiBox/MultiBoxBuilder.cpp b/Engine/Source/Runtime/Slate/Private/Framework/MultiBox/MultiBoxBuilder.cpp
--- a/Engine/Source/Runtime/Slate/Private/Framework/MultiBox/MultiBoxBuilder.cpp
+++ b/Engine/Source/Runtime/Slate/Private/Framework/MultiBox/MultiBoxBuilder.cpp
@@ -347,7 +347,7 @@ void FMenuBarBuilder::ApplyHook(FName InExtensionHook, EExtensionHook::Position
 	}
 }
 
-void FToolBarBuilder::AddToolBarButton( const TSharedPtr< const FUICommandInfo > InCommand, FName InExtensionHook, const TAttribute<FText>& InLabelOverride, const TAttribute<FText>& InToolTipOverride, const TAttribute<FSlateIcon>& InIconOverride, FName InTutorialHighlightName)
+void FToolBarBuilder::AddToolBarButton(const TSharedPtr< const FUICommandInfo > InCommand, FName InExtensionHook, const TAttribute<FText>& InLabelOverride, const TAttribute<FText>& InToolTipOverride, const TAttribute<FSlateIcon>& InIconOverride, FName InTutorialHighlightName, bool bSearchable )
 {
 	ApplySectionBeginning();
 
@@ -366,10 +366,13 @@ void FToolBarBuilder::AddToolBarButton( const TSharedPtr< const FUICommandInfo >
 
 	MultiBox->AddMultiBlock( NewToolBarButtonBlock );
 
+	// Set if this object should be searchable, defaulted to false
+	NewToolBarButtonBlock->SetSearchable(bSearchable);
+
 	ApplyHook(InExtensionHook, EExtensionHook::After);
 }
 
-void FToolBarBuilder::AddToolBarButton( const FUIAction& InAction, FName InExtensionHook, const TAttribute<FText>& InLabelOverride, const TAttribute<FText>& InToolTipOverride, const TAttribute<FSlateIcon>& InIconOverride, const EUserInterfaceActionType::Type UserInterfaceActionType, FName InTutorialHighlightName )
+void FToolBarBuilder::AddToolBarButton(const FUIAction& InAction, FName InExtensionHook, const TAttribute<FText>& InLabelOverride, const TAttribute<FText>& InToolTipOverride, const TAttribute<FSlateIcon>& InIconOverride, const EUserInterfaceActionType::Type UserInterfaceActionType, FName InTutorialHighlightName, bool bSearchable )
 {
 	ApplySectionBeginning();
 
@@ -386,6 +389,7 @@ void FToolBarBuilder::AddToolBarButton( const FUIAction& InAction, FName InExten
 	NewToolBarButtonBlock->SetForceSmallIcons(bForceSmallIcons);
 	NewToolBarButtonBlock->SetTutorialHighlightName(GenerateTutorialIdentfierName(TutorialHighlightName, InTutorialHighlightName, nullptr, MultiBox->GetBlocks().Num()));
 
+	// Set if this object should be searchable, defaulted to false
 	MultiBox->AddMultiBlock( NewToolBarButtonBlock );
 
 	ApplyHook(InExtensionHook, EExtensionHook::After);
diff --git a/Engine/Source/Runtime/Slate/Public/Framework/MultiBox/MultiBoxBuilder.h b/Engine/Source/Runtime/Slate/Public/Framework/MultiBox/MultiBoxBuilder.h
--- a/Engine/Source/Runtime/Slate/Public/Framework/MultiBox/MultiBoxBuilder.h
+++ b/Engine/Source/Runtime/Slate/Public/Framework/MultiBox/MultiBoxBuilder.h
@@ -364,8 +364,9 @@ public:
 	 * @param	InToolTipOverride		Optional tool tip override.	 If omitted, then the action's label will be used instead.
 	 * @param	InIconOverride			Optional name of the slate brush to use for the tool bar image.  If omitted, then the action's icon will be used instead.
 	 * @param	InTutorialHighlightName	Name to identify this widget and highlight during tutorials
+	 * @param	bSearchable	If true, widget will be searchable (default == false)
 	 */
-	void AddToolBarButton( const TSharedPtr< const FUICommandInfo > InCommand, FName InExtensionHook = NAME_None, const TAttribute<FText>& InLabelOverride = TAttribute<FText>(), const TAttribute<FText>& InToolTipOverride = TAttribute<FText>(), const TAttribute<FSlateIcon>& InIconOverride = TAttribute<FSlateIcon>(), FName InTutorialHighlightName = NAME_None );
+	void AddToolBarButton(const TSharedPtr< const FUICommandInfo > InCommand, FName InExtensionHook = NAME_None, const TAttribute<FText>& InLabelOverride = TAttribute<FText>(), const TAttribute<FText>& InToolTipOverride = TAttribute<FText>(), const TAttribute<FSlateIcon>& InIconOverride = TAttribute<FSlateIcon>(), FName InTutorialHighlightName = NAME_None, bool bSearchable = false );
 	
 	/**
 	 * Adds a tool bar button
@@ -376,8 +377,9 @@ public:
 	 * @param	InIcon		The icon to use		
 	 * @param	UserInterfaceActionType	Type of interface action
 	 * @param	InTutorialHighlightName	Name to identify this widget and highlight during tutorials
+	 * @param	bSearchable	If true, widget will be searchable (default == false)
 	 */
-	void AddToolBarButton( const FUIAction& InAction, FName InExtensionHook = NAME_None, const TAttribute<FText>& InLabelOverride = TAttribute<FText>(), const TAttribute<FText>& InToolTipOverride = TAttribute<FText>(), const TAttribute<FSlateIcon>& InIconOverride = TAttribute<FSlateIcon>(), const EUserInterfaceActionType::Type UserInterfaceActionType = EUserInterfaceActionType::Button, FName InTutorialHighlightName = NAME_None );
+	void AddToolBarButton(const FUIAction& InAction, FName InExtensionHook = NAME_None, const TAttribute<FText>& InLabelOverride = TAttribute<FText>(), const TAttribute<FText>& InToolTipOverride = TAttribute<FText>(), const TAttribute<FSlateIcon>& InIconOverride = TAttribute<FSlateIcon>(), const EUserInterfaceActionType::Type UserInterfaceActionType = EUserInterfaceActionType::Button, FName InTutorialHighlightName = NAME_None, bool bSearchable = false );
 
 	/**
 	 * Adds a combo button
