Added ability to disable searching on certain boxes. Made it so that only fmenubuilder elements use searching, other elements that were not intended to search such as Ftoolbarbuilder no longer do.

UE-20462

diff --git a/Engine/Source/Runtime/Slate/Private/Framework/MultiBox/MultiBox.cpp b/Engine/Source/Runtime/Slate/Private/Framework/MultiBox/MultiBox.cpp
--- a/Engine/Source/Runtime/Slate/Private/Framework/MultiBox/MultiBox.cpp
+++ b/Engine/Source/Runtime/Slate/Private/Framework/MultiBox/MultiBox.cpp
@@ -172,7 +172,6 @@ FMultiBox::FMultiBox( const EMultiBoxType::Type InType, FMultiBoxCustomization I
 	, Type( InType )
 	, bShouldCloseWindowAfterMenuSelection( bInShouldCloseWindowAfterMenuSelection )
 {
-
 }
 
 FMultiBox::~FMultiBox()
@@ -249,13 +248,16 @@ void FMultiBox::InsertCustomMultiBlock( TSharedRef<const FMultiBlock> InBlock, i
  *
  * @return  MultiBox widget object
  */
-TSharedRef< SMultiBoxWidget > FMultiBox::MakeWidget()
+TSharedRef< SMultiBoxWidget > FMultiBox::MakeWidget( bool bSearchable )
 {	
 	ApplyCustomizedBlocks();
 
 	TSharedRef< SMultiBoxWidget > NewMultiBoxWidget =
 		SNew( SMultiBoxWidget );
 
+	// Set whether this box should be searched
+	NewMultiBoxWidget->SetSearchable( bSearchable );
+
 	// Assign ourselves to the MultiBox widget
 	NewMultiBoxWidget->SetMultiBox( AsShared() );
 
@@ -563,6 +565,14 @@ void SMultiBoxWidget::AddBlockWidget( const FMultiBlock& Block, TSharedPtr<SHori
 	}
 }
 
+void SMultiBoxWidget::SetSearchable(bool InSearchable)
+{
+	bSearchable = InSearchable;
+}
+bool SMultiBoxWidget::GetSearchable() const
+{
+	return bSearchable;
+}
 
 /**
  * Builds this MultiBox widget up from the MultiBox associated with it
@@ -1054,11 +1064,11 @@ FReply SMultiBoxWidget::OnKeyDown( const FGeometry& MyGeometry, const FKeyEvent&
 {
 	SCompoundWidget::OnKeyDown( MyGeometry, KeyEvent );
 
-	if ( KeyEvent.GetKey() == EKeys::BackSpace )
+	if( bSearchable && KeyEvent.GetKey() == EKeys::BackSpace )
 	{
 		ReduceSearch();
 	}
-	else if( KeyEvent.GetKey() == EKeys::Delete )
+	else if( bSearchable && KeyEvent.GetKey() == EKeys::Delete )
 	{
 		ResetSearch();
 	}
@@ -1075,10 +1085,13 @@ FReply SMultiBoxWidget::OnKeyChar( const FGeometry& MyGeometry, const FCharacter
 {
 	FReply Reply = FReply::Unhandled();
 	
-	// Check for special characters
-	const TCHAR Character = InCharacterEvent.GetCharacter();
-	TypeChar(Character);
-	Reply = FReply::Handled();
+	if (bSearchable)
+	{
+		// Check for special characters
+		const TCHAR Character = InCharacterEvent.GetCharacter();
+		TypeChar(Character);
+		Reply = FReply::Handled();
+	}
 
 	return Reply;
 }
diff --git a/Engine/Source/Runtime/Slate/Private/Framework/MultiBox/MultiBoxBuilder.cpp b/Engine/Source/Runtime/Slate/Private/Framework/MultiBox/MultiBoxBuilder.cpp
--- a/Engine/Source/Runtime/Slate/Private/Framework/MultiBox/MultiBoxBuilder.cpp
+++ b/Engine/Source/Runtime/Slate/Private/Framework/MultiBox/MultiBoxBuilder.cpp
@@ -82,7 +82,7 @@ FMultiBoxCustomization FMultiBoxBuilder::GetCustomization() const
 
 TSharedRef< class SWidget > FMultiBoxBuilder::MakeWidget()
 {
-	return MultiBox->MakeWidget();
+	return MultiBox->MakeWidget( false );
 }
 
 TSharedRef< class FMultiBox > FMultiBoxBuilder::GetMultiBox()
@@ -163,6 +163,12 @@ void FBaseMenuBuilder::AddMenuEntry( const FUIAction& UIAction, const TSharedRef
 	ApplyHook(InExtensionHook, EExtensionHook::After);
 }
 
+TSharedRef< class SWidget > FMenuBuilder::MakeWidget()
+{
+	// Make menu builders searchable
+	return MultiBox->MakeWidget( true );
+}
+
 void FMenuBuilder::BeginSection( FName InExtensionHook, const TAttribute< FText >& InHeadingText )
 {
 	check(CurrentSectionExtensionHook == NAME_None && !bSectionNeedsToBeApplied);
@@ -333,7 +339,7 @@ void FMenuBarBuilder::AddPullDownMenu( const FText& InMenuLabel, const FText& In
 	NewMenuEntryBlock->SetTutorialHighlightName(GenerateTutorialIdentfierName(TutorialHighlightName, InTutorialHighlightName, nullptr, MultiBox->GetBlocks().Num()));
 
 	MultiBox->AddMultiBlock( NewMenuEntryBlock );
-	
+
 	ApplyHook(InExtensionHook, EExtensionHook::After);
 }
 
@@ -347,7 +353,7 @@ void FMenuBarBuilder::ApplyHook(FName InExtensionHook, EExtensionHook::Position
 	}
 }
 
-void FToolBarBuilder::AddToolBarButton(const TSharedPtr< const FUICommandInfo > InCommand, FName InExtensionHook, const TAttribute<FText>& InLabelOverride, const TAttribute<FText>& InToolTipOverride, const TAttribute<FSlateIcon>& InIconOverride, FName InTutorialHighlightName, bool bSearchable )
+void FToolBarBuilder::AddToolBarButton(const TSharedPtr< const FUICommandInfo > InCommand, FName InExtensionHook, const TAttribute<FText>& InLabelOverride, const TAttribute<FText>& InToolTipOverride, const TAttribute<FSlateIcon>& InIconOverride, FName InTutorialHighlightName )
 {
 	ApplySectionBeginning();
 
@@ -366,13 +372,10 @@ void FToolBarBuilder::AddToolBarButton(const TSharedPtr< const FUICommandInfo >
 
 	MultiBox->AddMultiBlock( NewToolBarButtonBlock );
 
-	// Set if this object should be searchable, defaulted to false
-	NewToolBarButtonBlock->SetSearchable(bSearchable);
-
 	ApplyHook(InExtensionHook, EExtensionHook::After);
 }
 
-void FToolBarBuilder::AddToolBarButton(const FUIAction& InAction, FName InExtensionHook, const TAttribute<FText>& InLabelOverride, const TAttribute<FText>& InToolTipOverride, const TAttribute<FSlateIcon>& InIconOverride, const EUserInterfaceActionType::Type UserInterfaceActionType, FName InTutorialHighlightName, bool bSearchable )
+void FToolBarBuilder::AddToolBarButton(const FUIAction& InAction, FName InExtensionHook, const TAttribute<FText>& InLabelOverride, const TAttribute<FText>& InToolTipOverride, const TAttribute<FSlateIcon>& InIconOverride, const EUserInterfaceActionType::Type UserInterfaceActionType, FName InTutorialHighlightName )
 {
 	ApplySectionBeginning();
 
@@ -389,9 +392,6 @@ void FToolBarBuilder::AddToolBarButton(const FUIAction& InAction, FName InExtens
 	NewToolBarButtonBlock->SetForceSmallIcons(bForceSmallIcons);
 	NewToolBarButtonBlock->SetTutorialHighlightName(GenerateTutorialIdentfierName(TutorialHighlightName, InTutorialHighlightName, nullptr, MultiBox->GetBlocks().Num()));
 
-	// Set if this object should be searchable, defaulted to false
-	MultiBox->AddMultiBlock( NewToolBarButtonBlock );
-
 	ApplyHook(InExtensionHook, EExtensionHook::After);
 }
 
diff --git a/Engine/Source/Runtime/Slate/Public/Framework/MultiBox/MultiBox.h b/Engine/Source/Runtime/Slate/Public/Framework/MultiBox/MultiBox.h
--- a/Engine/Source/Runtime/Slate/Public/Framework/MultiBox/MultiBox.h
+++ b/Engine/Source/Runtime/Slate/Public/Framework/MultiBox/MultiBox.h
@@ -244,7 +244,7 @@ public:
 	 *
 	 * @return  MultiBox widget object
 	 */
-	TSharedRef< class SMultiBoxWidget > MakeWidget();
+	TSharedRef< class SMultiBoxWidget > MakeWidget( bool bSearchable );
 
 
 	/**
@@ -484,6 +484,18 @@ public:
 		return MultiBox.ToSharedRef();
 	}
 
+	/**
+	* Sets the searchable state of this multibox
+	*
+	* @param	bSearchable		The searchable state to set
+	*/
+	void SetSearchable(bool bSearchable);
+	/**
+	* Gets the searchable state of this multibox
+	*
+	* @return	Whether this block is searchable
+	*/
+	bool GetSearchable() const;
 	
 	/**
 	 * Builds this MultiBox widget up from the MultiBox associated with it
@@ -656,4 +668,7 @@ private:
 
 	/* The text being searched for */
 	FText SearchText;
+
+	/** Whether this multibox can be searched */
+	bool bSearchable;
 };
diff --git a/Engine/Source/Runtime/Slate/Public/Framework/MultiBox/MultiBoxBuilder.h b/Engine/Source/Runtime/Slate/Public/Framework/MultiBox/MultiBoxBuilder.h
--- a/Engine/Source/Runtime/Slate/Public/Framework/MultiBox/MultiBoxBuilder.h
+++ b/Engine/Source/Runtime/Slate/Public/Framework/MultiBox/MultiBoxBuilder.h
@@ -45,7 +45,7 @@ public:
 	 * @param	InOnTextChanged		Called when the text is changed interactively
 	 * @param	bInReadOnly			Whether or not the text block is read only
 	 */
-	void AddEditableText( const FText& InLabel, const FText& InToolTip, const FSlateIcon& InIcon, const TAttribute< FText >& InTextAttribute, const FOnTextCommitted& InOnTextCommitted = FOnTextCommitted(), const FOnTextChanged& InOnTextChanged = FOnTextChanged(), bool bInReadOnly = false);
+	void AddEditableText( const FText& InLabel, const FText& InToolTip, const FSlateIcon& InIcon, const TAttribute< FText >& InTextAttribute, const FOnTextCommitted& InOnTextCommitted = FOnTextCommitted(), const FOnTextChanged& InOnTextChanged = FOnTextChanged(), bool bInReadOnly = false );
 
 
 	/**
@@ -53,7 +53,7 @@ public:
 	 *
 	 * @return  New widget object
 	 */
-	TSharedRef< class SWidget > MakeWidget();
+	virtual TSharedRef< class SWidget > MakeWidget();
 	
 
 	/** 
@@ -203,6 +203,12 @@ public:
 		AddSearchWidget();
 	}
 
+	/**
+	* Creates a widget for this MultiBox
+	*
+	* @return  New widget object
+	*/
+	virtual TSharedRef< class SWidget > MakeWidget() override;
 
 	/**
 	 * Adds a menu separator
@@ -364,9 +370,8 @@ public:
 	 * @param	InToolTipOverride		Optional tool tip override.	 If omitted, then the action's label will be used instead.
 	 * @param	InIconOverride			Optional name of the slate brush to use for the tool bar image.  If omitted, then the action's icon will be used instead.
 	 * @param	InTutorialHighlightName	Name to identify this widget and highlight during tutorials
-	 * @param	bSearchable	If true, widget will be searchable (default == false)
 	 */
-	void AddToolBarButton(const TSharedPtr< const FUICommandInfo > InCommand, FName InExtensionHook = NAME_None, const TAttribute<FText>& InLabelOverride = TAttribute<FText>(), const TAttribute<FText>& InToolTipOverride = TAttribute<FText>(), const TAttribute<FSlateIcon>& InIconOverride = TAttribute<FSlateIcon>(), FName InTutorialHighlightName = NAME_None, bool bSearchable = false );
+	void AddToolBarButton(const TSharedPtr< const FUICommandInfo > InCommand, FName InExtensionHook = NAME_None, const TAttribute<FText>& InLabelOverride = TAttribute<FText>(), const TAttribute<FText>& InToolTipOverride = TAttribute<FText>(), const TAttribute<FSlateIcon>& InIconOverride = TAttribute<FSlateIcon>(), FName InTutorialHighlightName = NAME_None );
 	
 	/**
 	 * Adds a tool bar button
@@ -377,9 +382,8 @@ public:
 	 * @param	InIcon		The icon to use		
 	 * @param	UserInterfaceActionType	Type of interface action
 	 * @param	InTutorialHighlightName	Name to identify this widget and highlight during tutorials
-	 * @param	bSearchable	If true, widget will be searchable (default == false)
 	 */
-	void AddToolBarButton(const FUIAction& InAction, FName InExtensionHook = NAME_None, const TAttribute<FText>& InLabelOverride = TAttribute<FText>(), const TAttribute<FText>& InToolTipOverride = TAttribute<FText>(), const TAttribute<FSlateIcon>& InIconOverride = TAttribute<FSlateIcon>(), const EUserInterfaceActionType::Type UserInterfaceActionType = EUserInterfaceActionType::Button, FName InTutorialHighlightName = NAME_None, bool bSearchable = false );
+	void AddToolBarButton(const FUIAction& InAction, FName InExtensionHook = NAME_None, const TAttribute<FText>& InLabelOverride = TAttribute<FText>(), const TAttribute<FText>& InToolTipOverride = TAttribute<FText>(), const TAttribute<FSlateIcon>& InIconOverride = TAttribute<FSlateIcon>(), const EUserInterfaceActionType::Type UserInterfaceActionType = EUserInterfaceActionType::Button, FName InTutorialHighlightName = NAME_None );
 
 	/**
 	 * Adds a combo button
