[INTEGRATE] merge ue4 <-ue4-launcherDev cl 2680106

diff --git a/Engine/Source/Runtime/Slate/Private/Widgets/Layout/SScrollBox.cpp b/Engine/Source/Runtime/Slate/Private/Widgets/Layout/SScrollBox.cpp
--- a/Engine/Source/Runtime/Slate/Private/Widgets/Layout/SScrollBox.cpp
+++ b/Engine/Source/Runtime/Slate/Private/Widgets/Layout/SScrollBox.cpp
@@ -376,6 +376,8 @@ void SScrollBox::ScrollDescendantIntoView(const TSharedPtr<SWidget>& WidgetToFin
 	bAnimateScrollingWidgetIntoView = InAnimateScroll;
 
 	// This will force the active timer system to wakeup for a frame to ensure we tick at least once.
+	bIsScrolling = true;
+	bIsScrollingActiveTimerRegistered = true;
 	RegisterActiveTimer(0.f, FWidgetActiveTimerDelegate::CreateSP(this, &SScrollBox::UpdateInertialScroll));
 }
 
@@ -483,7 +485,7 @@ void SScrollBox::SetScrollBarThickness(FVector2D InThickness)
 
 EActiveTimerReturnType SScrollBox::UpdateInertialScroll(double InCurrentTime, float InDeltaTime)
 {
-	bool bKeepTicking = false;
+	bool bKeepTicking = bIsScrolling;
 
 	//if (bIsScrolling)
 	//{
