Added a Never option to consuming the mouse wheel, because sometimes you don't want areas to scroll with the mouse.

diff --git a/Engine/Source/Runtime/Slate/Private/Widgets/Layout/SScrollBox.cpp b/Engine/Source/Runtime/Slate/Private/Widgets/Layout/SScrollBox.cpp
--- a/Engine/Source/Runtime/Slate/Private/Widgets/Layout/SScrollBox.cpp
+++ b/Engine/Source/Runtime/Slate/Private/Widgets/Layout/SScrollBox.cpp
@@ -719,7 +719,7 @@ void SScrollBox::OnMouseLeave( const FPointerEvent& MouseEvent )
 
 FReply SScrollBox::OnMouseWheel( const FGeometry& MyGeometry, const FPointerEvent& MouseEvent )
 {
-	if (ScrollBar->IsNeeded())
+	if (ScrollBar->IsNeeded() && ConsumeMouseWheel != EConsumeMouseWheel::Never )
 	{
 		// Make sure scroll velocity is cleared so it doesn't fight with the mouse wheel input
 		InertialScrollManager.ClearScrollVelocity();
diff --git a/Engine/Source/Runtime/SlateCore/Public/Styling/SlateTypes.h b/Engine/Source/Runtime/SlateCore/Public/Styling/SlateTypes.h
--- a/Engine/Source/Runtime/SlateCore/Public/Styling/SlateTypes.h
+++ b/Engine/Source/Runtime/SlateCore/Public/Styling/SlateTypes.h
@@ -18,6 +18,9 @@ enum class EConsumeMouseWheel
 
 	/** Always consume mouse wheel event even if we don't scroll at all. */
 	Always,
+
+	/** Never consume the mouse wheel */
+	Never,
 };
 
 
