[INTEGRATE] merge ue4 <-ue4-launcherDev cl 2684349

diff --git a/Engine/Source/Runtime/Slate/Private/Widgets/Views/SExpanderArrow.cpp b/Engine/Source/Runtime/Slate/Private/Widgets/Views/SExpanderArrow.cpp
--- a/Engine/Source/Runtime/Slate/Private/Widgets/Views/SExpanderArrow.cpp
+++ b/Engine/Source/Runtime/Slate/Private/Widgets/Views/SExpanderArrow.cpp
@@ -6,6 +6,7 @@
 void SExpanderArrow::Construct( const FArguments& InArgs, const TSharedPtr<class ITableRow>& TableRow  )
 {
 	OwnerRowPtr = TableRow;
+	StyleSet = InArgs._StyleSet;
 	IndentAmount = InArgs._IndentAmount;
 	BaseIndentLevel = InArgs._BaseIndentLevel;
 
@@ -94,6 +95,6 @@ const FSlateBrush* SExpanderArrow::GetExpanderImage() const
 		}
 	}
 
-	return FCoreStyle::Get().GetBrush( ResourceName );
+	return StyleSet->GetBrush(ResourceName);
 }
 
diff --git a/Engine/Source/Runtime/Slate/Public/Widgets/Views/SExpanderArrow.h b/Engine/Source/Runtime/Slate/Public/Widgets/Views/SExpanderArrow.h
--- a/Engine/Source/Runtime/Slate/Public/Widgets/Views/SExpanderArrow.h
+++ b/Engine/Source/Runtime/Slate/Public/Widgets/Views/SExpanderArrow.h
@@ -13,9 +13,11 @@ class SLATE_API SExpanderArrow : public SCompoundWidget
 public:
 
 	SLATE_BEGIN_ARGS( SExpanderArrow )
-		: _IndentAmount(10)
+		: _StyleSet(&FCoreStyle::Get())
+		, _IndentAmount(10)
 		, _BaseIndentLevel(0)
 	{ }
+		SLATE_ARGUMENT(const ISlateStyle*, StyleSet)
 		SLATE_ATTRIBUTE(float, IndentAmount)
 		SLATE_ATTRIBUTE(int32, BaseIndentLevel)
 	SLATE_END_ARGS()
@@ -41,6 +43,9 @@ protected:
 	/** A reference to the expander button */
 	TSharedPtr<SButton> ExpanderArrow;
 
+	/** The slate style to use */
+	const ISlateStyle* StyleSet;
+
 	/** The amount of space to indent at each level */
 	TAttribute<float> IndentAmount;
 
diff --git a/Engine/Source/Runtime/Slate/Public/Widgets/Views/STableRow.h b/Engine/Source/Runtime/Slate/Public/Widgets/Views/STableRow.h
--- a/Engine/Source/Runtime/Slate/Public/Widgets/Views/STableRow.h
+++ b/Engine/Source/Runtime/Slate/Public/Widgets/Views/STableRow.h
@@ -81,12 +81,14 @@ public:
 
 	SLATE_BEGIN_ARGS( STableRow< ItemType > )
 		: _Style( &FCoreStyle::Get().GetWidgetStyle<FTableRowStyle>("TableView.Row") )
+		, _ExpanderStyleSet( &FCoreStyle::Get() )
 		, _Padding( FMargin(0) )
-		, _ShowSelection(true)
+		, _ShowSelection( true )
 		, _Content()
 		{}
 	
 		SLATE_STYLE_ARGUMENT( FTableRowStyle, Style )
+		SLATE_ARGUMENT(const ISlateStyle*, ExpanderStyleSet)
 
 		// High Level DragAndDrop
 
@@ -175,6 +177,7 @@ public:
 				.VAlign(VAlign_Fill)
 				[
 					SNew(SExpanderArrow, SharedThis(this) )
+					.StyleSet(ExpanderStyleSet)
 				]
 
 				+ SHorizontalBox::Slot()
@@ -873,6 +876,9 @@ protected:
 		check(InArgs._Style);
 		Style = InArgs._Style;
 
+		check(InArgs._ExpanderStyleSet);
+		ExpanderStyleSet = InArgs._ExpanderStyleSet;
+
 		this->BorderImage = TAttribute<const FSlateBrush*>( this, &STableRow::GetBorder );
 
 		this->ForegroundColor = TAttribute<FSlateColor>( this, &STableRow::GetForegroundBasedOnSelection );
@@ -938,6 +944,9 @@ protected:
 	/** Style used to draw this table row */
 	const FTableRowStyle* Style;
 
+	/** The slate style to use with the expander */
+	const ISlateStyle* ExpanderStyleSet;
+
 	/** @see STableRow's OnCanAcceptDrop event */
 	FOnCanAcceptDrop OnCanAcceptDrop;
 
@@ -1000,6 +1009,7 @@ protected:
 		STableRow<ItemType>::Construct(
 			FTableRowArgs()
 			.Style(InArgs._Style)
+			.ExpanderStyleSet(InArgs._ExpanderStyleSet)
 			.Padding(InArgs._Padding)
 			.ShowSelection(InArgs._ShowSelection)
 			.OnCanAcceptDrop(InArgs._OnCanAcceptDrop)
