Fixed a crash when deleting a FSlateImageRun that isn't set to use a dynamic brush

diff --git a/Engine/Source/Runtime/Slate/Private/Framework/Text/SlateImageRun.cpp b/Engine/Source/Runtime/Slate/Private/Framework/Text/SlateImageRun.cpp
--- a/Engine/Source/Runtime/Slate/Private/Framework/Text/SlateImageRun.cpp
+++ b/Engine/Source/Runtime/Slate/Private/Framework/Text/SlateImageRun.cpp
@@ -82,7 +82,10 @@ FSlateImageRun::FSlateImageRun( const FRunInfo& InRunInfo, const TSharedRef< con
 
 FSlateImageRun::~FSlateImageRun()
 {
-	DynamicBrush->ReleaseResource();
+	if (DynamicBrush.IsValid())
+	{
+		DynamicBrush->ReleaseResource();
+	}
 }
 
 const TArray< TSharedRef<SWidget> >& FSlateImageRun::GetChildren() 
