Merging CL 2680933 from UE4-Orion/... to UE4/...

Setting mouse capture only makes sense for UserIndex == 0.

diff --git a/Engine/Source/Runtime/Slate/Private/Framework/Application/SlateApplication.cpp b/Engine/Source/Runtime/Slate/Private/Framework/Application/SlateApplication.cpp
--- a/Engine/Source/Runtime/Slate/Private/Framework/Application/SlateApplication.cpp
+++ b/Engine/Source/Runtime/Slate/Private/Framework/Application/SlateApplication.cpp
@@ -438,39 +438,42 @@ void FSlateApplication::MouseCaptorHelper::SetMouseCaptor(uint32 UserIndex, uint
 	// it still may not have a valid captor widget, this is ok
 	InvalidateCaptureForPointer(UserIndex, PointerIndex);
 
-	if ( Widget.IsValid() )
+	if (UserIndex == 0)
 	{
-		TSharedRef< SWidget > WidgetRef = Widget.ToSharedRef();
-		FWidgetPath NewMouseCaptorPath = EventPath.GetPathDownTo( WidgetRef );
-
-		const auto IsPathToCaptorFound = []( const FWidgetPath& PathToTest, const TSharedRef<SWidget>& WidgetToFind )
+		if ( Widget.IsValid() )
 		{
-			return PathToTest.Widgets.Num() > 0 && PathToTest.Widgets.Last().Widget == WidgetToFind;
-		};
+			TSharedRef< SWidget > WidgetRef = Widget.ToSharedRef();
+			FWidgetPath NewMouseCaptorPath = EventPath.GetPathDownTo( WidgetRef );
 
-		FWeakWidgetPath MouseCaptorWeakPath;
-		if ( IsPathToCaptorFound( NewMouseCaptorPath, WidgetRef ) )
-		{
-			MouseCaptorWeakPath = NewMouseCaptorPath;
-		}
-		else if (EventPath.Widgets.Num() > 0)
-		{
-			// If the target widget wasn't found on the event path then start the search from the root
-			NewMouseCaptorPath = EventPath.GetPathDownTo( EventPath.Widgets[0].Widget );
-			NewMouseCaptorPath.ExtendPathTo( FWidgetMatcher( WidgetRef ) );
+			const auto IsPathToCaptorFound = []( const FWidgetPath& PathToTest, const TSharedRef<SWidget>& WidgetToFind )
+			{
+				return PathToTest.Widgets.Num() > 0 && PathToTest.Widgets.Last().Widget == WidgetToFind;
+			};
+
+			FWeakWidgetPath MouseCaptorWeakPath;
+			if ( IsPathToCaptorFound( NewMouseCaptorPath, WidgetRef ) )
+			{
+				MouseCaptorWeakPath = NewMouseCaptorPath;
+			}
+			else if (EventPath.Widgets.Num() > 0)
+			{
+				// If the target widget wasn't found on the event path then start the search from the root
+				NewMouseCaptorPath = EventPath.GetPathDownTo( EventPath.Widgets[0].Widget );
+				NewMouseCaptorPath.ExtendPathTo( FWidgetMatcher( WidgetRef ) );
 			
-			MouseCaptorWeakPath = IsPathToCaptorFound( NewMouseCaptorPath, WidgetRef )
-				? NewMouseCaptorPath
-				: FWeakWidgetPath();
-		}
-		else
-		{
-			ensureMsgf(EventPath.Widgets.Num() > 0, TEXT("An unknown widget is attempting to set capture to %s"), *Widget->ToString() );
-		}
+				MouseCaptorWeakPath = IsPathToCaptorFound( NewMouseCaptorPath, WidgetRef )
+					? NewMouseCaptorPath
+					: FWeakWidgetPath();
+			}
+			else
+			{
+				ensureMsgf(EventPath.Widgets.Num() > 0, TEXT("An unknown widget is attempting to set capture to %s"), *Widget->ToString() );
+			}
 
-		if (MouseCaptorWeakPath.IsValid())
-		{
-			PointerIndexToMouseCaptorWeakPathMap.Add(FUserAndPointer(UserIndex,PointerIndex), MouseCaptorWeakPath);
+			if (MouseCaptorWeakPath.IsValid())
+			{
+				PointerIndexToMouseCaptorWeakPathMap.Add(FUserAndPointer(UserIndex,PointerIndex), MouseCaptorWeakPath);
+			}
 		}
 	}
 }
