Cleaning out programmer-only comments that were harvested by Doxygen.

diff --git a/Engine/Source/Runtime/Renderer/Private/LightRendering.h b/Engine/Source/Runtime/Renderer/Private/LightRendering.h
--- a/Engine/Source/Runtime/Renderer/Private/LightRendering.h
+++ b/Engine/Source/Runtime/Renderer/Private/LightRendering.h
@@ -199,14 +199,14 @@ public:
 		StencilingGeometryParameters.Set(RHICmdList, this, StencilingSpherePosAndScale);
 	}
 
-	// Begin FShader Interface
+	//~ Begin FShader Interface
 	virtual bool Serialize(FArchive& Ar) override
 	{
 		bool bShaderHasOutdatedParameters = FGlobalShader::Serialize(Ar);
 		Ar << StencilingGeometryParameters;
 		return bShaderHasOutdatedParameters;
 	}
-	// End FShader Interface
+	//~ End FShader Interface
 private:
 	FStencilingGeometryShaderParameters StencilingGeometryParameters;
 };
diff --git a/Engine/Source/Runtime/Renderer/Private/SceneOcclusion.h b/Engine/Source/Runtime/Renderer/Private/SceneOcclusion.h
--- a/Engine/Source/Runtime/Renderer/Private/SceneOcclusion.h
+++ b/Engine/Source/Runtime/Renderer/Private/SceneOcclusion.h
@@ -45,14 +45,14 @@ public:
 		StencilingGeometryParameters.Set(RHICmdList, this, FVector4(0,0,0,1));
 	}
 
-	// Begin FShader interface
+	//~ Begin FShader Interface
 	virtual bool Serialize(FArchive& Ar) override
 	{
 		bool bShaderHasOutdatedParameters = FGlobalShader::Serialize(Ar);
 		Ar << StencilingGeometryParameters;
 		return bShaderHasOutdatedParameters;
 	}
-	//  End FShader interface 
+	//~ Begin  End FShader Interface 
 
 private:
 	FStencilingGeometryShaderParameters StencilingGeometryParameters;
diff --git a/Engine/Source/Runtime/Renderer/Private/ShadowRendering.h b/Engine/Source/Runtime/Renderer/Private/ShadowRendering.h
--- a/Engine/Source/Runtime/Renderer/Private/ShadowRendering.h
+++ b/Engine/Source/Runtime/Renderer/Private/ShadowRendering.h
@@ -345,7 +345,7 @@ public:
 		return *this; 
 	}
 
-	// FMeshDrawingPolicy interface.
+	//~ Begin FMeshDrawingPolicy Interface.
 	bool Matches(const FShadowDepthDrawingPolicy& Other) const
 	{
 		return FMeshDrawingPolicy::Matches(Other) 
@@ -973,14 +973,14 @@ public:
 
 	void SetParameters(FRHICommandList& RHICmdList, const FSceneView& View, const FProjectedShadowInfo* ShadowInfo);
 
-	// Begin FShader interface
+	//~ Begin FShader Interface
 	virtual bool Serialize(FArchive& Ar) override
 	{
 		bool bShaderHasOutdatedParameters = FGlobalShader::Serialize(Ar);
 		Ar << StencilingGeometryParameters;
 		return bShaderHasOutdatedParameters;
 	}
-	//  End FShader interface 
+	//~ Begin  End FShader Interface 
 
 private:
 	FStencilingGeometryShaderParameters StencilingGeometryParameters;
diff --git a/Engine/Source/Runtime/Slate/Public/Framework/Application/SWindowTitleBar.h b/Engine/Source/Runtime/Slate/Public/Framework/Application/SWindowTitleBar.h
--- a/Engine/Source/Runtime/Slate/Public/Framework/Application/SWindowTitleBar.h
+++ b/Engine/Source/Runtime/Slate/Public/Framework/Application/SWindowTitleBar.h
@@ -122,7 +122,7 @@ public:
 
 public:
 
-	// Begin IWindowTitleBar interface
+	//~ Begin IWindowTitleBar Interface
 
 	void Flash( ) override
 	{
@@ -130,7 +130,7 @@ public:
 		TitleFlashSequence.Play(this->AsShared());
 	}
 
-	// Begin IWindowTitleBar interface
+	//~ Begin IWindowTitleBar Interface
 	
 protected:
 
diff --git a/Engine/Source/Runtime/Slate/Public/Framework/Application/SlateApplication.h b/Engine/Source/Runtime/Slate/Public/Framework/Application/SlateApplication.h
--- a/Engine/Source/Runtime/Slate/Public/Framework/Application/SlateApplication.h
+++ b/Engine/Source/Runtime/Slate/Public/Framework/Application/SlateApplication.h
@@ -1035,7 +1035,7 @@ public:
 
 public:
 
-	// Begin FSlateApplicationBase interface
+	//~ Begin FSlateApplicationBase Interface
 
 	virtual TSharedRef<SWindow> AddWindow( TSharedRef<SWindow> InSlateWindow, const bool bShowImmediately = true ) override;
 
@@ -1085,7 +1085,7 @@ protected:
 
 public:
 
-	// FSlateApplicationBase interface
+	//~ Begin FSlateApplicationBase Interface
 
 	virtual bool HasAnyMouseCaptor() const override;
 	virtual FSlateRect GetPreferredWorkArea() const override;
@@ -1105,7 +1105,7 @@ public:
 
 public:
 
-	// FGenericApplicationMessageHandler interface
+	//~ Begin FGenericApplicationMessageHandler Interface
 
 	virtual bool ShouldProcessUserInputMessages( const TSharedPtr< FGenericWindow >& PlatformWindow ) const override;
 	virtual bool OnKeyChar( const TCHAR Character, const bool IsRepeat ) override;
diff --git a/Engine/Source/Runtime/Slate/Public/Framework/Text/BaseTextLayoutMarshaller.h b/Engine/Source/Runtime/Slate/Public/Framework/Text/BaseTextLayoutMarshaller.h
--- a/Engine/Source/Runtime/Slate/Public/Framework/Text/BaseTextLayoutMarshaller.h
+++ b/Engine/Source/Runtime/Slate/Public/Framework/Text/BaseTextLayoutMarshaller.h
@@ -12,7 +12,7 @@ class SLATE_API FBaseTextLayoutMarshaller : public ITextLayoutMarshaller
 {
 public:
 
-	// BEGIN ITextLayoutMarshaller interface
+	//~ Begin ITextLayoutMarshaller Interface
 	virtual bool RequiresLiveUpdate() const override
 	{
 		return false;
@@ -32,7 +32,7 @@ public:
 	{
 		return bIsDirty;
 	}
-	// END ITextLayoutMarshaller interface
+	//~ End ITextLayoutMarshaller Interface
 
 protected:
 
diff --git a/Engine/Source/Runtime/Slate/Public/Widgets/Input/SEditableText.h b/Engine/Source/Runtime/Slate/Public/Widgets/Input/SEditableText.h
--- a/Engine/Source/Runtime/Slate/Public/Widgets/Input/SEditableText.h
+++ b/Engine/Source/Runtime/Slate/Public/Widgets/Input/SEditableText.h
@@ -332,7 +332,7 @@ protected:
 	};
 	
 public:
-	// BEGIN ITextEditorWidget interface
+	//~ Begin ITextEditorWidget Interface
 	virtual void StartChangingText() override;
 	virtual void FinishChangingText() override;
 	virtual bool GetIsReadOnly() const override;
@@ -369,10 +369,10 @@ public:
 	virtual TSharedRef< SWidget > GetWidget() override;
 	virtual void SummonContextMenu(const FVector2D& InLocation, TSharedPtr<SWindow> ParentWindow, const FWidgetPath& EventPath) override;
 	virtual void LoadText() override;
-	// END ITextEditorWidget interface
+	//~ End ITextEditorWidget Interface
 
 public:
-	// BEGIN IVirtualKeyboardEntry interface
+	//~ Begin IVirtualKeyboardEntry Interface
 	virtual void SetTextFromVirtualKeyboard(const FText& InNewText) override;
 
 	virtual const FText& GetText() const override
@@ -394,10 +394,10 @@ public:
 	{
 		return false;
 	}
-	// END IVirtualKeyboardEntry interface
+	//~ End IVirtualKeyboardEntry Interface
 
 protected:
-	// BEGIN SWidget interface
+	//~ Begin SWidget Interface
 	virtual void Tick( const FGeometry& AllottedGeometry, const double InCurrentTime, const float InDeltaTime ) override;
 	virtual FVector2D ComputeDesiredSize(float) const override;
 	virtual int32 OnPaint( const FPaintArgs& Args, const FGeometry& AllottedGeometry, const FSlateRect& MyClippingRect, FSlateWindowElementList& OutDrawElements, int32 LayerId, const FWidgetStyle& InWidgetStyle, bool bParentEnabled ) const override;
@@ -415,7 +415,7 @@ protected:
 	virtual FReply OnMouseButtonDoubleClick( const FGeometry& InMyGeometry, const FPointerEvent& InMouseEvent ) override;
 	virtual FCursorReply OnCursorQuery( const FGeometry& MyGeometry, const FPointerEvent& CursorEvent ) const override;
 	virtual const FSlateBrush* GetFocusBrush() const;
-	// END SWidget interface
+	//~ End SWidget Interface
 
 protected:
 	/**
diff --git a/Engine/Source/Runtime/Slate/Public/Widgets/Input/SVirtualKeyboardEntry.h b/Engine/Source/Runtime/Slate/Public/Widgets/Input/SVirtualKeyboardEntry.h
--- a/Engine/Source/Runtime/Slate/Public/Widgets/Input/SVirtualKeyboardEntry.h
+++ b/Engine/Source/Runtime/Slate/Public/Widgets/Input/SVirtualKeyboardEntry.h
@@ -87,7 +87,7 @@ public:
 	bool GetIsReadOnly() const;
 
 public:
-	// BEGIN IVirtualKeyboardEntry Interface
+	//~ Begin IVirtualKeyboardEntry Interface
 	virtual void SetTextFromVirtualKeyboard(const FText& InNewText) override;
 
 	virtual const FText& GetText() const override
@@ -109,18 +109,18 @@ public:
 	{
 		return false;
 	}
-	// END IVirtualKeyboardEntry Interface
+	//~ End IVirtualKeyboardEntry Interface
 
 protected:
 
-	// SWidget interface
+	//~ Begin SWidget Interface
 	virtual void Tick( const FGeometry& AllottedGeometry, const double InCurrentTime, const float InDeltaTime ) override;
 	virtual FVector2D ComputeDesiredSize(float) const override;
 	virtual int32 OnPaint( const FPaintArgs& Args, const FGeometry& AllottedGeometry, const FSlateRect& MyClippingRect, FSlateWindowElementList& OutDrawElements, int32 LayerId, const FWidgetStyle& InWidgetStyle, bool bParentEnabled ) const override;
 	virtual bool SupportsKeyboardFocus() const override;
 	virtual FReply OnFocusReceived( const FGeometry& MyGeometry, const FFocusEvent& InFocusEvent ) override;
 	virtual void OnFocusLost( const FFocusEvent& InFocusEvent ) override;
-	// End of SWidget interface
+	//~ End SWidget Interface
 
 private:
 
diff --git a/Engine/Source/Runtime/Slate/Public/Widgets/Layout/SUniformGridPanel.h b/Engine/Source/Runtime/Slate/Public/Widgets/Layout/SUniformGridPanel.h
--- a/Engine/Source/Runtime/Slate/Public/Widgets/Layout/SUniformGridPanel.h
+++ b/Engine/Source/Runtime/Slate/Public/Widgets/Layout/SUniformGridPanel.h
@@ -55,11 +55,11 @@ public:
 
 	void Construct( const FArguments& InArgs );
 
-	// BEGIN SPanel INTERFACE	
+	//~ Begin SPanel Interface	
 	virtual void OnArrangeChildren( const FGeometry& AllottedGeometry, FArrangedChildren& ArrangedChildren ) const override;
 	virtual FVector2D ComputeDesiredSize(float) const override;
 	virtual FChildren* GetChildren() override;
-	// END SPanel INTERFACE
+	//~ End SPanel Interface
 
 	/** See SlotPadding attribute */
 	void SetSlotPadding(TAttribute<FMargin> InSlotPadding);
diff --git a/Engine/Source/Runtime/Slate/Public/Widgets/Text/SMultiLineEditableText.h b/Engine/Source/Runtime/Slate/Public/Widgets/Text/SMultiLineEditableText.h
--- a/Engine/Source/Runtime/Slate/Public/Widgets/Text/SMultiLineEditableText.h
+++ b/Engine/Source/Runtime/Slate/Public/Widgets/Text/SMultiLineEditableText.h
@@ -428,7 +428,7 @@ private:
 
 private:
 
-	// BEGIN ITextEditorWidget interface
+	//~ Begin ITextEditorWidget Interface
 	virtual void StartChangingText() override;
 	virtual void FinishChangingText() override;
 	virtual void BackspaceChar() override;
@@ -463,10 +463,10 @@ private:
 	virtual TSharedRef< SWidget > GetWidget() override;
 	virtual void SummonContextMenu(const FVector2D& InLocation, TSharedPtr<SWindow> ParentWindow, const FWidgetPath& EventPath) override;
 	virtual void LoadText() override;
-	// END ITextEditorWidget interface
+	//~ End ITextEditorWidget Interface
 
 public:
-	// BEGIN IVirtualKeyboardEntry interface
+	//~ Begin IVirtualKeyboardEntry Interface
 	virtual void SetTextFromVirtualKeyboard(const FText& InNewText) override;
 
 	virtual const FText& GetText() const override
@@ -488,10 +488,10 @@ public:
 	{
 		return true;
 	}
-	// END IVirtualKeyboardEntry interface
+	//~ End IVirtualKeyboardEntry Interface
 
 protected:
-	// BEGIN SWidget interface
+	//~ Begin SWidget Interface
 	virtual void Tick( const FGeometry& AllottedGeometry, const double InCurrentTime, const float InDeltaTime ) override;
 	virtual int32 OnPaint( const FPaintArgs& Args, const FGeometry& AllottedGeometry, const FSlateRect& MyClippingRect, FSlateWindowElementList& OutDrawElements, int32 LayerId, const FWidgetStyle& InWidgetStyle, bool bParentEnabled ) const override;
 	virtual void CacheDesiredSize(float) override;
@@ -510,7 +510,7 @@ protected:
 	virtual FReply OnMouseWheel( const FGeometry& MyGeometry, const FPointerEvent& MouseEvent ) override;
 	virtual FReply OnMouseButtonDoubleClick(const FGeometry& MyGeometry, const FPointerEvent& MouseEvent) override;
 	virtual FCursorReply OnCursorQuery( const FGeometry& MyGeometry, const FPointerEvent& CursorEvent ) const override;
-	// END SWidget interface
+	//~ End SWidget Interface
 
 private:
 	
diff --git a/Engine/Source/Runtime/Sockets/Private/BSDIPv6Sockets/SocketSubsystemBSDIPv6.h b/Engine/Source/Runtime/Sockets/Private/BSDIPv6Sockets/SocketSubsystemBSDIPv6.h
--- a/Engine/Source/Runtime/Sockets/Private/BSDIPv6Sockets/SocketSubsystemBSDIPv6.h
+++ b/Engine/Source/Runtime/Sockets/Private/BSDIPv6Sockets/SocketSubsystemBSDIPv6.h
@@ -19,7 +19,7 @@ class FSocketSubsystemBSDIPv6
 {
 public:
 
-	// Begin ISocketSubsystem interface
+	//~ Begin ISocketSubsystem Interface
 
 	virtual TSharedRef<FInternetAddr> CreateInternetAddr( uint32 Address = 0, uint32 Port = 0 ) override;
 
@@ -56,7 +56,7 @@ public:
 
 	virtual ESocketErrors TranslateErrorCode( int32 Code ) override;
 
-	// End ISocketSubsystem interface
+	//~ End ISocketSubsystem Interface
 
 
 protected:
diff --git a/Engine/Source/Runtime/UMG/Public/Blueprint/UserWidget.h b/Engine/Source/Runtime/UMG/Public/Blueprint/UserWidget.h
--- a/Engine/Source/Runtime/UMG/Public/Blueprint/UserWidget.h
+++ b/Engine/Source/Runtime/UMG/Public/Blueprint/UserWidget.h
@@ -152,18 +152,18 @@ public:
 	virtual void PostDuplicate(bool bDuplicateForPIE) override;
 	virtual void BeginDestroy() override;
 	virtual void PostLoad() override;
-	// End of UObject interface
+	//~ End UObject Interface
 
 	virtual bool Initialize();
 	virtual void CustomNativeInitilize() {}
 
 	//UVisual interface
 	virtual void ReleaseSlateResources(bool bReleaseChildren) override;
-	// End of UVisual interface
+	//~ End UVisual Interface
 
-	// UWidget interface
+	//~ Begin UWidget Interface
 	virtual void SynchronizeProperties() override;
-	// End of UWidget interface
+	//~ End UWidget Interface
 
 	// UNamedSlotInterface Begin
 	virtual void GetSlotNames(TArray<FName>& SlotNames) const override;
@@ -766,19 +766,19 @@ public:
 	/** @returns The uobject widget corresponding to a given name */
 	UWidget* GetWidgetFromName(const FName& Name) const;
 
-	// Begin UObject interface
+	//~ Begin UObject Interface
 	virtual void PreSave() override;
-	// End UObject interface
+	//~ End UObject Interface
 
 	/** Are we currently playing any animations? */
 	UFUNCTION(BlueprintCallable, Category="User Interface|Animation")
 	bool IsPlayingAnimation() const { return ActiveSequencePlayers.Num() > 0; }
 
 #if WITH_EDITOR
-	// UWidget interface
+	//~ Begin UWidget Interface
 	virtual const FSlateBrush* GetEditorIcon() override;
 	virtual const FText GetPaletteCategory() override;
-	// End UWidget interface
+	//~ End UWidget Interface
 
 	void SetDesignerFlags(EWidgetDesignFlags::Type NewFlags);
 #endif
diff --git a/Engine/Source/Runtime/UMG/Public/Components/Border.h b/Engine/Source/Runtime/UMG/Public/Components/Border.h
--- a/Engine/Source/Runtime/UMG/Public/Components/Border.h
+++ b/Engine/Source/Runtime/UMG/Public/Components/Border.h
@@ -115,22 +115,22 @@ public:
 
 public:
 
-	// UWidget interface
+	//~ Begin UWidget Interface
 	virtual void SynchronizeProperties() override;
-	// End of UWidget interface
+	//~ End UWidget Interface
 
-	// UVisual interface
+	//~ Begin UVisual Interface
 	virtual void ReleaseSlateResources(bool bReleaseChildren) override;
-	// End of UVisual interface
+	//~ End UVisual Interface
 
-	// Begin UObject interface
+	//~ Begin UObject Interface
 	virtual void PostLoad() override;
-	// End of UObject interface
+	//~ End UObject Interface
 
 #if WITH_EDITOR
-	// UObject interface
+	//~ Begin UObject Interface
 	virtual void PostEditChangeProperty(struct FPropertyChangedEvent& PropertyChangedEvent) override;
-	// End of UObject interface
+	//~ End UObject Interface
 
 	virtual const FSlateBrush* GetEditorIcon() override;
 	virtual const FText GetPaletteCategory() override;
@@ -148,9 +148,9 @@ protected:
 	TSharedPtr<SBorder> MyBorder;
 
 protected:
-	// UWidget interface
+	//~ Begin UWidget Interface
 	virtual TSharedRef<SWidget> RebuildWidget() override;
-	// End of UWidget interface
+	//~ End UWidget Interface
 
 	FReply HandleMouseButtonDown(const FGeometry& Geometry, const FPointerEvent& MouseEvent);
 	FReply HandleMouseButtonUp(const FGeometry& Geometry, const FPointerEvent& MouseEvent);
diff --git a/Engine/Source/Runtime/UMG/Public/Components/Button.h b/Engine/Source/Runtime/UMG/Public/Components/Button.h
--- a/Engine/Source/Runtime/UMG/Public/Components/Button.h
+++ b/Engine/Source/Runtime/UMG/Public/Components/Button.h
@@ -90,17 +90,17 @@ public:
 
 public:
 
-	// UWidget interface
+	//~ Begin UWidget Interface
 	virtual void SynchronizeProperties() override;
-	// End of UWidget interface
+	//~ End UWidget Interface
 
-	// UVisual interface
+	//~ Begin UVisual Interface
 	virtual void ReleaseSlateResources(bool bReleaseChildren) override;
-	// End of UVisual interface
+	//~ End UVisual Interface
 
-	// Begin UObject interface
+	//~ Begin UObject Interface
 	virtual void PostLoad() override;
-	// End of UObject interface
+	//~ End UObject Interface
 
 #if WITH_EDITOR
 	virtual const FSlateBrush* GetEditorIcon() override;
@@ -122,9 +122,9 @@ protected:
 	void SlateHandleReleased();
 
 protected:
-	// UWidget interface
+	//~ Begin UWidget Interface
 	virtual TSharedRef<SWidget> RebuildWidget() override;
-	// End of UWidget interface
+	//~ End UWidget Interface
 
 protected:
 	/** Cached pointer to the underlying slate button owned by this UWidget */
diff --git a/Engine/Source/Runtime/UMG/Public/Components/CheckBox.h b/Engine/Source/Runtime/UMG/Public/Components/CheckBox.h
--- a/Engine/Source/Runtime/UMG/Public/Components/CheckBox.h
+++ b/Engine/Source/Runtime/UMG/Public/Components/CheckBox.h
@@ -117,17 +117,17 @@ public:
 
 public:
 	
-	// UWidget interface
+	//~ Begin UWidget Interface
 	virtual void SynchronizeProperties() override;
-	// End of UWidget interface
+	//~ End UWidget Interface
 
-	// UVisual interface
+	//~ Begin UVisual Interface
 	virtual void ReleaseSlateResources(bool bReleaseChildren) override;
-	// End of UVisual interface
+	//~ End UVisual Interface
 
-	// Begin UObject interface
+	//~ Begin UObject Interface
 	virtual void PostLoad() override;
-	// End of UObject interface
+	//~ End UObject Interface
 
 #if WITH_EDITOR
 	virtual const FSlateBrush* GetEditorIcon() override;
@@ -142,9 +142,9 @@ protected:
 	// End UPanelWidget
 
 protected:
-	// UWidget interface
+	//~ Begin UWidget Interface
 	virtual TSharedRef<SWidget> RebuildWidget() override;
-	// End of UWidget interface
+	//~ End UWidget Interface
 
 	void SlateOnCheckStateChangedCallback(ECheckBoxState NewState);
 	
diff --git a/Engine/Source/Runtime/UMG/Public/Components/ComboBoxString.h b/Engine/Source/Runtime/UMG/Public/Components/ComboBoxString.h
--- a/Engine/Source/Runtime/UMG/Public/Components/ComboBoxString.h
+++ b/Engine/Source/Runtime/UMG/Public/Components/ComboBoxString.h
@@ -100,9 +100,9 @@ public:
 
 	virtual void ReleaseSlateResources(bool bReleaseChildren) override;
 
-	// Begin UObject interface
+	//~ Begin UObject Interface
 	virtual void PostLoad() override;
-	// End of UObject interface
+	//~ End UObject Interface
 
 #if WITH_EDITOR
 	virtual const FSlateBrush* GetEditorIcon() override;
@@ -119,9 +119,9 @@ protected:
 	/** Called by slate when the underlying comobobox is opening */
 	virtual void HandleOpening();
 
-	// UWidget interface
+	//~ Begin UWidget Interface
 	virtual TSharedRef<SWidget> RebuildWidget() override;
-	// End of UWidget interface
+	//~ End UWidget Interface
 
 protected:
 	/** The true objects bound to the Slate combobox. */
diff --git a/Engine/Source/Runtime/UMG/Public/Components/EditableText.h b/Engine/Source/Runtime/UMG/Public/Components/EditableText.h
--- a/Engine/Source/Runtime/UMG/Public/Components/EditableText.h
+++ b/Engine/Source/Runtime/UMG/Public/Components/EditableText.h
@@ -138,26 +138,26 @@ public:
 
 public:
 	
-	// UWidget interface
+	//~ Begin UWidget Interface
 	virtual void SynchronizeProperties() override;
-	// End of UWidget interface
+	//~ End UWidget Interface
 
-	// UVisual interface
+	//~ Begin UVisual Interface
 	virtual void ReleaseSlateResources(bool bReleaseChildren) override;
-	// End of UVisual interface
+	//~ End UVisual Interface
 
-	// Begin UObject interface
+	//~ Begin UObject Interface
 	virtual void PostLoad() override;
-	// End of UObject interface
+	//~ End UObject Interface
 
 #if WITH_EDITOR
 	virtual const FSlateBrush* GetEditorIcon() override;
 	virtual const FText GetPaletteCategory() override;
 #endif
-	// End of UWidget interface
+	//~ End UWidget Interface
 
 protected:
-	// UWidget interface
+	//~ Begin UWidget Interface
 	virtual TSharedRef<SWidget> RebuildWidget() override;
 	// End of UWidget
 
diff --git a/Engine/Source/Runtime/UMG/Public/Components/EditableTextBox.h b/Engine/Source/Runtime/UMG/Public/Components/EditableTextBox.h
--- a/Engine/Source/Runtime/UMG/Public/Components/EditableTextBox.h
+++ b/Engine/Source/Runtime/UMG/Public/Components/EditableTextBox.h
@@ -128,17 +128,17 @@ public:
 	UFUNCTION(BlueprintCallable, Category="Widget")
 	void ClearError();
 
-	// UWidget interface
+	//~ Begin UWidget Interface
 	virtual void SynchronizeProperties() override;
-	// End of UWidget interface
+	//~ End UWidget Interface
 
-	// UVisual interface
+	//~ Begin UVisual Interface
 	virtual void ReleaseSlateResources(bool bReleaseChildren) override;
-	// End of UVisual interface
+	//~ End UVisual Interface
 
-	// Begin UObject interface
+	//~ Begin UObject Interface
 	virtual void PostLoad() override;
-	// End of UObject interface
+	//~ End UObject Interface
 
 #if WITH_EDITOR
 	virtual const FSlateBrush* GetEditorIcon() override;
@@ -146,7 +146,7 @@ public:
 #endif
 
 protected:
-	// UWidget interface
+	//~ Begin UWidget Interface
 	virtual TSharedRef<SWidget> RebuildWidget() override;
 	// End of UWidget
 
diff --git a/Engine/Source/Runtime/UMG/Public/Components/Image.h b/Engine/Source/Runtime/UMG/Public/Components/Image.h
--- a/Engine/Source/Runtime/UMG/Public/Components/Image.h
+++ b/Engine/Source/Runtime/UMG/Public/Components/Image.h
@@ -73,29 +73,29 @@ public:
 	UFUNCTION(BlueprintCallable, Category="Appearance")
 	UMaterialInstanceDynamic* GetDynamicMaterial();
 
-	// UWidget interface
+	//~ Begin UWidget Interface
 	virtual void SynchronizeProperties() override;
-	// End of UWidget interface
+	//~ End UWidget Interface
 
-	// UVisual interface
+	//~ Begin UVisual Interface
 	virtual void ReleaseSlateResources(bool bReleaseChildren) override;
-	// End of UVisual interface
+	//~ End UVisual Interface
 
-	// Begin UObject interface
+	//~ Begin UObject Interface
 	virtual void PostLoad() override;
-	// End of UObject interface
+	//~ End UObject Interface
 
 #if WITH_EDITOR
-	// UWidget interface
+	//~ Begin UWidget Interface
 	virtual const FSlateBrush* GetEditorIcon() override;
 	virtual const FText GetPaletteCategory() override;
-	// End UWidget interface
+	//~ End UWidget Interface
 #endif
 
 protected:
-	// UWidget interface
+	//~ Begin UWidget Interface
 	virtual TSharedRef<SWidget> RebuildWidget() override;
-	// End of UWidget interface
+	//~ End UWidget Interface
 
 	/** Translates the bound brush data and assigns it to the cached brush used by this widget. */
 	const FSlateBrush* ConvertImage(TAttribute<FSlateBrush> InImageAsset) const;
diff --git a/Engine/Source/Runtime/UMG/Public/Components/MultiLineEditableText.h b/Engine/Source/Runtime/UMG/Public/Components/MultiLineEditableText.h
--- a/Engine/Source/Runtime/UMG/Public/Components/MultiLineEditableText.h
+++ b/Engine/Source/Runtime/UMG/Public/Components/MultiLineEditableText.h
@@ -62,17 +62,17 @@ public:
 	UFUNCTION(BlueprintCallable, Category="Widget", meta=(DisplayName="SetText (Multi-Line Editable Text)"))
 	void SetText(FText InText);
 	
-	// UWidget interface
+	//~ Begin UWidget Interface
 	virtual void SynchronizeProperties() override;
-	// End of UWidget interface
+	//~ End UWidget Interface
 
-	// UVisual interface
+	//~ Begin UVisual Interface
 	virtual void ReleaseSlateResources(bool bReleaseChildren) override;
-	// End of UVisual interface
+	//~ End UVisual Interface
 
-	// Begin UObject interface
+	//~ Begin UObject Interface
 	virtual void PostLoad() override;
-	// End of UObject interface
+	//~ End UObject Interface
 
 #if WITH_EDITOR
 	virtual const FSlateBrush* GetEditorIcon() override;
@@ -80,7 +80,7 @@ public:
 #endif
 
 protected:
-	// UWidget interface
+	//~ Begin UWidget Interface
 	virtual TSharedRef<SWidget> RebuildWidget() override;
 	// End of UWidget
 
diff --git a/Engine/Source/Runtime/UMG/Public/Components/ProgressBar.h b/Engine/Source/Runtime/UMG/Public/Components/ProgressBar.h
--- a/Engine/Source/Runtime/UMG/Public/Components/ProgressBar.h
+++ b/Engine/Source/Runtime/UMG/Public/Components/ProgressBar.h
@@ -79,31 +79,31 @@ public:
 
 public:
 	
-	// UWidget interface
+	//~ Begin UWidget Interface
 	virtual void SynchronizeProperties() override;
-	// End of UWidget interface
+	//~ End UWidget Interface
 
-	// UVisual interface
+	//~ Begin UVisual Interface
 	virtual void ReleaseSlateResources(bool bReleaseChildren) override;
-	// End of UVisual interface
+	//~ End UVisual Interface
 
-	// Begin UObject interface
+	//~ Begin UObject Interface
 	virtual void PostLoad() override;
-	// End of UObject interface
+	//~ End UObject Interface
 
 #if WITH_EDITOR
-	// UWidget interface
+	//~ Begin UWidget Interface
 	virtual const FSlateBrush* GetEditorIcon() override;
 	virtual const FText GetPaletteCategory() override;
 	virtual void OnCreationFromPalette() override;
-	// End UWidget interface
+	//~ End UWidget Interface
 #endif
 
 protected:
 	/** Native Slate Widget */
 	TSharedPtr<SProgressBar> MyProgressBar;
 
-	// UWidget interface
+	//~ Begin UWidget Interface
 	virtual TSharedRef<SWidget> RebuildWidget() override;
-	// End of UWidget interface
+	//~ End UWidget Interface
 };
diff --git a/Engine/Source/Runtime/UMG/Public/Components/SpinBox.h b/Engine/Source/Runtime/UMG/Public/Components/SpinBox.h
--- a/Engine/Source/Runtime/UMG/Public/Components/SpinBox.h
+++ b/Engine/Source/Runtime/UMG/Public/Components/SpinBox.h
@@ -145,26 +145,26 @@ public:
 
 public:
 
-	// UWidget interface
+	//~ Begin UWidget Interface
 	virtual void SynchronizeProperties() override;
-	// End of UWidget interface
+	//~ End UWidget Interface
 
-	// UVisual interface
+	//~ Begin UVisual Interface
 	virtual void ReleaseSlateResources(bool bReleaseChildren) override;
-	// End of UVisual interface
+	//~ End UVisual Interface
 
-	// Begin UObject interface
+	//~ Begin UObject Interface
 	virtual void PostLoad() override;
-	// End of UObject interface
+	//~ End UObject Interface
 
 #if WITH_EDITOR
 	virtual const FSlateBrush* GetEditorIcon() override;
 	virtual const FText GetPaletteCategory() override;
 #endif
-	// End of UWidget interface
+	//~ End UWidget Interface
 
 protected:
-	// UWidget interface
+	//~ Begin UWidget Interface
 	virtual TSharedRef<SWidget> RebuildWidget() override;
 	// End of UWidget
 
diff --git a/Engine/Source/Runtime/Windows/D3D11RHI/Public/D3D11Util.h b/Engine/Source/Runtime/Windows/D3D11RHI/Public/D3D11Util.h
--- a/Engine/Source/Runtime/Windows/D3D11RHI/Public/D3D11Util.h
+++ b/Engine/Source/Runtime/Windows/D3D11RHI/Public/D3D11Util.h
@@ -223,7 +223,7 @@ public:
 	/** Unlocks the buffer returning the underlying D3D11 buffer to use as a resource. */
 	ID3D11Buffer* Unlock();
 
-	// Begin FRenderResource interface.
+	//~ Begin FRenderResource Interface.
 	virtual void InitRHI() override;
 	virtual void ReleaseRHI() override;
 	// End FRenderResource interface.
diff --git a/Engine/Source/Runtime/Windows/D3D12RHI/Public/D3D12Util.h b/Engine/Source/Runtime/Windows/D3D12RHI/Public/D3D12Util.h
--- a/Engine/Source/Runtime/Windows/D3D12RHI/Public/D3D12Util.h
+++ b/Engine/Source/Runtime/Windows/D3D12RHI/Public/D3D12Util.h
@@ -245,7 +245,7 @@ public:
 	/** Unlocks the buffer returning the underlying D3D12 buffer to use as a resource. */
 	FD3D12ResourceLocation* Unlock();
 
-	// Begin FRenderResource interface.
+	//~ Begin FRenderResource Interface.
 	virtual void InitRHI() override;
 	virtual void ReleaseRHI() override;
 	// End FRenderResource interface.
