UMG - Fixed SScaleBox not reporting its relative layout scale correctly, and making it better at laying out text by doing another prepass after calculating the newest scale.  Does add additional overhead for using the scalebox but ensures the proper result.

diff --git a/Engine/Source/Runtime/Slate/Private/Widgets/Layout/SScaleBox.cpp b/Engine/Source/Runtime/Slate/Private/Widgets/Layout/SScaleBox.cpp
--- a/Engine/Source/Runtime/Slate/Private/Widgets/Layout/SScaleBox.cpp
+++ b/Engine/Source/Runtime/Slate/Private/Widgets/Layout/SScaleBox.cpp
@@ -104,6 +104,13 @@ void SScaleBox::OnArrangeChildren( const FGeometry& AllottedGeometry, FArrangedC
 			}
 		}
 
+		if ( CurrentStretch != EStretch::UserSpecified )
+		{
+			// We need to run another pre-pass now that we know the final scale.
+			// This will allow things that don't scale linearly (such as text) to update their size and layout correctly.
+			ChildSlot.GetWidget()->SlatePrepass(AllottedGeometry.GetAccumulatedLayoutTransform().GetScale() * FinalScale);
+		}
+
 		ArrangedChildren.AddWidget(ChildVisibility, AllottedGeometry.MakeChild(
 			ChildSlot.GetWidget(),
 			FinalOffset,
@@ -150,3 +157,19 @@ void SScaleBox::SetUserSpecifiedScale(float InUserSpecifiedScale)
 {
 	UserSpecifiedScale = InUserSpecifiedScale;
 }
+
+float SScaleBox::GetRelativeLayoutScale(const FSlotBase& Child) const
+{
+	const EStretch::Type CurrentStretch = Stretch.Get();
+
+	switch ( CurrentStretch )
+	{
+	case EStretch::UserSpecified:
+		return UserSpecifiedScale.Get(1.0f);
+	default:
+		// Because our scale is determined by our size, we always report a scale of 1.0 here, 
+		// as reporting our actual scale can cause a feedback loop whereby the calculated size changes each frame.
+		// We workaround this by forcibly pre-passing our child content a second time once we know its final scale in OnArrangeChildren.
+		return 1.0f;
+	}
+}
diff --git a/Engine/Source/Runtime/Slate/Public/Widgets/Layout/SScaleBox.h b/Engine/Source/Runtime/Slate/Public/Widgets/Layout/SScaleBox.h
--- a/Engine/Source/Runtime/Slate/Public/Widgets/Layout/SScaleBox.h
+++ b/Engine/Source/Runtime/Slate/Public/Widgets/Layout/SScaleBox.h
@@ -122,6 +122,9 @@ public:
 	/** See UserSpecifiedScale argument */
 	void SetUserSpecifiedScale(float InUserSpecifiedScale);
 	
+protected:
+	virtual float GetRelativeLayoutScale(const FSlotBase& Child) const override;
+
 private:
 	/** The allowed direction of stretching of the content */
 	TAttribute<EStretchDirection::Type> StretchDirection;
