Keep track of the right stick on the analog cursor.

diff --git a/Engine/Source/Runtime/Slate/Private/Framework/Application/AnalogCursor.cpp b/Engine/Source/Runtime/Slate/Private/Framework/Application/AnalogCursor.cpp
--- a/Engine/Source/Runtime/Slate/Private/Framework/Application/AnalogCursor.cpp
+++ b/Engine/Source/Runtime/Slate/Private/Framework/Application/AnalogCursor.cpp
@@ -13,8 +13,9 @@ FAnalogCursor::FAnalogCursor()
 , StickySlowdown(0.5f)
 , DeadZone(0.1f)
 , Mode(AnalogCursorMode::Accelerated)
-, AnalogValues(FVector2D::ZeroVector)
 {
+	AnalogValues[ static_cast< uint8 >( EAnalogStick::Left ) ] = FVector2D::ZeroVector;
+	AnalogValues[ static_cast< uint8 >( EAnalogStick::Right ) ] = FVector2D::ZeroVector;
 }
 
 void FAnalogCursor::Tick(const float DeltaTime, FSlateApplication& SlateApp, TSharedRef<ICursor> Cursor)
@@ -22,7 +23,7 @@ void FAnalogCursor::Tick(const float DeltaTime, FSlateApplication& SlateApp, TSh
 	const FVector2D OldPosition = Cursor->GetPosition();
 
 	float SpeedMult = 1.0f; // Used to do a speed multiplication before adding the delta to the position to make widgets sticky
-	FVector2D AdjAnalogVals = AnalogValues; // A copy of the analog values so I can modify them based being over a widget, not currently doing this
+	FVector2D AdjAnalogVals = GetAnalogValue(EAnalogStick::Left); // A copy of the analog values so I can modify them based being over a widget, not currently doing this
 
 	// Adjust analog values according to dead zone
 	const float AnalogValsSize = AdjAnalogVals.Size();
@@ -184,16 +185,29 @@ bool FAnalogCursor::HandleAnalogInputEvent(FSlateApplication& SlateApp, const FA
 
 	if (Key == EKeys::Gamepad_LeftX)
 	{
-		AnalogValues.X = AnalogValue;
+		FVector2D& Value = GetAnalogValue( EAnalogStick::Left );
+		Value.X = AnalogValue;
 	}
 	else if (Key == EKeys::Gamepad_LeftY)
 	{
-		AnalogValues.Y = -AnalogValue;
+		FVector2D& Value = GetAnalogValue( EAnalogStick::Left );
+		Value.Y = -AnalogValue;
 	}
-	else
+	else if ( Key == EKeys::Gamepad_RightX )
+	{
+		FVector2D& Value = GetAnalogValue( EAnalogStick::Right );
+		Value.X = AnalogValue;
+	}
+	else if ( Key == EKeys::Gamepad_RightY )
+	{
+		FVector2D& Value = GetAnalogValue( EAnalogStick::Right );
+		Value.Y = -AnalogValue;
+	}
+	else 
 	{
 		return false;
 	}
+
 	return true;
 }
 
diff --git a/Engine/Source/Runtime/Slate/Public/Framework/Application/AnalogCursor.h b/Engine/Source/Runtime/Slate/Public/Framework/Application/AnalogCursor.h
--- a/Engine/Source/Runtime/Slate/Public/Framework/Application/AnalogCursor.h
+++ b/Engine/Source/Runtime/Slate/Public/Framework/Application/AnalogCursor.h
@@ -13,6 +13,13 @@ namespace AnalogCursorMode
 	};
 }
 
+enum class EAnalogStick : uint8
+{
+	Left,
+	Right,
+	Max,
+};
+
 /**
  * A class that simulates a cursor driven by an analog stick.
  */
@@ -40,9 +47,9 @@ public:
 protected:
 
 	/** Getter */
-	FORCEINLINE FVector2D GetAnalogValues() const
+	FORCEINLINE FVector2D GetAnalogValues( EAnalogStick Stick = EAnalogStick::Left ) const
 	{
-		return AnalogValues;
+		return AnalogValues[ static_cast< uint8 >( Stick ) ];
 	}
 
 	/** Handles updating the cursor position and processing a Mouse Move Event */
@@ -63,7 +70,12 @@ protected:
 
 private:
 
+	FORCEINLINE FVector2D& GetAnalogValue( EAnalogStick Stick )
+	{
+		return AnalogValues[ static_cast< uint8 >( Stick ) ];
+	}
+
 	/** Input from the gamepad */
-	FVector2D AnalogValues;
+	FVector2D AnalogValues[ static_cast<uint8>( EAnalogStick::Max ) ];
 };
 
