Merging 2695071 SetTextFromVirtualKeyboard now takes a bool so we can no if text was committed or just changed.

//depot/UE4-Orion/Engine/Source/Runtime/...

to //depot/UE4/Engine/Source/Runtime/...

#change PS4 platform now sends commit when dialog is closed via R2
#change IOS and Android platforms functionality remains unchanged
#change SMultilineEditableText now has ETextCommit::OnEnter for the broadcast delegate if the chat was committed
#change SMultilineEditableText will no longer open virtual keyboard on read only text boxes

#Platform
#Social

diff --git a/Engine/Source/Runtime/Launch/Private/Android/AndroidJNI.cpp b/Engine/Source/Runtime/Launch/Private/Android/AndroidJNI.cpp
--- a/Engine/Source/Runtime/Launch/Private/Android/AndroidJNI.cpp
+++ b/Engine/Source/Runtime/Launch/Private/Android/AndroidJNI.cpp
@@ -187,7 +187,7 @@ void EngineCrashHandler(const FGenericCrashContext& GenericContext)
 		// Walk the stack and dump it to the allocated memory.
 		FPlatformStackWalk::StackWalkAndDump(StackTrace, StackTraceSize, 0, Context.Context);
 		UE_LOG(LogEngine, Error, TEXT("\n%s\n"), ANSI_TO_TCHAR(StackTrace));
-		
+
 		if (GLog)
 		{
 			GLog->SetCurrentThreadAsMasterThread();
@@ -311,10 +311,7 @@ extern "C" void Java_com_epicgames_ue4_GameActivity_nativeVirtualKeyboardResult(
 			{
 				FGraphEventRef SetWidgetText = FFunctionGraphTask::CreateAndDispatchWhenReady([&]()
 				{
-					VirtualKeyboardWidget->SetTextFromVirtualKeyboard(FText::FromString(FString(UTF8_TO_TCHAR(javaChars))));
-				}, TStatId(), NULL, ENamedThreads::GameThread);
-				FTaskGraphInterface::Get().WaitUntilTaskCompletes(SetWidgetText);
-			}
+			VirtualKeyboardWidget->SetTextFromVirtualKeyboard(FText::FromString(FString(UTF8_TO_TCHAR(javaChars))),false);
 
 			// release string
 			jenv->ReleaseStringUTFChars(contents, javaChars);
@@ -530,7 +527,7 @@ JNIEXPORT jint JNI_OnLoad(JavaVM* InJavaVM, void* InReserved)
 	{
 		FPlatformMisc::SetCrashHandler(EngineCrashHandler);
 	}
-	
+
 	// Cache path to external storage
 	jclass EnvClass = Env->FindClass("android/os/Environment");
 	jmethodID getExternalStorageDir = Env->GetStaticMethodID(EnvClass, "getExternalStorageDirectory", "()Ljava/io/File;");
diff --git a/Engine/Source/Runtime/Slate/Private/Framework/Text/IOS/IOSPlatformTextField.cpp b/Engine/Source/Runtime/Slate/Private/Framework/Text/IOS/IOSPlatformTextField.cpp
--- a/Engine/Source/Runtime/Slate/Private/Framework/Text/IOS/IOSPlatformTextField.cpp
+++ b/Engine/Source/Runtime/Slate/Private/Framework/Text/IOS/IOSPlatformTextField.cpp
@@ -59,7 +59,7 @@ void FIOSPlatformTextField::ShowVirtualKeyboard(bool bShow, int32 UserIndex, TSh
 											FIOSAsyncTask* AsyncTask = [[FIOSAsyncTask alloc] init];
 											AsyncTask.GameThreadCallback = ^ bool(void)
 											{
-												TextWidget->SetTextFromVirtualKeyboard(TextEntry);
+												TextWidget->SetTextFromVirtualKeyboard(TextEntry, false);
 
 												// clear the TextWidget
 												TextWidget = nullptr;
@@ -193,12 +193,12 @@ void FIOSPlatformTextField::ShowVirtualKeyboard(bool bShow, int32 UserIndex, TSh
 	FIOSAsyncTask* AsyncTask = [[FIOSAsyncTask alloc] init];
     AsyncTask.GameThreadCallback = ^ bool(void)
     {
-		// index 1 is the OK button
-		if (buttonIndex == 1)
-		{
-			TextWidget->SetTextFromVirtualKeyboard(TextEntry);
-		}
-
+	// index 1 is the OK button
+	if(buttonIndex == 1)
+	{
+			TextWidget->SetTextFromVirtualKeyboard(TextEntry, false);
+	}
+    
         // clear the TextWidget
         TextWidget = nullptr;
         return true;
diff --git a/Engine/Source/Runtime/Slate/Private/Widgets/Input/SEditableText.cpp b/Engine/Source/Runtime/Slate/Private/Widgets/Input/SEditableText.cpp
--- a/Engine/Source/Runtime/Slate/Private/Widgets/Input/SEditableText.cpp
+++ b/Engine/Source/Runtime/Slate/Private/Widgets/Input/SEditableText.cpp
@@ -159,7 +159,7 @@ void SEditableText::SetText( const TAttribute< FText >& InNewText )
 	}
 }
 
-void SEditableText::SetTextFromVirtualKeyboard(const FText& InNewText)
+void SEditableText::SetTextFromVirtualKeyboard(const FText& InNewText, bool InCommit)
 {
 	// Only set the text if the text attribute doesn't have a getter binding (otherwise it would be blown away).
 	// If it is bound, we'll assume that OnTextCommitted will handle the update.
diff --git a/Engine/Source/Runtime/Slate/Private/Widgets/Input/SVirtualKeyboardEntry.cpp b/Engine/Source/Runtime/Slate/Private/Widgets/Input/SVirtualKeyboardEntry.cpp
--- a/Engine/Source/Runtime/Slate/Private/Widgets/Input/SVirtualKeyboardEntry.cpp
+++ b/Engine/Source/Runtime/Slate/Private/Widgets/Input/SVirtualKeyboardEntry.cpp
@@ -47,7 +47,7 @@ void SVirtualKeyboardEntry::SetText(const TAttribute< FText >& InNewText)
 	bNeedsUpdate = true;
 }
 
-void SVirtualKeyboardEntry::SetTextFromVirtualKeyboard( const FText& InNewText )
+void SVirtualKeyboardEntry::SetTextFromVirtualKeyboard(const FText& InNewText, bool InCommit)
 {
 	// Only set the text if the text attribute doesn't have a getter binding (otherwise it would be blown away).
 	// If it is bound, we'll assume that OnTextChanged will handle the update.
@@ -74,7 +74,7 @@ void SVirtualKeyboardEntry::RestoreOriginalText()
 {
 	if( HasTextChangedFromOriginal() )
 	{
-		SetTextFromVirtualKeyboard(OriginalText);
+		SetTextFromVirtualKeyboard(OriginalText, false);
 	}
 }
 
diff --git a/Engine/Source/Runtime/Slate/Private/Widgets/Text/SMultiLineEditableText.cpp b/Engine/Source/Runtime/Slate/Private/Widgets/Text/SMultiLineEditableText.cpp
--- a/Engine/Source/Runtime/Slate/Private/Widgets/Text/SMultiLineEditableText.cpp
+++ b/Engine/Source/Runtime/Slate/Private/Widgets/Text/SMultiLineEditableText.cpp
@@ -207,6 +207,7 @@ SMultiLineEditableText::SMultiLineEditableText()
 	, IsReadOnly(false)
 	, UICommandList(new FUICommandList())
 	, bTextChangedByVirtualKeyboard(false)
+	, bTextCommittedByVirtualKeyboard(false)
 	, AmountScrolledWhileRightMouseDown(0.0f)
 	, bIsSoftwareCursor(false)
 {
@@ -425,7 +426,7 @@ void SMultiLineEditableText::SetFont(const TAttribute< FSlateFontInfo >& InNewFo
 	// @todo: Chris.Wood - this doesn't update the font (needs the TextLayout updating?)
 }
 
-void SMultiLineEditableText::SetTextFromVirtualKeyboard(const FText& InNewText)
+void SMultiLineEditableText::SetTextFromVirtualKeyboard(const FText& InNewText, bool InCommit)
 {
 	// Only set the text if the text attribute doesn't have a getter binding (otherwise it would be blown away).
 	// If it is bound, we'll assume that OnTextCommitted will handle the update.
@@ -440,7 +441,14 @@ void SMultiLineEditableText::SetTextFromVirtualKeyboard(const FText& InNewText)
 		// This method is called from the main thread (i.e. not the game thread) of the device with the virtual keyboard
 		// This causes the app to crash on those devices, so we're using polling here to ensure delegates are
 		// fired on the game thread in Tick.
-		bTextChangedByVirtualKeyboard = true;
+		if (InCommit)
+		{
+			bTextCommittedByVirtualKeyboard = true;
+		}
+		else
+		{
+			bTextChangedByVirtualKeyboard = true;
+		}
 	}
 }
 
@@ -558,8 +566,11 @@ FReply SMultiLineEditableText::OnFocusReceived( const FGeometry& MyGeometry, con
 		FSlateApplication& SlateApplication = FSlateApplication::Get();
 		if (FPlatformMisc::GetRequiresVirtualKeyboard())
 		{
-			// @TODO: Create ITextInputMethodSystem derivations for mobile
-			SlateApplication.ShowVirtualKeyboard(true, InFocusEvent.GetUser(), SharedThis(this));
+			if (!GetIsReadOnly())
+			{
+				// @TODO: Create ITextInputMethodSystem derivations for mobile
+				SlateApplication.ShowVirtualKeyboard(true, InFocusEvent.GetUser(), SharedThis(this));
+			}
 		}
 		else
 		{
@@ -2432,11 +2443,12 @@ void SMultiLineEditableText::LoadText()
 
 void SMultiLineEditableText::Tick( const FGeometry& AllottedGeometry, const double InCurrentTime, const float InDeltaTime )
 {
-	if (bTextChangedByVirtualKeyboard)
+	if (bTextChangedByVirtualKeyboard || bTextCommittedByVirtualKeyboard)
 	{
 		// Let outsiders know that the text content has been changed
-		OnTextCommitted.ExecuteIfBound(GetEditableText(), ETextCommit::OnUserMovedFocus);
+		OnTextCommitted.ExecuteIfBound(GetEditableText(), bTextCommittedByVirtualKeyboard ? ETextCommit::OnEnter : ETextCommit::OnUserMovedFocus);
 		bTextChangedByVirtualKeyboard = false;
+		bTextCommittedByVirtualKeyboard = false;
 	}
 
 	if(TextInputMethodChangeNotifier.IsValid() && TextInputMethodContext.IsValid() && TextInputMethodContext->CachedGeometry != AllottedGeometry)
diff --git a/Engine/Source/Runtime/Slate/Public/Widgets/Input/IVirtualKeyboardEntry.h b/Engine/Source/Runtime/Slate/Public/Widgets/Input/IVirtualKeyboardEntry.h
--- a/Engine/Source/Runtime/Slate/Public/Widgets/Input/IVirtualKeyboardEntry.h
+++ b/Engine/Source/Runtime/Slate/Public/Widgets/Input/IVirtualKeyboardEntry.h
@@ -21,8 +21,9 @@ public:
 	* Sets the text to that entered by the virtual keyboard
 	*
 	* @param  InNewText  The new text
+	* @param InCommit If True the text was committed, False then it was just changed
 	*/
-	virtual void SetTextFromVirtualKeyboard(const FText& InNewText) = 0;
+	virtual void SetTextFromVirtualKeyboard(const FText& InNewText, bool InCommit) = 0;
 	
 	/**
 	* Returns the text.
diff --git a/Engine/Source/Runtime/Slate/Public/Widgets/Input/SEditableText.h b/Engine/Source/Runtime/Slate/Public/Widgets/Input/SEditableText.h
--- a/Engine/Source/Runtime/Slate/Public/Widgets/Input/SEditableText.h
+++ b/Engine/Source/Runtime/Slate/Public/Widgets/Input/SEditableText.h
@@ -373,7 +373,7 @@ public:
 
 public:
 	//~ Begin IVirtualKeyboardEntry Interface
-	virtual void SetTextFromVirtualKeyboard(const FText& InNewText) override;
+	virtual void SetTextFromVirtualKeyboard(const FText& InNewText, bool InCommit) override;
 
 	virtual const FText& GetText() const override
 	{
diff --git a/Engine/Source/Runtime/Slate/Public/Widgets/Input/SVirtualKeyboardEntry.h b/Engine/Source/Runtime/Slate/Public/Widgets/Input/SVirtualKeyboardEntry.h
--- a/Engine/Source/Runtime/Slate/Public/Widgets/Input/SVirtualKeyboardEntry.h
+++ b/Engine/Source/Runtime/Slate/Public/Widgets/Input/SVirtualKeyboardEntry.h
@@ -88,7 +88,7 @@ public:
 
 public:
 	//~ Begin IVirtualKeyboardEntry Interface
-	virtual void SetTextFromVirtualKeyboard(const FText& InNewText) override;
+	virtual void SetTextFromVirtualKeyboard(const FText& InNewText, bool InCommit) override;
 
 	virtual const FText& GetText() const override
 	{
diff --git a/Engine/Source/Runtime/Slate/Public/Widgets/Text/SMultiLineEditableText.h b/Engine/Source/Runtime/Slate/Public/Widgets/Text/SMultiLineEditableText.h
--- a/Engine/Source/Runtime/Slate/Public/Widgets/Text/SMultiLineEditableText.h
+++ b/Engine/Source/Runtime/Slate/Public/Widgets/Text/SMultiLineEditableText.h
@@ -467,7 +467,7 @@ private:
 
 public:
 	//~ Begin IVirtualKeyboardEntry Interface
-	virtual void SetTextFromVirtualKeyboard(const FText& InNewText) override;
+	virtual void SetTextFromVirtualKeyboard(const FText& InNewText, bool InCommit) override;
 
 	virtual const FText& GetText() const override
 	{
@@ -761,6 +761,9 @@ private:
 	/** Whether the text has been changed by a virtual keyboard */
 	bool bTextChangedByVirtualKeyboard;
 
+	/** Whether the text has been committed by a virtual keyboard */
+	bool bTextCommittedByVirtualKeyboard;
+
 	/** The optional modifier key necessary to create a newline when typing into the editor. */
 	EModifierKey::Type ModiferKeyForNewLine;
 
