UMG/Slate - Fixing a memory leak with the invalidation panel, it was handing out a shared pointer to the ICacheLayer interface creating a circiular reference.  We now use raw pointers only in the PaintArgs.

diff --git a/Engine/Source/Runtime/Slate/Private/Framework/Application/SlateApplication.cpp b/Engine/Source/Runtime/Slate/Private/Framework/Application/SlateApplication.cpp
--- a/Engine/Source/Runtime/Slate/Private/Framework/Application/SlateApplication.cpp
+++ b/Engine/Source/Runtime/Slate/Private/Framework/Application/SlateApplication.cpp
@@ -969,7 +969,7 @@ void FSlateApplication::DrawWindowAndChildren( const TSharedRef<SWindow>& Window
 
 			FPlatformMisc::BeginNamedEvent(FColor::Magenta, "Slate::DrawWindow");
 			MaxLayerId = WindowToDraw->PaintWindow(
-				FPaintArgs(WindowToDraw, *WindowToDraw->GetHittestGrid(), WindowToDraw->GetPositionInScreen(), GetCurrentTime(), GetDeltaTime()),
+				FPaintArgs(WindowToDraw.Get(), *WindowToDraw->GetHittestGrid(), WindowToDraw->GetPositionInScreen(), GetCurrentTime(), GetDeltaTime()),
 				WindowGeometry, WindowToDraw->GetClippingRectangleInWindow(),
 				WindowElementList,
 				0,
@@ -993,7 +993,7 @@ void FSlateApplication::DrawWindowAndChildren( const TSharedRef<SWindow>& Window
 						const FGeometry DragDropContentGeometry = FGeometry::MakeRoot(DecoratorWidget->GetDesiredSize(), FSlateLayoutTransform(DragDropContentInWindowSpace));
 
 						DecoratorWidget->Paint(
-							FPaintArgs(WindowToDraw, *WindowToDraw->GetHittestGrid(), WindowToDraw->GetPositionInScreen(), GetCurrentTime(), GetDeltaTime()),
+							FPaintArgs(WindowToDraw.Get(), *WindowToDraw->GetHittestGrid(), WindowToDraw->GetPositionInScreen(), GetCurrentTime(), GetDeltaTime()),
 							DragDropContentGeometry, WindowToDraw->GetClippingRectangleInWindow(),
 							WindowElementList,
 							++MaxLayerId,
@@ -1018,7 +1018,7 @@ void FSlateApplication::DrawWindowAndChildren( const TSharedRef<SWindow>& Window
 					const FGeometry CursorGeometry = FGeometry::MakeRoot(CursorWidget->GetDesiredSize(), FSlateLayoutTransform(CursorPosInWindowSpace));
 
 					CursorWidget->Paint(
-						FPaintArgs(WindowToDraw, *WindowToDraw->GetHittestGrid(), WindowToDraw->GetPositionInScreen(), GetCurrentTime(), GetDeltaTime()),
+						FPaintArgs(WindowToDraw.Get(), *WindowToDraw->GetHittestGrid(), WindowToDraw->GetPositionInScreen(), GetCurrentTime(), GetDeltaTime()),
 						CursorGeometry, WindowToDraw->GetClippingRectangleInWindow(),
 						WindowElementList,
 						++MaxLayerId,
diff --git a/Engine/Source/Runtime/Slate/Private/Widgets/SInvalidationPanel.cpp b/Engine/Source/Runtime/Slate/Private/Widgets/SInvalidationPanel.cpp
--- a/Engine/Source/Runtime/Slate/Private/Widgets/SInvalidationPanel.cpp
+++ b/Engine/Source/Runtime/Slate/Private/Widgets/SInvalidationPanel.cpp
@@ -130,7 +130,7 @@ void SInvalidationPanel::Tick( const FGeometry& AllottedGeometry, const double I
 		if ( bNeedsCaching )
 		{
 			SlatePrepass(AllottedGeometry.Scale);
-			CachePrepass(SharedThis(this));
+			CachePrepass(this);
 		}
 	}
 
@@ -212,6 +212,7 @@ int32 SInvalidationPanel::OnPaint( const FPaintArgs& Args, const FGeometry& Allo
 			//FPlatformMisc::BeginNamedEvent(FColor::Red, "Slate::Invalidation");
 
 			SInvalidationPanel* MutableThis = const_cast<SInvalidationPanel*>( this );
+			TSharedRef<SInvalidationPanel> SharedMutableThis = SharedThis(MutableThis);
 
 			// Always set the caching flag to false first, during the paint / tick pass we may change something
 			// to volatile and need to re-cache.
@@ -229,11 +230,11 @@ int32 SInvalidationPanel::OnPaint( const FPaintArgs& Args, const FGeometry& Allo
 			LastUsedCachedNodeIndex = 0;
 
 			RootCacheNode = CreateCacheNode();
-			RootCacheNode->Initialize(Args, SharedThis(MutableThis), AllottedGeometry, MyClippingRect);
+			RootCacheNode->Initialize(Args, SharedMutableThis, AllottedGeometry, MyClippingRect);
 
 			//TODO: When SWidget::Paint is called don't drag self if volatile, and we're doing a cache pass.
 			CachedMaxChildLayer = SCompoundWidget::OnPaint(
-				Args.EnableCaching(SharedThis(MutableThis), RootCacheNode, true, false),
+				Args.EnableCaching(MutableThis, RootCacheNode, true, false),
 				AllottedGeometry,
 				MyClippingRect,
 				*CachedWindowElements.Get(),
diff --git a/Engine/Source/Runtime/SlateCore/Private/Types/PaintArgs.cpp b/Engine/Source/Runtime/SlateCore/Private/Types/PaintArgs.cpp
--- a/Engine/Source/Runtime/SlateCore/Private/Types/PaintArgs.cpp
+++ b/Engine/Source/Runtime/SlateCore/Private/Types/PaintArgs.cpp
@@ -4,7 +4,7 @@
 #include "HittestGrid.h"
 #include "WidgetCaching.h"
 
-FPaintArgs::FPaintArgs( const TSharedRef<SWidget>& Parent, FHittestGrid& InHittestGrid, FVector2D InWindowOffset, double InCurrentTime, float InDeltaTime )
+FPaintArgs::FPaintArgs(const SWidget& Parent, FHittestGrid& InHittestGrid, FVector2D InWindowOffset, double InCurrentTime, float InDeltaTime)
 : ParentPtr(Parent)
 , Grid(InHittestGrid)
 , LastHittestIndex(INDEX_NONE)
@@ -14,13 +14,14 @@ FPaintArgs::FPaintArgs( const TSharedRef<SWidget>& Parent, FHittestGrid& InHitte
 , DeltaTime(InDeltaTime)
 , bIsCaching(false)
 , bIsVolatilityPass(false)
+, LayoutCache(nullptr)
 , ParentCacheNode(nullptr)
 {
 }
 
-FPaintArgs FPaintArgs::WithNewParent( const SWidget* Parent ) const
+FPaintArgs FPaintArgs::WithNewParent( const SWidget& Parent ) const
 {
-	FPaintArgs Args = FPaintArgs( const_cast<SWidget*>(Parent)->AsShared(), this->Grid, this->WindowOffset, this->CurrentTime, this->DeltaTime );
+	FPaintArgs Args = FPaintArgs(Parent, this->Grid, this->WindowOffset, this->CurrentTime, this->DeltaTime);
 	Args.LastHittestIndex = this->LastHittestIndex;
 	Args.LastRecordedVisibility = this->LastRecordedVisibility;
 	Args.LayoutCache = this->LayoutCache;
@@ -31,7 +32,7 @@ FPaintArgs FPaintArgs::WithNewParent( const SWidget* Parent ) const
 	return Args;
 }
 
-FPaintArgs FPaintArgs::EnableCaching(const TSharedPtr<ILayoutCache>& InLayoutCache, FCachedWidgetNode* InParentCacheNode, bool bEnableCaching, bool bEnableVolatile) const
+FPaintArgs FPaintArgs::EnableCaching(ILayoutCache* InLayoutCache, FCachedWidgetNode* InParentCacheNode, bool bEnableCaching, bool bEnableVolatile) const
 {
 	FPaintArgs UpdatedArgs(*this);
 	UpdatedArgs.LayoutCache = InLayoutCache;
diff --git a/Engine/Source/Runtime/SlateCore/Private/Widgets/SWidget.cpp b/Engine/Source/Runtime/SlateCore/Private/Widgets/SWidget.cpp
--- a/Engine/Source/Runtime/SlateCore/Private/Widgets/SWidget.cpp
+++ b/Engine/Source/Runtime/SlateCore/Private/Widgets/SWidget.cpp
@@ -29,6 +29,7 @@ SWidget::SWidget()
 	, RenderTransformPivot( FVector2D::ZeroVector )
 	, DesiredSize(FVector2D::ZeroVector)
 	, ToolTip()
+	, LayoutCache(nullptr)
 	, bIsHovered(false)
 	, bCanTick(true)
 	, bCanSupportFocus(true)
@@ -372,6 +373,13 @@ void SWidget::SlatePrepass(float LayoutScaleMultiplier)
 {
 	//SLATE_CYCLE_COUNTER_SCOPE_CUSTOM_DETAILED(SLATE_STATS_DETAIL_LEVEL_MED, GSlatePrepass, GetType());
 
+	// TODO Figure out a better way than to just reset the pointer.  This causes problems when we prepass
+	// volatile widgets, who still need to know about their invalidation panel incase they vanish themselves.
+
+	// Reset the layout cache object each pre-pass to ensure we never access a stale layout cache object 
+	// as this widget could have been moved in and out of a panel that was invalidated between frames.
+	//LayoutCache = nullptr;
+
 	if ( bCanHaveChildren )
 	{
 		// Cache child desired sizes first. This widget's desired size is
@@ -406,7 +414,7 @@ const FVector2D& SWidget::GetDesiredSize() const
 	return DesiredSize;
 }
 
-void SWidget::CachePrepass(TWeakPtr<ILayoutCache> InLayoutCache)
+void SWidget::CachePrepass(ILayoutCache* InLayoutCache)
 {
 	if ( bCanHaveChildren )
 	{
@@ -649,9 +657,6 @@ int32 SWidget::Paint(const FPaintArgs& Args, const FGeometry& AllottedGeometry,
 	// pass later when the layout cache draws.
 	if ( Args.IsCaching() && IsVolatile() )
 	{
-		// Volatile widgets don't have an associated layout cache.
-		LayoutCache.Reset();
-
 		const int32 VolatileLayerId = LayerId + 1;
 		OutDrawElements.QueueVolatilePainting(
 			FSlateWindowElementList::FVolatilePaint(SharedThis(this), Args, AllottedGeometry, MyClippingRect, VolatileLayerId, InWidgetStyle, bParentEnabled));
diff --git a/Engine/Source/Runtime/SlateCore/Public/Types/PaintArgs.h b/Engine/Source/Runtime/SlateCore/Public/Types/PaintArgs.h
--- a/Engine/Source/Runtime/SlateCore/Public/Types/PaintArgs.h
+++ b/Engine/Source/Runtime/SlateCore/Public/Types/PaintArgs.h
@@ -21,9 +21,14 @@ class FCachedWidgetNode;
 class SLATECORE_API FPaintArgs
 {
 public:
-	FPaintArgs( const TSharedRef<SWidget>& Parent, FHittestGrid& InHittestGrid, FVector2D InWindowOffset, double InCurrentTime, float InDeltaTime );
-	FPaintArgs WithNewParent( const SWidget* Parent ) const;
-	FPaintArgs EnableCaching(const TSharedPtr<ILayoutCache>& InLayoutCache, FCachedWidgetNode* InParentCacheNode, bool bEnableCaching, bool bEnableVolatility) const;
+	FPaintArgs( const SWidget& Parent, FHittestGrid& InHittestGrid, FVector2D InWindowOffset, double InCurrentTime, float InDeltaTime );
+	FORCEINLINE FPaintArgs WithNewParent(const SWidget* Parent) const
+	{
+		checkSlow(Parent);
+		return WithNewParent(*Parent);
+	}
+	FPaintArgs WithNewParent( const SWidget& Parent ) const;
+	FPaintArgs EnableCaching(ILayoutCache* InLayoutCache, FCachedWidgetNode* InParentCacheNode, bool bEnableCaching, bool bEnableVolatility) const;
 	FPaintArgs RecordHittestGeometry(const SWidget* Widget, const FGeometry& WidgetGeometry, const FSlateRect& InClippingRect) const;
 	FPaintArgs InsertCustomHitTestPath( TSharedRef<ICustomHitTestPath> CustomHitTestPath, int32 HitTestIndex ) const;
 
@@ -35,11 +40,11 @@ public:
 	float GetDeltaTime() const { return DeltaTime; }
 	bool IsCaching() const { return bIsCaching; }
 	bool IsVolatilityPass() const { return bIsVolatilityPass; }
-	TSharedPtr<ILayoutCache> GetLayoutCache() const { return LayoutCache; }
+	ILayoutCache* GetLayoutCache() const { return LayoutCache; }
 	FCachedWidgetNode* GetParentCacheNode() const { return ParentCacheNode; }
 
 private:
-	const TSharedRef<SWidget>& ParentPtr;
+	const SWidget& ParentPtr;
 	FHittestGrid& Grid;
 	int32 LastHittestIndex;
 	EVisibility LastRecordedVisibility;
@@ -48,6 +53,6 @@ private:
 	float DeltaTime;
 	bool bIsCaching;
 	bool bIsVolatilityPass;
-	TSharedPtr<ILayoutCache> LayoutCache;
+	ILayoutCache* LayoutCache;
 	FCachedWidgetNode* ParentCacheNode;
 };
\ No newline at end of file
diff --git a/Engine/Source/Runtime/SlateCore/Public/Widgets/SWidget.h b/Engine/Source/Runtime/SlateCore/Public/Widgets/SWidget.h
--- a/Engine/Source/Runtime/SlateCore/Public/Widgets/SWidget.h
+++ b/Engine/Source/Runtime/SlateCore/Public/Widgets/SWidget.h
@@ -92,6 +92,7 @@ enum class EInvalidateWidget
 class ILayoutCache
 {
 public:
+	virtual ~ILayoutCache() { }
 	virtual void InvalidateWidget(class SWidget* InvalidateWidget) = 0;
 	virtual FCachedWidgetNode* CreateCacheNode() const = 0;
 };
@@ -792,10 +793,9 @@ protected:
 	 */
 	FORCEINLINE void Advanced_ForceInvalidateLayout()
 	{
-		TSharedPtr<ILayoutCache> Cache = LayoutCache.Pin();
-		if ( Cache.IsValid() )
+		if ( LayoutCache )
 		{
-			Cache->InvalidateWidget(this);
+			LayoutCache->InvalidateWidget(this);
 		}
 	}
 
@@ -1029,7 +1029,7 @@ protected:
 	 * on invisible children that never get the opportunity to paint and receive the layout cache through the normal
 	 * means.  That way if an invisible widget becomes visible, we still properly invalidate the hierarchy.
 	 */
-	void CachePrepass(TWeakPtr<ILayoutCache> LayoutCache);
+	void CachePrepass(ILayoutCache* LayoutCache);
 
 protected:
 	/**
@@ -1114,7 +1114,7 @@ private:
 	TSharedPtr<IToolTip> ToolTip;
 
 	/** The current layout cache that may need to invalidated by changes to this widget. */
-	mutable TWeakPtr<ILayoutCache> LayoutCache;
+	mutable ILayoutCache* LayoutCache;
 
 protected:
 	/** Is this widget hovered? */
