Build Fix - Fixing some compiling errors, always offering a do-nothing inlined alternative when function isn't defined in the build.

diff --git a/Engine/Source/Runtime/Slate/Public/Widgets/SInvalidationPanel.h b/Engine/Source/Runtime/Slate/Public/Widgets/SInvalidationPanel.h
--- a/Engine/Source/Runtime/Slate/Public/Widgets/SInvalidationPanel.h
+++ b/Engine/Source/Runtime/Slate/Public/Widgets/SInvalidationPanel.h
@@ -48,6 +48,12 @@ public:
 
 	static bool GetEnableWidgetCaching();
 	static void SetEnableWidgetCaching(bool bEnable);
+#else
+	static bool IsInvalidationDebuggingEnabled() { return false; }
+	static void EnableInvalidationDebugging(bool bEnable) { }
+
+	static bool GetEnableWidgetCaching() { return true; }
+	static void SetEnableWidgetCaching(bool bEnable) { }
 #endif
 
 private:
