PR #1563: Exposing ClearGameViewportFocusWhenWindowLosesFocus as a console variable (Contributed by [REDACTED_PERSON])

diff --git a/Engine/Source/Runtime/Slate/Private/Framework/Application/SlateApplication.cpp b/Engine/Source/Runtime/Slate/Private/Framework/Application/SlateApplication.cpp
--- a/Engine/Source/Runtime/Slate/Private/Framework/Application/SlateApplication.cpp
+++ b/Engine/Source/Runtime/Slate/Private/Framework/Application/SlateApplication.cpp
@@ -335,11 +335,14 @@ namespace SlateDefs
 
 	// How far tool tips should be pushed out from a force field border, in pixels
 	static const FVector2D ToolTipOffsetFromForceField( 4.0f, 3.0f );
-
-	// When true we will clear keyboard focus from widgets within the game view port when the window loses focus
-	static bool ClearGameViewportFocusWhenWindowLosesFocus = false;
 }
 
+/** When true we will clear keyboard focus from widgets within the game viewport when the window loses focus */
+TAutoConsoleVariable<int32> CVarClearGameViewportFocusWhenWindowLosesFocus(
+	TEXT("Slate.ClearGameViewportFocusWhenWindowLosesFocus"),
+	false,
+	TEXT("When true we will clear keyboard focus from widgets within the game viewport when the window loses focus"));
+
 
 /** True if we should allow throttling based on mouse movement activity.  int32 instead of bool only for console variable system. */
 TAutoConsoleVariable<int32> ThrottleWhenMouseIsMoving( 
@@ -5395,7 +5398,7 @@ void FSlateApplication::ProcessApplicationActivationEvent( bool InAppActivated )
 		// Clear keyboard focus when the app is deactivated, if the widget isn't a child of the game viewport.
 		bool bClearKeyboardFocus = true;
 
-		if (!SlateDefs::ClearGameViewportFocusWhenWindowLosesFocus && GameViewportWidget.IsValid())
+		if (!CVarClearGameViewportFocusWhenWindowLosesFocus.GetValueOnGameThread() && GameViewportWidget.IsValid())
 		{
 			FUserFocusEntry& UserFocusEntry = UserFocusEntries[GetUserIndexForKeyboard()];
 			bClearKeyboardFocus = !(UserFocusEntry.WidgetPath.ContainsWidget(GameViewportWidget.Pin().ToSharedRef()));
