Fixed rendering issue with submenu button entries.

Since the introduction of SButton::OnPaint, we can just use the default foreground color; no need to make a special case when hovered.

diff --git a/Engine/Source/Runtime/Slate/Private/Framework/MultiBox/SMenuEntryBlock.cpp b/Engine/Source/Runtime/Slate/Private/Framework/MultiBox/SMenuEntryBlock.cpp
--- a/Engine/Source/Runtime/Slate/Private/Framework/MultiBox/SMenuEntryBlock.cpp
+++ b/Engine/Source/Runtime/Slate/Private/Framework/MultiBox/SMenuEntryBlock.cpp
@@ -625,7 +625,7 @@ public:
 
 		SButton::FArguments ButtonArgs;
 		ButtonArgs.Text(InArgs._Label);
-		ButtonArgs.ForegroundColor( this, &SSubMenuButton::InvertOnHover );
+		ButtonArgs.ForegroundColor(FSlateColor::UseForeground());
 		ButtonArgs.HAlign(HAlign_Fill);
 		ButtonArgs.VAlign(VAlign_Fill);
 		ButtonArgs.ContentPadding(FMargin(0,2));
@@ -663,18 +663,6 @@ private:
 		}
 	}
 
-	FSlateColor InvertOnHover() const
-	{
-		if ( this->IsHovered() || ShouldAppearHovered.Get() )
-		{
-			return FLinearColor::Black;
-		}
-		else
-		{
-			return FSlateColor::UseForeground();
-		}
-	}
-
 private:
 	/** Attribute to indicate if the sub-menu is open or not */
 	TAttribute<bool> ShouldAppearHovered;
