Fix to SSubMenuButton to correct its visual appearance when its submenu is opened.

diff --git a/Engine/Source/Runtime/Slate/Private/Framework/MultiBox/SMenuEntryBlock.cpp b/Engine/Source/Runtime/Slate/Private/Framework/MultiBox/SMenuEntryBlock.cpp
--- a/Engine/Source/Runtime/Slate/Private/Framework/MultiBox/SMenuEntryBlock.cpp
+++ b/Engine/Source/Runtime/Slate/Private/Framework/MultiBox/SMenuEntryBlock.cpp
@@ -625,7 +625,7 @@ public:
 
 		SButton::FArguments ButtonArgs;
 		ButtonArgs.Text(InArgs._Label);
-		ButtonArgs.ForegroundColor(FSlateColor::UseForeground());
+		ButtonArgs.ForegroundColor( this, &SSubMenuButton::InvertOnHover );
 		ButtonArgs.HAlign(HAlign_Fill);
 		ButtonArgs.VAlign(VAlign_Fill);
 		ButtonArgs.ContentPadding(FMargin(0,2));
@@ -642,24 +642,24 @@ public:
 		];
 
 		SButton::Construct( ButtonArgs );
+	}
 
-		// Replace SButtons delegate for getting the border image with our own
-		// so we can control hovered state
-		BorderImage.Bind( this, &SSubMenuButton::GetBorderImage );
+	virtual bool IsHovered() const override
+	{
+		// Submenu widgets which have been opened should remain as if hovered, even if the cursor is outside them
+		return SWidget::IsHovered() || ShouldAppearHovered.Get();
 	}
 
 private:
-	const FSlateBrush* GetBorderImage() const
+	FSlateColor InvertOnHover() const
 	{
-		// If the button should appear hovered always show the hovered image
-		// Otherwise let SButton decide
-		if( ShouldAppearHovered.Get() )
+		if ( this->IsHovered() )
 		{
-			return HoverImage;
+			return FLinearColor::Black;
 		}
 		else
 		{
-			return SButton::GetBorder();
+			return FSlateColor::UseForeground();
 		}
 	}
 
