Slate - Fixing placement of hint text when scaled, also multiply HintText's Color by the incoming widget tint, so that it's properly alpha blended if the parent's alpha is modified.

diff --git a/Engine/Source/Runtime/Slate/Private/Widgets/Input/SEditableText.cpp b/Engine/Source/Runtime/Slate/Private/Widgets/Input/SEditableText.cpp
--- a/Engine/Source/Runtime/Slate/Private/Widgets/Input/SEditableText.cpp
+++ b/Engine/Source/Runtime/Slate/Private/Widgets/Input/SEditableText.cpp
@@ -1260,12 +1260,12 @@ int32 SEditableText::OnPaint( const FPaintArgs& Args, const FGeometry& AllottedG
 	if (VisibleText.Len() == 0)
 	{
 		// Draw the hint text.
-		const FLinearColor HintTextColor = FLinearColor(ColorAndOpacitySRGB.R, ColorAndOpacitySRGB.G, ColorAndOpacitySRGB.B, 0.35f);
+		const FLinearColor HintTextColor = FLinearColor(ColorAndOpacitySRGB.R, ColorAndOpacitySRGB.G, ColorAndOpacitySRGB.B, 0.35f) * InWidgetStyle.GetColorAndOpacityTint();
 		const FString ThisHintText = this->HintText.Get().ToString();
 		FSlateDrawElement::MakeText(
 			OutDrawElements,
 			LayerId + TextLayer,
-			AllottedGeometry.ToPaintGeometry( FVector2D( 0, DrawPositionY ), AllottedGeometry.Size ),
+			AllottedGeometry.ToPaintGeometry(FVector2D(0, DrawPositionY * ScaleInverse), AllottedGeometry.Size),
 			ThisHintText,          // Text
 			FontInfo,              // Font information (font name, size)
 			MyClippingRect,        // Clipping rect
