Button Press and Release should be virtual

diff --git a/Engine/Source/Runtime/Slate/Public/Widgets/Input/SButton.h b/Engine/Source/Runtime/Slate/Public/Widgets/Input/SButton.h
--- a/Engine/Source/Runtime/Slate/Public/Widgets/Input/SButton.h
+++ b/Engine/Source/Runtime/Slate/Public/Widgets/Input/SButton.h
@@ -213,10 +213,10 @@ protected:
 	bool bIsFocusable;
 
 	/** Press the button */
-	void Press();
+	virtual void Press();
 
 	/** Release the button */
-	void Release();
+	virtual void Release();
 
 	/** Utility function to determine if the incoming mouse event is for a precise tap or click */
 	bool IsPreciseTapOrClick(const FPointerEvent& MouseEvent) const;
