Slate - Slate Application no longer attempts to run prepass on child windows if they are minimized or invisible.

diff --git a/Engine/Source/Runtime/Slate/Private/Framework/Application/SlateApplication.cpp b/Engine/Source/Runtime/Slate/Private/Framework/Application/SlateApplication.cpp
--- a/Engine/Source/Runtime/Slate/Private/Framework/Application/SlateApplication.cpp
+++ b/Engine/Source/Runtime/Slate/Private/Framework/Application/SlateApplication.cpp
@@ -948,10 +948,11 @@ struct FDrawWindowArgs
 
 void FSlateApplication::DrawWindowAndChildren( const TSharedRef<SWindow>& WindowToDraw, FDrawWindowArgs& DrawWindowArgs )
 {
-	SLATE_CYCLE_COUNTER_SCOPE_CUSTOM(GSlateDrawWindowAndChildren, WindowToDraw->GetCreatedInLocation());
 	// Only draw visible windows
 	if( WindowToDraw->IsVisible() && !WindowToDraw->IsWindowMinimized() )
 	{
+		SLATE_CYCLE_COUNTER_SCOPE_CUSTOM(GSlateDrawWindowAndChildren, WindowToDraw->GetCreatedInLocation());
+	
 		// Switch to the appropriate world for drawing
 		FScopedSwitchWorldHack SwitchWorld( WindowToDraw );
 
@@ -1060,22 +1061,25 @@ void FSlateApplication::DrawWindowAndChildren( const TSharedRef<SWindow>& Window
 
 static void PrepassWindowAndChildren( TSharedRef<SWindow> WindowToPrepass )
 {
-	SLATE_CYCLE_COUNTER_SCOPE_CUSTOM(GSlatePrepassWindowAndChildren, WindowToPrepass->GetCreatedInLocation());
-	FScopedSwitchWorldHack SwitchWorld( WindowToPrepass );
-
+	if ( WindowToPrepass->IsVisible() && !WindowToPrepass->IsWindowMinimized() )
 	{
-		SCOPE_CYCLE_COUNTER(STAT_SlatePrepass);
-		WindowToPrepass->SlatePrepass(FSlateApplication::Get().GetApplicationScale());
-	}
+		SLATE_CYCLE_COUNTER_SCOPE_CUSTOM(GSlatePrepassWindowAndChildren, WindowToPrepass->GetCreatedInLocation());
+		FScopedSwitchWorldHack SwitchWorld(WindowToPrepass);
 
-	if (WindowToPrepass->IsAutosized())
-	{
-		WindowToPrepass->Resize(WindowToPrepass->GetDesiredSize());
-	}
+		{
+			SCOPE_CYCLE_COUNTER(STAT_SlatePrepass);
+			WindowToPrepass->SlatePrepass(FSlateApplication::Get().GetApplicationScale());
+		}
 
-	for (const TSharedRef<SWindow>& ChildWindow : WindowToPrepass->GetChildWindows() )
-	{
-		PrepassWindowAndChildren(ChildWindow);
+		if ( WindowToPrepass->IsAutosized() )
+		{
+			WindowToPrepass->Resize(WindowToPrepass->GetDesiredSize());
+		}
+
+		for ( const TSharedRef<SWindow>& ChildWindow : WindowToPrepass->GetChildWindows() )
+		{
+			PrepassWindowAndChildren(ChildWindow);
+		}
 	}
 }
 
@@ -1106,17 +1110,14 @@ void FSlateApplication::DrawPrepass( TSharedPtr<SWindow> DrawOnlyThisWindow )
 	}
 	else if (DrawOnlyThisWindow.IsValid())
 	{
-		PrepassWindowAndChildren( DrawOnlyThisWindow.ToSharedRef() );
+		PrepassWindowAndChildren(DrawOnlyThisWindow.ToSharedRef());
 	}
 	else
 	{
 		// Draw all windows
 		for (const TSharedRef<SWindow>& CurrentWindow : SlateWindows)
 		{
-			if (CurrentWindow->IsVisible() && !CurrentWindow->IsWindowMinimized())
-			{
-				PrepassWindowAndChildren(CurrentWindow );
-			}
+			PrepassWindowAndChildren(CurrentWindow);
 		}
 	}
 }
@@ -3194,7 +3195,7 @@ void FSlateApplication::GetAllVisibleWindowsOrdered(TArray< TSharedRef<SWindow>
 	for( TArray< TSharedRef<SWindow> >::TConstIterator CurrentWindowIt( SlateWindows ); CurrentWindowIt; ++CurrentWindowIt )
 	{
 		TSharedRef<SWindow> CurrentWindow = *CurrentWindowIt;
-		if ( CurrentWindow->IsVisible() )
+		if ( CurrentWindow->IsVisible() && !CurrentWindow->IsWindowMinimized() )
 		{
 			GetAllVisibleChildWindows(OutWindows, CurrentWindow);
 		}
@@ -3203,7 +3204,7 @@ void FSlateApplication::GetAllVisibleWindowsOrdered(TArray< TSharedRef<SWindow>
 
 void FSlateApplication::GetAllVisibleChildWindows(TArray< TSharedRef<SWindow> >& OutWindows, TSharedRef<SWindow> CurrentWindow)
 {
-	if ( CurrentWindow->IsVisible() )
+	if ( CurrentWindow->IsVisible() && !CurrentWindow->IsWindowMinimized() )
 	{
 		OutWindows.Add(CurrentWindow);
 
