Fixed Slate crash while dragging a tab.

diff --git a/Engine/Source/Runtime/Slate/Private/Widgets/Docking/SDockTab.cpp b/Engine/Source/Runtime/Slate/Private/Widgets/Docking/SDockTab.cpp
--- a/Engine/Source/Runtime/Slate/Private/Widgets/Docking/SDockTab.cpp
+++ b/Engine/Source/Runtime/Slate/Private/Widgets/Docking/SDockTab.cpp
@@ -68,8 +68,11 @@ FReply SDockTab::OnDragDetected( const FGeometry& MyGeometry, const FPointerEven
 	const FVector2D TabGrabOffsetFraction = FVector2D(
 		FMath::Clamp(TabGrabOffset.X / TabSize.X, 0.0f, 1.0f),
 		FMath::Clamp(TabGrabOffset.Y / TabSize.Y, 0.0f, 1.0f) );
-			
-	return ParentPtr.Pin()->StartDraggingTab( SharedThis(this), TabGrabOffsetFraction, MouseEvent );
+	
+	auto PinnedParent = ParentPtr.Pin();
+	return ensure(PinnedParent.IsValid())
+		? PinnedParent->StartDraggingTab(SharedThis(this), TabGrabOffsetFraction, MouseEvent)
+		: FReply::Unhandled();
 }
 
 FReply SDockTab::OnMouseButtonUp( const FGeometry& MyGeometry, const FPointerEvent& MouseEvent )
