Added minimum drag width for SSpinBox mouse interaction

This alleviates issues with small SSpinBoxes with a finite range where mouse drags are too sensitive (since it's derived from the widget's DrawSize)

diff --git a/Engine/Source/Runtime/Slate/Public/Widgets/Input/SSpinBox.h b/Engine/Source/Runtime/Slate/Public/Widgets/Input/SSpinBox.h
--- a/Engine/Source/Runtime/Slate/Public/Widgets/Input/SSpinBox.h
+++ b/Engine/Source/Runtime/Slate/Public/Widgets/Input/SSpinBox.h
@@ -326,6 +326,10 @@ public:
 
 				// Increments the spin based on delta mouse movement.
 
+				// A minimum slider width to use for calculating deltas in the slider-range space
+				const float MinSliderWidth = 100.f;
+				const float SliderWidthInSlateUnits = FMath::Max(MyGeometry.GetDrawSize().X, MinSliderWidth);
+				
 				//if we have a range to draw in
 				if ( !bUnlimitedSpinRange) 
 				{
@@ -337,16 +341,16 @@ public:
 					{
 						FractionFilled = 1.0f - FMath::Pow( 1.0f - FractionFilled, CachedSliderExponent);
 					}
-					FractionFilled *= MyGeometry.GetDrawSize().X;
+					FractionFilled *= SliderWidthInSlateUnits;
 
 					// Now add the delta to the fraction filled, this causes the spin.
 					FractionFilled += MouseEvent.GetCursorDelta().X;
 						
 					// Clamp the fraction to be within the bounds of the geometry.
-					FractionFilled = FMath::Clamp(FractionFilled, 0.0f, MyGeometry.GetDrawSize().X);
+					FractionFilled = FMath::Clamp(FractionFilled, 0.0f, SliderWidthInSlateUnits);
 
 					// Convert the fraction filled to a percent.
-					float Percent = FMath::Clamp(FractionFilled / MyGeometry.GetDrawSize().X, 0.0f, 1.0f);
+					float Percent = FMath::Clamp(FractionFilled / SliderWidthInSlateUnits, 0.0f, 1.0f);
 					if (CachedSliderExponent != 1)
 					{
 						// Have to convert the percent to the proper value due to the exponent component to the spin.
@@ -359,7 +363,7 @@ public:
 				{
 					//we have no range specified, so increment scaled by our current value		
 					const double CurrentValue = FMath::Clamp<double>(FMath::Abs(InternalValue),1.0,TNumericLimits<NumericType>::Max());
-					const float MouseDelta = FMath::Abs(MouseEvent.GetCursorDelta().X / MyGeometry.GetDrawSize().X);
+					const float MouseDelta = FMath::Abs(MouseEvent.GetCursorDelta().X / SliderWidthInSlateUnits);
 					const float Sign = (MouseEvent.GetCursorDelta().X > 0) ? 1.f : -1.f;
 					NewValue = InternalValue + (Sign * MouseDelta * FMath::Pow(CurrentValue, SliderExponent.Get()));
 				}
