Fix for Mac window activation

diff --git a/Engine/Source/Runtime/Slate/Private/Framework/Application/SlateApplication.cpp b/Engine/Source/Runtime/Slate/Private/Framework/Application/SlateApplication.cpp
--- a/Engine/Source/Runtime/Slate/Private/Framework/Application/SlateApplication.cpp
+++ b/Engine/Source/Runtime/Slate/Private/Framework/Application/SlateApplication.cpp
@@ -4343,7 +4343,7 @@ FReply FSlateApplication::RoutePointerDownEvent(FWidgetPath& WidgetsUnderPointer
 
 #if PLATFORM_MAC
 		TSharedPtr<SWindow> TopLevelWindow = WidgetsUnderPointer.TopLevelWindow;
-		if ( bNeedToActivateWindow )
+		if ( bNeedToActivateWindow || (TopLevelWindow.IsValid() && TopLevelWindow->GetNativeWindow()->GetOSWindowHandle() != ActiveWindow) )
 		{
 			// Clicking on a context menu should not activate anything
 			// @todo: This needs to be updated when we have window type in SWindow and we no longer have to guess if WidgetsUnderCursor.TopLevelWindow is a menu
