Copying //UE4/Dev-Sequencer to Dev-Main (//UE4/Dev-Main)

diff --git a/Engine/Source/Runtime/MovieSceneTracks/Public/Sections/MovieSceneSubSection.h b/Engine/Source/Runtime/MovieSceneTracks/Public/Sections/MovieSceneSubSection.h
--- a/Engine/Source/Runtime/MovieSceneTracks/Public/Sections/MovieSceneSubSection.h
+++ b/Engine/Source/Runtime/MovieSceneTracks/Public/Sections/MovieSceneSubSection.h
@@ -2,6 +2,7 @@
 
 #pragma once
 
+#include "MovieSceneSection.h"
 #include "MovieSceneSubSection.generated.h"
 
 class UMovieSceneSequence;
@@ -16,11 +17,11 @@ class MOVIESCENETRACKS_API UMovieSceneSubSection
 {
 	GENERATED_BODY()
 
+public:
+
 	/** Default constructor. */
 	UMovieSceneSubSection();
 
-public:
-
 	/**
 	 * Get the sequence that is assigned to this section.
 	 *
@@ -38,9 +39,9 @@ public:
 	 * @param Sequence The sequence to play.
 	 * @see GetSequence
 	 */
-	void SetSequence(UMovieSceneSequence& Sequence)
+	void SetSequence(UMovieSceneSequence* Sequence)
 	{
-		SubSequence = &Sequence;
+		SubSequence = Sequence;
 	}
 
 public:
@@ -53,7 +54,11 @@ public:
 	UPROPERTY(EditAnywhere, Category="Timing")
 	float TimeScale;
 
-private:
+	/** Preroll time. */
+	UPROPERTY(EditAnywhere, Category="Timing")
+	float PrerollTime;
+
+protected:
 
 	/**
 	 * Movie scene being played by this section.
diff --git a/Engine/Source/Runtime/MovieSceneTracks/Public/Sections/MovieSceneVectorSection.h b/Engine/Source/Runtime/MovieSceneTracks/Public/Sections/MovieSceneVectorSection.h
--- a/Engine/Source/Runtime/MovieSceneTracks/Public/Sections/MovieSceneVectorSection.h
+++ b/Engine/Source/Runtime/MovieSceneTracks/Public/Sections/MovieSceneVectorSection.h
@@ -2,6 +2,8 @@
 
 #pragma once
 
+#include "IKeyframeSection.h"
+#include "MovieSceneKeyStruct.h"
 #include "MovieSceneSection.h"
 #include "MovieSceneVectorSection.generated.h"
 
@@ -17,7 +19,7 @@ enum class EKeyVectorChannel
 
 struct MOVIESCENETRACKS_API FVectorKey
 {
-	FVectorKey( EKeyVectorChannel InChannel, float InValue )
+	FVectorKey(EKeyVectorChannel InChannel, float InValue)
 	{
 		Channel = InChannel;
 		Value = InValue;
@@ -27,10 +29,29 @@ struct MOVIESCENETRACKS_API FVectorKey
 };
 
 
+/**
+ * Proxy structure for vector section key data.
+ */
+USTRUCT()
+struct FMovieSceneVectorKeyStruct
+	: public FMovieSceneKeyStruct
+{
+	GENERATED_BODY()
+
+	/** They key's vector value. */
+	UPROPERTY(EditAnywhere, Category=Key)
+	FVector4 Vector;
+
+	FRichCurveKey* Keys[4];
+
+	virtual void PropagateChanges(const FPropertyChangedEvent& ChangeEvent) override;
+};
+
+
 /**
  * A vector section
  */
-UCLASS(MinimalAPI )
+UCLASS(MinimalAPI)
 class UMovieSceneVectorSection
 	: public UMovieSceneSection
 	, public IKeyframeSection<FVectorKey>
@@ -44,20 +65,7 @@ public:
 	 *
 	 * @param Position	The position in time within the movie scene
 	 */
-	FVector4 Eval( float Position, const FVector4& DefaultVector ) const;
-
-	/** MovieSceneSection interface */
-	virtual void MoveSection(float DeltaPosition, TSet<FKeyHandle>& KeyHandles) override;
-	virtual void DilateSection(float DilationFactor, float Origin, TSet<FKeyHandle>& KeyHandles) override;
-	virtual void GetKeyHandles(TSet<FKeyHandle>& KeyHandles) const override;
-
-
-	// IKeyframeSection interface.
-
-	virtual void AddKey( float Time, const FVectorKey& Key, EMovieSceneKeyInterpolation KeyInterpolation ) override;
-	virtual bool NewKeyIsNewData( float Time, const FVectorKey& Key ) const override;
-	virtual bool HasKeys( const FVectorKey& Key ) const override;
-	virtual void SetDefault( const FVectorKey& Key ) override;
+	FVector4 Eval(float Position, const FVector4& DefaultVector) const;
 
 	/** Gets one of four curves in this section */
 	FRichCurve& GetCurve(const int32& Index) { return Curves[Index]; }
@@ -65,13 +73,31 @@ public:
 	/** Sets how many channels are to be used */
 	void SetChannelsUsed(int32 InChannelsUsed) 
 	{
-		checkf(InChannelsUsed >= 2 && InChannelsUsed <= 4, TEXT("Only 2-4 channels are supported.") );
+		checkf(InChannelsUsed >= 2 && InChannelsUsed <= 4, TEXT("Only 2-4 channels are supported."));
 		ChannelsUsed = InChannelsUsed;
 	}
 	
 	/** Gets the number of channels in use */
 	int32 GetChannelsUsed() const {return ChannelsUsed;}
 
+public:
+
+	// UMovieSceneSection interface
+
+	virtual void MoveSection(float DeltaPosition, TSet<FKeyHandle>& KeyHandles) override;
+	virtual void DilateSection(float DilationFactor, float Origin, TSet<FKeyHandle>& KeyHandles) override;
+	virtual void GetKeyHandles(TSet<FKeyHandle>& KeyHandles) const override;
+	virtual TSharedPtr<FStructOnScope> GetKeyStruct(const TArray<FKeyHandle>& KeyHandles) override;
+
+public:
+
+	// IKeyframeSection interface.
+
+	virtual void AddKey(float Time, const FVectorKey& Key, EMovieSceneKeyInterpolation KeyInterpolation) override;
+	virtual bool NewKeyIsNewData(float Time, const FVectorKey& Key) const override;
+	virtual bool HasKeys(const FVectorKey& Key) const override;
+	virtual void SetDefault(const FVectorKey& Key) override;
+
 private:
 
 	/** Vector t */
diff --git a/Engine/Source/Runtime/MovieSceneTracks/Public/TrackInstances/MovieScene3DConstraintTrackInstance.h b/Engine/Source/Runtime/MovieSceneTracks/Public/TrackInstances/MovieScene3DConstraintTrackInstance.h
--- a/Engine/Source/Runtime/MovieSceneTracks/Public/TrackInstances/MovieScene3DConstraintTrackInstance.h
+++ b/Engine/Source/Runtime/MovieSceneTracks/Public/TrackInstances/MovieScene3DConstraintTrackInstance.h
@@ -19,7 +19,7 @@ public:
 	/** IMovieSceneTrackInstance interface */
 	virtual void SaveState (const TArray<UObject*>& RuntimeObjects, IMovieScenePlayer& Player, FMovieSceneSequenceInstance& SequenceInstance) override;
 	virtual void RestoreState (const TArray<UObject*>& RuntimeObjects, IMovieScenePlayer& Player, FMovieSceneSequenceInstance& SequenceInstance) override;
-	virtual void Update( float Position, float LastPosition, const TArray<UObject*>& RuntimeObjects, IMovieScenePlayer& Player, FMovieSceneSequenceInstance& SequenceInstance, EMovieSceneUpdatePass UpdatePass ) override;
+	virtual void Update(EMovieSceneUpdateData& UpdateData, const TArray<UObject*>& RuntimeObjects, IMovieScenePlayer& Player, FMovieSceneSequenceInstance& SequenceInstance) override;
 	virtual void RefreshInstance( const TArray<UObject*>& RuntimeObjects, IMovieScenePlayer& Player, FMovieSceneSequenceInstance& SequenceInstance ) override {}
 	virtual void ClearInstance( IMovieScenePlayer& Player, FMovieSceneSequenceInstance& SequenceInstance ) override {}
 	virtual float EvalOrder() { return -1.f; }
diff --git a/Engine/Source/Runtime/MovieSceneTracks/Public/TrackInstances/MovieScene3DTransformTrackInstance.h b/Engine/Source/Runtime/MovieSceneTracks/Public/TrackInstances/MovieScene3DTransformTrackInstance.h
--- a/Engine/Source/Runtime/MovieSceneTracks/Public/TrackInstances/MovieScene3DTransformTrackInstance.h
+++ b/Engine/Source/Runtime/MovieSceneTracks/Public/TrackInstances/MovieScene3DTransformTrackInstance.h
@@ -21,7 +21,7 @@ public:
 	/** IMovieSceneTrackInstance interface */
 	virtual void SaveState (const TArray<UObject*>& RuntimeObjects, IMovieScenePlayer& Player, FMovieSceneSequenceInstance& SequenceInstance) override;
 	virtual void RestoreState (const TArray<UObject*>& RuntimeObjects, IMovieScenePlayer& Player, FMovieSceneSequenceInstance& SequenceInstance) override;
-	virtual void Update( float Position, float LastPosition, const TArray<UObject*>& RuntimeObjects, IMovieScenePlayer& Player, FMovieSceneSequenceInstance& SequenceInstance, EMovieSceneUpdatePass UpdatePass ) override;
+	virtual void Update(EMovieSceneUpdateData& UpdateData, const TArray<UObject*>& RuntimeObjects, IMovieScenePlayer& Player, FMovieSceneSequenceInstance& SequenceInstance) override;
 	virtual EMovieSceneUpdatePass HasUpdatePasses() override { return (EMovieSceneUpdatePass)(MSUP_PreUpdate | MSUP_Update); }
 	virtual void RefreshInstance( const TArray<UObject*>& RuntimeObjects, IMovieScenePlayer& Player, FMovieSceneSequenceInstance& SequenceInstance ) override;
 	virtual void ClearInstance( IMovieScenePlayer& Player, FMovieSceneSequenceInstance& SequenceInstance ) override {}
diff --git a/Engine/Source/Runtime/MovieSceneTracks/Public/TrackInstances/MovieSceneCameraCutTrackInstance.h b/Engine/Source/Runtime/MovieSceneTracks/Public/TrackInstances/MovieSceneCameraCutTrackInstance.h
new file mode 100644
--- /dev/null
+++ b/Engine/Source/Runtime/MovieSceneTracks/Public/TrackInstances/MovieSceneCameraCutTrackInstance.h
@@ -0,0 +1,52 @@
+// Copyright 1998-2016 Epic Games, Inc. All Rights Reserved.
+
+#pragma once
+
+#include "MovieSceneSubTrackInstance.h"
+
+class UMovieSceneCameraCutTrack;
+
+
+/**
+ * Instance of a UMovieSceneCameraCutTrack
+ */
+class FMovieSceneCameraCutTrackInstance
+	: public IMovieSceneTrackInstance
+{
+public:
+
+	FMovieSceneCameraCutTrackInstance(UMovieSceneCameraCutTrack& InCameraCutTrack);
+
+public:
+
+	// IMovieSceneTrackInstance interface
+
+	virtual void ClearInstance(IMovieScenePlayer& Player, FMovieSceneSequenceInstance& SequenceInstance) override;
+	virtual void RefreshInstance(const TArray<UObject*>& RuntimeObjects, IMovieScenePlayer& Player, FMovieSceneSequenceInstance& SequenceInstance) override;
+	virtual void RestoreState(const TArray<UObject*>& RuntimeObjects, IMovieScenePlayer& Player, FMovieSceneSequenceInstance& SequenceInstance) override;
+	virtual void SaveState(const TArray<UObject*>& RuntimeObjects, IMovieScenePlayer& Player, FMovieSceneSequenceInstance& SequenceInstance) override { }
+	virtual void Update(EMovieSceneUpdateData& UpdateData, const TArray<UObject*>& RuntimeObjects, IMovieScenePlayer& Player, FMovieSceneSequenceInstance& SequenceInstance) override;
+
+protected:
+
+	/**
+	 * Acquire the camera object for the specified camera identifier.
+	 *
+	 * @param CameraCutIndex The index of the CameraCut section that needs the camera.
+	 * @param CameraGuid The unique identifier of the camera object to get.
+	 * @param SequenceInstance The sequence instance that owns this track instance.
+	 * @return The camera object, or nullptr if not found.
+	 */
+	UObject* AcquireCameraForCameraCut(int32 CameraCutIndex, const FGuid& CameraGuid, FMovieSceneSequenceInstance& SequenceInstance, const IMovieScenePlayer& Player);
+
+private:
+	
+	/** Runtime camera objects.  One for each CameraCut.  Must be the same number of entries as sections */
+	TArray<TWeakObjectPtr<UObject>> CachedCameraObjects;
+
+	/** Current camera object we are looking through.  Used to determine when making a new cut */
+	TWeakObjectPtr<UObject> LastCameraObject;
+
+	/** Track that is being instanced */
+	TWeakObjectPtr<UMovieSceneCameraCutTrack> CameraCutTrack;
+};
diff --git a/Engine/Source/Runtime/MovieSceneTracks/Public/TrackInstances/MovieSceneCinematicShotTrackInstance.h b/Engine/Source/Runtime/MovieSceneTracks/Public/TrackInstances/MovieSceneCinematicShotTrackInstance.h
new file mode 100644
--- /dev/null
+++ b/Engine/Source/Runtime/MovieSceneTracks/Public/TrackInstances/MovieSceneCinematicShotTrackInstance.h
@@ -0,0 +1,31 @@
+// Copyright 1998-2016 Epic Games, Inc. All Rights Reserved.
+
+#pragma once
+
+#include "MovieSceneSubTrackInstance.h"
+
+
+class FMovieSceneSequenceInstance;
+class UMovieSceneCinematicShotSection;
+class UMovieSceneCinematicShotTrack;
+
+
+/**
+ * Instance of a UMovieSceneSubTrack.
+ */
+class FMovieSceneCinematicShotTrackInstance
+	: public FMovieSceneSubTrackInstance
+{
+public:
+
+	/** Create and initialize a new instance. */
+	FMovieSceneCinematicShotTrackInstance(UMovieSceneCinematicShotTrack& InTrack);
+
+	// IMovieSceneTrackInstance interface
+
+	virtual void ClearInstance(IMovieScenePlayer& Player, FMovieSceneSequenceInstance& SequenceInstance) override;
+	virtual void RefreshInstance(const TArray<UObject*>& RuntimeObjects, IMovieScenePlayer& Player, FMovieSceneSequenceInstance& SequenceInstance) override;
+	virtual void RestoreState (const TArray<UObject*>& RuntimeObjects, IMovieScenePlayer& Player, FMovieSceneSequenceInstance& SequenceInstance) override;
+	virtual void SaveState(const TArray<UObject*>& RuntimeObjects, IMovieScenePlayer& Player, FMovieSceneSequenceInstance& SequenceInstance) override { }
+	virtual void Update(EMovieSceneUpdateData& UpdateData, const TArray<UObject*>& RuntimeObjects, IMovieScenePlayer& Player, FMovieSceneSequenceInstance& SequenceInstance) override;
+};
diff --git a/Engine/Source/Runtime/MovieSceneTracks/Public/TrackInstances/MovieSceneEventTrackInstance.h b/Engine/Source/Runtime/MovieSceneTracks/Public/TrackInstances/MovieSceneEventTrackInstance.h
--- a/Engine/Source/Runtime/MovieSceneTracks/Public/TrackInstances/MovieSceneEventTrackInstance.h
+++ b/Engine/Source/Runtime/MovieSceneTracks/Public/TrackInstances/MovieSceneEventTrackInstance.h
@@ -34,7 +34,7 @@ public:
 	virtual void RefreshInstance(const TArray<UObject*>& RuntimeObjects, IMovieScenePlayer& Player, FMovieSceneSequenceInstance& SequenceInstance) override { }
 	virtual void RestoreState(const TArray<UObject*>& RuntimeObjects, IMovieScenePlayer& Player, FMovieSceneSequenceInstance& SequenceInstance) override { }
 	virtual void SaveState(const TArray<UObject*>& RuntimeObjects, IMovieScenePlayer& Player, FMovieSceneSequenceInstance& SequenceInstance) override { }
-	virtual void Update(float Position, float LastPosition, const TArray<UObject*>& RuntimeObjects, IMovieScenePlayer& Player, FMovieSceneSequenceInstance& SequenceInstance, EMovieSceneUpdatePass UpdatePass) override;
+	virtual void Update(EMovieSceneUpdateData& UpdateData, const TArray<UObject*>& RuntimeObjects, IMovieScenePlayer& Player, FMovieSceneSequenceInstance& SequenceInstance) override;
 
 private:
 
diff --git a/Engine/Source/Runtime/MovieSceneTracks/Public/TrackInstances/MovieSceneFadeTrackInstance.h b/Engine/Source/Runtime/MovieSceneTracks/Public/TrackInstances/MovieSceneFadeTrackInstance.h
--- a/Engine/Source/Runtime/MovieSceneTracks/Public/TrackInstances/MovieSceneFadeTrackInstance.h
+++ b/Engine/Source/Runtime/MovieSceneTracks/Public/TrackInstances/MovieSceneFadeTrackInstance.h
@@ -34,7 +34,7 @@ public:
 	virtual void RefreshInstance(const TArray<UObject*>& RuntimeObjects, IMovieScenePlayer& Player, FMovieSceneSequenceInstance& SequenceInstance) override { }
 	virtual void RestoreState(const TArray<UObject*>& RuntimeObjects, IMovieScenePlayer& Player, FMovieSceneSequenceInstance& SequenceInstance) override;
 	virtual void SaveState(const TArray<UObject*>& RuntimeObjects, IMovieScenePlayer& Player, FMovieSceneSequenceInstance& SequenceInstance) override { }
-	virtual void Update(float Position, float LastPosition, const TArray<UObject*>& RuntimeObjects, IMovieScenePlayer& Player, FMovieSceneSequenceInstance& SequenceInstance, EMovieSceneUpdatePass UpdatePass) override;
+	virtual void Update(EMovieSceneUpdateData& UpdateData, const TArray<UObject*>& RuntimeObjects, IMovieScenePlayer& Player, FMovieSceneSequenceInstance& SequenceInstance) override;
 
 private:
 
diff --git a/Engine/Source/Runtime/MovieSceneTracks/Public/TrackInstances/MovieSceneSpawnTrackInstance.h b/Engine/Source/Runtime/MovieSceneTracks/Public/TrackInstances/MovieSceneSpawnTrackInstance.h
--- a/Engine/Source/Runtime/MovieSceneTracks/Public/TrackInstances/MovieSceneSpawnTrackInstance.h
+++ b/Engine/Source/Runtime/MovieSceneTracks/Public/TrackInstances/MovieSceneSpawnTrackInstance.h
@@ -21,7 +21,7 @@ public:
 	virtual void ClearInstance(IMovieScenePlayer& Player, FMovieSceneSequenceInstance& SequenceInstance) override {}
 	virtual EMovieSceneUpdatePass HasUpdatePasses() override { return MSUP_PreUpdate; }
 
-	virtual void Update(float Position, float LastPosition, const TArray<UObject*>& RuntimeObjects, class IMovieScenePlayer& Player, FMovieSceneSequenceInstance& SequenceInstance, EMovieSceneUpdatePass UpdatePass) override;
+	virtual void Update(EMovieSceneUpdateData& UpdateData, const TArray<UObject*>& RuntimeObjects, class IMovieScenePlayer& Player, FMovieSceneSequenceInstance& SequenceInstance) override;
 
 	virtual float EvalOrder() override { return -10000.f; }
 
diff --git a/Engine/Source/Runtime/MovieSceneTracks/Public/TrackInstances/MovieSceneSubTrackInstance.h b/Engine/Source/Runtime/MovieSceneTracks/Public/TrackInstances/MovieSceneSubTrackInstance.h
--- a/Engine/Source/Runtime/MovieSceneTracks/Public/TrackInstances/MovieSceneSubTrackInstance.h
+++ b/Engine/Source/Runtime/MovieSceneTracks/Public/TrackInstances/MovieSceneSubTrackInstance.h
@@ -29,7 +29,12 @@ public:
 	virtual void RefreshInstance(const TArray<UObject*>& RuntimeObjects, IMovieScenePlayer& Player, FMovieSceneSequenceInstance& SequenceInstance) override;
 	virtual void RestoreState (const TArray<UObject*>& RuntimeObjects, IMovieScenePlayer& Player, FMovieSceneSequenceInstance& SequenceInstance) override;
 	virtual void SaveState(const TArray<UObject*>& RuntimeObjects, IMovieScenePlayer& Player, FMovieSceneSequenceInstance& SequenceInstance) override;
-	virtual void Update(float Position, float LastPosition, const TArray<UObject*>& RuntimeObjects, IMovieScenePlayer& Player, FMovieSceneSequenceInstance& SequenceInstance, EMovieSceneUpdatePass UpdatePass) override;
+	virtual void Update(EMovieSceneUpdateData& UpdateData, const TArray<UObject*>& RuntimeObjects, IMovieScenePlayer& Player, FMovieSceneSequenceInstance& SequenceInstance) override;
+
+protected:
+
+	TArray<UMovieSceneSection*> GetAllTraversedSectionsWithPreroll( const TArray<UMovieSceneSection*>& Sections, float CurrentTime, float PreviousTime );
+	TArray<UMovieSceneSection*> GetTraversedSectionsWithPreroll( const TArray<UMovieSceneSection*>& Sections, float CurrentTime, float PreviousTime );
 
 protected:
 
diff --git a/Engine/Source/Runtime/MovieSceneTracks/Public/TrackInstances/MovieSceneVectorTrackInstance.h b/Engine/Source/Runtime/MovieSceneTracks/Public/TrackInstances/MovieSceneVectorTrackInstance.h
--- a/Engine/Source/Runtime/MovieSceneTracks/Public/TrackInstances/MovieSceneVectorTrackInstance.h
+++ b/Engine/Source/Runtime/MovieSceneTracks/Public/TrackInstances/MovieSceneVectorTrackInstance.h
@@ -22,7 +22,7 @@ public:
 	/** IMovieSceneTrackInstance interface */
 	virtual void SaveState (const TArray<UObject*>& RuntimeObjects, IMovieScenePlayer& Player, FMovieSceneSequenceInstance& SequenceInstance) override;
 	virtual void RestoreState (const TArray<UObject*>& RuntimeObjects, IMovieScenePlayer& Player, FMovieSceneSequenceInstance& SequenceInstance) override;
-	virtual void Update( float Position, float LastPosition, const TArray<UObject*>& RuntimeObjects, IMovieScenePlayer& Player, FMovieSceneSequenceInstance& SequenceInstance, EMovieSceneUpdatePass UpdatePass ) override;
+	virtual void Update( EMovieSceneUpdateData& UpdateData, const TArray<UObject*>& RuntimeObjects, IMovieScenePlayer& Player, FMovieSceneSequenceInstance& SequenceInstance ) override;
 	virtual void RefreshInstance( const TArray<UObject*>& RuntimeObjects, IMovieScenePlayer& Player, FMovieSceneSequenceInstance& SequenceInstance ) override;
 	virtual void ClearInstance( IMovieScenePlayer& Player, FMovieSceneSequenceInstance& SequenceInstance ) override {}
 
diff --git a/Engine/Source/Runtime/MovieSceneTracks/Public/TrackInstances/MovieSceneVisibilityTrackInstance.h b/Engine/Source/Runtime/MovieSceneTracks/Public/TrackInstances/MovieSceneVisibilityTrackInstance.h
--- a/Engine/Source/Runtime/MovieSceneTracks/Public/TrackInstances/MovieSceneVisibilityTrackInstance.h
+++ b/Engine/Source/Runtime/MovieSceneTracks/Public/TrackInstances/MovieSceneVisibilityTrackInstance.h
@@ -21,7 +21,7 @@ public:
 	/** IMovieSceneTrackInstance interface */
 	virtual void SaveState (const TArray<UObject*>& RuntimeObjects, IMovieScenePlayer& Player, FMovieSceneSequenceInstance& SequenceInstance) override;
 	virtual void RestoreState (const TArray<UObject*>& RuntimeObjects, IMovieScenePlayer& Player, FMovieSceneSequenceInstance& SequenceInstance) override;
-	virtual void Update( float Position, float LastPosition, const TArray<UObject*>& RuntimeObjects, class IMovieScenePlayer& Player, FMovieSceneSequenceInstance& SequenceInstance, EMovieSceneUpdatePass UpdatePass ) override;
+	virtual void Update( EMovieSceneUpdateData& UpdateData, const TArray<UObject*>& RuntimeObjects, class IMovieScenePlayer& Player, FMovieSceneSequenceInstance& SequenceInstance ) override;
 	virtual void RefreshInstance( const TArray<UObject*>& RuntimeObjects, class IMovieScenePlayer& Player, FMovieSceneSequenceInstance& SequenceInstance ) override;
 	virtual void ClearInstance( IMovieScenePlayer& Player, FMovieSceneSequenceInstance& SequenceInstance ) override {}
 
diff --git a/Engine/Source/Runtime/MovieSceneTracks/Public/Tracks/MovieScene3DConstraintTrack.h b/Engine/Source/Runtime/MovieSceneTracks/Public/Tracks/MovieScene3DConstraintTrack.h
--- a/Engine/Source/Runtime/MovieSceneTracks/Public/Tracks/MovieScene3DConstraintTrack.h
+++ b/Engine/Source/Runtime/MovieSceneTracks/Public/Tracks/MovieScene3DConstraintTrack.h
@@ -2,7 +2,7 @@
 
 #pragma once
 
-#include "MovieSceneTrack.h"
+#include "MovieSceneNameableTrack.h"
 #include "MovieScene3DConstraintTrack.generated.h"
 
 
@@ -14,7 +14,7 @@ class UMovieSceneSection;
  */
 UCLASS( MinimalAPI )
 class UMovieScene3DConstraintTrack
-	: public UMovieSceneTrack
+	: public UMovieSceneNameableTrack
 {
 	GENERATED_UCLASS_BODY()
 
@@ -43,7 +43,7 @@ public:
 	virtual const TArray<UMovieSceneSection*>& GetAllSections() const override;
 
 #if WITH_EDITORONLY_DATA
-	virtual FText GetDisplayName() const override;
+	virtual FText GetDefaultDisplayName() const override;
 #endif
 
 protected:
diff --git a/Engine/Source/Runtime/MovieSceneTracks/Public/Tracks/MovieScene3DTransformTrack.h b/Engine/Source/Runtime/MovieSceneTracks/Public/Tracks/MovieScene3DTransformTrack.h
--- a/Engine/Source/Runtime/MovieSceneTracks/Public/Tracks/MovieScene3DTransformTrack.h
+++ b/Engine/Source/Runtime/MovieSceneTracks/Public/Tracks/MovieScene3DTransformTrack.h
@@ -11,24 +11,30 @@
 /**
  * Handles manipulation of component transforms in a movie scene
  */
-UCLASS( MinimalAPI )
-class UMovieScene3DTransformTrack : public UMovieScenePropertyTrack
+UCLASS(MinimalAPI)
+class UMovieScene3DTransformTrack
+	: public UMovieScenePropertyTrack
 {
 	GENERATED_UCLASS_BODY()
+
 public:
-	/** UMovieSceneTrack interface */
-	virtual UMovieSceneSection* CreateNewSection() override;
-	virtual TSharedPtr<IMovieSceneTrackInstance> CreateInstance() override;
 
 	/**
 	 * Evaluates the track at the playback position
 	 *
-	 * @param Position		The current playback position
-	 * @param LastPosition		The last plackback position
-	 * @param OutTranslation	The evalulated translation component of the transform
-	 * @param OutRotation		The evalulated rotation component of the transform
-	 * @param OutScale		The evalulated scale component of the transform
-	 * @return true if anything was evaluated. Note: if false is returned OutBool remains unchanged
+	 * @param Position The current playback position.
+	 * @param LastPosition The last playback position.
+	 * @param OutTranslation The evaluated translation component of the transform.
+	 * @param OutRotation The evaluated rotation component of the transform.
+	 * @param OutScale The evaluated scale component of the transform.
+	 * @return true if anything was evaluated. Note: if false is returned OutBool remains unchanged.
 	 */
-	virtual bool Eval( float Position, float LastPosition, FVector& OutTranslation, FRotator& OutRotation, FVector& OutScale ) const;
+	virtual bool Eval(float Position, float LastPosition, FVector& OutTranslation, FRotator& OutRotation, FVector& OutScale) const;
+
+public:
+
+	// UMovieSceneTrack interface
+
+	virtual UMovieSceneSection* CreateNewSection() override;
+	virtual TSharedPtr<IMovieSceneTrackInstance> CreateInstance() override;
 };
diff --git a/Engine/Source/Runtime/MovieSceneTracks/Public/Tracks/MovieSceneCameraCutTrack.h b/Engine/Source/Runtime/MovieSceneTracks/Public/Tracks/MovieSceneCameraCutTrack.h
new file mode 100644
--- /dev/null
+++ b/Engine/Source/Runtime/MovieSceneTracks/Public/Tracks/MovieSceneCameraCutTrack.h
@@ -0,0 +1,59 @@
+// Copyright 1998-2016 Epic Games, Inc. All Rights Reserved.
+
+#pragma once
+
+#include "MovieSceneNameableTrack.h"
+#include "MovieSceneCameraCutTrack.generated.h"
+
+
+class IMovieSceneTrackInstance;
+class UMovieSceneSection;
+
+
+/**
+ * Handles manipulation of CameraCut properties in a movie scene.
+ */
+UCLASS(MinimalAPI)
+class UMovieSceneCameraCutTrack
+	: public UMovieSceneNameableTrack
+{
+	GENERATED_BODY()
+
+public:
+
+	/** 
+	 * Adds a new CameraCut at the specified time.
+	 *	
+	 * @param CameraHandle Handle to the camera that the CameraCut switches to when active.
+	 * @param TimeRange The range within this track's movie scene where the CameraCut is initially placed.
+	 */
+	MOVIESCENETRACKS_API void AddNewCameraCut(FGuid CameraHandle, float Time);
+
+public:
+
+	// UMovieSceneTrack interface
+
+	virtual void AddSection(UMovieSceneSection& Section) override;
+	virtual TSharedPtr<IMovieSceneTrackInstance> CreateInstance() override;
+	virtual UMovieSceneSection* CreateNewSection() override;
+	virtual const TArray<UMovieSceneSection*>& GetAllSections() const override;
+	virtual void RemoveSection(UMovieSceneSection& Section) override;
+
+#if WITH_EDITORONLY_DATA
+	virtual FText GetDefaultDisplayName() const override;
+#endif
+
+#if WITH_EDITOR
+	virtual void OnSectionMoved(UMovieSceneSection& Section) override;
+#endif
+
+protected:
+
+	float FindEndTimeForCameraCut(float StartTime);
+
+private:
+
+	/** All movie scene sections. */
+	UPROPERTY()
+	TArray<UMovieSceneSection*> Sections;
+};
diff --git a/Engine/Source/Runtime/MovieSceneTracks/Public/Tracks/MovieSceneCinematicShotTrack.h b/Engine/Source/Runtime/MovieSceneTracks/Public/Tracks/MovieSceneCinematicShotTrack.h
new file mode 100644
--- /dev/null
+++ b/Engine/Source/Runtime/MovieSceneTracks/Public/Tracks/MovieSceneCinematicShotTrack.h
@@ -0,0 +1,44 @@
+// Copyright 1998-2016 Epic Games, Inc. All Rights Reserved.
+
+#pragma once
+
+#include "MovieSceneSubTrack.h"
+#include "MovieSceneCinematicShotTrack.generated.h"
+
+
+class UMovieSceneSection;
+class UMovieSceneSequence;
+
+
+/**
+ * A track that holds consecutive sub sequences.
+ */
+UCLASS(MinimalAPI)
+class UMovieSceneCinematicShotTrack
+	: public UMovieSceneSubTrack
+{
+	GENERATED_BODY()
+
+public:
+
+	// UMovieSceneSubTrack interface
+
+	MOVIESCENETRACKS_API virtual UMovieSceneSubSection* AddSequence(UMovieSceneSequence* Sequence, float StartTime, float Duration, const bool& bInsertSequence = false);
+
+	// UMovieSceneTrack interface
+
+	virtual void AddSection(UMovieSceneSection& Section) override;
+	virtual TSharedPtr<IMovieSceneTrackInstance> CreateInstance() override;
+	virtual UMovieSceneSection* CreateNewSection() override;
+	virtual void RemoveSection(UMovieSceneSection& Section) override;
+	virtual bool SupportsMultipleRows() const override;
+	virtual bool AddsSectionBoundsToPlayRange() const override;
+
+#if WITH_EDITOR
+	virtual void OnSectionMoved(UMovieSceneSection& Section) override;
+#endif
+
+#if WITH_EDITORONLY_DATA
+	virtual FText GetDefaultDisplayName() const override;
+#endif
+};
diff --git a/Engine/Source/Runtime/MovieSceneTracks/Public/Tracks/MovieSceneEventTrack.h b/Engine/Source/Runtime/MovieSceneTracks/Public/Tracks/MovieSceneEventTrack.h
--- a/Engine/Source/Runtime/MovieSceneTracks/Public/Tracks/MovieSceneEventTrack.h
+++ b/Engine/Source/Runtime/MovieSceneTracks/Public/Tracks/MovieSceneEventTrack.h
@@ -58,6 +58,10 @@ public:
 	virtual void RemoveAllAnimationData() override;
 	virtual void RemoveSection(UMovieSceneSection& Section) override;
 
+#if WITH_EDITORONLY_DATA
+	virtual FText GetDefaultDisplayName() const override;
+#endif
+
 private:
 
 	/** If events should be fired when passed playing the sequence forwards. */
diff --git a/Engine/Source/Runtime/MovieSceneTracks/Public/Tracks/MovieSceneMaterialTrack.h b/Engine/Source/Runtime/MovieSceneTracks/Public/Tracks/MovieSceneMaterialTrack.h
--- a/Engine/Source/Runtime/MovieSceneTracks/Public/Tracks/MovieSceneMaterialTrack.h
+++ b/Engine/Source/Runtime/MovieSceneTracks/Public/Tracks/MovieSceneMaterialTrack.h
@@ -7,10 +7,14 @@
 #include "MovieSceneMaterialTrack.generated.h"
 
 
+class IMovieSceneTrackInstance;
+class UMovieSceneSection;
+
+
 /**
  * Handles manipulation of material parameters in a movie scene.
  */
-UCLASS( MinimalAPI )
+UCLASS(MinimalAPI)
 class UMovieSceneMaterialTrack
 	: public UMovieSceneNameableTrack
 {
@@ -37,7 +41,7 @@ public:
 	 * @param Time The time to add the new key.
 	 * @param The value for the new key.
 	 */
-	void MOVIESCENETRACKS_API AddScalarParameterKey( FName ParameterName, float Position, float Value );
+	void MOVIESCENETRACKS_API AddScalarParameterKey(FName ParameterName, float Position, float Value);
 
 	/**
 	* Adds a color parameter key to the track.
@@ -45,7 +49,7 @@ public:
 	* @param Time The time to add the new key.
 	* @param The value for the new key.
 	*/
-	void MOVIESCENETRACKS_API AddColorParameterKey( FName ParameterName, float Position, FLinearColor Value );
+	void MOVIESCENETRACKS_API AddColorParameterKey(FName ParameterName, float Position, FLinearColor Value);
 
 	/**
 	 * Gets the animated values for this track.
@@ -53,7 +57,7 @@ public:
 	 * @param OutScalarValues An array of FScalarParameterNameAndValue objects representing each animated scalar parameter and it's animated value.
 	 * @param OutVectorValues An array of FVectorParameterNameAndValue objects representing each animated vector parameter and it's animated value.
 	 */
-	void Eval( float Position, TArray<FScalarParameterNameAndValue>& OutScalarValues, TArray<FColorParameterNameAndValue>& OutVectorValues ) const;
+	void Eval(float Position, TArray<FScalarParameterNameAndValue>& OutScalarValues, TArray<FColorParameterNameAndValue>& OutVectorValues) const;
 
 private:
 
@@ -66,7 +70,7 @@ private:
 /**
  * A material track which is specialized for animation materials which are owned by actor components.
  */
-UCLASS( MinimalAPI )
+UCLASS(MinimalAPI)
 class UMovieSceneComponentMaterialTrack
 	: public UMovieSceneMaterialTrack
 {
@@ -79,7 +83,7 @@ public:
 	virtual TSharedPtr<IMovieSceneTrackInstance> CreateInstance() override;
 
 #if WITH_EDITORONLY_DATA
-	virtual FText GetDisplayName() const override;
+	virtual FText GetDefaultDisplayName() const override;
 #endif
 
 public:
diff --git a/Engine/Source/Runtime/MovieSceneTracks/Public/Tracks/MovieSceneParticleParameterTrack.h b/Engine/Source/Runtime/MovieSceneTracks/Public/Tracks/MovieSceneParticleParameterTrack.h
--- a/Engine/Source/Runtime/MovieSceneTracks/Public/Tracks/MovieSceneParticleParameterTrack.h
+++ b/Engine/Source/Runtime/MovieSceneTracks/Public/Tracks/MovieSceneParticleParameterTrack.h
@@ -2,7 +2,7 @@
 
 #pragma once
 
-#include "MovieSceneTrack.h"
+#include "MovieSceneNameableTrack.h"
 #include "MovieSceneParameterSection.h"
 #include "MovieSceneParticleParameterTrack.generated.h"
 
@@ -12,7 +12,7 @@
  */
 UCLASS( MinimalAPI )
 class UMovieSceneParticleParameterTrack
-	: public UMovieSceneTrack
+	: public UMovieSceneNameableTrack
 {
 	GENERATED_UCLASS_BODY()
 
@@ -31,7 +31,7 @@ public:
 	virtual const TArray<UMovieSceneSection*>& GetAllSections() const override;
 
 #if WITH_EDITORONLY_DATA
-	virtual FText GetDisplayName() const override;
+	virtual FText GetDefaultDisplayName() const override;
 #endif
 
 
diff --git a/Engine/Source/Runtime/MovieSceneTracks/Public/Tracks/MovieSceneParticleTrack.h b/Engine/Source/Runtime/MovieSceneTracks/Public/Tracks/MovieSceneParticleTrack.h
--- a/Engine/Source/Runtime/MovieSceneTracks/Public/Tracks/MovieSceneParticleTrack.h
+++ b/Engine/Source/Runtime/MovieSceneTracks/Public/Tracks/MovieSceneParticleTrack.h
@@ -46,7 +46,7 @@ public:
 	virtual UMovieSceneSection* CreateNewSection() override;
 
 #if WITH_EDITORONLY_DATA
-	virtual FText GetDisplayName() const override;
+	virtual FText GetDefaultDisplayName() const override;
 #endif
 
 private:
diff --git a/Engine/Source/Runtime/MovieSceneTracks/Public/Tracks/MovieScenePropertyTrack.h b/Engine/Source/Runtime/MovieSceneTracks/Public/Tracks/MovieScenePropertyTrack.h
--- a/Engine/Source/Runtime/MovieSceneTracks/Public/Tracks/MovieScenePropertyTrack.h
+++ b/Engine/Source/Runtime/MovieSceneTracks/Public/Tracks/MovieScenePropertyTrack.h
@@ -2,17 +2,20 @@
 
 #pragma once
 
-#include "MovieSceneTrack.h"
 #include "KeyParams.h"
+#include "MovieSceneNameableTrack.h"
 #include "MovieScenePropertyTrack.generated.h"
 
 
+class UMovieSceneSection;
+
+
 /**
  * Base class for tracks that animate an object property
  */
-UCLASS( abstract )
+UCLASS(abstract)
 class MOVIESCENETRACKS_API UMovieScenePropertyTrack
-	: public UMovieSceneTrack
+	: public UMovieSceneNameableTrack
 {
 	GENERATED_UCLASS_BODY()
 
@@ -22,7 +25,7 @@ public:
 
 	virtual FName GetTrackName() const override;
 	virtual void RemoveAllAnimationData() override;
-	virtual bool HasSection(const UMovieSceneSection& Section ) const override;
+	virtual bool HasSection(const UMovieSceneSection& Section) const override;
 	virtual void AddSection(UMovieSceneSection& Section) override;
 	virtual void RemoveSection(UMovieSceneSection& Section) override;
 	virtual bool IsEmpty() const override;
@@ -30,7 +33,7 @@ public:
 	virtual const TArray<UMovieSceneSection*>& GetAllSections() const override;
 
 #if WITH_EDITORONLY_DATA
-	virtual FText GetDisplayName() const override;
+	virtual FText GetDefaultDisplayName() const override;
 #endif
 
 public:
@@ -38,9 +41,9 @@ public:
 	/**
 	 * Sets the property name for this animatable property
 	 *
-	 * @param InPropertyName	The property being animated
+	 * @param InPropertyName The property being animated
 	 */
-	virtual void SetPropertyNameAndPath( FName InPropertyName, const FString& InPropertyPath );
+	virtual void SetPropertyNameAndPath(FName InPropertyName, const FString& InPropertyPath);
 
 	/** @return the name of the property being animated by this track */
 	FName GetPropertyName() const { return PropertyName; }
@@ -52,10 +55,9 @@ public:
 	 * Finds a section at the current time.
 	 *
 	 * @param Time	The time relative to the owning movie scene where the section should be
-	 *
 	 * @return The found section, or the new section.
 	 */
-	class UMovieSceneSection* FindOrAddSection(  float Time );
+	class UMovieSceneSection* FindOrAddSection(float Time);
 
 protected:
 
diff --git a/Engine/Source/Runtime/MovieSceneTracks/Public/Tracks/MovieSceneSkeletalAnimationTrack.h b/Engine/Source/Runtime/MovieSceneTracks/Public/Tracks/MovieSceneSkeletalAnimationTrack.h
--- a/Engine/Source/Runtime/MovieSceneTracks/Public/Tracks/MovieSceneSkeletalAnimationTrack.h
+++ b/Engine/Source/Runtime/MovieSceneTracks/Public/Tracks/MovieSceneSkeletalAnimationTrack.h
@@ -7,13 +7,14 @@
 #include "MovieSceneSkeletalAnimationTrack.generated.h"
 
 
+class IMovieSceneTrackInstance;
 class UMovieSceneSection;
 
 
 /**
  * Handles animation of skeletal mesh actors
  */
-UCLASS( MinimalAPI )
+UCLASS(MinimalAPI)
 class UMovieSceneSkeletalAnimationTrack
 	: public UMovieSceneNameableTrack
 {
@@ -33,7 +34,7 @@ public:
 
 	virtual TSharedPtr<IMovieSceneTrackInstance> CreateInstance() override;
 	virtual void RemoveAllAnimationData() override;
-	virtual bool HasSection(const UMovieSceneSection& Section ) const override;
+	virtual bool HasSection(const UMovieSceneSection& Section) const override;
 	virtual void AddSection(UMovieSceneSection& Section) override;
 	virtual void RemoveSection(UMovieSceneSection& Section) override;
 	virtual bool IsEmpty() const override;
@@ -43,7 +44,7 @@ public:
 	virtual UMovieSceneSection* CreateNewSection() override;
 
 #if WITH_EDITORONLY_DATA
-	virtual FText GetDisplayName() const override;
+	virtual FText GetDefaultDisplayName() const override;
 #endif
 
 private:
diff --git a/Engine/Source/Runtime/MovieSceneTracks/Public/Tracks/MovieSceneSlomoTrack.h b/Engine/Source/Runtime/MovieSceneTracks/Public/Tracks/MovieSceneSlomoTrack.h
--- a/Engine/Source/Runtime/MovieSceneTracks/Public/Tracks/MovieSceneSlomoTrack.h
+++ b/Engine/Source/Runtime/MovieSceneTracks/Public/Tracks/MovieSceneSlomoTrack.h
@@ -6,6 +6,10 @@
 #include "MovieSceneSlomoTrack.generated.h"
 
 
+class IMovieSceneTrackInstance;
+class UMovieSceneSection;
+
+
 /**
  * Implements a movie scene track that controls a movie scene's playback speed.
  */
@@ -23,6 +27,6 @@ public:
 	virtual UMovieSceneSection* CreateNewSection() override;
 
 #if WITH_EDITORONLY_DATA
-	virtual FText GetDisplayName() const override;
+	virtual FText GetDefaultDisplayName() const override;
 #endif
 };
diff --git a/Engine/Source/Runtime/MovieSceneTracks/Public/Tracks/MovieSceneSpawnTrack.h b/Engine/Source/Runtime/MovieSceneTracks/Public/Tracks/MovieSceneSpawnTrack.h
--- a/Engine/Source/Runtime/MovieSceneTracks/Public/Tracks/MovieSceneSpawnTrack.h
+++ b/Engine/Source/Runtime/MovieSceneTracks/Public/Tracks/MovieSceneSpawnTrack.h
@@ -6,13 +6,14 @@
 #include "MovieSceneSpawnTrack.generated.h"
 
 
+class IMovieSceneTrackInstance;
 class UMovieSceneSection;
 
 
 /**
  * Handles when a spawnable should be spawned and destroyed
  */
-UCLASS( MinimalAPI )
+UCLASS(MinimalAPI)
 class UMovieSceneSpawnTrack
 	: public UMovieSceneNameableTrack
 {
@@ -44,7 +45,7 @@ public:
 
 	virtual UMovieSceneSection* CreateNewSection() override;
 	virtual TSharedPtr<IMovieSceneTrackInstance> CreateInstance() override;
-	virtual bool HasSection(const UMovieSceneSection& Section ) const override;
+	virtual bool HasSection(const UMovieSceneSection& Section) const override;
 	virtual void AddSection(UMovieSceneSection& Section) override;
 	virtual void RemoveSection(UMovieSceneSection& Section) override;
 	virtual bool IsEmpty() const override;
@@ -52,7 +53,7 @@ public:
 	virtual const TArray<UMovieSceneSection*>& GetAllSections() const override;
 
 #if WITH_EDITORONLY_DATA
-	virtual FText GetDisplayName() const override;
+	virtual FText GetDefaultDisplayName() const override;
 #endif
 
 protected:
diff --git a/Engine/Source/Runtime/MovieSceneTracks/Public/Tracks/MovieSceneSubTrack.h b/Engine/Source/Runtime/MovieSceneTracks/Public/Tracks/MovieSceneSubTrack.h
--- a/Engine/Source/Runtime/MovieSceneTracks/Public/Tracks/MovieSceneSubTrack.h
+++ b/Engine/Source/Runtime/MovieSceneTracks/Public/Tracks/MovieSceneSubTrack.h
@@ -8,6 +8,7 @@
 
 class UMovieSceneSection;
 class UMovieSceneSequence;
+class UMovieSceneSubSection;
 
 
 /**
@@ -22,12 +23,15 @@ class UMovieSceneSubTrack
 public:
 
 	/**
-	 * Adds a movie scene section
+	 * Adds a movie scene section to the end of all the existing sections. If there are no sections, the new movie scene section is added at the requested time.
 	 *
 	 * @param Sequence The sequence to add
-	 * @param Time The time to add the section at
+	 * @param StartTime The time to add the section at
+	 * @param Duration The duration of the section
+	 * @param bInsertSequence Whether or not to insert the sequence and push existing sequences out
+	 * @return The newly created sub section
 	 */
-	MOVIESCENETRACKS_API void AddSequence(class UMovieSceneSequence& Sequence, float Time);
+	MOVIESCENETRACKS_API virtual UMovieSceneSubSection* AddSequence(UMovieSceneSequence* Sequence, float StartTime, float Duration, const bool& bInsertSequence = false);
 
 	/**
 	 * Check whether this track contains the given sequence.
@@ -54,10 +58,10 @@ public:
 	virtual bool SupportsMultipleRows() const override;
 
 #if WITH_EDITORONLY_DATA
-	virtual FText GetDisplayName() const override;
+	virtual FText GetDefaultDisplayName() const override;
 #endif
 
-private:
+protected:
 
 	/** All movie scene sections. */
 	UPROPERTY()
diff --git a/Engine/Source/Runtime/MovieSceneTracks/Public/Tracks/MovieSceneVisibilityTrack.h b/Engine/Source/Runtime/MovieSceneTracks/Public/Tracks/MovieSceneVisibilityTrack.h
--- a/Engine/Source/Runtime/MovieSceneTracks/Public/Tracks/MovieSceneVisibilityTrack.h
+++ b/Engine/Source/Runtime/MovieSceneTracks/Public/Tracks/MovieSceneVisibilityTrack.h
@@ -2,22 +2,31 @@
 
 #pragma once
 
-#include "MovieScene.h"
-#include "MovieSceneTrack.h"
-#include "MovieScenePropertyTrack.h"
 #include "MovieSceneBoolTrack.h"
 #include "MovieSceneVisibilityTrack.generated.h"
 
+
+class IMovieSceneTrackInstance;
+class UMovieSceneSection;
+
+
 /**
  * Handles manipulation of visibility properties in a movie scene
  */
-UCLASS( MinimalAPI )
-class UMovieSceneVisibilityTrack : public UMovieSceneBoolTrack
+UCLASS(MinimalAPI)
+class UMovieSceneVisibilityTrack
+	: public UMovieSceneBoolTrack
 {
 	GENERATED_UCLASS_BODY()
 
 public:
-	/** UMovieSceneTrack interface */
+
+	// UMovieSceneTrack interface
+
 	virtual UMovieSceneSection* CreateNewSection() override;
 	virtual TSharedPtr<IMovieSceneTrackInstance> CreateInstance() override;
+
+#if WITH_EDITORONLY_DATA
+	virtual FText GetDefaultDisplayName() const override;
+#endif
 };
diff --git a/Engine/Source/Runtime/Slate/Private/Widgets/Colors/SColorWheel.cpp b/Engine/Source/Runtime/Slate/Private/Widgets/Colors/SColorWheel.cpp
--- a/Engine/Source/Runtime/Slate/Private/Widgets/Colors/SColorWheel.cpp
+++ b/Engine/Source/Runtime/Slate/Private/Widgets/Colors/SColorWheel.cpp
@@ -7,7 +7,7 @@
 /* SColorWheel methods
  *****************************************************************************/
 
-void SColorWheel::Construct( const FArguments& InArgs )
+void SColorWheel::Construct(const FArguments& InArgs)
 {
 	Image = FCoreStyle::Get().GetBrush("ColorWheel.HueValueCircle");
 	SelectorImage = FCoreStyle::Get().GetBrush("ColorWheel.Selector");
@@ -22,19 +22,19 @@ void SColorWheel::Construct( const FArguments& InArgs )
 /* SWidget overrides
  *****************************************************************************/
 
-FVector2D SColorWheel::ComputeDesiredSize( float ) const
+FVector2D SColorWheel::ComputeDesiredSize(float) const
 {
 	return Image->ImageSize + SelectorImage->ImageSize;
 }
 
 
-FReply SColorWheel::OnMouseButtonDoubleClick( const FGeometry& InMyGeometry, const FPointerEvent& InMouseEvent )
+FReply SColorWheel::OnMouseButtonDoubleClick(const FGeometry& InMyGeometry, const FPointerEvent& InMouseEvent)
 {
 	return FReply::Handled();
 }
 
 
-FReply SColorWheel::OnMouseButtonDown( const FGeometry& MyGeometry, const FPointerEvent& MouseEvent )
+FReply SColorWheel::OnMouseButtonDown(const FGeometry& MyGeometry, const FPointerEvent& MouseEvent)
 {
 	if (MouseEvent.GetEffectingButton() == EKeys::LeftMouseButton)
 	{
@@ -48,14 +48,12 @@ FReply SColorWheel::OnMouseButtonDown( const FGeometry& MyGeometry, const FPoint
 
 		return FReply::Handled().CaptureMouse(SharedThis(this));
 	}
-	else
-	{
-		return FReply::Unhandled();
-	}
+
+	return FReply::Unhandled();
 }
 
 
-FReply SColorWheel::OnMouseButtonUp( const FGeometry& MyGeometry, const FPointerEvent& MouseEvent )
+FReply SColorWheel::OnMouseButtonUp(const FGeometry& MyGeometry, const FPointerEvent& MouseEvent)
 {
 	if (MouseEvent.GetEffectingButton() == EKeys::LeftMouseButton && HasMouseCapture())
 	{
@@ -63,14 +61,12 @@ FReply SColorWheel::OnMouseButtonUp( const FGeometry& MyGeometry, const FPointer
 
 		return FReply::Handled().ReleaseMouseCapture();
 	}
-	else
-	{
-		return FReply::Unhandled();
-	}
+
+	return FReply::Unhandled();
 }
 
 
-FReply SColorWheel::OnMouseMove( const FGeometry& MyGeometry, const FPointerEvent& MouseEvent )
+FReply SColorWheel::OnMouseMove(const FGeometry& MyGeometry, const FPointerEvent& MouseEvent)
 {
 	if (!HasMouseCapture())
 	{
@@ -83,28 +79,32 @@ FReply SColorWheel::OnMouseMove( const FGeometry& MyGeometry, const FPointerEven
 }
 
 
-int32 SColorWheel::OnPaint( const FPaintArgs& Args, const FGeometry& AllottedGeometry, const FSlateRect& MyClippingRect, FSlateWindowElementList& OutDrawElements, int32 LayerId, const FWidgetStyle& InWidgetStyle, bool bParentEnabled ) const
+int32 SColorWheel::OnPaint(const FPaintArgs& Args, const FGeometry& AllottedGeometry, const FSlateRect& MyClippingRect, FSlateWindowElementList& OutDrawElements, int32 LayerId, const FWidgetStyle& InWidgetStyle, bool bParentEnabled) const
 {
 	const bool bIsEnabled = ShouldBeEnabled(bParentEnabled);
 	const uint32 DrawEffects = bIsEnabled ? ESlateDrawEffect::None : ESlateDrawEffect::DisabledEffect;
+	const FVector2D& SelectorSize = SelectorImage->ImageSize;
+	const FVector2D CircleSize = AllottedGeometry.Size - SelectorSize;
 	
 	FSlateDrawElement::MakeBox(
 		OutDrawElements,
 		LayerId,
-		AllottedGeometry.ToPaintGeometry(),
+		AllottedGeometry.ToPaintGeometry(0.5 * SelectorSize, CircleSize),
 		Image,
 		MyClippingRect,
 		DrawEffects,
-		InWidgetStyle.GetColorAndOpacityTint() * Image->GetTint(InWidgetStyle));
+		InWidgetStyle.GetColorAndOpacityTint() * Image->GetTint(InWidgetStyle)
+	);
 
 	FSlateDrawElement::MakeBox(
 		OutDrawElements,
 		LayerId + 1,
-		AllottedGeometry.ToPaintGeometry(CalcRelativeSelectedPosition() * AllottedGeometry.Size * 0.5f - SelectorImage->ImageSize * 0.5f, SelectorImage->ImageSize),
+		AllottedGeometry.ToPaintGeometry(0.5f * (AllottedGeometry.Size + CalcRelativePositionFromCenter() * CircleSize - SelectorSize), SelectorSize),
 		SelectorImage,
 		MyClippingRect,
 		DrawEffects,
-		InWidgetStyle.GetColorAndOpacityTint() * SelectorImage->GetTint(InWidgetStyle));
+		InWidgetStyle.GetColorAndOpacityTint() * SelectorImage->GetTint(InWidgetStyle)
+	);
 
 	return LayerId + 1;
 }
@@ -113,26 +113,26 @@ int32 SColorWheel::OnPaint( const FPaintArgs& Args, const FGeometry& AllottedGeo
 /* SColorWheel implementation
  *****************************************************************************/
 
-FVector2D SColorWheel::CalcRelativeSelectedPosition( ) const
+FVector2D SColorWheel::CalcRelativePositionFromCenter() const
 {
 	float Hue = SelectedColor.Get().R;
 	float Saturation = SelectedColor.Get().G;
 	float Angle = Hue / 180.0f * PI;
 	float Radius = Saturation;
 
-	return FVector2D(FMath::Cos(Angle), FMath::Sin(Angle)) * Radius + FVector2D(1.0f, 1.0f);
+	return FVector2D(FMath::Cos(Angle), FMath::Sin(Angle)) * Radius;
 }
 
 
 bool SColorWheel::ProcessMouseAction(const FGeometry& MyGeometry, const FPointerEvent& MouseEvent, bool bProcessWhenOutsideColorWheel)
 {
 	const FVector2D LocalMouseCoordinate = MyGeometry.AbsoluteToLocal(MouseEvent.GetScreenSpacePosition());
-	const FVector2D Location = LocalMouseCoordinate / (MyGeometry.Size * 0.5f) - FVector2D(1.0f, 1.0f);
-	const float Radius = Location.Size();
+	const FVector2D RelativePositionFromCenter = (2.0f * LocalMouseCoordinate - MyGeometry.Size) / (MyGeometry.Size - SelectorImage->ImageSize);
+	const float RelativeRadius = RelativePositionFromCenter.Size();
 
-	if (Radius <= 1.0f || bProcessWhenOutsideColorWheel)
+	if (RelativeRadius <= 1.0f || bProcessWhenOutsideColorWheel)
 	{
-		float Angle = FMath::Atan2(Location.Y, Location.X);
+		float Angle = FMath::Atan2(RelativePositionFromCenter.Y, RelativePositionFromCenter.X);
 
 		if (Angle < 0.0f)
 		{
@@ -140,11 +140,15 @@ bool SColorWheel::ProcessMouseAction(const FGeometry& MyGeometry, const FPointer
 		}
 
 		FLinearColor NewColor = SelectedColor.Get();
-		NewColor.R = Angle * 180.0f * INV_PI;
-		NewColor.G = FMath::Min(Radius, 1.0f);
+		{
+			NewColor.R = Angle * 180.0f * INV_PI;
+			NewColor.G = FMath::Min(RelativeRadius, 1.0f);
+		}
 
 		OnValueChanged.ExecuteIfBound(NewColor);
 	}
 
-	return (Radius <= 1.0f);
+	GLog->Logf(TEXT("%f"), RelativeRadius);
+
+	return (RelativeRadius <= 1.0f);
 }
diff --git a/Engine/Source/Runtime/Slate/Private/Widgets/Input/SEditableLabel.cpp b/Engine/Source/Runtime/Slate/Private/Widgets/Input/SEditableLabel.cpp
--- a/Engine/Source/Runtime/Slate/Private/Widgets/Input/SEditableLabel.cpp
+++ b/Engine/Source/Runtime/Slate/Private/Widgets/Input/SEditableLabel.cpp
@@ -56,18 +56,6 @@ void SEditableLabel::Construct(const FArguments& InArgs)
 					.VirtualKeyboardType(EKeyboardType::Keyboard_Number)
 					.Visibility(EVisibility::Collapsed)
 			]
-
-		+ SHorizontalBox::Slot()
-			.FillWidth(1.0f)
-			.Padding(4.0f, 0.0f, 0.0f, 0.0f)
-			.HAlign(HAlign_Left)
-			.VAlign(VAlign_Center)
-			[
-				SNew(SImage)
-					.Image(FCoreStyle::Get().GetBrush(TEXT("Icons.Rename")))
-					.ToolTipText(LOCTEXT("RenameToolTip", "Press F2 or double-click the text label to rename it"))
-					.Visibility(this, &SEditableLabel::HandleIconVisibility)
-			]
 	];
 }
 
@@ -147,12 +135,4 @@ FReply SEditableLabel::HandleTextBlockDoubleClicked()
 }
 
 
-EVisibility SEditableLabel::HandleIconVisibility() const
-{
-	return (IsHovered() && !EditableText->HasKeyboardFocus() && CanEditAttribute.Get())
-		? EVisibility::Visible
-		: EVisibility::Collapsed;
-}
-
-
 #undef LOCTEXT_NAMESPACE
diff --git a/Engine/Source/Runtime/Slate/Public/Widgets/Colors/SColorWheel.h b/Engine/Source/Runtime/Slate/Public/Widgets/Colors/SColorWheel.h
--- a/Engine/Source/Runtime/Slate/Public/Widgets/Colors/SColorWheel.h
+++ b/Engine/Source/Runtime/Slate/Public/Widgets/Colors/SColorWheel.h
@@ -39,18 +39,18 @@ public:
 	 *
 	 * @param InArgs The declaration data for this widget.
 	 */
-	void Construct( const FArguments& InArgs );
+	void Construct(const FArguments& InArgs);
 
 public:
 
 	// SWidget overrides
 
 	virtual FVector2D ComputeDesiredSize(float) const override;
-	virtual FReply OnMouseButtonDoubleClick( const FGeometry& InMyGeometry, const FPointerEvent& InMouseEvent ) override;
-	virtual FReply OnMouseButtonDown( const FGeometry& MyGeometry, const FPointerEvent& MouseEvent ) override;
-	virtual FReply OnMouseButtonUp( const FGeometry& MyGeometry, const FPointerEvent& MouseEvent ) override;
-	virtual FReply OnMouseMove( const FGeometry& MyGeometry, const FPointerEvent& MouseEvent ) override;
-	virtual int32 OnPaint( const FPaintArgs& Args, const FGeometry& AllottedGeometry, const FSlateRect& MyClippingRect, FSlateWindowElementList& OutDrawElements, int32 LayerId, const FWidgetStyle& InWidgetStyle, bool bParentEnabled ) const override;
+	virtual FReply OnMouseButtonDoubleClick(const FGeometry& InMyGeometry, const FPointerEvent& InMouseEvent) override;
+	virtual FReply OnMouseButtonDown(const FGeometry& MyGeometry, const FPointerEvent& MouseEvent) override;
+	virtual FReply OnMouseButtonUp(const FGeometry& MyGeometry, const FPointerEvent& MouseEvent) override;
+	virtual FReply OnMouseMove(const FGeometry& MyGeometry, const FPointerEvent& MouseEvent) override;
+	virtual int32 OnPaint(const FPaintArgs& Args, const FGeometry& AllottedGeometry, const FSlateRect& MyClippingRect, FSlateWindowElementList& OutDrawElements, int32 LayerId, const FWidgetStyle& InWidgetStyle, bool bParentEnabled) const override;
 	
 protected:
 
@@ -59,7 +59,7 @@ protected:
 	 *
 	 * @return The position relative to the widget.
 	 */
-	FVector2D CalcRelativeSelectedPosition( ) const;
+	FVector2D CalcRelativePositionFromCenter() const;
 
 	/**
 	 * Performs actions according to mouse click / move
diff --git a/Engine/Source/Runtime/Slate/Public/Widgets/Input/NumericTypeInterface.h b/Engine/Source/Runtime/Slate/Public/Widgets/Input/NumericTypeInterface.h
--- a/Engine/Source/Runtime/Slate/Public/Widgets/Input/NumericTypeInterface.h
+++ b/Engine/Source/Runtime/Slate/Public/Widgets/Input/NumericTypeInterface.h
@@ -2,6 +2,9 @@
 
 #pragma once
 
+#include "BasicMathExpressionEvaluator.h"
+
+
 /** Interface to provide specific functionality for dealing with a numeric type. Currently includes string conversion functionality. */
 template<typename NumericType>
 struct INumericTypeInterface
@@ -10,7 +13,7 @@ struct INumericTypeInterface
 
 	/** Convert the type to/from a string */
 	virtual FString ToString(const NumericType& Value) const = 0;
-	virtual TOptional<NumericType> FromString(const FString& InString) = 0;
+	virtual TOptional<NumericType> FromString(const FString& InString, const NumericType& ExistingValue) = 0;
 
 	/** Check whether the typed character is valid */
 	virtual bool IsCharacterValid(TCHAR InChar) const = 0;
@@ -25,18 +28,17 @@ struct TDefaultNumericTypeInterface : INumericTypeInterface<NumericType>
 	{
 		return LexicalConversion::ToSanitizedString(Value);
 	}
-	virtual TOptional<NumericType> FromString(const FString& InString) override
+	virtual TOptional<NumericType> FromString(const FString& InString, const NumericType& InExistingValue) override
 	{
-		NumericType NewValue;
-		bool bEvalResult = LexicalConversion::TryParseString( NewValue, *InString );
-		if (!bEvalResult)
+		static FBasicMathExpressionEvaluator Parser;
+
+		TValueOrError<double, FExpressionError> Result = Parser.Evaluate(*InString, double(InExistingValue));
+		if (Result.IsValid())
 		{
-			float FloatValue = 0.f;
-			bEvalResult = FMath::Eval( *InString, FloatValue  );
-			NewValue = FloatValue;
+			return NumericType(Result.GetValue());
 		}
 
-		return bEvalResult ? NewValue : TOptional<NumericType>();
+		return TOptional<NumericType>();
 	}
 
 	/** Check whether the typed character is valid */
@@ -84,7 +86,7 @@ struct TNumericUnitTypeInterface : TDefaultNumericTypeInterface<NumericType>
 	virtual FString ToString(const NumericType& Value) const override;
 
 	/** Attempt to parse a numeral with our units from the specified string. */
-	virtual TOptional<NumericType> FromString(const FString& ValueString) override;
+	virtual TOptional<NumericType> FromString(const FString& ValueString, const NumericType& InExistingValue) override;
 
 	/** Check whether the specified typed character is valid */
 	virtual bool IsCharacterValid(TCHAR InChar) const override;
diff --git a/Engine/Source/Runtime/Slate/Public/Widgets/Input/NumericUnitTypeInterface.inl b/Engine/Source/Runtime/Slate/Public/Widgets/Input/NumericUnitTypeInterface.inl
--- a/Engine/Source/Runtime/Slate/Public/Widgets/Input/NumericUnitTypeInterface.inl
+++ b/Engine/Source/Runtime/Slate/Public/Widgets/Input/NumericUnitTypeInterface.inl
@@ -12,6 +12,11 @@ TNumericUnitTypeInterface<NumericType>::TNumericUnitTypeInterface(EUnit InUnits)
 template<typename NumericType>
 FString TNumericUnitTypeInterface<NumericType>::ToString(const NumericType& Value) const
 {
+	if (UnderlyingUnits == EUnit::Unspecified)
+	{
+		return TDefaultNumericTypeInterface<NumericType>::ToString(Value);
+	}
+
 	using namespace LexicalConversion;
 
 	FNumericUnit<NumericType> FinalValue(Value, UnderlyingUnits);
@@ -29,12 +34,19 @@ FString TNumericUnitTypeInterface<NumericType>::ToString(const NumericType& Valu
 }
 
 template<typename NumericType>
-TOptional<NumericType> TNumericUnitTypeInterface<NumericType>::FromString(const FString& InString)
+TOptional<NumericType> TNumericUnitTypeInterface<NumericType>::FromString(const FString& InString, const NumericType& InExistingValue)
 {
+	if (UnderlyingUnits == EUnit::Unspecified)
+	{
+		return TDefaultNumericTypeInterface<NumericType>::FromString(InString, InExistingValue);
+	}
+
 	using namespace LexicalConversion;
 
+	EUnit DefaultUnits = FixedDisplayUnits.IsSet() ? FixedDisplayUnits.GetValue() : UnderlyingUnits;
+
 	// Always parse in as a double, to allow for input of higher-order units with decimal numerals into integral types (eg, inputting 0.5km as 500m)
-	auto NewValue = FNumericUnit<double>::TryParseExpression(*InString, FixedDisplayUnits.IsSet() ? FixedDisplayUnits.GetValue() : UnderlyingUnits);
+	TValueOrError<FNumericUnit<double>, FText> NewValue = FNumericUnit<double>::TryParseExpression( *InString, DefaultUnits, InExistingValue );
 	if (NewValue.IsValid())
 	{
 		// Convert the number into the correct units
@@ -64,4 +76,4 @@ void TNumericUnitTypeInterface<NumericType>::SetupFixedDisplay(const NumericType
 	{
 		FixedDisplayUnits = DisplayUnit;
 	}
-}
\ No newline at end of file
+}
diff --git a/Engine/Source/Runtime/Slate/Public/Widgets/Input/SEditableLabel.h b/Engine/Source/Runtime/Slate/Public/Widgets/Input/SEditableLabel.h
--- a/Engine/Source/Runtime/Slate/Public/Widgets/Input/SEditableLabel.h
+++ b/Engine/Source/Runtime/Slate/Public/Widgets/Input/SEditableLabel.h
@@ -90,7 +90,6 @@ private:
 
 	void HandleEditableTextTextCommitted(const FText& NewText, ETextCommit::Type CommitInfo);
 	FReply HandleTextBlockDoubleClicked();
-	EVisibility HandleIconVisibility() const;
 
 private:
 
diff --git a/Engine/Source/Runtime/Slate/Public/Widgets/Input/SNumericEntryBox.h b/Engine/Source/Runtime/Slate/Public/Widgets/Input/SNumericEntryBox.h
--- a/Engine/Source/Runtime/Slate/Public/Widgets/Input/SNumericEntryBox.h
+++ b/Engine/Source/Runtime/Slate/Public/Widgets/Input/SNumericEntryBox.h
@@ -2,6 +2,7 @@
 
 #pragma once
 
+
 /**
  * Implementation for a box that only accepts a numeric value or that can display an undetermined value via a string
  * Supports an optional spin box for manipulating a value by dragging with the mouse
@@ -11,18 +12,26 @@ template<typename NumericType>
 class SNumericEntryBox : public SCompoundWidget
 {
 public: 
+
 	static const FLinearColor RedLabelBackgroundColor;
 	static const FLinearColor GreenLabelBackgroundColor;
 	static const FLinearColor BlueLabelBackgroundColor;
 	static const FText DefaultUndeterminedString;
 
 	/** Notification for numeric value change */
-	DECLARE_DELEGATE_OneParam( FOnValueChanged, NumericType );
+	DECLARE_DELEGATE_OneParam(FOnValueChanged, NumericType /*NewValue*/);
 
 	/** Notification for numeric value committed */
-	DECLARE_DELEGATE_TwoParams( FOnValueCommitted, NumericType, ETextCommit::Type);
+	DECLARE_DELEGATE_TwoParams(FOnValueCommitted, NumericType /*NewValue*/, ETextCommit::Type /*CommitType*/);
+
+	/** Notification for change of undetermined values */
+	DECLARE_DELEGATE_OneParam(FOnUndeterminedValueChanged, FText /*NewValue*/ );
+
+	/** Notification for committing undetermined values */
+	DECLARE_DELEGATE_TwoParams(FOnUndeterminedValueCommitted, FText /*NewValue*/, ETextCommit::Type /*CommitType*/);
 
 public:
+
 	SLATE_BEGIN_ARGS( SNumericEntryBox<NumericType> )
 		: _EditableTextBoxStyle( &FCoreStyle::Get().GetWidgetStyle<FEditableTextBoxStyle>("NormalEditableTextBox") )
 		, _Label()
@@ -82,6 +91,10 @@ public:
 		SLATE_EVENT( FOnValueChanged, OnValueChanged )
 		/** Called whenever the text is committed.  This happens when the user presses enter or the text box loses focus. */
 		SLATE_EVENT( FOnValueCommitted, OnValueCommitted )
+		/** Called whenever the text is changed interactively by the user */
+		SLATE_EVENT( FOnUndeterminedValueChanged, OnUndeterminedValueChanged )
+		/** Called whenever the text is committed.  This happens when the user presses enter or the text box loses focus. */
+		SLATE_EVENT( FOnUndeterminedValueCommitted, OnUndeterminedValueCommitted )
 		/** Called right before the slider begins to move */
 		SLATE_EVENT( FSimpleDelegate, OnBeginSliderMovement )
 		/** Called right after the slider handle is released by the user */
@@ -100,21 +113,22 @@ public:
 	{
 		check(InArgs._EditableTextBoxStyle);
 
-		const bool bAllowSpin = InArgs._AllowSpin;
 		OnValueChanged = InArgs._OnValueChanged;
 		OnValueCommitted = InArgs._OnValueCommitted;
+		OnUndeterminedValueChanged = InArgs._OnUndeterminedValueChanged;
+		OnUndeterminedValueCommitted = InArgs._OnUndeterminedValueCommitted;
 		ValueAttribute = InArgs._Value;
 		UndeterminedString = InArgs._UndeterminedString;
 		MinDesiredValueWidth = InArgs._MinDesiredValueWidth;
-		
 		BorderImageNormal = &InArgs._EditableTextBoxStyle->BackgroundImageNormal;
 		BorderImageHovered = &InArgs._EditableTextBoxStyle->BackgroundImageHovered;
 		BorderImageFocused = &InArgs._EditableTextBoxStyle->BackgroundImageFocused;
-		TAttribute<FMargin> TextMargin = InArgs._OverrideTextMargin.IsSet() ? InArgs._OverrideTextMargin : InArgs._EditableTextBoxStyle->Padding;
-
 		Interface = InArgs._TypeInterface.IsValid() ? InArgs._TypeInterface : MakeShareable( new TDefaultNumericTypeInterface<NumericType> );
 
+		TAttribute<FMargin> TextMargin = InArgs._OverrideTextMargin.IsSet() ? InArgs._OverrideTextMargin : InArgs._EditableTextBoxStyle->Padding;
+		const bool bAllowSpin = InArgs._AllowSpin;
 		TSharedPtr<SWidget> FinalWidget;
+
 		if( bAllowSpin )
 		{
 			SAssignNew( SpinBox, SSpinBox<NumericType> )
@@ -183,7 +197,6 @@ public:
 				EditableText.ToSharedRef()
 			];
 
-
 		ChildSlot
 			[
 				SNew( SBorder )
@@ -217,7 +230,8 @@ public:
 
 private:
 
-	//~ Begin SWidget Interface
+	//~ SWidget Interface
+
 	virtual bool SupportsKeyboardFocus() const override
 	{
 		return StaticCastSharedPtr<SWidget>(EditableText)->SupportsKeyboardFocus();
@@ -258,7 +272,8 @@ private:
 
 		return FReply::Unhandled();
 	}
-	//~ End SWidget Interface
+
+private:
 
 	/**
 	 * @return the Label that should be displayed                   
@@ -281,10 +296,8 @@ private:
 		{
 			return Value.GetValue();
 		}
-		else
-		{
-			return 0;
-		}
+
+		return 0;
 	}
 
 	/**
@@ -321,12 +334,14 @@ private:
 	{
 		const auto& Value = ValueAttribute.Get();
 
-		// Do not sent change events if the current value cannot be determined or else next tick the spin box could be swapped in if the value becomes determined
-		// while a user is typing in the box.  This causes keyboard focus switch which is bad
-		if( Value.IsSet() == true )
+		if (Value.IsSet() || !OnUndeterminedValueChanged.IsBound())
 		{
 			SendChangesFromText( NewValue, false, ETextCommit::Default );
 		}
+		else
+		{
+			OnUndeterminedValueChanged.Execute(NewValue);
+		}
 	}
 
 	/**
@@ -334,7 +349,16 @@ private:
 	 */
 	void OnTextCommitted( const FText& NewValue, ETextCommit::Type CommitInfo )
 	{
-		SendChangesFromText( NewValue, true, CommitInfo );
+		const auto& Value = ValueAttribute.Get();
+
+		if (Value.IsSet() || !OnUndeterminedValueCommitted.IsBound())
+		{
+			SendChangesFromText( NewValue, true, CommitInfo );
+		}
+		else
+		{
+			OnUndeterminedValueCommitted.Execute(NewValue, CommitInfo);
+		}
 	}
 
 	/**
@@ -343,6 +367,7 @@ private:
 	const FSlateBrush* GetBorderImage() const
 	{
 		TSharedPtr<const SWidget> EditingWidget;
+
 		if (SpinBox.IsValid() && SpinBox->GetVisibility() == EVisibility::Visible) 
 		{
 			EditingWidget = SpinBox;
@@ -356,17 +381,13 @@ private:
 		{
 			return BorderImageFocused;
 		}
-		else
+
+		if ( EditingWidget->IsHovered() )
 		{
-			if ( EditingWidget->IsHovered() )
-			{
-				return BorderImageHovered;
-			}
-			else
-			{
-				return BorderImageNormal;
-			}
+			return BorderImageHovered;
 		}
+
+		return BorderImageNormal;
 	}
 
 	/**
@@ -377,24 +398,28 @@ private:
 	 */
 	void SendChangesFromText( const FText& NewValue, bool bCommit, ETextCommit::Type CommitInfo )
 	{
+		if (NewValue.IsEmpty())
+		{
+			return;
+		}
+
 		// Only call the delegates if we have a valid numeric value
-		if( !NewValue.IsEmpty()  )
+		if (bCommit)
 		{
-			if( bCommit )
+			TOptional<NumericType> ExistingValue = ValueAttribute.Get();
+			TOptional<NumericType> NumericValue = Interface->FromString(NewValue.ToString(), ExistingValue.Get(0));
+
+			if (NumericValue.IsSet())
 			{
-				TOptional<NumericType> NumericValue = Interface->FromString(NewValue.ToString());
-				if( NumericValue.IsSet() )
-				{
-					OnValueCommitted.ExecuteIfBound(NumericValue.GetValue(), CommitInfo );
-				}
+				OnValueCommitted.ExecuteIfBound(NumericValue.GetValue(), CommitInfo );
 			}
-			else
+		}
+		else
+		{
+			NumericType NumericValue;
+			if (LexicalConversion::TryParseString(NumericValue, *NewValue.ToString()))
 			{
-				NumericType NumericValue;
-				if( LexicalConversion::TryParseString(NumericValue, *NewValue.ToString()) )
-				{
-					OnValueChanged.ExecuteIfBound( NumericValue );
-				}
+				OnValueChanged.ExecuteIfBound( NumericValue );
 			}
 		}
 	}
@@ -405,30 +430,33 @@ private:
 	virtual void Tick( const FGeometry& AllottedGeometry, const double InCurrentTime, const float InDeltaTime ) override
 	{
 		// Visibility toggle only matters if the spinbox is used
-		if( SpinBox.IsValid() )
+		if (!SpinBox.IsValid())
 		{
-			const auto& Value = ValueAttribute.Get();
+			return;
+		}
 
-			if( Value.IsSet() == true )
-			{
-				if( SpinBox->GetVisibility() != EVisibility::Visible )
-				{
-					// Set the visibility of the spinbox to visible if we have a valid value
-					SpinBox->SetVisibility( EVisibility::Visible );
-					// The text box should be invisible
-					EditableText->SetVisibility( EVisibility::Collapsed );
-				}
-			}
-			else
+		const auto& Value = ValueAttribute.Get();
+
+		if (Value.IsSet() == true)
+		{
+			if (SpinBox->GetVisibility() != EVisibility::Visible)
 			{
-				// The value isn't set so the spinbox should be hidden and the text box shown
-				SpinBox->SetVisibility(EVisibility::Collapsed);
-				EditableText->SetVisibility(EVisibility::Visible);
+				// Set the visibility of the spinbox to visible if we have a valid value
+				SpinBox->SetVisibility( EVisibility::Visible );
+				// The text box should be invisible
+				EditableText->SetVisibility( EVisibility::Collapsed );
 			}
 		}
+		else
+		{
+			// The value isn't set so the spinbox should be hidden and the text box shown
+			SpinBox->SetVisibility(EVisibility::Collapsed);
+			EditableText->SetVisibility(EVisibility::Visible);
+		}
 	}
 
 private:
+
 	/** Attribute for getting the label */
 	TAttribute< TOptional<FString > > LabelAttribute;
 	/** Attribute for getting the value.  If the value is not set we display the undetermined string */
@@ -441,6 +469,10 @@ private:
 	FOnValueChanged OnValueChanged;
 	/** Delegate to call when the value is committed */
 	FOnValueCommitted OnValueCommitted;
+	/** Delegate to call when an undetermined value changes */
+	FOnUndeterminedValueChanged OnUndeterminedValueChanged;
+	/** Delegate to call when an undetermined is committed */
+	FOnUndeterminedValueCommitted OnUndeterminedValueCommitted;
 	/** The undetermined string to display when needed */
 	FText UndeterminedString;
 	/** Styling: border image to draw when not hovered or focused */
diff --git a/Engine/Source/Runtime/Slate/Public/Widgets/Input/SSpinBox.h b/Engine/Source/Runtime/Slate/Public/Widgets/Input/SSpinBox.h
--- a/Engine/Source/Runtime/Slate/Public/Widgets/Input/SSpinBox.h
+++ b/Engine/Source/Runtime/Slate/Public/Widgets/Input/SSpinBox.h
@@ -542,7 +542,7 @@ protected:
 			ExitTextMode();
 		}
 
-		TOptional<NumericType> NewValue = Interface->FromString(NewText.ToString());
+		TOptional<NumericType> NewValue = Interface->FromString(NewText.ToString(), ValueAttribute.Get());
 		if (NewValue.IsSet())
 		{
 			CommitValue( NewValue.GetValue(), CommittedViaTypeIn, CommitInfo );
diff --git a/Engine/Source/Runtime/SlateCore/Private/Styling/CoreStyle.cpp b/Engine/Source/Runtime/SlateCore/Private/Styling/CoreStyle.cpp
--- a/Engine/Source/Runtime/SlateCore/Private/Styling/CoreStyle.cpp
+++ b/Engine/Source/Runtime/SlateCore/Private/Styling/CoreStyle.cpp
@@ -269,7 +269,6 @@ TSharedRef<ISlateStyle> FCoreStyle::Create( const FName& InStyleSetName )
 		Style->Set( "Icons.Info", new IMAGE_BRUSH( "Icons/icon_info_16x", Icon16x16) );
 		Style->Set( "Icons.Warning", new IMAGE_BRUSH( "Icons/icon_warning_16x", Icon16x16) );
 		Style->Set( "Icons.Download", new IMAGE_BRUSH( "Icons/icon_Downloads_16x", Icon16x16) );
-		Style->Set( "Icons.Rename", new IMAGE_BRUSH( "Icons/icon_rename_12x", Icon12x12) );
 	}
 
 	// Tool panels
diff --git a/Engine/Source/Runtime/UMG/Private/Animation/MovieScene2DTransformTrackInstance.cpp b/Engine/Source/Runtime/UMG/Private/Animation/MovieScene2DTransformTrackInstance.cpp
--- a/Engine/Source/Runtime/UMG/Private/Animation/MovieScene2DTransformTrackInstance.cpp
+++ b/Engine/Source/Runtime/UMG/Private/Animation/MovieScene2DTransformTrackInstance.cpp
@@ -11,12 +11,12 @@ FMovieScene2DTransformTrackInstance::FMovieScene2DTransformTrackInstance( UMovie
 	PropertyBindings = MakeShareable(new FTrackInstancePropertyBindings(TransformTrack->GetPropertyName(), TransformTrack->GetPropertyPath()));
 }
 
-void FMovieScene2DTransformTrackInstance::Update( float Position, float LastPosition, const TArray<UObject*>& RuntimeObjects, class IMovieScenePlayer& Player, FMovieSceneSequenceInstance& SequenceInstance, EMovieSceneUpdatePass UpdatePass ) 
+void FMovieScene2DTransformTrackInstance::Update(EMovieSceneUpdateData& UpdateData, const TArray<UObject*>& RuntimeObjects, class IMovieScenePlayer& Player, FMovieSceneSequenceInstance& SequenceInstance) 
 {
 	for(UObject* Object : RuntimeObjects)
 	{
 		FWidgetTransform TransformValue = PropertyBindings->GetCurrentValue<FWidgetTransform>(Object);
-		if(TransformTrack->Eval(Position, LastPosition, TransformValue))
+		if(TransformTrack->Eval(UpdateData.Position, UpdateData.LastPosition, TransformValue))
 		{
 			PropertyBindings->CallFunction<FWidgetTransform>(Object, &TransformValue);
 		}
diff --git a/Engine/Source/Runtime/UMG/Private/Animation/MovieScene2DTransformTrackInstance.h b/Engine/Source/Runtime/UMG/Private/Animation/MovieScene2DTransformTrackInstance.h
--- a/Engine/Source/Runtime/UMG/Private/Animation/MovieScene2DTransformTrackInstance.h
+++ b/Engine/Source/Runtime/UMG/Private/Animation/MovieScene2DTransformTrackInstance.h
@@ -15,7 +15,7 @@ public:
 	/** IMovieSceneTrackInstance interface */
 	virtual void SaveState(const TArray<UObject*>& RuntimeObjects, IMovieScenePlayer& Player, FMovieSceneSequenceInstance& SequenceInstance) override {}
 	virtual void RestoreState(const TArray<UObject*>& RuntimeObjects, IMovieScenePlayer& Player, FMovieSceneSequenceInstance& SequenceInstance) override {}
-	virtual void Update( float Position, float LastPosition, const TArray<UObject*>& RuntimeObjects, class IMovieScenePlayer& Player, FMovieSceneSequenceInstance& SequenceInstance, EMovieSceneUpdatePass UpdatePass ) override;
+	virtual void Update(EMovieSceneUpdateData& UpdateData, const TArray<UObject*>& RuntimeObjects, class IMovieScenePlayer& Player, FMovieSceneSequenceInstance& SequenceInstance) override;
 	virtual void RefreshInstance( const TArray<UObject*>& RuntimeObjects, class IMovieScenePlayer& Player, FMovieSceneSequenceInstance& SequenceInstance ) override;
 	virtual void ClearInstance( IMovieScenePlayer& Player, FMovieSceneSequenceInstance& SequenceInstance ) override {};
 private:
diff --git a/Engine/Source/Runtime/UMG/Private/Animation/MovieSceneMarginTrackInstance.cpp b/Engine/Source/Runtime/UMG/Private/Animation/MovieSceneMarginTrackInstance.cpp
--- a/Engine/Source/Runtime/UMG/Private/Animation/MovieSceneMarginTrackInstance.cpp
+++ b/Engine/Source/Runtime/UMG/Private/Animation/MovieSceneMarginTrackInstance.cpp
@@ -12,12 +12,12 @@ FMovieSceneMarginTrackInstance::FMovieSceneMarginTrackInstance( UMovieSceneMargi
 	PropertyBindings = MakeShareable( new FTrackInstancePropertyBindings( MarginTrack->GetPropertyName(), MarginTrack->GetPropertyPath() ) );
 }
 
-void FMovieSceneMarginTrackInstance::Update( float Position, float LastPosition, const TArray<UObject*>& RuntimeObjects, class IMovieScenePlayer& Player, FMovieSceneSequenceInstance& SequenceInstance, EMovieSceneUpdatePass UpdatePass ) 
+void FMovieSceneMarginTrackInstance::Update(EMovieSceneUpdateData& UpdateData, const TArray<UObject*>& RuntimeObjects, class IMovieScenePlayer& Player, FMovieSceneSequenceInstance& SequenceInstance) 
 {
 	for( UObject* Object : RuntimeObjects )
 	{
 		FMargin MarginValue = PropertyBindings->GetCurrentValue<FMargin>( Object );
-		if(MarginTrack->Eval(Position, LastPosition, MarginValue))
+		if(MarginTrack->Eval(UpdateData.Position, UpdateData.LastPosition, MarginValue))
 		{
 			PropertyBindings->CallFunction<FMargin>(Object, &MarginValue);
 		}
diff --git a/Engine/Source/Runtime/UMG/Private/Animation/MovieSceneMarginTrackInstance.h b/Engine/Source/Runtime/UMG/Private/Animation/MovieSceneMarginTrackInstance.h
--- a/Engine/Source/Runtime/UMG/Private/Animation/MovieSceneMarginTrackInstance.h
+++ b/Engine/Source/Runtime/UMG/Private/Animation/MovieSceneMarginTrackInstance.h
@@ -15,7 +15,7 @@ public:
 	/** IMovieSceneTrackInstance interface */
 	virtual void SaveState( const TArray<UObject*>& RuntimeObjects, IMovieScenePlayer& Player, FMovieSceneSequenceInstance& SequenceInstance ) override {}
 	virtual void RestoreState( const TArray<UObject*>& RuntimeObjects, IMovieScenePlayer& Player, FMovieSceneSequenceInstance& SequenceInstance ) override {}
-	virtual void Update( float Position, float LastPosition, const TArray<UObject*>& RuntimeObjects, class IMovieScenePlayer& Player, FMovieSceneSequenceInstance& SequenceInstance, EMovieSceneUpdatePass UpdatePass ) override;
+	virtual void Update(EMovieSceneUpdateData& UpdateData, const TArray<UObject*>& RuntimeObjects, class IMovieScenePlayer& Player, FMovieSceneSequenceInstance& SequenceInstance) override;
 	virtual void RefreshInstance( const TArray<UObject*>& RuntimeObjects, class IMovieScenePlayer& Player, FMovieSceneSequenceInstance& SequenceInstance ) override;
 	virtual void ClearInstance( IMovieScenePlayer& Player, FMovieSceneSequenceInstance& SequenceInstance ) override {};
 private:
diff --git a/Engine/Source/Runtime/UMG/Private/Animation/UMGSequencePlayer.cpp b/Engine/Source/Runtime/UMG/Private/Animation/UMGSequencePlayer.cpp
--- a/Engine/Source/Runtime/UMG/Private/Animation/UMGSequencePlayer.cpp
+++ b/Engine/Source/Runtime/UMG/Private/Animation/UMGSequencePlayer.cpp
@@ -99,8 +99,9 @@ void UUMGSequencePlayer::Tick(float DeltaTime)
 		}
 
 		if (RootMovieSceneInstance.IsValid())
-		{
-			RootMovieSceneInstance->Update(TimeCursorPosition + AnimationStartOffset, LastTimePosition + AnimationStartOffset, *this);
+		{			
+			EMovieSceneUpdateData UpdateData(TimeCursorPosition + AnimationStartOffset, LastTimePosition + AnimationStartOffset);
+			RootMovieSceneInstance->Update(UpdateData, *this);
 		}
 	}
 }
diff --git a/Engine/Source/Runtime/UMG/Private/Components/Border.cpp b/Engine/Source/Runtime/UMG/Private/Components/Border.cpp
--- a/Engine/Source/Runtime/UMG/Private/Components/Border.cpp
+++ b/Engine/Source/Runtime/UMG/Private/Components/Border.cpp
@@ -326,11 +326,6 @@ void UBorder::PostEditChangeProperty(struct FPropertyChangedEvent& PropertyChang
 	}
 }
 
-const FSlateBrush* UBorder::GetEditorIcon()
-{
-	return FUMGStyle::Get().GetBrush("Widget.Border");
-}
-
 const FText UBorder::GetPaletteCategory()
 {
 	return LOCTEXT("Common", "Common");
diff --git a/Engine/Source/Runtime/UMG/Private/Components/Button.cpp b/Engine/Source/Runtime/UMG/Private/Components/Button.cpp
--- a/Engine/Source/Runtime/UMG/Private/Components/Button.cpp
+++ b/Engine/Source/Runtime/UMG/Private/Components/Button.cpp
@@ -180,11 +180,6 @@ void UButton::SlateHandleUnhovered()
 
 #if WITH_EDITOR
 
-const FSlateBrush* UButton::GetEditorIcon()
-{
-	return FUMGStyle::Get().GetBrush("Widget.Button");
-}
-
 const FText UButton::GetPaletteCategory()
 {
 	return LOCTEXT("Common", "Common");
diff --git a/Engine/Source/Runtime/UMG/Private/Components/CanvasPanel.cpp b/Engine/Source/Runtime/UMG/Private/Components/CanvasPanel.cpp
--- a/Engine/Source/Runtime/UMG/Private/Components/CanvasPanel.cpp
+++ b/Engine/Source/Runtime/UMG/Private/Components/CanvasPanel.cpp
@@ -110,11 +110,6 @@ bool UCanvasPanel::GetGeometryForSlot(UCanvasPanelSlot* Slot, FGeometry& Arrange
 
 #if WITH_EDITOR
 
-const FSlateBrush* UCanvasPanel::GetEditorIcon()
-{
-	return FUMGStyle::Get().GetBrush("Widget.Canvas");
-}
-
 const FText UCanvasPanel::GetPaletteCategory()
 {
 	return LOCTEXT("Panel", "Panel");
diff --git a/Engine/Source/Runtime/UMG/Private/Components/CheckBox.cpp b/Engine/Source/Runtime/UMG/Private/Components/CheckBox.cpp
--- a/Engine/Source/Runtime/UMG/Private/Components/CheckBox.cpp
+++ b/Engine/Source/Runtime/UMG/Private/Components/CheckBox.cpp
@@ -214,11 +214,6 @@ void UCheckBox::PostLoad()
 
 #if WITH_EDITOR
 
-const FSlateBrush* UCheckBox::GetEditorIcon()
-{
-	return FUMGStyle::Get().GetBrush("Widget.CheckBox");
-}
-
 const FText UCheckBox::GetPaletteCategory()
 {
 	return LOCTEXT("Common", "Common");
diff --git a/Engine/Source/Runtime/UMG/Private/Components/CircularThrobber.cpp b/Engine/Source/Runtime/UMG/Private/Components/CircularThrobber.cpp
--- a/Engine/Source/Runtime/UMG/Private/Components/CircularThrobber.cpp
+++ b/Engine/Source/Runtime/UMG/Private/Components/CircularThrobber.cpp
@@ -103,11 +103,6 @@ void UCircularThrobber::PostLoad()
 
 #if WITH_EDITOR
 
-const FSlateBrush* UCircularThrobber::GetEditorIcon()
-{
-	return FUMGStyle::Get().GetBrush("Widget.CircularThrobber");
-}
-
 const FText UCircularThrobber::GetPaletteCategory()
 {
 	return LOCTEXT("Primitive", "Primitive");
diff --git a/Engine/Source/Runtime/UMG/Private/Components/ComboBox.cpp b/Engine/Source/Runtime/UMG/Private/Components/ComboBox.cpp
--- a/Engine/Source/Runtime/UMG/Private/Components/ComboBox.cpp
+++ b/Engine/Source/Runtime/UMG/Private/Components/ComboBox.cpp
@@ -50,11 +50,6 @@ TSharedRef<SWidget> UComboBox::HandleGenerateWidget(UObject* Item) const
 
 #if WITH_EDITOR
 
-const FSlateBrush* UComboBox::GetEditorIcon()
-{
-	return FUMGStyle::Get().GetBrush("Widget.ComboBox");
-}
-
 const FText UComboBox::GetPaletteCategory()
 {
 	return LOCTEXT("Misc", "Misc");
diff --git a/Engine/Source/Runtime/UMG/Private/Components/ComboBoxString.cpp b/Engine/Source/Runtime/UMG/Private/Components/ComboBoxString.cpp
--- a/Engine/Source/Runtime/UMG/Private/Components/ComboBoxString.cpp
+++ b/Engine/Source/Runtime/UMG/Private/Components/ComboBoxString.cpp
@@ -227,11 +227,6 @@ void UComboBoxString::HandleOpening()
 
 #if WITH_EDITOR
 
-const FSlateBrush* UComboBoxString::GetEditorIcon()
-{
-	return FUMGStyle::Get().GetBrush("Widget.ComboBox");
-}
-
 const FText UComboBoxString::GetPaletteCategory()
 {
 	return LOCTEXT("Input", "Input");
diff --git a/Engine/Source/Runtime/UMG/Private/Components/EditableText.cpp b/Engine/Source/Runtime/UMG/Private/Components/EditableText.cpp
--- a/Engine/Source/Runtime/UMG/Private/Components/EditableText.cpp
+++ b/Engine/Source/Runtime/UMG/Private/Components/EditableText.cpp
@@ -183,11 +183,6 @@ void UEditableText::PostLoad()
 
 #if WITH_EDITOR
 
-const FSlateBrush* UEditableText::GetEditorIcon()
-{
-	return FUMGStyle::Get().GetBrush("Widget.EditableText");
-}
-
 const FText UEditableText::GetPaletteCategory()
 {
 	return LOCTEXT("Primitive", "Primitive");
diff --git a/Engine/Source/Runtime/UMG/Private/Components/EditableTextBox.cpp b/Engine/Source/Runtime/UMG/Private/Components/EditableTextBox.cpp
--- a/Engine/Source/Runtime/UMG/Private/Components/EditableTextBox.cpp
+++ b/Engine/Source/Runtime/UMG/Private/Components/EditableTextBox.cpp
@@ -181,11 +181,6 @@ void UEditableTextBox::PostLoad()
 
 #if WITH_EDITOR
 
-const FSlateBrush* UEditableTextBox::GetEditorIcon()
-{
-	return FUMGStyle::Get().GetBrush("Widget.EditableTextBox");
-}
-
 const FText UEditableTextBox::GetPaletteCategory()
 {
 	return LOCTEXT("Common", "Common");
diff --git a/Engine/Source/Runtime/UMG/Private/Components/ExpandableArea.cpp b/Engine/Source/Runtime/UMG/Private/Components/ExpandableArea.cpp
--- a/Engine/Source/Runtime/UMG/Private/Components/ExpandableArea.cpp
+++ b/Engine/Source/Runtime/UMG/Private/Components/ExpandableArea.cpp
@@ -141,11 +141,6 @@ void UExpandableArea::SlateExpansionChanged(bool NewState)
 
 #if WITH_EDITOR
 
-const FSlateBrush* UExpandableArea::GetEditorIcon()
-{
-	return FUMGStyle::Get().GetBrush("Widget.Spacer");
-}
-
 const FText UExpandableArea::GetPaletteCategory()
 {
 	return LOCTEXT("Extra", "Extra");
diff --git a/Engine/Source/Runtime/UMG/Private/Components/GridPanel.cpp b/Engine/Source/Runtime/UMG/Private/Components/GridPanel.cpp
--- a/Engine/Source/Runtime/UMG/Private/Components/GridPanel.cpp
+++ b/Engine/Source/Runtime/UMG/Private/Components/GridPanel.cpp
@@ -88,11 +88,6 @@ void UGridPanel::SynchronizeProperties()
 
 #if WITH_EDITOR
 
-const FSlateBrush* UGridPanel::GetEditorIcon()
-{
-	return FUMGStyle::Get().GetBrush("Widget.Grid");
-}
-
 const FText UGridPanel::GetPaletteCategory()
 {
 	return LOCTEXT("Panel", "Panel");
diff --git a/Engine/Source/Runtime/UMG/Private/Components/HorizontalBox.cpp b/Engine/Source/Runtime/UMG/Private/Components/HorizontalBox.cpp
--- a/Engine/Source/Runtime/UMG/Private/Components/HorizontalBox.cpp
+++ b/Engine/Source/Runtime/UMG/Private/Components/HorizontalBox.cpp
@@ -73,11 +73,6 @@ TSharedRef<SWidget> UHorizontalBox::RebuildWidget()
 
 #if WITH_EDITOR
 
-const FSlateBrush* UHorizontalBox::GetEditorIcon()
-{
-	return FUMGStyle::Get().GetBrush("Widget.HorizontalBox");
-}
-
 const FText UHorizontalBox::GetPaletteCategory()
 {
 	return LOCTEXT("Panel", "Panel");
diff --git a/Engine/Source/Runtime/UMG/Private/Components/Image.cpp b/Engine/Source/Runtime/UMG/Private/Components/Image.cpp
--- a/Engine/Source/Runtime/UMG/Private/Components/Image.cpp
+++ b/Engine/Source/Runtime/UMG/Private/Components/Image.cpp
@@ -169,11 +169,6 @@ FReply UImage::HandleMouseButtonDown(const FGeometry& Geometry, const FPointerEv
 
 #if WITH_EDITOR
 
-const FSlateBrush* UImage::GetEditorIcon()
-{
-	return FUMGStyle::Get().GetBrush("Widget.Image");
-}
-
 const FText UImage::GetPaletteCategory()
 {
 	return LOCTEXT("Common", "Common");
diff --git a/Engine/Source/Runtime/UMG/Private/Components/InvalidationBox.cpp b/Engine/Source/Runtime/UMG/Private/Components/InvalidationBox.cpp
--- a/Engine/Source/Runtime/UMG/Private/Components/InvalidationBox.cpp
+++ b/Engine/Source/Runtime/UMG/Private/Components/InvalidationBox.cpp
@@ -86,11 +86,6 @@ void UInvalidationBox::SetCanCache(bool CanCache)
 
 #if WITH_EDITOR
 
-const FSlateBrush* UInvalidationBox::GetEditorIcon()
-{
-	return FUMGStyle::Get().GetBrush("Widget.MenuAnchor");
-}
-
 const FText UInvalidationBox::GetPaletteCategory()
 {
 	return LOCTEXT("Optimization", "Optimization");
diff --git a/Engine/Source/Runtime/UMG/Private/Components/ListView.cpp b/Engine/Source/Runtime/UMG/Private/Components/ListView.cpp
--- a/Engine/Source/Runtime/UMG/Private/Components/ListView.cpp
+++ b/Engine/Source/Runtime/UMG/Private/Components/ListView.cpp
@@ -61,11 +61,6 @@ TSharedRef<ITableRow> UListView::HandleOnGenerateRow(UObject* Item, const TShare
 
 #if WITH_EDITOR
 
-const FSlateBrush* UListView::GetEditorIcon()
-{
-	return FUMGStyle::Get().GetBrush("Widget.ListView");
-}
-
 const FText UListView::GetPaletteCategory()
 {
 	return LOCTEXT("Misc", "Misc");
diff --git a/Engine/Source/Runtime/UMG/Private/Components/MenuAnchor.cpp b/Engine/Source/Runtime/UMG/Private/Components/MenuAnchor.cpp
--- a/Engine/Source/Runtime/UMG/Private/Components/MenuAnchor.cpp
+++ b/Engine/Source/Runtime/UMG/Private/Components/MenuAnchor.cpp
@@ -155,11 +155,6 @@ bool UMenuAnchor::HasOpenSubMenus() const
 
 #if WITH_EDITOR
 
-const FSlateBrush* UMenuAnchor::GetEditorIcon()
-{
-	return FUMGStyle::Get().GetBrush("Widget.MenuAnchor");
-}
-
 const FText UMenuAnchor::GetPaletteCategory()
 {
 	return LOCTEXT("Primitive", "Primitive");
diff --git a/Engine/Source/Runtime/UMG/Private/Components/MultiLineEditableText.cpp b/Engine/Source/Runtime/UMG/Private/Components/MultiLineEditableText.cpp
--- a/Engine/Source/Runtime/UMG/Private/Components/MultiLineEditableText.cpp
+++ b/Engine/Source/Runtime/UMG/Private/Components/MultiLineEditableText.cpp
@@ -114,11 +114,6 @@ void UMultiLineEditableText::PostLoad()
 
 #if WITH_EDITOR
 
-const FSlateBrush* UMultiLineEditableText::GetEditorIcon()
-{
-	return FUMGStyle::Get().GetBrush("Widget.MultiLineEditableText");
-}
-
 const FText UMultiLineEditableText::GetPaletteCategory()
 {
 	return LOCTEXT("Primitive", "Primitive");
diff --git a/Engine/Source/Runtime/UMG/Private/Components/MultiLineEditableTextBox.cpp b/Engine/Source/Runtime/UMG/Private/Components/MultiLineEditableTextBox.cpp
--- a/Engine/Source/Runtime/UMG/Private/Components/MultiLineEditableTextBox.cpp
+++ b/Engine/Source/Runtime/UMG/Private/Components/MultiLineEditableTextBox.cpp
@@ -161,11 +161,6 @@ void UMultiLineEditableTextBox::PostLoad()
 
 #if WITH_EDITOR
 
-const FSlateBrush* UMultiLineEditableTextBox::GetEditorIcon()
-{
-	return FUMGStyle::Get().GetBrush("Widget.MultiLineEditableTextBox");
-}
-
 const FText UMultiLineEditableTextBox::GetPaletteCategory()
 {
 	return LOCTEXT("Input", "Input");
diff --git a/Engine/Source/Runtime/UMG/Private/Components/NamedSlot.cpp b/Engine/Source/Runtime/UMG/Private/Components/NamedSlot.cpp
--- a/Engine/Source/Runtime/UMG/Private/Components/NamedSlot.cpp
+++ b/Engine/Source/Runtime/UMG/Private/Components/NamedSlot.cpp
@@ -89,11 +89,6 @@ void UNamedSlot::OnSlotRemoved(UPanelSlot* Slot)
 
 #if WITH_EDITOR
 
-const FSlateBrush* UNamedSlot::GetEditorIcon()
-{
-	return FUMGStyle::Get().GetBrush("Widget.NamedSlot");
-}
-
 const FText UNamedSlot::GetPaletteCategory()
 {
 	return LOCTEXT("Common", "Common");
diff --git a/Engine/Source/Runtime/UMG/Private/Components/NativeWidgetHost.cpp b/Engine/Source/Runtime/UMG/Private/Components/NativeWidgetHost.cpp
--- a/Engine/Source/Runtime/UMG/Private/Components/NativeWidgetHost.cpp
+++ b/Engine/Source/Runtime/UMG/Private/Components/NativeWidgetHost.cpp
@@ -65,11 +65,6 @@ TSharedRef<SWidget> UNativeWidgetHost::GetDefaultContent()
 
 #if WITH_EDITOR
 
-const FSlateBrush* UNativeWidgetHost::GetEditorIcon()
-{
-	return FUMGStyle::Get().GetBrush("Widget.NativeWidgetHost");
-}
-
 const FText UNativeWidgetHost::GetPaletteCategory()
 {
 	return LOCTEXT("Primitive", "Primitive");
diff --git a/Engine/Source/Runtime/UMG/Private/Components/Overlay.cpp b/Engine/Source/Runtime/UMG/Private/Components/Overlay.cpp
--- a/Engine/Source/Runtime/UMG/Private/Components/Overlay.cpp
+++ b/Engine/Source/Runtime/UMG/Private/Components/Overlay.cpp
@@ -73,11 +73,6 @@ TSharedRef<SWidget> UOverlay::RebuildWidget()
 
 #if WITH_EDITOR
 
-const FSlateBrush* UOverlay::GetEditorIcon()
-{
-	return FUMGStyle::Get().GetBrush("Widget.Overlay");
-}
-
 const FText UOverlay::GetPaletteCategory()
 {
 	return LOCTEXT("Panel", "Panel");
diff --git a/Engine/Source/Runtime/UMG/Private/Components/ProgressBar.cpp b/Engine/Source/Runtime/UMG/Private/Components/ProgressBar.cpp
--- a/Engine/Source/Runtime/UMG/Private/Components/ProgressBar.cpp
+++ b/Engine/Source/Runtime/UMG/Private/Components/ProgressBar.cpp
@@ -115,11 +115,6 @@ void UProgressBar::PostLoad()
 
 #if WITH_EDITOR
 
-const FSlateBrush* UProgressBar::GetEditorIcon()
-{
-	return FUMGStyle::Get().GetBrush("Widget.ProgressBar");
-}
-
 const FText UProgressBar::GetPaletteCategory()
 {
 	return LOCTEXT("Common", "Common");
diff --git a/Engine/Source/Runtime/UMG/Private/Components/RetainerBox.cpp b/Engine/Source/Runtime/UMG/Private/Components/RetainerBox.cpp
--- a/Engine/Source/Runtime/UMG/Private/Components/RetainerBox.cpp
+++ b/Engine/Source/Runtime/UMG/Private/Components/RetainerBox.cpp
@@ -105,11 +105,6 @@ void URetainerBox::OnSlotRemoved(UPanelSlot* Slot)
 
 #if WITH_EDITOR
 
-const FSlateBrush* URetainerBox::GetEditorIcon()
-{
-	return FUMGStyle::Get().GetBrush("Widget.MenuAnchor");
-}
-
 const FText URetainerBox::GetPaletteCategory()
 {
 	return LOCTEXT("Optimization", "Optimization");
diff --git a/Engine/Source/Runtime/UMG/Private/Components/RichTextBlock.cpp b/Engine/Source/Runtime/UMG/Private/Components/RichTextBlock.cpp
--- a/Engine/Source/Runtime/UMG/Private/Components/RichTextBlock.cpp
+++ b/Engine/Source/Runtime/UMG/Private/Components/RichTextBlock.cpp
@@ -70,11 +70,6 @@ void URichTextBlock::SynchronizeProperties()
 
 #if WITH_EDITOR
 
-const FSlateBrush* URichTextBlock::GetEditorIcon()
-{
-	return FUMGStyle::Get().GetBrush("Widget.RichTextBlock");
-}
-
 const FText URichTextBlock::GetPaletteCategory()
 {
 	return LOCTEXT("Common", "Common");
diff --git a/Engine/Source/Runtime/UMG/Private/Components/SafeZone.cpp b/Engine/Source/Runtime/UMG/Private/Components/SafeZone.cpp
--- a/Engine/Source/Runtime/UMG/Private/Components/SafeZone.cpp
+++ b/Engine/Source/Runtime/UMG/Private/Components/SafeZone.cpp
@@ -15,10 +15,6 @@ USafeZone::USafeZone()
 }
 
 #if WITH_EDITOR
-const FSlateBrush* USafeZone::GetEditorIcon()
-{
-	return FUMGStyle::Get().GetBrush( "Widget.SafeZone" );
-}
 
 const FText USafeZone::GetPaletteCategory()
 {
diff --git a/Engine/Source/Runtime/UMG/Private/Components/ScaleBox.cpp b/Engine/Source/Runtime/UMG/Private/Components/ScaleBox.cpp
--- a/Engine/Source/Runtime/UMG/Private/Components/ScaleBox.cpp
+++ b/Engine/Source/Runtime/UMG/Private/Components/ScaleBox.cpp
@@ -73,11 +73,6 @@ void UScaleBox::OnSlotRemoved(UPanelSlot* Slot)
 
 #if WITH_EDITOR
 
-const FSlateBrush* UScaleBox::GetEditorIcon()
-{
-	return FUMGStyle::Get().GetBrush("Widget.ScaleBox");
-}
-
 const FText UScaleBox::GetPaletteCategory()
 {
 	return LOCTEXT("Panel", "Panel");
diff --git a/Engine/Source/Runtime/UMG/Private/Components/ScrollBar.cpp b/Engine/Source/Runtime/UMG/Private/Components/ScrollBar.cpp
--- a/Engine/Source/Runtime/UMG/Private/Components/ScrollBar.cpp
+++ b/Engine/Source/Runtime/UMG/Private/Components/ScrollBar.cpp
@@ -76,11 +76,6 @@ void UScrollBar::PostLoad()
 
 #if WITH_EDITOR
 
-const FSlateBrush* UScrollBar::GetEditorIcon()
-{
-	return FUMGStyle::Get().GetBrush("Widget.ScrollBar");
-}
-
 const FText UScrollBar::GetPaletteCategory()
 {
 	return LOCTEXT("Primitive", "Primitive");
diff --git a/Engine/Source/Runtime/UMG/Private/Components/ScrollBox.cpp b/Engine/Source/Runtime/UMG/Private/Components/ScrollBox.cpp
--- a/Engine/Source/Runtime/UMG/Private/Components/ScrollBox.cpp
+++ b/Engine/Source/Runtime/UMG/Private/Components/ScrollBox.cpp
@@ -174,11 +174,6 @@ void UScrollBox::PostLoad()
 
 #if WITH_EDITOR
 
-const FSlateBrush* UScrollBox::GetEditorIcon()
-{
-	return FUMGStyle::Get().GetBrush("Widget.ScrollBox");
-}
-
 const FText UScrollBox::GetPaletteCategory()
 {
 	return LOCTEXT("Panel", "Panel");
diff --git a/Engine/Source/Runtime/UMG/Private/Components/SizeBox.cpp b/Engine/Source/Runtime/UMG/Private/Components/SizeBox.cpp
--- a/Engine/Source/Runtime/UMG/Private/Components/SizeBox.cpp
+++ b/Engine/Source/Runtime/UMG/Private/Components/SizeBox.cpp
@@ -232,11 +232,6 @@ void USizeBox::OnSlotRemoved(UPanelSlot* Slot)
 
 #if WITH_EDITOR
 
-const FSlateBrush* USizeBox::GetEditorIcon()
-{
-	return FUMGStyle::Get().GetBrush("Widget.SizeBox");
-}
-
 const FText USizeBox::GetPaletteCategory()
 {
 	return LOCTEXT("Panel", "Panel");
diff --git a/Engine/Source/Runtime/UMG/Private/Components/Slider.cpp b/Engine/Source/Runtime/UMG/Private/Components/Slider.cpp
--- a/Engine/Source/Runtime/UMG/Private/Components/Slider.cpp
+++ b/Engine/Source/Runtime/UMG/Private/Components/Slider.cpp
@@ -145,11 +145,6 @@ void USlider::SetSliderBarColor(FLinearColor InValue)
 
 #if WITH_EDITOR
 
-const FSlateBrush* USlider::GetEditorIcon()
-{
-	return FUMGStyle::Get().GetBrush("Widget.Slider");
-}
-
 const FText USlider::GetPaletteCategory()
 {
 	return LOCTEXT("Common", "Common");
diff --git a/Engine/Source/Runtime/UMG/Private/Components/Spacer.cpp b/Engine/Source/Runtime/UMG/Private/Components/Spacer.cpp
--- a/Engine/Source/Runtime/UMG/Private/Components/Spacer.cpp
+++ b/Engine/Source/Runtime/UMG/Private/Components/Spacer.cpp
@@ -51,11 +51,6 @@ void USpacer::SynchronizeProperties()
 
 #if WITH_EDITOR
 
-const FSlateBrush* USpacer::GetEditorIcon()
-{
-	return FUMGStyle::Get().GetBrush("Widget.Spacer");
-}
-
 const FText USpacer::GetPaletteCategory()
 {
 	return LOCTEXT("Primitive", "Primitive");
diff --git a/Engine/Source/Runtime/UMG/Private/Components/SpinBox.cpp b/Engine/Source/Runtime/UMG/Private/Components/SpinBox.cpp
--- a/Engine/Source/Runtime/UMG/Private/Components/SpinBox.cpp
+++ b/Engine/Source/Runtime/UMG/Private/Components/SpinBox.cpp
@@ -304,11 +304,6 @@ void USpinBox::PostLoad()
 
 #if WITH_EDITOR
 
-const FSlateBrush* USpinBox::GetEditorIcon()
-{
-	return FUMGStyle::Get().GetBrush("Widget.SpinBox");
-}
-
 const FText USpinBox::GetPaletteCategory()
 {
 	return LOCTEXT("Input", "Input");
diff --git a/Engine/Source/Runtime/UMG/Private/Components/TextBlock.cpp b/Engine/Source/Runtime/UMG/Private/Components/TextBlock.cpp
--- a/Engine/Source/Runtime/UMG/Private/Components/TextBlock.cpp
+++ b/Engine/Source/Runtime/UMG/Private/Components/TextBlock.cpp
@@ -187,11 +187,6 @@ void UTextBlock::HandleTextCommitted(const FText& InText, ETextCommit::Type Comm
 	//     Need a way to recognize one particular widget and forward things to them!
 }
 
-const FSlateBrush* UTextBlock::GetEditorIcon()
-{
-	return FUMGStyle::Get().GetBrush("Widget.TextBlock");
-}
-
 const FText UTextBlock::GetPaletteCategory()
 {
 	return LOCTEXT("Common", "Common");
diff --git a/Engine/Source/Runtime/UMG/Private/Components/Throbber.cpp b/Engine/Source/Runtime/UMG/Private/Components/Throbber.cpp
--- a/Engine/Source/Runtime/UMG/Private/Components/Throbber.cpp
+++ b/Engine/Source/Runtime/UMG/Private/Components/Throbber.cpp
@@ -109,11 +109,6 @@ void UThrobber::PostLoad()
 
 #if WITH_EDITOR
 
-const FSlateBrush* UThrobber::GetEditorIcon()
-{
-	return FUMGStyle::Get().GetBrush("Widget.Throbber");
-}
-
 const FText UThrobber::GetPaletteCategory()
 {
 	return LOCTEXT("Primitive", "Primitive");
diff --git a/Engine/Source/Runtime/UMG/Private/Components/TileView.cpp b/Engine/Source/Runtime/UMG/Private/Components/TileView.cpp
--- a/Engine/Source/Runtime/UMG/Private/Components/TileView.cpp
+++ b/Engine/Source/Runtime/UMG/Private/Components/TileView.cpp
@@ -85,11 +85,6 @@ void UTileView::ReleaseSlateResources(bool bReleaseChildren)
 
 #if WITH_EDITOR
 
-const FSlateBrush* UTileView::GetEditorIcon()
-{
-	return FUMGStyle::Get().GetBrush("Widget.TileView");
-}
-
 const FText UTileView::GetPaletteCategory()
 {
 	return LOCTEXT("Misc", "Misc");
diff --git a/Engine/Source/Runtime/UMG/Private/Components/UniformGridPanel.cpp b/Engine/Source/Runtime/UMG/Private/Components/UniformGridPanel.cpp
--- a/Engine/Source/Runtime/UMG/Private/Components/UniformGridPanel.cpp
+++ b/Engine/Source/Runtime/UMG/Private/Components/UniformGridPanel.cpp
@@ -109,11 +109,6 @@ void UUniformGridPanel::SynchronizeProperties()
 
 #if WITH_EDITOR
 
-const FSlateBrush* UUniformGridPanel::GetEditorIcon()
-{
-	return FUMGStyle::Get().GetBrush("Widget.UniformGrid");
-}
-
 const FText UUniformGridPanel::GetPaletteCategory()
 {
 	return LOCTEXT("Panel", "Panel");
diff --git a/Engine/Source/Runtime/UMG/Private/Components/VerticalBox.cpp b/Engine/Source/Runtime/UMG/Private/Components/VerticalBox.cpp
--- a/Engine/Source/Runtime/UMG/Private/Components/VerticalBox.cpp
+++ b/Engine/Source/Runtime/UMG/Private/Components/VerticalBox.cpp
@@ -73,11 +73,6 @@ TSharedRef<SWidget> UVerticalBox::RebuildWidget()
 
 #if WITH_EDITOR
 
-const FSlateBrush* UVerticalBox::GetEditorIcon()
-{
-	return FUMGStyle::Get().GetBrush("Widget.VerticalBox");
-}
-
 const FText UVerticalBox::GetPaletteCategory()
 {
 	return LOCTEXT("Panel", "Panel");
diff --git a/Engine/Source/Runtime/UMG/Private/Components/Viewport.cpp b/Engine/Source/Runtime/UMG/Private/Components/Viewport.cpp
--- a/Engine/Source/Runtime/UMG/Private/Components/Viewport.cpp
+++ b/Engine/Source/Runtime/UMG/Private/Components/Viewport.cpp
@@ -483,11 +483,6 @@ AActor* UViewport::Spawn(TSubclassOf<AActor> ActorClass)
 
 #if WITH_EDITOR
 
-const FSlateBrush* UViewport::GetEditorIcon()
-{
-	return FUMGStyle::Get().GetBrush("Widget.Viewport");
-}
-
 const FText UViewport::GetPaletteCategory()
 {
 	return LOCTEXT("Primitive", "Primitive");
diff --git a/Engine/Source/Runtime/UMG/Private/Components/Widget.cpp b/Engine/Source/Runtime/UMG/Private/Components/Widget.cpp
--- a/Engine/Source/Runtime/UMG/Private/Components/Widget.cpp
+++ b/Engine/Source/Runtime/UMG/Private/Components/Widget.cpp
@@ -678,7 +678,7 @@ const FText UWidget::GetPaletteCategory()
 
 const FSlateBrush* UWidget::GetEditorIcon()
 {
-	return FUMGStyle::Get().GetBrush("Widget");
+	return nullptr;
 }
 
 EVisibility UWidget::GetVisibilityInDesigner() const
diff --git a/Engine/Source/Runtime/UMG/Private/Components/WidgetSwitcher.cpp b/Engine/Source/Runtime/UMG/Private/Components/WidgetSwitcher.cpp
--- a/Engine/Source/Runtime/UMG/Private/Components/WidgetSwitcher.cpp
+++ b/Engine/Source/Runtime/UMG/Private/Components/WidgetSwitcher.cpp
@@ -128,11 +128,6 @@ void UWidgetSwitcher::SynchronizeProperties()
 
 #if WITH_EDITOR
 
-const FSlateBrush* UWidgetSwitcher::GetEditorIcon()
-{
-	return FUMGStyle::Get().GetBrush("Widget.WidgetSwitcher");
-}
-
 const FText UWidgetSwitcher::GetPaletteCategory()
 {
 	return LOCTEXT("Panel", "Panel");
diff --git a/Engine/Source/Runtime/UMG/Private/Components/WrapBox.cpp b/Engine/Source/Runtime/UMG/Private/Components/WrapBox.cpp
--- a/Engine/Source/Runtime/UMG/Private/Components/WrapBox.cpp
+++ b/Engine/Source/Runtime/UMG/Private/Components/WrapBox.cpp
@@ -91,11 +91,6 @@ void UWrapBox::SetInnerSlotPadding(FVector2D InPadding)
 
 #if WITH_EDITOR
 
-const FSlateBrush* UWrapBox::GetEditorIcon()
-{
-	return FUMGStyle::Get().GetBrush("Widget.WrapBox");
-}
-
 const FText UWrapBox::GetPaletteCategory()
 {
 	return LOCTEXT("Panel", "Panel");
