Copying //UE4/Orion-Staging to //UE4/Main (Origin //Orion/Dev-General @ 2861092)

diff --git a/Engine/Source/Runtime/RHI/Private/RHICommandList.cpp b/Engine/Source/Runtime/RHI/Private/RHICommandList.cpp
--- a/Engine/Source/Runtime/RHI/Private/RHICommandList.cpp
+++ b/Engine/Source/Runtime/RHI/Private/RHICommandList.cpp
@@ -181,23 +181,26 @@ void FRHIAsyncComputeCommandListImmediate::ImmediateDispatch(FRHIAsyncComputeCom
 	//queue a final command to submit all the async compute commands up to this point to the GPU.
 	RHIComputeCmdList.SubmitCommandsHint();
 
-	FRHIAsyncComputeCommandListImmediate* SwapCmdList;
+	if (!RHIComputeCmdList.Bypass())
 	{
-		QUICK_SCOPE_CYCLE_COUNTER(STAT_FRHICommandListExecutor_SwapCmdLists);
-		SwapCmdList = new FRHIAsyncComputeCommandListImmediate();
+		FRHIAsyncComputeCommandListImmediate* SwapCmdList;
+		{
+			QUICK_SCOPE_CYCLE_COUNTER(STAT_FRHICommandListExecutor_SwapCmdLists);
+			SwapCmdList = new FRHIAsyncComputeCommandListImmediate();
 
-		//hack stolen from Gfx commandlist.  transfer
-		static_assert(sizeof(FRHICommandList) == sizeof(FRHIAsyncComputeCommandListImmediate), "We are memswapping FRHICommandList and FRHICommandListImmediate; they need to be swappable.");
-		check(RHIComputeCmdList.IsImmediateAsyncCompute());
-		SwapCmdList->ExchangeCmdList(RHIComputeCmdList);
+			//hack stolen from Gfx commandlist.  transfer
+			static_assert(sizeof(FRHICommandList) == sizeof(FRHIAsyncComputeCommandListImmediate), "We are memswapping FRHICommandList and FRHICommandListImmediate; they need to be swappable.");
+			check(RHIComputeCmdList.IsImmediateAsyncCompute());
+			SwapCmdList->ExchangeCmdList(RHIComputeCmdList);
 
-		//queue the execution of this async commandlist amongst other commands in the immediate gfx list.
-		//this guarantees resource update commands made on the gfx commandlist will be executed before the async compute.
-		FRHICommandListImmediate& RHIImmCmdList = FRHICommandListExecutor::GetImmediateCommandList();		
-		RHIImmCmdList.QueueAsyncCompute(*SwapCmdList);
+			//queue the execution of this async commandlist amongst other commands in the immediate gfx list.
+			//this guarantees resource update commands made on the gfx commandlist will be executed before the async compute.
+			FRHICommandListImmediate& RHIImmCmdList = FRHICommandListExecutor::GetImmediateCommandList();
+			RHIImmCmdList.QueueAsyncCompute(*SwapCmdList);
 
-		//dispatch immediately to RHI Thread so we can get the async compute on the GPU ASAP.
-		RHIImmCmdList.ImmediateFlush(EImmediateFlushType::DispatchToRHIThread);
+			//dispatch immediately to RHI Thread so we can get the async compute on the GPU ASAP.
+			RHIImmCmdList.ImmediateFlush(EImmediateFlushType::DispatchToRHIThread);
+		}
 	}
 }
 
diff --git a/Engine/Source/Runtime/RHI/Public/GPUDefragAllocator.h b/Engine/Source/Runtime/RHI/Public/GPUDefragAllocator.h
--- a/Engine/Source/Runtime/RHI/Public/GPUDefragAllocator.h
+++ b/Engine/Source/Runtime/RHI/Public/GPUDefragAllocator.h
@@ -7,7 +7,7 @@
 #define LOG_EVERY_ALLOCATION			0
 #define DUMP_ALLOC_FREQUENCY			0 // 100
 #define VALIDATE_SYNC_SIZE				(!(UE_BUILD_TEST || UE_BUILD_SHIPPING))
-
+#define VALIDATE_MEMORY_PROTECTION		1 //(!(UE_BUILD_TEST || UE_BUILD_SHIPPING))
 
 /*-----------------------------------------------------------------------------
 Custom fixed size pool best fit texture memory allocator
@@ -677,6 +677,30 @@ protected:
 	TArray<FRelocationEntry> Relocations;
 #endif
 
+#if VALIDATE_MEMORY_PROTECTION
+	struct FMemProtectTracker
+	{
+		FMemProtectTracker(const void* InMemory, uint64 InBlockSize, uint32 InSyncIndex)
+		: Memory(InMemory)
+		, BlockSize(InBlockSize)
+		, SyncIndex(InSyncIndex)
+		{
+
+		}
+		const void* Memory;
+		uint64 BlockSize;
+		uint32 SyncIndex;
+	};
+
+	TArray<FMemProtectTracker> BlocksToProtect;
+	TArray<FMemProtectTracker> BlocksToUnProtect;
+
+	void AllowCPUAccessToBlocks();
+	virtual void	PlatformAllowCPUAccessToBlock(const FMemProtectTracker& BlocksToAllow) {};
+	virtual void	PlatformRemoveCPUAccessTo(const TArray<FMemProtectTracker>& BlocksToRemove) {};
+#endif
+
+
 	/**
 	* Copy memory from one location to another. If it returns false, the defragmentation
 	* process will assume the memory is not relocatable and keep it in place.
@@ -689,6 +713,8 @@ protected:
 	*/
 	virtual void	PlatformRelocate(void* Dest, const void* Source, int64 Size, void* UserPayload) = 0;
 
+	
+
 	/**
 	* Inserts a fence to synchronize relocations.
 	* The fence can be blocked on at a later time to ensure that all relocations initiated
@@ -744,6 +770,10 @@ protected:
 
 		if (!bBenchmarkMode)
 		{
+#if VALIDATE_MEMORY_PROTECTION
+			BlocksToProtect.Emplace(DestAddr, Size, CurrentSyncIndex);
+			BlocksToProtect.Emplace(Source, Size, CurrentSyncIndex);
+#endif
 			PlatformRelocate(DestAddr, Source, Size, UserPayload);
 		}
 		int64 RelocateSize = bOverlappedMove ? (Size * Settings.MaxDefragRelocations) : Size;
@@ -1055,6 +1085,10 @@ protected:
 
 	int32 NumLockedChunks;
 
+#if VALIDATE_MEMORY_PROTECTION
+	double TimeInMemProtect;
+#endif
+
 	/** When in benchmark mode, don't call any Platform functions.	*/
 	bool			bBenchmarkMode;
 
diff --git a/Engine/Source/Runtime/Renderer/Private/CapsuleShadowRendering.cpp b/Engine/Source/Runtime/Renderer/Private/CapsuleShadowRendering.cpp
--- a/Engine/Source/Runtime/Renderer/Private/CapsuleShadowRendering.cpp
+++ b/Engine/Source/Runtime/Renderer/Private/CapsuleShadowRendering.cpp
@@ -605,8 +605,21 @@ bool FDeferredShadingSceneRenderer::RenderCapsuleDirectShadows(
 	FRHICommandListImmediate& RHICmdList, 
 	const TArray<FProjectedShadowInfo*, SceneRenderingAllocator>& CapsuleShadows) const
 {
+	bool bAllViewsHaveViewState = true;
+
+	for (int32 ViewIndex = 0; ViewIndex < Views.Num(); ViewIndex++)
+	{
+		const FViewInfo& View = Views[ViewIndex];
+
+		if (!View.ViewState)
+		{
+			bAllViewsHaveViewState = false;
+		}
+	}
+
 	if (SupportsCapsuleShadows(FeatureLevel, GShaderPlatformForFeatureLevel[FeatureLevel])
-		&& CapsuleShadows.Num() > 0)
+		&& CapsuleShadows.Num() > 0
+		&& bAllViewsHaveViewState)
 	{
 		QUICK_SCOPE_CYCLE_COUNTER(STAT_RenderCapsuleShadows);
 
diff --git a/Engine/Source/Runtime/Slate/Private/Framework/Application/AnalogCursor.cpp b/Engine/Source/Runtime/Slate/Private/Framework/Application/AnalogCursor.cpp
--- a/Engine/Source/Runtime/Slate/Private/Framework/Application/AnalogCursor.cpp
+++ b/Engine/Source/Runtime/Slate/Private/Framework/Application/AnalogCursor.cpp
@@ -44,7 +44,7 @@ void FAnalogCursor::Tick(const float DeltaTime, FSlateApplication& SlateApp, TSh
 		for(const FArrangedWidget& ArrangedWidget : AllArrangedWidgets)
 		{
 			TSharedRef<SWidget> Widget = ArrangedWidget.Widget;
-			if (Widget->SupportsKeyboardFocus() && Widget->IsInteractable())
+			if (Widget->IsInteractable())
 			{
 				SpeedMult = StickySlowdown;
 				//FVector2D Adjustment = WidgetsAndCursors.Last().Geometry.Position - OldPosition; // example of calculating distance from cursor to widget center
diff --git a/Engine/Source/Runtime/Slate/Private/Framework/Text/RichTextLayoutMarshaller.cpp b/Engine/Source/Runtime/Slate/Private/Framework/Text/RichTextLayoutMarshaller.cpp
--- a/Engine/Source/Runtime/Slate/Private/Framework/Text/RichTextLayoutMarshaller.cpp
+++ b/Engine/Source/Runtime/Slate/Private/Framework/Text/RichTextLayoutMarshaller.cpp
@@ -9,7 +9,7 @@
 
 TSharedRef< FRichTextLayoutMarshaller > FRichTextLayoutMarshaller::Create(TArray< TSharedRef< ITextDecorator > > InDecorators, const ISlateStyle* const InDecoratorStyleSet)
 {
-	return MakeShareable(new FRichTextLayoutMarshaller(FDefaultRichTextMarkupParser::Create(), FDefaultRichTextMarkupWriter::Create(), MoveTemp(InDecorators), InDecoratorStyleSet));
+	return MakeShareable(new FRichTextLayoutMarshaller(MoveTemp(InDecorators), InDecoratorStyleSet));
 }
 
 TSharedRef< FRichTextLayoutMarshaller > FRichTextLayoutMarshaller::Create(TSharedPtr< IRichTextMarkupParser > InParser, TSharedPtr< IRichTextMarkupWriter > InWriter, TArray< TSharedRef< ITextDecorator > > InDecorators, const ISlateStyle* const InDecoratorStyleSet)
@@ -40,43 +40,7 @@ void FRichTextLayoutMarshaller::SetText(const FString& SourceString, FTextLayout
 
 		for (const FTextRunParseResults& RunParseResult : LineParseResults.Runs)
 		{
-			TSharedPtr< ISlateRun > Run;
-
-			TSharedPtr< ITextDecorator > Decorator = TryGetDecorator(ProcessedString, RunParseResult);
-
-			if (Decorator.IsValid())
-			{
-				// Create run and update model string.
-				Run = Decorator->Create(TargetTextLayout.AsShared(), RunParseResult, ProcessedString, ModelString, DecoratorStyleSet);
-			}
-			else
-			{
-				FRunInfo RunInfo(RunParseResult.Name);
-				for (const TPair<FString, FTextRange>& Pair : RunParseResult.MetaData)
-				{
-					RunInfo.MetaData.Add(Pair.Key, ProcessedString.Mid( Pair.Value.BeginIndex, Pair.Value.EndIndex - Pair.Value.BeginIndex));
-				}
-
-				const FTextBlockStyle* TextBlockStyle;
-				FTextRange ModelRange;
-				ModelRange.BeginIndex = ModelString->Len();
-				if (!(RunParseResult.Name.IsEmpty()) && DecoratorStyleSet->HasWidgetStyle< FTextBlockStyle >(FName(*RunParseResult.Name)))
-				{
-					*ModelString += ProcessedString.Mid(RunParseResult.ContentRange.BeginIndex, RunParseResult.ContentRange.EndIndex - RunParseResult.ContentRange.BeginIndex);
-					TextBlockStyle = &(DecoratorStyleSet->GetWidgetStyle< FTextBlockStyle >( FName(*RunParseResult.Name) ));
-				}
-				else
-				{
-					*ModelString += ProcessedString.Mid(RunParseResult.OriginalRange.BeginIndex, RunParseResult.OriginalRange.EndIndex - RunParseResult.OriginalRange.BeginIndex);
-					TextBlockStyle = &DefaultTextStyle;
-				}
-				ModelRange.EndIndex = ModelString->Len();
-
-				// Create run.
-				Run = FSlateTextRun::Create(RunInfo, ModelString, *TextBlockStyle, ModelRange);
-			}
-
-			Runs.Add( Run.ToSharedRef() );
+			AppendRunsForText(RunParseResult, ProcessedString, DefaultTextStyle, ModelString, TargetTextLayout, Runs);
 		}
 
 		LinesToAdd.Emplace(MoveTemp(ModelString), MoveTemp(Runs));
@@ -110,6 +74,15 @@ void FRichTextLayoutMarshaller::GetText(FString& TargetString, const FTextLayout
 	Writer->Write(WriterLines, TargetString);
 }
 
+FRichTextLayoutMarshaller::FRichTextLayoutMarshaller(TArray< TSharedRef< ITextDecorator > > InDecorators, const ISlateStyle* const InDecoratorStyleSet)
+	: Parser(MoveTemp(FDefaultRichTextMarkupParser::Create()))
+	, Writer(MoveTemp(FDefaultRichTextMarkupWriter::Create()))
+	, Decorators(MoveTemp(InDecorators))
+	, InlineDecorators()
+	, DecoratorStyleSet(InDecoratorStyleSet)
+{
+}
+
 FRichTextLayoutMarshaller::FRichTextLayoutMarshaller(TSharedPtr< IRichTextMarkupParser > InParser, TSharedPtr< IRichTextMarkupWriter > InWriter, TArray< TSharedRef< ITextDecorator > > InDecorators, const ISlateStyle* const InDecoratorStyleSet)
 	: Parser(MoveTemp(InParser))
 	, Writer(MoveTemp(InWriter))
@@ -140,4 +113,49 @@ TSharedPtr< ITextDecorator > FRichTextLayoutMarshaller::TryGetDecorator(const FS
 	return nullptr;
 }
 
+void FRichTextLayoutMarshaller::AppendRunsForText(const FTextRunParseResults& TextRun,
+	const FString& ProcessedString,
+	const FTextBlockStyle& DefaultTextStyle,
+	const TSharedRef<FString>& InOutModelText,
+	FTextLayout& TargetTextLayout,
+	TArray<TSharedRef<IRun>>& Runs)
+{
+	TSharedPtr< ISlateRun > Run;
+	TSharedPtr< ITextDecorator > Decorator = TryGetDecorator(ProcessedString, TextRun);
+
+	if (Decorator.IsValid())
+	{
+		// Create run and update model string.
+		Run = Decorator->Create(TargetTextLayout.AsShared(), TextRun, ProcessedString, InOutModelText, DecoratorStyleSet);
+	}
+	else
+	{
+		FRunInfo RunInfo(TextRun.Name);
+		for (const TPair<FString, FTextRange>& Pair : TextRun.MetaData)
+		{
+			RunInfo.MetaData.Add(Pair.Key, ProcessedString.Mid(Pair.Value.BeginIndex, Pair.Value.EndIndex - Pair.Value.BeginIndex));
+		}
+
+		const FTextBlockStyle* TextBlockStyle;
+		FTextRange ModelRange;
+		ModelRange.BeginIndex = InOutModelText->Len();
+		if (!(TextRun.Name.IsEmpty()) && DecoratorStyleSet->HasWidgetStyle< FTextBlockStyle >(FName(*TextRun.Name)))
+		{
+			*InOutModelText += ProcessedString.Mid(TextRun.ContentRange.BeginIndex, TextRun.ContentRange.EndIndex - TextRun.ContentRange.BeginIndex);
+			TextBlockStyle = &(DecoratorStyleSet->GetWidgetStyle< FTextBlockStyle >(FName(*TextRun.Name)));
+		}
+		else
+		{
+			*InOutModelText += ProcessedString.Mid(TextRun.OriginalRange.BeginIndex, TextRun.OriginalRange.EndIndex - TextRun.OriginalRange.BeginIndex);
+			TextBlockStyle = &DefaultTextStyle;
+		}
+		ModelRange.EndIndex = InOutModelText->Len();
+
+		// Create run.
+		Run = FSlateTextRun::Create(RunInfo, InOutModelText, *TextBlockStyle, ModelRange);
+	}
+
+	Runs.Add(Run.ToSharedRef());
+}
+
 #endif //WITH_FANCY_TEXT
diff --git a/Engine/Source/Runtime/Slate/Private/Framework/Text/TextLayout.cpp b/Engine/Source/Runtime/Slate/Private/Framework/Text/TextLayout.cpp
--- a/Engine/Source/Runtime/Slate/Private/Framework/Text/TextLayout.cpp
+++ b/Engine/Source/Runtime/Slate/Private/Framework/Text/TextLayout.cpp
@@ -125,8 +125,36 @@ FTextLayout::FBreakCandidate FTextLayout::CreateBreakCandidate( int32& OutRunInd
 			break;
 		}
 	}
+	
+#if DO_CHECK
 	if (!SuccessfullyMeasuredSlice)
 	{
+		FString AnonymizedText;
+		AnonymizedText.Reserve(Line.Text->Len());
+		for (const TCHAR Char : *Line.Text)
+		{
+			if (Char == TCHAR('\\'))
+			{
+				AnonymizedText += TEXT("\\\\");
+			}
+			else if (FChar::IsWhitespace(Char) || FChar::IsPunct(Char))
+			{
+				AnonymizedText += Char;
+			}
+			else if (FChar::IsDigit(Char))
+			{
+				AnonymizedText += TEXT("0");
+			}
+			else if (FChar::IsLower(Char))
+			{
+				AnonymizedText += TEXT("a");
+			}
+			else
+			{
+				AnonymizedText += TEXT("A");
+			}
+		}
+
 		FString RunDebugData;
 		for (int32 RunIndex = 0; RunIndex < Line.Runs.Num(); ++RunIndex)
 		{
@@ -135,8 +163,9 @@ FTextLayout::FBreakCandidate FTextLayout::CreateBreakCandidate( int32& OutRunInd
 			RunDebugData.Append(FString::Printf(TEXT("\t\t[%d] - Range: {%d, %d}\n"), RunIndex, RunRange.BeginIndex, RunRange.EndIndex));
 		}
 
-		checkf( SuccessfullyMeasuredSlice, TEXT("Failed to measure a slice of text!\n\tStart Index: %d\n\tEnd Index: %d\n\tStart Run Index: %d\n\tLine Runs:\n%s"), PreviousBreak, CurrentBreak, FirstRunIndexChecked, *RunDebugData );
+		checkf(SuccessfullyMeasuredSlice, TEXT("Failed to measure a slice of text!\n\tDebug Source: %s\n\tAnonymized Text: %s\n\tStart Index: %d\n\tEnd Index: %d\n\tStart Run Index: %d\n\tLine Runs:\n%s"), *DebugSourceInfo.Get(FString()), *AnonymizedText, PreviousBreak, CurrentBreak, FirstRunIndexChecked, *RunDebugData);
 	}
+#endif // DO_CHECK
 
 	BreakSize.Y = BreakSizeWithoutTrailingWhitespace.Y = MaxAboveBaseline + MaxBelowBaseline;
 
@@ -222,7 +251,7 @@ void FTextLayout::CreateLineViewBlocks( int32 LineModelIndex, const int32 StopIn
 			const TextBiDi::FTextDirectionInfo& CurrentTextDirectionInfo = SortedTextDirectionInfos[CurrentSortedTextDirectionInfoIndex];
 			CurrentTextDirectionStopIndex = CurrentTextDirectionInfo.StartIndex + CurrentTextDirectionInfo.Length;
 
-			check(BlockBeginIndex >= CurrentTextDirectionInfo.StartIndex);
+			checkf(BlockBeginIndex >= CurrentTextDirectionInfo.StartIndex, TEXT("Debug Source: %s"), *DebugSourceInfo.Get(FString()));
 
 			BlockStopIndex = FMath::Min(BlockStopIndex, CurrentTextDirectionStopIndex);
 			BlockTextDirection = CurrentTextDirectionInfo.TextDirection;
@@ -278,7 +307,7 @@ void FTextLayout::CreateLineViewBlocks( int32 LineModelIndex, const int32 StopIn
 
 		const bool IsLastBlock = BlockStopIndex == StopIndex;
 
-		check( BlockBeginIndex <= BlockStopIndex );
+		checkf(BlockBeginIndex <= BlockStopIndex, TEXT("Debug Source: %s"), *DebugSourceInfo.Get(FString()));
 
 		// Add the new block
 		{
@@ -343,7 +372,7 @@ void FTextLayout::CreateLineViewBlocks( int32 LineModelIndex, const int32 StopIn
 						return !InLineBlock->GetTextRange().IsEmpty() && InLineBlock->GetTextRange().BeginIndex == CurrentBlockStartIndex;
 					});
 
-					check(FoundLineBlock);
+					checkf(FoundLineBlock, TEXT("Debug Source: %s"), *DebugSourceInfo.Get(FString()));
 
 					const TSharedRef<ILayoutBlock>& FoundLineBlockRef = *FoundLineBlock;
 					if (VisualTextDirectionInfo.TextDirection == TextBiDi::ETextDirection::LeftToRight)
@@ -449,7 +478,7 @@ void FTextLayout::JustifyLayout()
 
 float FTextLayout::GetWrappingDrawWidth() const
 {
-	check( WrappingWidth >= 0 );
+	checkf(WrappingWidth >= 0, TEXT("Debug Source: %s"), *DebugSourceInfo.Get(FString()));
 	return FMath::Max( 0.01f, ( WrappingWidth - Margin.GetTotalSpaceAlong<Orient_Horizontal>() ) * Scale );
 }
 
@@ -512,7 +541,7 @@ void FTextLayout::FlowLineLayout(const int32 LineModelIndex, const float Wrappin
 	{
 		//Then iterate over all of it's runs
 		CreateLineViewBlocks( LineModelIndex, INDEX_NONE, 0.0f, /*OUT*/CurrentRunIndex, /*OUT*/CurrentRendererIndex, /*OUT*/PreviousBlockEnd, SoftLine );
-		check( CurrentRunIndex == LineModel.Runs.Num() );
+		checkf(CurrentRunIndex == LineModel.Runs.Num(), TEXT("Debug Source: %s"), *DebugSourceInfo.Get(FString()));
 		CurrentWidth = 0;
 		SoftLine.Reset();
 	}
@@ -584,7 +613,7 @@ void FTextLayout::FlowHighlights()
 	};
 
 	// FlowLayout must have been called first
-	check(!(DirtyFlags & ETextLayoutDirtyState::Layout));
+	checkf(!(DirtyFlags & ETextLayoutDirtyState::Layout), TEXT("Debug Source: %s"), *DebugSourceInfo.Get(FString()));
 
 	for (FLineView& LineView : LineViews)
 	{
@@ -925,7 +954,7 @@ void FTextLayout::DirtyAllLineModels(const ELineModelDirtyState::Flags InDirtyFl
 	}
 }
 
-FTextLayout::FTextLayout() 
+FTextLayout::FTextLayout()
 	: LineModels()
 	, LineViews()
 	, LineViewsToJustify()
@@ -1045,7 +1074,7 @@ void FTextLayout::SetRunRenderers( const TArray< FTextRunRenderer >& Renderers )
 
 void FTextLayout::AddRunRenderer( const FTextRunRenderer& Renderer )
 {
-	checkf( LineModels.IsValidIndex( Renderer.LineIndex ), TEXT("Renderers must be for a valid Line Index") );
+	checkf(LineModels.IsValidIndex(Renderer.LineIndex), TEXT("Renderers must be for a valid Line Index!\n\tDebug Source: %s"), *DebugSourceInfo.Get(FString()));
 
 	FLineModel& LineModel = LineModels[ Renderer.LineIndex ];
 
@@ -1055,13 +1084,13 @@ void FTextLayout::AddRunRenderer( const FTextRunRenderer& Renderer )
 	{
 		if ( LineModel.RunRenderers[ Index ].Range.BeginIndex > Renderer.Range.BeginIndex )
 		{
-			checkf( Index == 0 || LineModel.RunRenderers[ Index - 1 ].Range.EndIndex <= Renderer.Range.BeginIndex, TEXT("Renderers cannot overlap") );
+			checkf(Index == 0 || LineModel.RunRenderers[Index - 1].Range.EndIndex <= Renderer.Range.BeginIndex, TEXT("Renderers cannot overlap!\n\tDebug Source: %s"), *DebugSourceInfo.Get(FString()));
 			LineModel.RunRenderers.Insert( Renderer, Index - 1 );
 			bWasInserted = true;
 		}
 		else if ( LineModel.RunRenderers[ Index ].Range.EndIndex > Renderer.Range.EndIndex )
 		{
-			checkf( LineModel.RunRenderers[ Index ].Range.BeginIndex >= Renderer.Range.EndIndex, TEXT("Renderers cannot overlap") );
+			checkf(LineModel.RunRenderers[Index].Range.BeginIndex >= Renderer.Range.EndIndex, TEXT("Renderers cannot overlap!\n\tDebug Source: %s"), *DebugSourceInfo.Get(FString()));
 			LineModel.RunRenderers.Insert( Renderer, Index - 1 );
 			bWasInserted = true;
 		}
@@ -1099,8 +1128,8 @@ void FTextLayout::SetLineHighlights( const TArray< FTextLineHighlight >& Highlig
 
 void FTextLayout::AddLineHighlight( const FTextLineHighlight& Highlight )
 {
-	checkf( LineModels.IsValidIndex( Highlight.LineIndex ), TEXT("Highlights must be for a valid Line Index") );
-	checkf( Highlight.ZOrder, TEXT("The highlight Z-order must be <0 to create an underlay, or >0 to create an overlay") );
+	checkf(LineModels.IsValidIndex(Highlight.LineIndex), TEXT("Highlights must be for a valid Line Index!\n\tDebug Source: %s"), *DebugSourceInfo.Get(FString()));
+	checkf(Highlight.ZOrder, TEXT("The highlight Z-order must be <0 to create an underlay, or >0 to create an overlay!\n\tDebug Source: %s"), *DebugSourceInfo.Get(FString()));
 
 	FLineModel& LineModel = LineModels[ Highlight.LineIndex ];
 
@@ -1335,7 +1364,7 @@ bool FTextLayout::InsertAt(const FTextLocation& Location, TCHAR Character)
 		const bool bIsLastRun = RunIndex == LineModel.Runs.Num() - 1;
 		if (RunRange.Contains(InsertLocation) || (bIsLastRun && !RunIsAfterInsertLocation))
 		{
-			check(RunIsAfterInsertLocation == false);
+			checkf(RunIsAfterInsertLocation == false, TEXT("Debug Source: %s"), *DebugSourceInfo.Get(FString()));
 			RunIsAfterInsertLocation = true;
 
 			if ((RunModel.GetRun()->GetRunAttributes() & ERunAttributes::SupportsText) != ERunAttributes::None)
@@ -1345,7 +1374,7 @@ bool FTextLayout::InsertAt(const FTextLocation& Location, TCHAR Character)
 			else
 			{
 				// Non-text runs are supposed to have a single dummy character in them
-				check(RunRange.Len() == 1);
+				checkf(RunRange.Len() == 1, TEXT("Debug Source: %s"), *DebugSourceInfo.Get(FString()));
 
 				// This run doesn't support text, so we need to insert a new text run before or after the current run depending on the insertion point
 				const bool bIsInsertingToTheLeft = InsertLocation == RunRange.BeginIndex;
@@ -1400,7 +1429,7 @@ bool FTextLayout::InsertAt(const FTextLocation& Location, const FString& Text)
 		const bool bIsLastRun = RunIndex == LineModel.Runs.Num() - 1;
 		if (RunRange.Contains(InsertLocation) || (bIsLastRun && !RunIsAfterInsertLocation))
 		{
-			check(RunIsAfterInsertLocation == false);
+			checkf(RunIsAfterInsertLocation == false, TEXT("Debug Source: %s"), *DebugSourceInfo.Get(FString()));
 			RunIsAfterInsertLocation = true;
 
 			if ((RunModel.GetRun()->GetRunAttributes() & ERunAttributes::SupportsText) != ERunAttributes::None)
@@ -1410,7 +1439,7 @@ bool FTextLayout::InsertAt(const FTextLocation& Location, const FString& Text)
 			else
 			{
 				// Non-text runs are supposed to have a single dummy character in them
-				check(RunRange.Len() == 1);
+				checkf(RunRange.Len() == 1, TEXT("Debug Source: %s"), *DebugSourceInfo.Get(FString()));
 
 				// This run doesn't support text, so we need to insert a new text run before or after the current run depending on the insertion point
 				const bool bIsInsertingToTheLeft = InsertLocation == RunRange.BeginIndex;
@@ -1468,7 +1497,7 @@ bool FTextLayout::InsertAt(const FTextLocation& Location, TSharedRef<IRun> InRun
 		const bool bIsLastRun = RunIndex == LineModel.Runs.Num() - 1;
 		if (RunRange.Contains(InsertLocation) || (bIsLastRun && !RunIsAfterInsertLocation))
 		{
-			check(RunIsAfterInsertLocation == false);
+			checkf(RunIsAfterInsertLocation == false, TEXT("Debug Source: %s"), *DebugSourceInfo.Get(FString()));
 			RunIsAfterInsertLocation = true;
 
 			const int32 InsertLocationEnd = InsertLocation + NewRunText.Len();
@@ -1487,7 +1516,7 @@ bool FTextLayout::InsertAt(const FTextLocation& Location, TSharedRef<IRun> InRun
 			else
 			{
 				// Non-text runs are supposed to have a single dummy character in them
-				check(RunRange.Len() == 1);
+				checkf(RunRange.Len() == 1, TEXT("Debug Source: %s"), *DebugSourceInfo.Get(FString()));
 
 				// This run doesn't support text, so we need to insert a new text run before or after the current run depending on the insertion point
 				const bool bIsInsertingToTheLeft = InsertLocation == RunRange.BeginIndex;
@@ -1600,7 +1629,7 @@ bool FTextLayout::SplitLineAt(const FTextLocation& Location)
 	FLineModel LeftLineModel(MakeShareable(new FString(BreakLocation, **LineModel.Text)));
 	FLineModel RightLineModel(MakeShareable(new FString(LineModel.Text->Len() - BreakLocation, **LineModel.Text + BreakLocation)));
 
-	check(LeftLineModel.Text->Len() == BreakLocation);
+	checkf(LeftLineModel.Text->Len() == BreakLocation, TEXT("Debug Source: %s"), *DebugSourceInfo.Get(FString()));
 
 	bool RunIsToTheLeftOfTheBreakLocation = true;
 	for (int32 RunIndex = 0; RunIndex < LineModel.Runs.Num(); RunIndex++)
@@ -1611,7 +1640,7 @@ bool FTextLayout::SplitLineAt(const FTextLocation& Location)
 		const bool bIsLastRun = RunIndex == LineModel.Runs.Num() - 1;
 		if (RunRange.Contains(BreakLocation) || (bIsLastRun && RunIsToTheLeftOfTheBreakLocation))
 		{
-			check(RunIsToTheLeftOfTheBreakLocation == true);
+			checkf(RunIsToTheLeftOfTheBreakLocation == true, TEXT("Debug Source: %s"), *DebugSourceInfo.Get(FString()));
 			RunIsToTheLeftOfTheBreakLocation = false;
 
 			TSharedPtr<IRun> LeftRun;
@@ -1627,7 +1656,7 @@ bool FTextLayout::SplitLineAt(const FTextLocation& Location)
 			else
 			{
 				// Non-text runs are supposed to have a single dummy character in them
-				check(RunRange.Len() == 1);
+				checkf(RunRange.Len() == 1, TEXT("Debug Source: %s"), *DebugSourceInfo.Get(FString()));
 
 				// This run doesn't support text, so we need to insert a new text run before or after the current run depending on the insertion point
 				const bool bIsInsertingToTheLeft = BreakLocation == RunRange.BeginIndex;
@@ -1735,7 +1764,7 @@ bool FTextLayout::RemoveAt( const FTextLocation& Location, int32 Count )
 				// Some of this run has been removed, and this run is the right hand part of the removal
 				// So we need to adjust the range so that we start at the removal point since the text has been removed from the beginning of this run
 				const FTextRange NewRange(RemoveTextRange.BeginIndex, RunRange.EndIndex - Count);
-				check(!NewRange.IsEmpty());
+				checkf(!NewRange.IsEmpty(), TEXT("Debug Source: %s"), *DebugSourceInfo.Get(FString()));
 				RunModel.SetTextRange(NewRange);
 			}
 			else
@@ -1743,7 +1772,7 @@ bool FTextLayout::RemoveAt( const FTextLocation& Location, int32 Count )
 				// Some of this run has been removed, and this run is the left hand part of the removal
 				// So we need to adjust the range by the amount of text that has been removed from the end of this run
 				const FTextRange NewRange(RunRange.BeginIndex, RunRange.EndIndex - IntersectedLength);
-				check(!NewRange.IsEmpty());
+				checkf(!NewRange.IsEmpty(), TEXT("Debug Source: %s"), *DebugSourceInfo.Get(FString()));
 				RunModel.SetTextRange(NewRange);
 			}
 
@@ -2275,6 +2304,11 @@ void FTextLayout::SetWrappingWidth( float Value )
 	}
 }
 
+void FTextLayout::SetDebugSourceInfo(const TAttribute<FString>& InDebugSourceInfo)
+{
+	DebugSourceInfo = InDebugSourceInfo;
+}
+
 FVector2D FTextLayout::GetDrawSize() const
 {
 	return TextLayoutSize.GetDrawSize();
diff --git a/Engine/Source/Runtime/Slate/Private/Widgets/Input/SButton.cpp b/Engine/Source/Runtime/Slate/Private/Widgets/Input/SButton.cpp
--- a/Engine/Source/Runtime/Slate/Private/Widgets/Input/SButton.cpp
+++ b/Engine/Source/Runtime/Slate/Private/Widgets/Input/SButton.cpp
@@ -359,7 +359,7 @@ void SButton::Release()
 
 bool SButton::IsInteractable() const
 {
-	return IsEnabled() && SupportsKeyboardFocus();
+	return IsEnabled();
 }
 
 bool SButton::IsPreciseTapOrClick(const FPointerEvent& MouseEvent) const
diff --git a/Engine/Source/Runtime/Slate/Private/Widgets/Input/SCheckBox.cpp b/Engine/Source/Runtime/Slate/Private/Widgets/Input/SCheckBox.cpp
--- a/Engine/Source/Runtime/Slate/Private/Widgets/Input/SCheckBox.cpp
+++ b/Engine/Source/Runtime/Slate/Private/Widgets/Input/SCheckBox.cpp
@@ -223,7 +223,7 @@ void SCheckBox::OnMouseLeave( const FPointerEvent& MouseEvent )
 
 bool SCheckBox::IsInteractable() const
 {
-	return IsEnabled() && SupportsKeyboardFocus();
+	return IsEnabled();
 }
 
 /**
diff --git a/Engine/Source/Runtime/Slate/Private/Widgets/Text/SMultiLineEditableText.cpp b/Engine/Source/Runtime/Slate/Private/Widgets/Text/SMultiLineEditableText.cpp
--- a/Engine/Source/Runtime/Slate/Private/Widgets/Text/SMultiLineEditableText.cpp
+++ b/Engine/Source/Runtime/Slate/Private/Widgets/Text/SMultiLineEditableText.cpp
@@ -7,6 +7,7 @@
 #include "TextEditHelper.h"
 #include "PlainTextLayoutMarshaller.h"
 #include "TextBlockLayout.h"
+#include "ReflectionMetadata.h"
 #include "GenericCommands.h"
 #include "IMenu.h"
 
@@ -328,6 +329,7 @@ void SMultiLineEditableText::Construct( const FArguments& InArgs )
 	TextCompositionHighlighter = FTextCompositionHighlighter::Create();
 	TextSelectionRunRenderer = FTextSelectionRunRenderer::Create();
 	TextLayout = FSlateTextLayout::Create(TextStyle);
+	TextLayout->SetDebugSourceInfo(TAttribute<FString>::Create(TAttribute<FString>::FGetter::CreateLambda([this]{ return FReflectionMetaData::GetWidgetDebugInfo(this); })));
 
 	if (InArgs._TextShapingMethod.IsSet())
 	{
@@ -365,6 +367,7 @@ void SMultiLineEditableText::Construct( const FArguments& InArgs )
 	{
 		HintTextStyle = MakeShareable(new FTextBlockStyle(TextStyle));
 		HintTextLayout = FTextBlockLayout::Create(*HintTextStyle, TextLayout->GetTextShapingMethod(), TextLayout->GetTextFlowDirection(), Marshaller.ToSharedRef(), nullptr);
+		HintTextLayout->SetDebugSourceInfo(TAttribute<FString>::Create(TAttribute<FString>::FGetter::CreateLambda([this]{ return FReflectionMetaData::GetWidgetDebugInfo(this); })));
 	}
 
 	// Map UI commands to delegates which are called when the command should be executed
@@ -457,6 +460,7 @@ void SMultiLineEditableText::SetHintText(const TAttribute< FText >& InHintText)
 	{
 		HintTextStyle = MakeShareable(new FTextBlockStyle(TextStyle));
 		HintTextLayout = FTextBlockLayout::Create(*HintTextStyle, TextLayout->GetTextShapingMethod(), TextLayout->GetTextFlowDirection(), Marshaller.ToSharedRef(), nullptr);
+		HintTextLayout->SetDebugSourceInfo(TAttribute<FString>::Create(TAttribute<FString>::FGetter::CreateLambda([this]{ return FReflectionMetaData::GetWidgetDebugInfo(this); })));
 	}
 	else
 	{
diff --git a/Engine/Source/Runtime/Slate/Private/Widgets/Text/SRichTextBlock.cpp b/Engine/Source/Runtime/Slate/Private/Widgets/Text/SRichTextBlock.cpp
--- a/Engine/Source/Runtime/Slate/Private/Widgets/Text/SRichTextBlock.cpp
+++ b/Engine/Source/Runtime/Slate/Private/Widgets/Text/SRichTextBlock.cpp
@@ -9,6 +9,7 @@
 #include "IRichTextMarkupParser.h"
 #include "RichTextMarkupProcessing.h"
 #include "RichTextLayoutMarshaller.h"
+#include "ReflectionMetadata.h"
 
 void SRichTextBlock::Construct( const FArguments& InArgs )
 {
@@ -29,13 +30,19 @@ void SRichTextBlock::Construct( const FArguments& InArgs )
 			Parser = FDefaultRichTextMarkupParser::Create();
 		}
 
-		TSharedRef<FRichTextLayoutMarshaller> Marshaller = FRichTextLayoutMarshaller::Create(Parser, nullptr, InArgs._Decorators, InArgs._DecoratorStyleSet);
-		for ( const TSharedRef< ITextDecorator >& Decorator : InArgs.InlineDecorators )
+		TSharedPtr<FRichTextLayoutMarshaller> Marshaller = InArgs._Marshaller;
+		if (!Marshaller.IsValid())
 		{
-			Marshaller->AppendInlineDecorator( Decorator );
+			Marshaller = FRichTextLayoutMarshaller::Create(Parser, nullptr, InArgs._Decorators, InArgs._DecoratorStyleSet);
+		}
+		
+		for (const TSharedRef< ITextDecorator >& Decorator : InArgs.InlineDecorators)
+		{
+			Marshaller->AppendInlineDecorator(Decorator);
 		}
 
-		TextLayoutCache = FTextBlockLayout::Create(TextStyle, InArgs._TextShapingMethod, InArgs._TextFlowDirection, Marshaller, nullptr);
+		TextLayoutCache = FTextBlockLayout::Create(TextStyle, InArgs._TextShapingMethod, InArgs._TextFlowDirection, Marshaller.ToSharedRef(), nullptr);
+		TextLayoutCache->SetDebugSourceInfo(TAttribute<FString>::Create(TAttribute<FString>::FGetter::CreateLambda([this]{ return FReflectionMetaData::GetWidgetDebugInfo(this); })));
 	}
 }
 
diff --git a/Engine/Source/Runtime/Slate/Private/Widgets/Text/STextBlock.cpp b/Engine/Source/Runtime/Slate/Private/Widgets/Text/STextBlock.cpp
--- a/Engine/Source/Runtime/Slate/Private/Widgets/Text/STextBlock.cpp
+++ b/Engine/Source/Runtime/Slate/Private/Widgets/Text/STextBlock.cpp
@@ -3,6 +3,7 @@
 #include "SlatePrivatePCH.h"
 #include "PlainTextLayoutMarshaller.h"
 #include "TextBlockLayout.h"
+#include "ReflectionMetadata.h"
 
 DECLARE_CYCLE_STAT(TEXT("STextBlock::SetText Time"), Stat_SlateTextBlockSetText, STATGROUP_SlateVerbose)
 DECLARE_CYCLE_STAT(TEXT("STextBlock::OnPaint Time"), Stat_SlateTextBlockOnPaint, STATGROUP_SlateVerbose)
@@ -36,6 +37,7 @@ void STextBlock::Construct( const FArguments& InArgs )
 
 	// We use a dummy style here (as it may not be safe to call the delegates used to compute the style), but the correct style is set by ComputeDesiredSize
 	TextLayoutCache = FTextBlockLayout::Create(FTextBlockStyle::GetDefault(), InArgs._TextShapingMethod, InArgs._TextFlowDirection, FPlainTextLayoutMarshaller::Create(), InArgs._LineBreakPolicy);
+	TextLayoutCache->SetDebugSourceInfo(TAttribute<FString>::Create(TAttribute<FString>::FGetter::CreateLambda([this]{ return FReflectionMetaData::GetWidgetDebugInfo(this); })));
 
 #endif//WITH_FANCY_TEXT
 }
diff --git a/Engine/Source/Runtime/Slate/Private/Widgets/Text/TextBlockLayout.cpp b/Engine/Source/Runtime/Slate/Private/Widgets/Text/TextBlockLayout.cpp
--- a/Engine/Source/Runtime/Slate/Private/Widgets/Text/TextBlockLayout.cpp
+++ b/Engine/Source/Runtime/Slate/Private/Widgets/Text/TextBlockLayout.cpp
@@ -158,6 +158,11 @@ void FTextBlockLayout::SetTextFlowDirection(const TOptional<ETextFlowDirection>&
 	TextLayout->SetTextFlowDirection((InTextFlowDirection.IsSet()) ? InTextFlowDirection.GetValue() : GetDefaultTextFlowDirection());
 }
 
+void FTextBlockLayout::SetDebugSourceInfo(const TAttribute<FString>& InDebugSourceInfo)
+{
+	TextLayout->SetDebugSourceInfo(InDebugSourceInfo);
+}
+
 FChildren* FTextBlockLayout::GetChildren()
 {
 	return TextLayout->GetChildren();
diff --git a/Engine/Source/Runtime/Slate/Public/Framework/Text/RichTextLayoutMarshaller.h b/Engine/Source/Runtime/Slate/Public/Framework/Text/RichTextLayoutMarshaller.h
--- a/Engine/Source/Runtime/Slate/Public/Framework/Text/RichTextLayoutMarshaller.h
+++ b/Engine/Source/Runtime/Slate/Public/Framework/Text/RichTextLayoutMarshaller.h
@@ -35,10 +35,18 @@ public:
 
 protected:
 
+	FRichTextLayoutMarshaller(TArray< TSharedRef< ITextDecorator > > InDecorators, const ISlateStyle* const InDecoratorStyleSet);
 	FRichTextLayoutMarshaller(TSharedPtr< IRichTextMarkupParser > InParser, TSharedPtr< IRichTextMarkupWriter > InWriter, TArray< TSharedRef< ITextDecorator > > InDecorators, const ISlateStyle* const InDecoratorStyleSet);
 
 	TSharedPtr< ITextDecorator > TryGetDecorator(const FString& Line, const FTextRunParseResults& TextRun) const;
 
+	virtual void AppendRunsForText(const FTextRunParseResults& TextRun, 
+		const FString& ProcessedString, 
+		const FTextBlockStyle& DefaultTextStyle,
+		const TSharedRef<FString>& InOutModelText, 
+		FTextLayout& TargetTextLayout,
+		TArray<TSharedRef<IRun>>& Runs);
+
 	/** The parser used to resolve any markup used in the provided string. */
 	TSharedPtr< IRichTextMarkupParser > Parser;
 
diff --git a/Engine/Source/Runtime/Slate/Public/Framework/Text/TextLayout.h b/Engine/Source/Runtime/Slate/Public/Framework/Text/TextLayout.h
--- a/Engine/Source/Runtime/Slate/Public/Framework/Text/TextLayout.h
+++ b/Engine/Source/Runtime/Slate/Public/Framework/Text/TextLayout.h
@@ -363,6 +363,9 @@ public:
 	/** Set the iterator to use to detect appropriate soft-wrapping points for lines (or null to go back to using the default) */
 	void SetLineBreakIterator( TSharedPtr<IBreakIterator> InLineBreakIterator );
 
+	/** Set the information used to help identify who owns this text layout in the case of an error */
+	void SetDebugSourceInfo(const TAttribute<FString>& InDebugSourceInfo);
+
 	void ClearLines();
 
 	struct FNewLineData
@@ -656,4 +659,7 @@ protected:
 
 	/** Unicode BiDi text detection */
 	TUniquePtr<TextBiDi::ITextBiDi> TextBiDiDetection;
+
+	/** Information given to use by our an external source (typically our owner widget) to help identify who owns this text layout in the case of an error */
+	TAttribute<FString> DebugSourceInfo;
 };
diff --git a/Engine/Source/Runtime/Slate/Public/Widgets/Text/SRichTextBlock.h b/Engine/Source/Runtime/Slate/Public/Widgets/Text/SRichTextBlock.h
--- a/Engine/Source/Runtime/Slate/Public/Widgets/Text/SRichTextBlock.h
+++ b/Engine/Source/Runtime/Slate/Public/Widgets/Text/SRichTextBlock.h
@@ -16,6 +16,7 @@ public:
 		, _HighlightText()
 		, _WrapTextAt( 0.0f )
 		, _AutoWrapText(false)
+		, _Marshaller()
 		, _DecoratorStyleSet( &FCoreStyle::Get() )
 		, _TextStyle( &FCoreStyle::Get().GetWidgetStyle<FTextBlockStyle>( "NormalText" ) )
 		, _Margin( FMargin() )
@@ -40,6 +41,9 @@ public:
 			desired size will not be clamped.  This works best in cases where the text block's size is not affecting other widget's layout. */
 		SLATE_ATTRIBUTE( bool, AutoWrapText )
 
+		/** The marshaller used to get/set the raw text to/from the text layout. */
+		SLATE_ARGUMENT(TSharedPtr<class FRichTextLayoutMarshaller>, Marshaller)
+
 		/** The style set used for looking up styles used by decorators*/
 		SLATE_ARGUMENT( const ISlateStyle*, DecoratorStyleSet )
 
diff --git a/Engine/Source/Runtime/Slate/Public/Widgets/Text/TextBlockLayout.h b/Engine/Source/Runtime/Slate/Public/Widgets/Text/TextBlockLayout.h
--- a/Engine/Source/Runtime/Slate/Public/Widgets/Text/TextBlockLayout.h
+++ b/Engine/Source/Runtime/Slate/Public/Widgets/Text/TextBlockLayout.h
@@ -76,6 +76,11 @@ public:
 	 */
 	void SetTextFlowDirection(const TOptional<ETextFlowDirection>& InTextFlowDirection);
 
+	/**
+	 * Set the information used to help identify who owns this text layout in the case of an error
+	 */
+	void SetDebugSourceInfo(const TAttribute<FString>& InDebugSourceInfo);
+
 	/**
 	 * Get the child widgets of this layout
 	 */
diff --git a/Engine/Source/Runtime/UMG/Private/Components/Widget.cpp b/Engine/Source/Runtime/UMG/Private/Components/Widget.cpp
--- a/Engine/Source/Runtime/UMG/Private/Components/Widget.cpp
+++ b/Engine/Source/Runtime/UMG/Private/Components/Widget.cpp
@@ -404,7 +404,7 @@ bool UWidget::HasUserFocusedDescendants(APlayerController* PlayerController) con
 
 void UWidget::SetUserFocus(APlayerController* PlayerController)
 {
-	if ( PlayerController == nullptr || !PlayerController->IsLocalPlayerController() )
+	if ( PlayerController == nullptr || !PlayerController->IsLocalPlayerController() || PlayerController->Player == nullptr )
 	{
 		return;
 	}
