Fix for member variable shadowing

diff --git a/Engine/Source/Runtime/Slate/Private/Framework/Docking/TabManager.cpp b/Engine/Source/Runtime/Slate/Private/Framework/Docking/TabManager.cpp
--- a/Engine/Source/Runtime/Slate/Private/Framework/Docking/TabManager.cpp
+++ b/Engine/Source/Runtime/Slate/Private/Framework/Docking/TabManager.cpp
@@ -39,8 +39,8 @@ TSharedPtr<SDockTab> FTabManager::FRequireClosedTab::Search(const FTabManager& M
 	return TSharedPtr<SDockTab>();
 }
 
-FTabManager::FLastMajorOrNomadTab::FLastMajorOrNomadTab(FName FallbackTabId)
-	: FallbackTabId(FallbackTabId)
+FTabManager::FLastMajorOrNomadTab::FLastMajorOrNomadTab(FName InFallbackTabId)
+	: FallbackTabId(InFallbackTabId)
 {
 }
 
diff --git a/Engine/Source/Runtime/Slate/Public/Framework/Docking/TabManager.h b/Engine/Source/Runtime/Slate/Public/Framework/Docking/TabManager.h
--- a/Engine/Source/Runtime/Slate/Public/Framework/Docking/TabManager.h
+++ b/Engine/Source/Runtime/Slate/Public/Framework/Docking/TabManager.h
@@ -527,7 +527,7 @@ class SLATE_API FTabManager : public TSharedFromThis<FTabManager>
 		class SLATE_API FLastMajorOrNomadTab : public FSearchPreference
 		{
 		public:
-			FLastMajorOrNomadTab(FName FallbackTabId);
+			FLastMajorOrNomadTab(FName InFallbackTabId);
 
 			virtual TSharedPtr<SDockTab> Search(const FTabManager& Manager, FName PlaceholderId, const TSharedRef<SDockTab>& UnmanagedTab) const override;
 		private:
