Copying //UE4/Dev-Editor to Dev-Main (//UE4/Dev-Main)

diff --git a/Engine/Source/Runtime/Engine/Private/Slate/SceneViewport.cpp b/Engine/Source/Runtime/Engine/Private/Slate/SceneViewport.cpp
--- a/Engine/Source/Runtime/Engine/Private/Slate/SceneViewport.cpp
+++ b/Engine/Source/Runtime/Engine/Private/Slate/SceneViewport.cpp
@@ -24,6 +24,8 @@ FSceneViewport::FSceneViewport( FViewportClient* InViewportClient, TSharedPtr<SV
 	, NumMouseSamplesY( 0 )
 	, MouseDelta( 0, 0 )
 	, bIsCursorVisible( true )
+	, bShouldCaptureMouseOnActivate( true )
+	, bIgnoreCaptureClick( true )
 	, bRequiresVsync( false )
 	, bUseSeparateRenderTarget( InViewportWidget.IsValid() ? !InViewportWidget->ShouldRenderDirectly() : true )
 	, bIsResizing( false )
@@ -153,7 +155,7 @@ void FSceneViewport::GetMousePos( FIntPoint& MousePosition, const bool bLocalPos
 	}
 	else
 	{
-		const FVector2D AbsoluteMousePos = CachedGeometry.LocalToAbsolute(FVector2D(CachedMousePos.X, CachedMousePos.Y));
+		const FVector2D AbsoluteMousePos = CachedGeometry.LocalToAbsolute(FVector2D(CachedMousePos.X / CachedGeometry.Scale, CachedMousePos.Y / CachedGeometry.Scale));
 		MousePosition.X = AbsoluteMousePos.X;
 		MousePosition.Y = AbsoluteMousePos.Y;
 	}
@@ -173,7 +175,11 @@ void FSceneViewport::ProcessInput( float DeltaTime )
 
 void FSceneViewport::UpdateCachedMousePos( const FGeometry& InGeometry, const FPointerEvent& InMouseEvent )
 {
-	CachedMousePos = InGeometry.AbsoluteToLocal( InMouseEvent.GetScreenSpacePosition() ).IntPoint();
+	FVector2D LocalPixelMousePos = InGeometry.AbsoluteToLocal(InMouseEvent.GetScreenSpacePosition());
+	LocalPixelMousePos.X *= CachedGeometry.Scale;
+	LocalPixelMousePos.Y *= CachedGeometry.Scale;
+
+	CachedMousePos = LocalPixelMousePos.IntPoint();
 }
 
 void FSceneViewport::UpdateCachedGeometry( const FGeometry& InGeometry )
@@ -189,6 +195,8 @@ void FSceneViewport::UpdateModifierKeys( const FPointerEvent& InMouseEvent )
 	KeyStateMap.Add(EKeys::RightControl, InMouseEvent.IsRightControlDown());
 	KeyStateMap.Add(EKeys::LeftShift, InMouseEvent.IsLeftShiftDown());
 	KeyStateMap.Add(EKeys::RightShift, InMouseEvent.IsRightShiftDown());
+	KeyStateMap.Add(EKeys::LeftCommand, InMouseEvent.IsLeftCommandDown());
+	KeyStateMap.Add(EKeys::RightCommand, InMouseEvent.IsRightCommandDown());
 }
 
 void FSceneViewport::ApplyModifierKeys( const FModifierKeysState& InKeysState )
@@ -356,74 +364,99 @@ FReply FSceneViewport::OnMouseButtonDown( const FGeometry& InGeometry, const FPo
 	CurrentReplyState = FReply::Handled().PreventThrottling();
 
 	KeyStateMap.Add(InMouseEvent.GetEffectingButton(), true);
-	UpdateModifierKeys( InMouseEvent );
-	UpdateCachedMousePos( InGeometry, InMouseEvent );
+	UpdateModifierKeys(InMouseEvent);
+	UpdateCachedMousePos(InGeometry, InMouseEvent);
 	UpdateCachedGeometry(InGeometry);
 
 	// Switch to the viewport clients world before processing input
-	FScopedConditionalWorldSwitcher WorldSwitcher( ViewportClient );
-	if( ViewportClient && GetSizeXY() != FIntPoint::ZeroValue )
+	FScopedConditionalWorldSwitcher WorldSwitcher(ViewportClient);
+	if (ViewportClient && GetSizeXY() != FIntPoint::ZeroValue)
 	{
 		// If we're obtaining focus, we have to copy the modifier key states prior to processing this mouse button event, as this is the only point at which the mouse down
 		// event is processed when focus initially changes and the modifier keys need to be in-place to detect any unique drag-like events.
-		if ( !HasFocus() )
+		if (!HasFocus())
 		{
 			FModifierKeysState KeysState = FSlateApplication::Get().GetModifierKeys();
-			ApplyModifierKeys( KeysState );
+			ApplyModifierKeys(KeysState);
 		}
 
+		const bool bTemporaryCapture = 
+			ViewportClient->CaptureMouseOnClick() == EMouseCaptureMode::CaptureDuringMouseDown ||
+			(ViewportClient->CaptureMouseOnClick() == EMouseCaptureMode::CaptureDuringRightMouseDown && InMouseEvent.GetEffectingButton() == EKeys::RightMouseButton);
+
+		// Capture if we aren't currently a game viewport, or if we've been set to IgnoreCaptureClick or we already have capture.
+		const bool bProcessInputPrimary = !IsCurrentlyGameViewport() || !bIgnoreCaptureClick || HasMouseCapture();
+
 		const bool bAnyMenuWasVisible = FSlateApplication::Get().AnyMenusVisible();
 
 		// Process the mouse event
-		if (!ViewportClient->InputKey(this, InMouseEvent.GetUserIndex(), InMouseEvent.GetEffectingButton(), IE_Pressed))
+		if (bTemporaryCapture || bProcessInputPrimary)
 		{
-			CurrentReplyState = FReply::Unhandled(); 
+			if (!ViewportClient->InputKey(this, InMouseEvent.GetUserIndex(), InMouseEvent.GetEffectingButton(), IE_Pressed))
+			{
+				CurrentReplyState = FReply::Unhandled();
+			}
 		}
 
 		// a new menu was opened if there was previously not a menu visible but now there is
 		const bool bNewMenuWasOpened = !bAnyMenuWasVisible && FSlateApplication::Get().AnyMenusVisible();
 
-		if (!ViewportClient->IgnoreInput() &&
+		if (FSlateApplication::Get().IsActive() && !ViewportClient->IgnoreInput() &&
 			!bNewMenuWasOpened && // We should not focus the viewport if a menu was opened as it would close the menu
-			( ViewportClient->CaptureMouseOnClick() == EMouseCaptureMode::CapturePermanently ||
-			  ViewportClient->CaptureMouseOnClick() == EMouseCaptureMode::CaptureDuringMouseDown ||
-			  ( ViewportClient->CaptureMouseOnClick() == EMouseCaptureMode::CaptureDuringRightMouseDown && InMouseEvent.GetEffectingButton() == EKeys::RightMouseButton ) ) )
+			(ViewportClient->CaptureMouseOnClick() == EMouseCaptureMode::CapturePermanently || bTemporaryCapture))
 		{
-			TSharedRef<SViewport> ViewportWidgetRef = ViewportWidget.Pin().ToSharedRef();
-
-			// Mouse down should focus viewport for user input
-			CurrentReplyState.SetUserFocus(ViewportWidgetRef, EFocusCause::SetDirectly, true);
-			
-			UWorld* World = ViewportClient->GetWorld();
-			if (World && World->IsGameWorld() && World->GetGameInstance() && World->GetGameInstance()->GetFirstLocalPlayerController())
-			{
-				CurrentReplyState.CaptureMouse(ViewportWidgetRef);
-				CurrentReplyState.LockMouseToWidget(ViewportWidgetRef);
-
-				bool bShouldShowMouseCursor = World->GetGameInstance()->GetFirstLocalPlayerController()->ShouldShowMouseCursor();
-				if (ViewportClient->HideCursorDuringCapture() && bShouldShowMouseCursor)
-				{
-					bCursorHiddenDueToCapture = true;
-					MousePosBeforeHiddenDueToCapture = FIntPoint( InMouseEvent.GetScreenSpacePosition().X, InMouseEvent.GetScreenSpacePosition().Y );
-				}
-				if (bCursorHiddenDueToCapture || !bShouldShowMouseCursor)
-				{
-					CurrentReplyState.UseHighPrecisionMouseMovement(ViewportWidgetRef);
-				}
-			}
-			else
-			{
-				CurrentReplyState.UseHighPrecisionMouseMovement(ViewportWidgetRef);
-			}
+			CurrentReplyState = AcquireFocusAndCapture(FIntPoint(InMouseEvent.GetScreenSpacePosition().X, InMouseEvent.GetScreenSpacePosition().Y));
 		}
 	}
-	
+
 	// Re-set prevent throttling here as it can get reset when inside of InputKey()
 	CurrentReplyState.PreventThrottling();
 
 	return CurrentReplyState;
 }
 
+FReply FSceneViewport::AcquireFocusAndCapture(FIntPoint MousePosition)
+{
+	bShouldCaptureMouseOnActivate = false;
+
+	FReply ReplyState = FReply::Handled().PreventThrottling();
+
+	TSharedRef<SViewport> ViewportWidgetRef = ViewportWidget.Pin().ToSharedRef();
+
+	// Mouse down should focus viewport for user input
+	ReplyState.SetUserFocus(ViewportWidgetRef, EFocusCause::SetDirectly, true);
+
+	UWorld* World = ViewportClient->GetWorld();
+	if (World && World->IsGameWorld() && World->GetGameInstance() && World->GetGameInstance()->GetFirstLocalPlayerController())
+	{
+		ReplyState.CaptureMouse(ViewportWidgetRef);
+		ReplyState.LockMouseToWidget(ViewportWidgetRef);
+
+		bool bShouldShowMouseCursor = World->GetGameInstance()->GetFirstLocalPlayerController()->ShouldShowMouseCursor();
+		if (ViewportClient->HideCursorDuringCapture() && bShouldShowMouseCursor)
+		{
+			bCursorHiddenDueToCapture = true;
+			MousePosBeforeHiddenDueToCapture = MousePosition;
+		}
+		if (bCursorHiddenDueToCapture || !bShouldShowMouseCursor)
+		{
+			ReplyState.UseHighPrecisionMouseMovement(ViewportWidgetRef);
+		}
+	}
+	else
+	{
+		ReplyState.UseHighPrecisionMouseMovement(ViewportWidgetRef);
+	}
+
+	return ReplyState;
+}
+
+bool FSceneViewport::IsCurrentlyGameViewport()
+{
+	// Either were game code only or were are currently play in editor.
+	return (FApp::IsGame() && !GIsEditor) || IsPlayInEditorViewport();
+}
+
 FReply FSceneViewport::OnMouseButtonUp( const FGeometry& InGeometry, const FPointerEvent& InMouseEvent )
 {
 	// Start a new reply state
@@ -450,7 +483,7 @@ FReply FSceneViewport::OnMouseButtonUp( const FGeometry& InGeometry, const FPoin
 			ViewportClient->CaptureMouseOnClick() == EMouseCaptureMode::CaptureDuringMouseDown ||
 			( ViewportClient->CaptureMouseOnClick() == EMouseCaptureMode::CaptureDuringRightMouseDown && InMouseEvent.GetEffectingButton() == EKeys::RightMouseButton );
 	}
-	if (!((FApp::IsGame() && !GIsEditor) || bIsPlayInEditorViewport) || bReleaseMouse)
+	if (!IsCurrentlyGameViewport() || bReleaseMouse)
 	{
 		// On mouse up outside of the game (editor viewport) or if the cursor is visible in game, we should make sure the mouse is no longer captured
 		// as long as the left or right mouse buttons are not still down
@@ -485,7 +518,7 @@ void FSceneViewport::OnMouseLeave( const FPointerEvent& MouseEvent )
 	{
 		ViewportClient->MouseLeave( this );
 	
-		if ( (FApp::IsGame() && !GIsEditor) || IsPlayInEditorViewport() )
+		if (IsCurrentlyGameViewport())
 		{
 			CachedMousePos = FIntPoint(-1, -1);
 		}
@@ -847,42 +880,6 @@ FReply FSceneViewport::OnFocusReceived(const FFocusEvent& InFocusEvent)
 	{
 		FScopedConditionalWorldSwitcher WorldSwitcher(ViewportClient);
 		ViewportClient->ReceivedFocus(this);
-
-		if ((FApp::IsGame() && !GIsEditor) || bIsPlayInEditorViewport)
-		{
-			if (IsForegroundWindow())
-			{
-				bool bIsCursorForcedVisible = false;
-				APlayerController* FirstLocalPC = nullptr;
-				if (UWorld* ViewportClientWorld = ViewportClient->GetWorld())
-				{
-					if (UGameInstance* GameInstance = ViewportClientWorld->GetGameInstance())
-					{
-						FirstLocalPC = GameInstance->GetFirstLocalPlayerController();
-					}
-				}
-				if (FirstLocalPC)
-				{
-					bIsCursorForcedVisible = FirstLocalPC->GetMouseCursor() != EMouseCursor::None;
-				}
-
-				const bool bPlayInEditorCapture = !bIsPlayInEditorViewport || InFocusEvent.GetCause() != EFocusCause::SetDirectly || bPlayInEditorGetsMouseControl;
-
-				// capturing the mouse interferes with slate UI (like the virtual joysticks)
-				if (FPlatformProperties::SupportsWindowedMode() && bPlayInEditorCapture && !bIsCursorForcedVisible && !FSlateApplication::Get().IsFakingTouchEvents())
-				{
-					// Only require the user to click in the window the first time - after that return focus to the game so long as it was the last focused widget.
-					// Means that tabbing in/out will return the mouse control to where it was & the in-game console won't leave the mouse under editor control.
-					bPlayInEditorGetsMouseControl = true;
-					CurrentReplyState.UseHighPrecisionMouseMovement(ViewportWidget.Pin().ToSharedRef());
-					CurrentReplyState.LockMouseToWidget(ViewportWidget.Pin().ToSharedRef());
-				}
-			}
-			else
-			{
-				CurrentReplyState.ClearUserFocus(true);
-			}
-		}
 	}
 
 	// Update key state mappings so that the the viewport modifier states are valid upon focus.
@@ -893,29 +890,21 @@ FReply FSceneViewport::OnFocusReceived(const FFocusEvent& InFocusEvent)
 	KeyStateMap.Add( EKeys::RightControl, KeysState.IsRightControlDown());
 	KeyStateMap.Add( EKeys::LeftShift, KeysState.IsLeftShiftDown());
 	KeyStateMap.Add( EKeys::RightShift, KeysState.IsRightShiftDown());
+	KeyStateMap.Add( EKeys::LeftCommand, KeysState.IsLeftCommandDown());
+	KeyStateMap.Add( EKeys::RightCommand, KeysState.IsRightCommandDown());
 
 	return CurrentReplyState;
 }
 
 void FSceneViewport::OnFocusLost( const FFocusEvent& InFocusEvent )
 {
+	bShouldCaptureMouseOnActivate = false;
+
 	KeyStateMap.Empty();
 	if (ViewportClient != nullptr)
 	{
 		FScopedConditionalWorldSwitcher WorldSwitcher( ViewportClient );
 		ViewportClient->LostFocus( this );
-
-		TSharedPtr<SWidget> ViewportWidgetPin = ViewportWidget.Pin();
-		if( ViewportWidgetPin.IsValid() )
-		{
-			for (int32 UserIndex = 0; UserIndex < SlateApplicationDefs::MaxUsers; ++UserIndex)
-			{
-				if (FSlateApplication::Get().GetUserFocusedWidget(UserIndex) == ViewportWidgetPin)
-				{
-					FSlateApplication::Get().ClearUserFocus(UserIndex);
-				}
-			}
-		}
 	}
 }
 
@@ -928,6 +917,51 @@ void FSceneViewport::OnViewportClosed()
 	}
 }
 
+FReply FSceneViewport::OnViewportActivated(const FWindowActivateEvent& InActivateEvent)
+{
+	if (ViewportClient != nullptr)
+	{
+		FScopedConditionalWorldSwitcher WorldSwitcher(ViewportClient);
+		ViewportClient->Activated(this, InActivateEvent);
+		
+		// If we are activating and had Mouse Capture on deactivate then we should get focus again
+		// It's important to note in the case of:
+		//    InActivateEvent.ActivationType == FWindowActivateEvent::EA_ActivateByMouse
+		// we do NOT acquire focus the reasoning is that the click itself will give us a chance on Mouse down to get capture.
+		// This also means we don't go and grab capture in situations like:
+		//    - the user clicked on the application header
+		//    - the user clicked on some UI
+		//    - the user clicked in our window but not an area our viewport covers.
+		if (InActivateEvent.GetActivationType() == FWindowActivateEvent::EA_Activate && bShouldCaptureMouseOnActivate)
+		{
+			return AcquireFocusAndCapture(GetSizeXY() / 2);
+		}
+	}
+	return FReply::Unhandled();
+}
+void FSceneViewport::OnViewportDeactivated(const FWindowActivateEvent& InActivateEvent)
+{
+	// We backup if we have capture for us on activation, however we also maintain "true" if it's already true!
+	// The reasoning behind maintaining "true" is that if the viewport is activated, 
+	// however doesn't reclaim capture we want to claim capture next time we activate unless something else gets focus.
+	// So we reset bHadMouseCaptureOnDeactivate in AcquireFocusAndCapture() and in OnFocusLost()
+	//
+	// This is not ideal, however the better fix probably requires that slate fundamentally chance when it "activates" a window or maybe just the viewport
+	// Which there simply doesn't exist the right hooks currently.
+	//
+	// This fixes the case where the application is deactivated, then the user click on the windows header
+	// this activates the window but we do not capture the mouse, then the User Alt-Tabs to the application.
+	// We properly acquire capture because we maintained the "true" through the activation where nothing was focuses
+	bShouldCaptureMouseOnActivate = bShouldCaptureMouseOnActivate || HasMouseCapture();
+
+	KeyStateMap.Empty();
+	if (ViewportClient != nullptr)
+	{
+		FScopedConditionalWorldSwitcher WorldSwitcher(ViewportClient);
+		ViewportClient->Deactivated(this, InActivateEvent);
+	}
+}
+
 FSlateShaderResource* FSceneViewport::GetViewportRenderTargetTexture() const
 { 
 	check(IsThreadSafeForSlateRendering());
diff --git a/Engine/Source/Runtime/Engine/Private/Tests/AutomationCommon.cpp b/Engine/Source/Runtime/Engine/Private/Tests/AutomationCommon.cpp
--- a/Engine/Source/Runtime/Engine/Private/Tests/AutomationCommon.cpp
+++ b/Engine/Source/Runtime/Engine/Private/Tests/AutomationCommon.cpp
@@ -32,9 +32,10 @@ namespace AutomationCommon
 
 		TArray<FColor> OutImageData;
 		FIntVector OutImageSize;
-		FSlateApplication::Get().TakeScreenshot(WindowRef,OutImageData,OutImageSize);
-
-		FAutomationTestFramework::GetInstance().OnScreenshotCaptured().ExecuteIfBound(OutImageSize.X, OutImageSize.Y, OutImageData, FileName);
+		if (FSlateApplication::Get().TakeScreenshot(WindowRef, OutImageData, OutImageSize))
+		{
+			FAutomationTestFramework::GetInstance().OnScreenshotCaptured().ExecuteIfBound(OutImageSize.X, OutImageSize.Y, OutImageData, FileName);
+		}
 	}
 }
 
diff --git a/Engine/Source/Runtime/Engine/Private/UserInterface/CanvasItem.cpp b/Engine/Source/Runtime/Engine/Private/UserInterface/CanvasItem.cpp
--- a/Engine/Source/Runtime/Engine/Private/UserInterface/CanvasItem.cpp
+++ b/Engine/Source/Runtime/Engine/Private/UserInterface/CanvasItem.cpp
@@ -820,11 +820,11 @@ void FCanvasBorderItem::Draw( class FCanvas* InCanvas )
 }
 
 
-void FCanvasTextItem::Draw( class FCanvas* InCanvas )
+void FCanvasTextItemBase::Draw( class FCanvas* InCanvas )
 {	
 	SCOPE_CYCLE_COUNTER(STAT_Canvas_TextItemTime);
 
-	if (InCanvas == nullptr || Font == nullptr || Text.IsEmpty())
+	if (InCanvas == nullptr || !HasValidText())
 	{
 		return;
 	}
@@ -835,55 +835,15 @@ void FCanvasTextItem::Draw( class FCanvas* InCanvas )
 		// EnableShadow will set a default ShadowOffset value
 		EnableShadow( FLinearColor::Black );
 	}
-	if (Font->ImportOptions.bUseDistanceFieldAlpha)
-	{
-		// convert blend mode to distance field type
-		switch(BlendMode)
-		{
-		case SE_BLEND_Translucent:
-			BlendMode = (bHasShadow) ? SE_BLEND_TranslucentDistanceFieldShadowed : SE_BLEND_TranslucentDistanceField;
-			break;
-		case SE_BLEND_Masked:
-			BlendMode = (bHasShadow) ? SE_BLEND_MaskedDistanceFieldShadowed : SE_BLEND_MaskedDistanceField;
-			break;
-		}
-	}
-	if (GetFontCacheType() == EFontCacheType::Runtime)
-	{
-		// The runtime font cache uses an alpha-only texture, so we have to force this blend mode so we use the correct shader
-		check(BlendMode == SE_BLEND_Translucent || BlendMode == SE_BLEND_TranslucentAlphaOnly);
-		BlendMode = SE_BLEND_TranslucentAlphaOnly;
-	}
+	BlendMode = GetTextBlendMode( bHasShadow );
 
 	FVector2D DrawPos( Position.X , Position.Y );
 
 	// If we are centering the string or we want to fix stereoscopic rendering issues we need to measure the string
 	if( ( bCentreX || bCentreY ) || ( !bDontCorrectStereoscopic ) )
 	{
-		FVector2D MeasuredTextSize;
-		switch( GetFontCacheType() )
-		{
-		case EFontCacheType::Offline:
-			{
-				FTextSizingParameters Parameters( Font, Scale.X ,Scale.Y );
-				UCanvas::CanvasStringSize( Parameters, *Text.ToString() );
-				MeasuredTextSize.X = Parameters.DrawXL;
-				MeasuredTextSize.Y = Parameters.DrawYL;
-			}
-			break;
-
-		case EFontCacheType::Runtime:
-			{
-				const FSlateFontInfo LegacyFontInfo = (SlateFontInfo.IsSet()) ? SlateFontInfo.GetValue() : Font->GetLegacySlateFontInfo();
-				const TSharedRef<FSlateFontMeasure> FontMeasure = FSlateApplication::Get().GetRenderer()->GetFontMeasureService();
-				MeasuredTextSize = FontMeasure->Measure( Text, LegacyFontInfo ) * Scale;
-			}
-			break;
-
-		default:
-			break;
-		}
-				
+		const FVector2D MeasuredTextSize = GetTextSize();
+		
 		// Calculate the offset if we are centering
 		if( bCentreX || bCentreY )
 		{		
@@ -940,11 +900,70 @@ void FCanvasTextItem::Draw( class FCanvas* InCanvas )
 	DrawStringInternal( InCanvas, DrawPos, DrawColor );
 }
 
+
 EFontCacheType FCanvasTextItem::GetFontCacheType() const
 {
 	return Font->FontCacheType;
 }
 
+bool FCanvasTextItem::HasValidText() const
+{
+	return Font && !Text.IsEmpty();
+}
+
+ESimpleElementBlendMode FCanvasTextItem::GetTextBlendMode( const bool bHasShadow ) const
+{
+	ESimpleElementBlendMode BlendModeToUse = BlendMode;
+	if (Font->ImportOptions.bUseDistanceFieldAlpha)
+	{
+		// convert blend mode to distance field type
+		switch(BlendMode)
+		{
+		case SE_BLEND_Translucent:
+			BlendModeToUse = (bHasShadow) ? SE_BLEND_TranslucentDistanceFieldShadowed : SE_BLEND_TranslucentDistanceField;
+			break;
+		case SE_BLEND_Masked:
+			BlendModeToUse = (bHasShadow) ? SE_BLEND_MaskedDistanceFieldShadowed : SE_BLEND_MaskedDistanceField;
+			break;
+		}
+	}
+	if (GetFontCacheType() == EFontCacheType::Runtime)
+	{
+		// The runtime font cache uses an alpha-only texture, so we have to force this blend mode so we use the correct shader
+		check(BlendModeToUse == SE_BLEND_Translucent || BlendModeToUse == SE_BLEND_TranslucentAlphaOnly);
+		BlendModeToUse = SE_BLEND_TranslucentAlphaOnly;
+	}
+	return BlendModeToUse;
+}
+
+FVector2D FCanvasTextItem::GetTextSize() const
+{
+	FVector2D MeasuredTextSize = FVector2D::ZeroVector;
+	switch( GetFontCacheType() )
+	{
+	case EFontCacheType::Offline:
+		{
+			FTextSizingParameters Parameters( Font, Scale.X ,Scale.Y );
+			UCanvas::CanvasStringSize( Parameters, *Text.ToString() );
+			MeasuredTextSize.X = Parameters.DrawXL;
+			MeasuredTextSize.Y = Parameters.DrawYL;
+		}
+		break;
+
+	case EFontCacheType::Runtime:
+		{
+			const FSlateFontInfo LegacyFontInfo = (SlateFontInfo.IsSet()) ? SlateFontInfo.GetValue() : Font->GetLegacySlateFontInfo();
+			const TSharedRef<FSlateFontMeasure> FontMeasure = FSlateApplication::Get().GetRenderer()->GetFontMeasureService();
+			MeasuredTextSize = FontMeasure->Measure( Text, LegacyFontInfo ) * Scale;
+		}
+		break;
+
+	default:
+		break;
+	}
+	return MeasuredTextSize;
+}
+
 void FCanvasTextItem::DrawStringInternal( FCanvas* InCanvas, const FVector2D& DrawPos, const FLinearColor& InColor )
 {
 	switch(GetFontCacheType())
@@ -1042,11 +1061,10 @@ void FCanvasTextItem::DrawStringInternal_OfflineCache( FCanvas* InCanvas, const
 			const float SizeU	= Char.USize * InvTextureSize.X;
 			const float SizeV	= Char.VSize * InvTextureSize.Y;				
 
-			float Left, Top, Right, Bottom;
-			Left = X * Depth;
-			Top = Y * Depth;
-			Right = (X + SizeX) * Depth;
-			Bottom = (Y + SizeY) * Depth;
+			const float Left = X * Depth;
+			const float Top = Y * Depth;
+			const float Right = (X + SizeX) * Depth;
+			const float Bottom = (Y + SizeY) * Depth;
 
 			int32 V00 = BatchedElements->AddVertex(
 				FVector4( Left, Top, 0.f, Depth ),
@@ -1119,12 +1137,8 @@ void FCanvasTextItem::DrawStringInternal_RuntimeCache( FCanvas* InCanvas, const
 	float InvTextureSizeX = 0;
 	float InvTextureSizeY = 0;
 
-	float LineX = 0;
-
 	FCharacterEntry PreviousCharEntry;
 
-	int32 Kerning = 0;
-
 	FVector2D TopLeft(0,0);
 
 	const float PosX = TopLeft.X;
@@ -1133,7 +1147,7 @@ void FCanvasTextItem::DrawStringInternal_RuntimeCache( FCanvas* InCanvas, const
 	const float ScaledHorizSpacingAdjust = HorizSpacingAdjust * Scale.X;
 	const float ScaledMaxHeight = CharacterList.GetMaxHeight() * Scale.Y;
 
-	LineX = PosX;
+	float LineX = PosX;
 	
 	for( int32 CharIndex = 0; CharIndex < TextString.Len(); ++CharIndex )
 	{
@@ -1177,14 +1191,11 @@ void FCanvasTextItem::DrawStringInternal_RuntimeCache( FCanvas* InCanvas, const
 
 			const bool bIsWhitespace = FChar::IsWhitespace(CurrentChar);
 
+			int32 Kerning = 0;
 			if( !bIsWhitespace && PreviousCharEntry.IsCached() )
 			{
 				Kerning = CharacterList.GetKerning( PreviousCharEntry, Entry ) * Scale.X;
 			}
-			else
-			{
-				Kerning = 0;
-			}
 
 			LineX += Kerning;
 			PreviousCharEntry = Entry;
@@ -1202,11 +1213,10 @@ void FCanvasTextItem::DrawStringInternal_RuntimeCache( FCanvas* InCanvas, const
 				const float SizeU = Entry.USize * InvTextureSizeX;
 				const float SizeV = Entry.VSize * InvTextureSizeY;
 
-				float Left, Top, Right, Bottom;
-				Left = X * Depth;
-				Top = Y * Depth;
-				Right = (X + SizeX) * Depth;
-				Bottom = (Y + SizeY) * Depth;
+				const float Left = X * Depth;
+				const float Top = Y * Depth;
+				const float Right = (X + SizeX) * Depth;
+				const float Bottom = (Y + SizeY) * Depth;
 
 				int32 V00 = BatchedElements->AddVertex(
 					FVector4( Left, Top, 0.f, Depth ),
@@ -1245,6 +1255,137 @@ void FCanvasTextItem::DrawStringInternal_RuntimeCache( FCanvas* InCanvas, const
 	}
 }
 
+
+bool FCanvasShapedTextItem::HasValidText() const
+{
+	return ShapedGlyphSequence.IsValid() && ShapedGlyphSequence->GetGlyphsToRender().Num();
+}
+
+ESimpleElementBlendMode FCanvasShapedTextItem::GetTextBlendMode( const bool bHasShadow ) const
+{
+	ESimpleElementBlendMode BlendModeToUse = BlendMode;
+
+	// The runtime font cache uses an alpha-only texture, so we have to force this blend mode so we use the correct shader
+	check(BlendModeToUse == SE_BLEND_Translucent || BlendModeToUse == SE_BLEND_TranslucentAlphaOnly);
+	BlendModeToUse = SE_BLEND_TranslucentAlphaOnly;
+
+	return BlendModeToUse;
+}
+
+FVector2D FCanvasShapedTextItem::GetTextSize() const
+{
+	return FVector2D(ShapedGlyphSequence->GetMeasuredWidth(), ShapedGlyphSequence->GetMaxTextHeight());
+}
+
+void FCanvasShapedTextItem::DrawStringInternal( FCanvas* InCanvas, const FVector2D& DrawPos, const FLinearColor& InColor )
+{
+	DrawnSize = FVector2D::ZeroVector;
+
+	TSharedPtr<FSlateFontCache> FontCache = FEngineFontServices::Get().GetFontCache();
+	if( !FontCache.IsValid() )
+	{
+		return;
+	}
+
+	FHitProxyId HitProxyId = InCanvas->GetHitProxyId();
+
+	uint32 FontTextureIndex = 0;
+	FTextureResource* FontTexture = nullptr;
+
+	float InvTextureSizeX = 0;
+	float InvTextureSizeY = 0;
+
+	FVector2D TopLeft(0,0);
+
+	const float PosX = TopLeft.X;
+	float PosY = TopLeft.Y;
+
+	const float ScaledHorizSpacingAdjust = HorizSpacingAdjust * Scale.X;
+	const float ScaledMaxHeight = ShapedGlyphSequence->GetMaxTextHeight() * Scale.Y;
+	const float ScaledBaseline = ShapedGlyphSequence->GetTextBaseline() * Scale.Y;
+
+	float LineX = PosX;
+	
+	for( const auto& GlyphToRender : ShapedGlyphSequence->GetGlyphsToRender() )
+	{
+		if (DrawnSize.Y == 0)
+		{
+			// We have a valid glyph so initialize vertical DrawnSize
+			DrawnSize.Y = ScaledMaxHeight;
+		}
+
+		if( GlyphToRender.bIsVisible )
+		{
+			const FShapedGlyphFontAtlasData GlyphAtlasData = FontCache->GetShapedGlyphFontAtlasData(GlyphToRender);
+
+			if( FontTexture == nullptr || GlyphAtlasData.TextureIndex != FontTextureIndex )
+			{
+				// Font has a new texture for this glyph. Refresh the batch we use and the index we are currently using
+				FontTextureIndex = GlyphAtlasData.TextureIndex;
+				FontTexture = FontCache->GetEngineTextureResource( FontTextureIndex );
+				check(FontTexture);
+
+				FBatchedElementParameters* BatchedElementParameters = nullptr;
+				BatchedElements = InCanvas->GetBatchedElements(FCanvas::ET_Triangle, BatchedElementParameters, FontTexture, BlendMode, FontRenderInfo.GlowInfo);
+				check(BatchedElements);
+
+				InvTextureSizeX = 1.0f/FontTexture->GetSizeX();
+				InvTextureSizeY = 1.0f/FontTexture->GetSizeY();
+			}
+
+			const float X = DrawPos.X + LineX + (GlyphAtlasData.HorizontalOffset * Scale.X) + (GlyphToRender.XOffset  * Scale.X);
+			// Note PosX,PosY is the upper left corner of the bounding box representing the string.  This computes the Y position of the baseline where text will sit
+
+			const float Y = DrawPos.Y + PosY - (GlyphAtlasData.VerticalOffset * Scale.Y) + (GlyphToRender.YOffset * Scale.Y) + ScaledBaseline + ScaledMaxHeight;
+			const float U = GlyphAtlasData.StartU * InvTextureSizeX;
+			const float V = GlyphAtlasData.StartV * InvTextureSizeY;
+			const float SizeX = GlyphAtlasData.USize * Scale.X;
+			const float SizeY = GlyphAtlasData.VSize * Scale.Y;
+			const float SizeU = GlyphAtlasData.USize * InvTextureSizeX;
+			const float SizeV = GlyphAtlasData.VSize * InvTextureSizeY;
+
+			const float Left = X * Depth;
+			const float Top = Y * Depth;
+			const float Right = (X + SizeX) * Depth;
+			const float Bottom = (Y + SizeY) * Depth;
+
+			int32 V00 = BatchedElements->AddVertex(
+				FVector4( Left, Top, 0.f, Depth ),
+				FVector2D( U, V ),
+				InColor,
+				HitProxyId );
+			int32 V10 = BatchedElements->AddVertex(
+				FVector4( Right, Top, 0.0f, Depth ),
+				FVector2D( U + SizeU, V ),			
+				InColor,
+				HitProxyId );
+			int32 V01 = BatchedElements->AddVertex(
+				FVector4( Left, Bottom, 0.0f, Depth ),
+				FVector2D( U, V + SizeV ),	
+				InColor,
+				HitProxyId);
+			int32 V11 = BatchedElements->AddVertex(
+				FVector4( Right, Bottom, 0.0f, Depth ),
+				FVector2D( U + SizeU, V + SizeV ),
+				InColor,
+				HitProxyId);
+
+			BatchedElements->AddTriangle(V00, V10, V11, FontTexture, BlendMode, FontRenderInfo.GlowInfo);
+			BatchedElements->AddTriangle(V00, V11, V01, FontTexture, BlendMode, FontRenderInfo.GlowInfo);
+		}
+
+		LineX += GlyphToRender.XAdvance * Scale.X;
+		LineX += ScaledHorizSpacingAdjust;
+
+		// Increase the Horizontal DrawnSize
+		if (LineX > DrawnSize.X)
+		{
+			DrawnSize.X = LineX;
+		}
+	}
+}
+
+
 void FCanvasLineItem::Draw( class FCanvas* InCanvas )
 {
 	SCOPE_CYCLE_COUNTER(STAT_Canvas_LineItemTime);
diff --git a/Engine/Source/Runtime/Engine/Public/CanvasItem.h b/Engine/Source/Runtime/Engine/Public/CanvasItem.h
--- a/Engine/Source/Runtime/Engine/Public/CanvasItem.h
+++ b/Engine/Source/Runtime/Engine/Public/CanvasItem.h
@@ -8,6 +8,7 @@
 
 #include "Engine/EngineTypes.h"
 #include "Engine/Font.h"
+#include "ShapedTextFwd.h"
 #include "SceneTypes.h"
 
 class FCanvas;
@@ -290,24 +291,13 @@
 	FVector2D   CornerSize;
 };
 
-/* Text item with misc optional items such as shadow, centering etc. */
-class ENGINE_API FCanvasTextItem : public FCanvasItem
+/* Base item used for drawing text */
+class ENGINE_API FCanvasTextItemBase : public FCanvasItem
 {
 public:
-	/** 	 
-	 * Text item
-	 *
-	 * @param	InPosition		Draw position
-	 * @param	InText			String to draw
-	 * @param	InFont			Font to draw with
-	 */
-	FCanvasTextItem( const FVector2D& InPosition, const FText& InText, const UFont* InFont, const FLinearColor& InColor )
+	FCanvasTextItemBase( const FVector2D& InPosition, const FLinearColor& InColor )
 		: FCanvasItem( InPosition )
-		, Text( InText )
-		, Font( InFont )
-		, SlateFontInfo()
 		, HorizSpacingAdjust( 0.0f )
-	  	, ForcedViewportHeight( nullptr )
 		, Depth( 1.0f )
 		, ShadowColor( FLinearColor::Black )
 		, ShadowOffset( FVector2D::ZeroVector )
@@ -322,38 +312,8 @@
 		SetColor( InColor );
 		Scale.Set( 1.0f, 1.0f );
 		BlendMode = SE_BLEND_Translucent;
-	};
+	}
 
-	/** 	 
-	 * Text item
-	 *
-	 * @param	InPosition		Draw position
-	 * @param	InText			String to draw
-	 * @param	InFontInfo		Font info to draw with
-	 */
-	FCanvasTextItem( const FVector2D& InPosition, const FText& InText, const FSlateFontInfo& InFontInfo, const FLinearColor& InColor )
-		: FCanvasItem( InPosition )
-		, Text( InText )
-		, Font( Cast<const UFont>(InFontInfo.FontObject) )
-		, SlateFontInfo( InFontInfo )
-		, HorizSpacingAdjust( 0.0f )
-	  	, ForcedViewportHeight( nullptr )
-		, Depth( 1.0f )
-		, ShadowColor( FLinearColor::Black )
-		, ShadowOffset( FVector2D::ZeroVector )
-		, DrawnSize( FVector2D::ZeroVector )
-		, bCentreX( false )
-		, bCentreY( false )
-		, bOutlined( false )
-		, OutlineColor( FLinearColor::Black )
-		, bDontCorrectStereoscopic( true )
-		, TileItem( InPosition, FVector2D::ZeroVector, InColor )
-	{
-		SetColor( InColor );
-		Scale.Set( 1.0f, 1.0f );
-		BlendMode = SE_BLEND_TranslucentAlphaOnly;
-	};
-		
 	/** 
 	 * Set the shadow offset and color. 
 	 *
@@ -386,20 +346,8 @@
 	/* Expose the functions defined in the base class. */
 	using FCanvasItem::Draw;
 
-	/* The text to draw. */
-	FText Text;
-	
-	/* Font to draw text with. */
-	const UFont* Font;
-
-	/** Font info to draw the text with. */
-	TOptional<FSlateFontInfo> SlateFontInfo;
-
 	/* Horizontal spacing adjustment. */
 	float HorizSpacingAdjust;
-	
-	/* Override viewport height. */
-	float* ForcedViewportHeight;
 
 	/* Depth sort key. */
 	float Depth;
@@ -438,17 +386,27 @@
 	/* Background tile used to fixup 3d text issues. */
 	FCanvasTileItem	TileItem;
 
-	/** Get the type of font cache the UFont is using */
-	EFontCacheType GetFontCacheType() const;
+	/**
+	 * Does this text item have valid text to draw?
+	 */
+	virtual bool HasValidText() const = 0;
+
+	/**
+	 * Get the blend mode we should use to draw this text
+	 */
+	virtual ESimpleElementBlendMode GetTextBlendMode( const bool bHasShadow ) const = 0;
+
+	/**
+	 * Get the measured size of the text to draw
+	 */
+	virtual FVector2D GetTextSize() const = 0;
 
 	/** 
 	 * Internal string draw
 	 *
 	 * In a method to make it simpler to do effects like shadow, outline
 	 */
-	void DrawStringInternal( FCanvas* InCanvas, const FVector2D& DrawPos, const FLinearColor& DrawColor );
-	void DrawStringInternal_OfflineCache( FCanvas* InCanvas, const FVector2D& DrawPos, const FLinearColor& DrawColor );
-	void DrawStringInternal_RuntimeCache( FCanvas* InCanvas, const FVector2D& DrawPos, const FLinearColor& DrawColor );
+	virtual void DrawStringInternal( FCanvas* InCanvas, const FVector2D& DrawPos, const FLinearColor& DrawColor ) = 0;
 
 	/** 
 	 * These are used by the DrawStringInternal function. 
@@ -457,6 +415,99 @@
 	FBatchedElements* BatchedElements;
 };
 
+/* Text item with misc optional items such as shadow, centering etc. */
+class ENGINE_API FCanvasTextItem : public FCanvasTextItemBase
+{
+public:
+	/** 	 
+	 * Text item
+	 *
+	 * @param	InPosition		Draw position
+	 * @param	InText			String to draw
+	 * @param	InFont			Font to draw with
+	 */
+	FCanvasTextItem( const FVector2D& InPosition, const FText& InText, const UFont* InFont, const FLinearColor& InColor )
+		: FCanvasTextItemBase( InPosition, InColor )
+		, Text( InText )
+		, Font( InFont )
+		, SlateFontInfo()
+	{
+		BlendMode = SE_BLEND_Translucent;
+	}
+
+	/** 	 
+	 * Text item
+	 *
+	 * @param	InPosition		Draw position
+	 * @param	InText			String to draw
+	 * @param	InFontInfo		Font info to draw with
+	 */
+	FCanvasTextItem( const FVector2D& InPosition, const FText& InText, const FSlateFontInfo& InFontInfo, const FLinearColor& InColor )
+		: FCanvasTextItemBase( InPosition, InColor )
+		, Text( InText )
+		, Font( Cast<const UFont>(InFontInfo.FontObject) )
+		, SlateFontInfo( InFontInfo )
+	{
+		BlendMode = SE_BLEND_TranslucentAlphaOnly;
+	}
+	
+	/* The text to draw. */
+	FText Text;
+	
+	/* Font to draw text with. */
+	const UFont* Font;
+
+	/** Font info to draw the text with. */
+	TOptional<FSlateFontInfo> SlateFontInfo;
+
+protected:
+	/** Get the type of font cache the UFont is using */
+	EFontCacheType GetFontCacheType() const;
+
+	//~ FCanvasTextItemBase overrides
+	virtual bool HasValidText() const override;
+	virtual ESimpleElementBlendMode GetTextBlendMode( const bool bHasShadow ) const override;
+	virtual FVector2D GetTextSize() const override;
+	virtual void DrawStringInternal( FCanvas* InCanvas, const FVector2D& DrawPos, const FLinearColor& DrawColor ) override;
+
+	/** 
+	 * Internal string draw
+	 *
+	 * In a method to make it simpler to do effects like shadow, outline
+	 */
+	void DrawStringInternal_OfflineCache( FCanvas* InCanvas, const FVector2D& DrawPos, const FLinearColor& DrawColor );
+	void DrawStringInternal_RuntimeCache( FCanvas* InCanvas, const FVector2D& DrawPos, const FLinearColor& DrawColor );
+};
+
+/* Text item which can handle complex shaped text. */
+class ENGINE_API FCanvasShapedTextItem : public FCanvasTextItemBase
+{
+public:
+	/** 	 
+	 * Text item
+	 *
+	 * @param	InPosition		Draw position
+	 * @param	InText			String to draw
+	 * @param	InFontInfo		Font info to draw with
+	 */
+	FCanvasShapedTextItem( const FVector2D& InPosition, FShapedGlyphSequenceRef InShapedGlyphSequence, const FLinearColor& InColor )
+		: FCanvasTextItemBase( InPosition, InColor )
+		, ShapedGlyphSequence( InShapedGlyphSequence )
+	{
+		BlendMode = SE_BLEND_TranslucentAlphaOnly;
+	}
+
+	/** The shaped text sequence to draw */
+	FShapedGlyphSequencePtr ShapedGlyphSequence;
+
+protected:
+	//~ FCanvasTextItemBase overrides
+	virtual bool HasValidText() const override;
+	virtual ESimpleElementBlendMode GetTextBlendMode( const bool bHasShadow ) const override;
+	virtual FVector2D GetTextSize() const override;
+	virtual void DrawStringInternal( FCanvas* InCanvas, const FVector2D& DrawPos, const FLinearColor& DrawColor ) override;
+};
+
 /* Line item. */
 class ENGINE_API FCanvasLineItem : public FCanvasItem
 {
diff --git a/Engine/Source/Runtime/Engine/Public/Model.h b/Engine/Source/Runtime/Engine/Public/Model.h
--- a/Engine/Source/Runtime/Engine/Public/Model.h
+++ b/Engine/Source/Runtime/Engine/Public/Model.h
@@ -208,9 +208,9 @@ public:
 
 	int32					iLightmassIndex;// 4 Index to the lightmass settings
 
-	bool				bHiddenEdTemporary;	// 4 Marks whether this surface is temporarily hidden in the editor or not. Not serialized.
-	bool				bHiddenEdLevel;		// 4 Marks whether this surface is hidden by the level browser or not. Not serialized.
-
+	bool				bHiddenEdTemporary;	// 1 Marks whether this surface is temporarily hidden in the editor or not. Not serialized.
+	bool				bHiddenEdLevel;		// 1 Marks whether this surface is hidden by the level browser or not. Not serialized.
+	bool				bHiddenEdLayer;		// 1 Marks whether this surface is hidden by the layer browser or not. Not serialized.
 	// Functions.
 
 
diff --git a/Engine/Source/Runtime/Engine/Public/Slate/SceneViewport.h b/Engine/Source/Runtime/Engine/Public/Slate/SceneViewport.h
--- a/Engine/Source/Runtime/Engine/Public/Slate/SceneViewport.h
+++ b/Engine/Source/Runtime/Engine/Public/Slate/SceneViewport.h
@@ -215,6 +215,8 @@ public:
 	virtual FReply OnFocusReceived( const FFocusEvent& InFocusEvent ) override;
 	virtual void OnFocusLost( const FFocusEvent& InFocusEvent ) override;
 	virtual void OnViewportClosed() override;
+	virtual FReply OnViewportActivated(const FWindowActivateEvent& InActivateEvent) override;
+	virtual void OnViewportDeactivated(const FWindowActivateEvent& InActivateEvent) override;
 	virtual FIntPoint GetSize() const override { return GetSizeXY(); }
 	
 	void SetViewportSize(uint32 NewSizeX,uint32 NewSizeY);
@@ -291,6 +293,11 @@ private:
 	 */
 	void ApplyModifierKeys( const FModifierKeysState& InKeysState );
 
+	/** Utility function to create an FReply that properly gets Focus and capture based on the settings*/
+	FReply AcquireFocusAndCapture(FIntPoint MousePosition);
+
+	/** Utility function to figure out if we are currently a game viewport */
+	bool IsCurrentlyGameViewport();
 
 	void WindowRenderTargetUpdate(FSlateRenderer* Renderer, SWindow* Window);
 
@@ -321,6 +328,11 @@ private:
 	FIntPoint MouseDelta;
 	/** true if the cursor is currently visible */
 	bool bIsCursorVisible;
+	/** true if we had Capture when deactivated */
+	bool bShouldCaptureMouseOnActivate;
+	/** true if we want to ignore input processing on the click that give the viewport capture.
+	  * This only applies to Permanent Capture mode, the temporary mode always process the input*/
+	bool bIgnoreCaptureClick;
 	/** true if this viewport requires vsync. */
 	bool bRequiresVsync;
 	/** true if this viewport renders to a separate render target.  false to render directly to the windows back buffer */
diff --git a/Engine/Source/Runtime/Engine/Public/UnrealClient.h b/Engine/Source/Runtime/Engine/Public/UnrealClient.h
--- a/Engine/Source/Runtime/Engine/Public/UnrealClient.h
+++ b/Engine/Source/Runtime/Engine/Public/UnrealClient.h
@@ -17,6 +17,7 @@ class FViewportClient;
 
 class SWidget;
 class FCursorReply;
+class FWindowActivateEvent;
 
 enum class EFocusCause : uint8;
 
@@ -807,6 +808,9 @@ public:
 	virtual void ReceivedFocus(FViewport* Viewport) {}
 	virtual bool IsFocused(FViewport* Viewport) { return true; }
 
+	virtual void Activated(FViewport* Viewport, const FWindowActivateEvent& InActivateEvent) {}
+	virtual void Deactivated(FViewport* Viewport, const FWindowActivateEvent& InActivateEvent) {}
+
 	virtual void CloseRequested(FViewport* Viewport) {}
 
 	virtual bool RequiresHitProxyStorage() { return true; }
diff --git a/Engine/Source/Runtime/Foliage/Public/FoliageType.h b/Engine/Source/Runtime/Foliage/Public/FoliageType.h
--- a/Engine/Source/Runtime/Foliage/Public/FoliageType.h
+++ b/Engine/Source/Runtime/Foliage/Public/FoliageType.h
@@ -226,7 +226,7 @@ public:
 	uint32 bReceivesDecals : 1;
 	
 	/** Whether to override the lightmap resolution defined in the static mesh. */
-	UPROPERTY(BlueprintReadOnly, Category=InstanceSettings)
+	UPROPERTY(EditAnywhere, BlueprintReadOnly, Category=InstanceSettings, meta=(InlineEditConditionToggle))
 	uint32 bOverrideLightMapRes:1;
 
 	/** Overrides the lightmap resolution defined in the static mesh */
diff --git a/Engine/Source/Runtime/GameplayTasks/Classes/GameplayTaskResource.h b/Engine/Source/Runtime/GameplayTasks/Classes/GameplayTaskResource.h
--- a/Engine/Source/Runtime/GameplayTasks/Classes/GameplayTaskResource.h
+++ b/Engine/Source/Runtime/GameplayTasks/Classes/GameplayTaskResource.h
@@ -18,7 +18,7 @@ private:
 	int8 AutoResourceID;
 
 public:
-	UPROPERTY()
+	UPROPERTY(EditDefaultsOnly, Category = "Task Resource", meta=(InlineEditConditionToggle))
 	uint32 bManuallySetID : 1;
 
 public:
diff --git a/Engine/Source/Runtime/Slate/Private/Framework/Application/SlateApplication.cpp b/Engine/Source/Runtime/Slate/Private/Framework/Application/SlateApplication.cpp
--- a/Engine/Source/Runtime/Slate/Private/Framework/Application/SlateApplication.cpp
+++ b/Engine/Source/Runtime/Slate/Private/Framework/Application/SlateApplication.cpp
@@ -10,6 +10,7 @@
 #include "GenericCommands.h"
 #include "NotificationManager.h"
 #include "IInputProcessor.h"
+#include "ITextInputMethodSystem.h"
 #include "ToolboxModule.h"
 #include "TabCommands.h"
 
@@ -347,9 +348,6 @@ namespace SlateDefs
 
 	// How far tool tips should be pushed out from a force field border, in pixels
 	static const FVector2D ToolTipOffsetFromForceField( 4.0f, 3.0f );
-
-	// When true we will clear keyboard focus from widgets within the game view port when the window loses focus
-	static bool ClearGameViewportFocusWhenWindowLosesFocus = false;
 }
 
 
@@ -2491,15 +2489,6 @@ bool FSlateApplication::SetUserFocus(const uint32 InUserIndex, const FWidgetPath
 	if (NewFocusedWidget.IsValid())
 	{
 		TSharedPtr<SWindow> FocusedWindow = NewFocusedWidgetPath.GetWindow();
-		
-		// Set the windows restore state - we do this before calling OnFocusReceived to ensure that if
-		// OnFocusReceived returns a new UserFocus, the WidgetToFocusOnActivate will get set correctly
-		//@Todo Slate: We need to store the full focus state
-		//@Todo Slate: Why are we checking FocusedWindow != NewFocusedWidget?
-		if (FocusedWindow.IsValid() && FocusedWindow != NewFocusedWidget)
-		{
-			FocusedWindow->SetWidgetToFocusOnActivate(NewFocusedWidget);
-		}
 
 		// Switch worlds for widgets in the new path
 		FScopedSwitchWorldHack SwitchWorld(FocusedWindow);
@@ -2525,6 +2514,18 @@ void FSlateApplication::SetAllUserFocus(const FWidgetPath& InFocusPath, const EF
 	}
 }
 
+void FSlateApplication::SetAllUserFocusAllowingDescendantFocus(const FWidgetPath& InFocusPath, const EFocusCause InCause)
+{
+	TSharedRef<SWidget> FocusWidget = InFocusPath.Widgets.Last().Widget;
+	for (int32 SlateUserIndex = 0; SlateUserIndex < SlateApplicationDefs::MaxUsers; ++SlateUserIndex)
+	{
+		const FUserFocusEntry& UserFocusEntry = UserFocusEntries[SlateUserIndex];
+		if (!UserFocusEntry.WidgetPath.ContainsWidget(FocusWidget))
+		{
+			SetUserFocus(SlateUserIndex, InFocusPath, InCause);
+		}
+	}
+}
 
 FModifierKeysState FSlateApplication::GetModifierKeys() const
 {
@@ -2698,48 +2699,60 @@ void FSlateApplication::ProcessReply( const FWidgetPath& CurrentEventPath, const
 		});
 	}
 	
-	TSharedPtr<SWidget> RequestedMouseCaptor = TheReply.GetMouseCaptor();
-	// Do not capture the mouse if we are also starting a drag and drop.
-	if( RequestedMouseCaptor.IsValid() && !bStartingDragDrop )
-	{
-		MouseCaptor.SetMouseCaptor(UserIndex, PointerIndex, CurrentEventPath, RequestedMouseCaptor );
-	}
-	
-	if( CurrentEventPath.IsValid() && ( TheReply.ShouldReleaseMouse() || RequestedMouseCaptor.IsValid() ) )
+	// Setting mouse capture, mouse position, and locking the mouse
+	// are all operations that we shouldn't do if our application isn't Active (The OS ignores half of it, and we'd be in a half state)
+	// We do allow the release of capture and lock when deactivated, this is innocuous of some platforms but required on others when 
+	// the Application deactivated before the window. (Mac is an example of this)
+	if (bAppIsActive)
 	{
-		// If the mouse is being captured or released, toggle high precision raw input if requested by the reply.
-		// Raw input is only used with mouse capture
-		const TSharedRef< SWindow> Window = CurrentEventPath.GetWindow();
+		TSharedPtr<SWidget> RequestedMouseCaptor = TheReply.GetMouseCaptor();
+		// Do not capture the mouse if we are also starting a drag and drop.
+		if (RequestedMouseCaptor.IsValid() && !bStartingDragDrop)
+		{
+			MouseCaptor.SetMouseCaptor(UserIndex, PointerIndex, CurrentEventPath, RequestedMouseCaptor);
+		}
 
-		if ( TheReply.ShouldUseHighPrecisionMouse() )
+		if (CurrentEventPath.IsValid() && RequestedMouseCaptor.IsValid())
 		{
-			PlatformApplication->SetCapture( Window->GetNativeWindow() );
-			PlatformApplication->SetHighPrecisionMouseMode( true, Window->GetNativeWindow() );
+			// If the mouse is being captured or released, toggle high precision raw input if requested by the reply.
+			// Raw input is only used with mouse capture
+			if (TheReply.ShouldUseHighPrecisionMouse())
+			{
+				const TSharedRef< SWindow> Window = CurrentEventPath.GetWindow();
+				PlatformApplication->SetCapture(Window->GetNativeWindow());
+				PlatformApplication->SetHighPrecisionMouseMode(true, Window->GetNativeWindow());
+			}
 		}
-		else if ( PlatformApplication->IsUsingHighPrecisionMouseMode() )
+
+		TOptional<FIntPoint> RequestedMousePos = TheReply.GetRequestedMousePos();
+		if (RequestedMousePos.IsSet())
 		{
-			PlatformApplication->SetHighPrecisionMouseMode( false, nullptr );
-			PlatformApplication->SetCapture( nullptr );
+			const FVector2D Position = RequestedMousePos.GetValue();
+			PointerIndexLastPositionMap.Add(CursorPointerIndex, Position);
+			SetCursorPos(Position);
 		}
-	}
 
-	TOptional<FIntPoint> RequestedMousePos = TheReply.GetRequestedMousePos();
-	if( RequestedMousePos.IsSet() )
-	{
-		const FVector2D Position = RequestedMousePos.GetValue();
-		PointerIndexLastPositionMap.Add(CursorPointerIndex, Position);
-		SetCursorPos( Position );
+		if (TheReply.GetMouseLockWidget().IsValid())
+		{
+			// The reply requested mouse lock so tell the native application to lock the mouse to the widget receiving the event
+			LockCursor(TheReply.GetMouseLockWidget());
+		}
 	}
 
-	if( TheReply.GetMouseLockWidget().IsValid() )
+	// Releasing Mouse capture
+	if (CurrentEventPath.IsValid() && TheReply.ShouldReleaseMouse())
 	{
-		// The reply requested mouse lock so tell the native application to lock the mouse to the widget receiving the event
-		LockCursor( TheReply.GetMouseLockWidget() );
+		if (PlatformApplication->IsUsingHighPrecisionMouseMode())
+		{
+			PlatformApplication->SetHighPrecisionMouseMode(false, nullptr);
+			PlatformApplication->SetCapture(nullptr);
+		}
 	}
-	else if( TheReply.ShouldReleaseMouseLock() )
+
+	// Releasing Mouse Lock
+	if (TheReply.ShouldReleaseMouseLock())
 	{
-		// Unlock the mouse
-		LockCursor( nullptr );
+		LockCursor(nullptr);
 	}
 	
 	// If we have a valid Navigation request and enough time has passed since the last navigation then try to navigate.
@@ -2919,6 +2932,7 @@ void FSlateApplication::QueryCursor()
 				for (int32 WidgetIndex = WidgetsToQueryForCursor.Widgets.Num() - 1; WidgetIndex >= 0; --WidgetIndex)
 				{
 					const FArrangedWidget& ArrangedWidget = WidgetsToQueryForCursor.Widgets[WidgetIndex];
+
 					CursorReply = ArrangedWidget.Widget->OnCursorQuery(ArrangedWidget.Geometry, CursorEvent);
 					if (CursorReply.IsEventHandled())
 					{
@@ -3066,7 +3080,7 @@ void FSlateApplication::UpdateToolTip( bool AllowSpawningOfNewToolTips )
 	const bool bCheckForToolTipChanges =
 		bAllowToolTips &&					// Tool-tips must be enabled
 		!IsUsingHighPrecisionMouseMovment() && // If we are using HighPrecision movement then we can't rely on the OS cursor to be accurate
-		!IsDragDropping();					// We must not currently be in the middle of a drag-drop action
+		!IsDragDropping();					// We must not currwently be in the middle of a drag-drop action
 	
 	// We still want to show tooltips for widgets that are disabled
 	const bool bIgnoreEnabledStatus = true;
@@ -3748,7 +3762,7 @@ bool FSlateApplication::TakeScreenshot(const TSharedRef<SWidget>& Widget, const
 	OutSize.X = ScreenshotRect.Size().X;
 	OutSize.Y = ScreenshotRect.Size().Y;
 
-	return true;
+	return (OutSize.X != 0 && OutSize.Y != 0);
 }
 
 TSharedPtr< FSlateWindowElementList > FSlateApplication::GetCachableElementList(const TSharedPtr<SWindow>& CurrentWindow, const ILayoutCache* LayoutCache)
@@ -4963,7 +4977,7 @@ bool FSlateApplication::ProcessMouseWheelOrGestureEvent( FPointerEvent& InWheelE
 	QueueSynthesizedMouseMove();
 
 	const bool bShouldProcessEvent = InWheelEvent.GetWheelDelta() != 0 ||
-		(InGestureEvent != nullptr && InGestureEvent->GetGestureDelta()!=FVector2D::ZeroVector);
+		(InGestureEvent != nullptr && InGestureEvent->GetGestureDelta() != FVector2D::ZeroVector);
 	
 	if ( !bShouldProcessEvent )
 	{
@@ -4974,34 +4988,45 @@ bool FSlateApplication::ProcessMouseWheelOrGestureEvent( FPointerEvent& InWheelE
 	
 	// NOTE: We intentionally don't reset LastUserInteractionTimeForThrottling here so that the UI can be responsive while scrolling
 
-	FWidgetPath EventPath = ( MouseCaptor.HasCaptureForPointerIndex(InWheelEvent.GetUserIndex(), InWheelEvent.GetPointerIndex()) )
-		? MouseCaptor.ToWidgetPath( InWheelEvent.GetUserIndex(), InWheelEvent.GetPointerIndex() )
-		: LocateWindowUnderMouse( InWheelEvent.GetScreenSpacePosition(), GetInteractiveTopLevelWindows() );
+	FWidgetPath EventPath = LocateWindowUnderMouse(InWheelEvent.GetScreenSpacePosition(), GetInteractiveTopLevelWindows());
+
+	return RouteMouseWheelOrGestureEvent(EventPath, InWheelEvent, InGestureEvent).IsEventHandled();
+}
+
+FReply FSlateApplication::RouteMouseWheelOrGestureEvent(const FWidgetPath& WidgetsUnderPointer, const FPointerEvent& InWheelEvent, const FPointerEvent* InGestureEvent)
+{
+	FWidgetPath MouseCaptorPath;
+	if ( MouseCaptor.HasCaptureForPointerIndex(InWheelEvent.GetUserIndex(), InWheelEvent.GetPointerIndex()) )
+	{
+		MouseCaptorPath = MouseCaptor.ToWidgetPath(FWeakWidgetPath::EInterruptedPathHandling::ReturnInvalid, &InWheelEvent);
+	}
+
+	const FWidgetPath& EventPath = MouseCaptorPath.IsValid() ? MouseCaptorPath : WidgetsUnderPointer;
 
 	// Switch worlds widgets in the current path
-	FScopedSwitchWorldHack SwitchWorld( EventPath );
+	FScopedSwitchWorldHack SwitchWorld(EventPath);
 
-	FReply Reply = FEventRouter::Route<FReply>( this, FEventRouter::FBubblePolicy(EventPath), InWheelEvent, [&InGestureEvent](const FArrangedWidget& CurWidget, const FPointerEvent& Event)
+	FReply Reply = FEventRouter::Route<FReply>(this, FEventRouter::FBubblePolicy(EventPath), InWheelEvent, [&InGestureEvent] (const FArrangedWidget& CurWidget, const FPointerEvent& Event)
 	{
 		FReply TempReply = FReply::Unhandled();
 		// Gesture event gets first shot, if slate doesn't respond to it, we'll try the wheel event.
-		if( InGestureEvent != nullptr && InGestureEvent->GetGestureDelta() != FVector2D::ZeroVector )
+		if ( InGestureEvent != nullptr && InGestureEvent->GetGestureDelta() != FVector2D::ZeroVector )
 		{
-			TempReply = CurWidget.Widget->OnTouchGesture( CurWidget.Geometry, *InGestureEvent );
+			TempReply = CurWidget.Widget->OnTouchGesture(CurWidget.Geometry, *InGestureEvent);
 		}
-		
+
 		// Send the mouse wheel event if we haven't already handled the gesture version of this event.
-		if( !TempReply.IsEventHandled() && Event.GetWheelDelta() != 0 )
+		if ( !TempReply.IsEventHandled() && Event.GetWheelDelta() != 0 )
 		{
-			TempReply = CurWidget.Widget->OnMouseWheel( CurWidget.Geometry, Event );
+			TempReply = CurWidget.Widget->OnMouseWheel(CurWidget.Geometry, Event);
 		}
 
 		return TempReply;
-	} );
+	});
 
-	LOG_EVENT( bIsGestureEvent ? EEventLog::TouchGesture : EEventLog::MouseWheel, Reply );
+	LOG_EVENT(bIsGestureEvent ? EEventLog::TouchGesture : EEventLog::MouseWheel, Reply);
 
-	return Reply.IsEventHandled();
+	return Reply;
 }
 
 bool FSlateApplication::OnMouseMove()
@@ -5453,8 +5478,10 @@ bool FSlateApplication::ProcessWindowActivatedEvent( const FWindowActivateEvent&
 
 	if ( ActivateEvent.GetActivationType() != FWindowActivateEvent::EA_Deactivate )
 	{
+		const bool bActivatedByMouse = ActivateEvent.GetActivationType() == FWindowActivateEvent::EA_ActivateByMouse;
+		
 		// Only window activate by mouse is considered a user interaction
-		if (ActivateEvent.GetActivationType() == FWindowActivateEvent::EA_ActivateByMouse)
+		if (bActivatedByMouse)
 		{
 			LastUserInteractionTime = this->GetCurrentTime();
 		}
@@ -5511,6 +5538,25 @@ bool FSlateApplication::ProcessWindowActivatedEvent( const FWindowActivateEvent&
 			ActiveModalWindow->BringToFront();
 			ActiveModalWindow->FlashWindow();
 		}
+
+		// Notify the GameViewport that it's been activated if it's a descendant of the top level window were activating 
+		if (GameViewportWidget.IsValid())
+		{
+			TSharedPtr<SViewport> GameViewportWidgetPtr = GameViewportWidget.Pin();
+			FWidgetPath PathToViewport;
+			// If we cannot find the window it could have been destroyed.
+			if (FSlateWindowHelper::FindPathToWidget(SlateWindows, GameViewportWidgetPtr.ToSharedRef(), PathToViewport, EVisibility::All))
+			{
+				if (PathToViewport.GetWindow() == ActiveTopLevelWindow)
+				{
+					FReply ViewportActivatedReply = GameViewportWidgetPtr->OnViewportActivated(ActivateEvent);
+					if (ViewportActivatedReply.IsEventHandled())
+					{
+						ProcessReply(PathToViewport, ViewportActivatedReply, nullptr, nullptr);
+					}
+				}
+			}
+		}
 	}
 	else
 	{
@@ -5529,18 +5575,15 @@ bool FSlateApplication::ProcessWindowActivatedEvent( const FWindowActivateEvent&
 		// Switch worlds for the activated window
 		FScopedSwitchWorldHack SwitchWorld( ActivateEvent.GetAffectedWindow() );
 		ActivateEvent.GetAffectedWindow()->OnIsActiveChanged( ActivateEvent );
+		
+		// Notify the GameViewport that it's been deactivated
+		if (GameViewportWidget.IsValid())
+		{
+			GameViewportWidget.Pin()->OnViewportDeactivated(ActivateEvent);
+		}
 
 		// A window was deactivated; mouse capture should be cleared
 		ResetToDefaultPointerInputSettings();
-
-		// Clear user focus entries, so SetUserFocus restores mouse capture etc. after the window regains focus
-		for (int32 UserIndex = 0; UserIndex < SlateApplicationDefs::MaxUsers; ++UserIndex)
-		{
-			FUserFocusEntry& UserFocusEntry = UserFocusEntries[UserIndex];
-			UserFocusEntry.WidgetPath = FWidgetPath();
-			UserFocusEntry.FocusCause = EFocusCause::Cleared;
-			UserFocusEntry.ShowFocus = false;
-		}
 	}
 
 	return true;
@@ -5552,13 +5595,12 @@ bool FSlateApplication::OnApplicationActivationChanged( const bool IsActive )
 	return true;
 }
 
-void FSlateApplication::ProcessApplicationActivationEvent( bool InAppActivated )
+void FSlateApplication::ProcessApplicationActivationEvent(bool InAppActivated)
 {
 	const bool UserSwitchedAway = bAppIsActive && !InAppActivated;
 
 	bAppIsActive = InAppActivated;
 
-
 	// If the user switched to a different application then we should dismiss our pop-ups.  In the case
 	// where a user clicked on a different Slate window, OnWindowActivatedMessage() will be call MenuStack.OnWindowActivated()
 	// to destroy any windows in our stack that are no longer appropriate to be displayed.
@@ -5573,21 +5615,8 @@ void FSlateApplication::ProcessApplicationActivationEvent( bool InAppActivated )
 		// No slate window is active when our entire app becomes inactive
 		bSlateWindowActive = false;
 
-		// Clear keyboard focus when the app is deactivated, if the widget isn't a child of the game viewport.
-		bool bClearKeyboardFocus = true;
-
-		if (!SlateDefs::ClearGameViewportFocusWhenWindowLosesFocus && GameViewportWidget.IsValid())
-		{
-			FUserFocusEntry& UserFocusEntry = UserFocusEntries[GetUserIndexForKeyboard()];
-			bClearKeyboardFocus = !(UserFocusEntry.WidgetPath.ContainsWidget(GameViewportWidget.Pin().ToSharedRef()));
-		}
-		if (bClearKeyboardFocus)
-		{
-			ClearKeyboardFocus(EFocusCause::OtherWidgetLostFocus);
-		}
-
 		// If we have a slate-only drag-drop occurring, stop the drag drop.
-		if ( IsDragDropping() && !DragDropContent->IsExternalOperation() )
+		if (IsDragDropping() && !DragDropContent->IsExternalOperation())
 		{
 			DragDropContent.Reset();
 		}
diff --git a/Engine/Source/Runtime/Slate/Private/Framework/Text/PlainTextLayoutMarshaller.cpp b/Engine/Source/Runtime/Slate/Private/Framework/Text/PlainTextLayoutMarshaller.cpp
--- a/Engine/Source/Runtime/Slate/Private/Framework/Text/PlainTextLayoutMarshaller.cpp
+++ b/Engine/Source/Runtime/Slate/Private/Framework/Text/PlainTextLayoutMarshaller.cpp
@@ -3,8 +3,7 @@
 #include "SlatePrivatePCH.h"
 #include "PlainTextLayoutMarshaller.h"
 #include "SlateTextLayout.h"
-
-#if WITH_FANCY_TEXT
+#include "SlatePasswordRun.h"
 
 TSharedRef< FPlainTextLayoutMarshaller > FPlainTextLayoutMarshaller::Create()
 {
@@ -15,6 +14,11 @@ FPlainTextLayoutMarshaller::~FPlainTextLayoutMarshaller()
 {
 }
 
+void FPlainTextLayoutMarshaller::SetIsPassword(const TAttribute<bool>& InIsPassword)
+{
+	bIsPassword = InIsPassword;
+}
+
 void FPlainTextLayoutMarshaller::SetText(const FString& SourceString, FTextLayout& TargetTextLayout)
 {
 	const FTextBlockStyle& DefaultTextStyle = static_cast<FSlateTextLayout&>(TargetTextLayout).GetDefaultTextStyle();
@@ -25,12 +29,20 @@ void FPlainTextLayoutMarshaller::SetText(const FString& SourceString, FTextLayou
 	TArray<FTextLayout::FNewLineData> LinesToAdd;
 	LinesToAdd.Reserve(LineRanges.Num());
 
-	for(const FTextRange& LineRange : LineRanges)
+	const bool bUsePasswordRun = bIsPassword.Get(false);
+	for (const FTextRange& LineRange : LineRanges)
 	{
 		TSharedRef<FString> LineText = MakeShareable(new FString(SourceString.Mid(LineRange.BeginIndex, LineRange.Len())));
 
 		TArray<TSharedRef<IRun>> Runs;
-		Runs.Add(FSlateTextRun::Create(FRunInfo(), LineText, DefaultTextStyle));
+		if (bUsePasswordRun)
+		{
+			Runs.Add(FSlatePasswordRun::Create(FRunInfo(), LineText, DefaultTextStyle));
+		}
+		else
+		{
+			Runs.Add(FSlateTextRun::Create(FRunInfo(), LineText, DefaultTextStyle));
+		}
 
 		LinesToAdd.Emplace(MoveTemp(LineText), MoveTemp(Runs));
 	}
@@ -45,6 +57,5 @@ void FPlainTextLayoutMarshaller::GetText(FString& TargetString, const FTextLayou
 
 FPlainTextLayoutMarshaller::FPlainTextLayoutMarshaller()
 {
+	bIsPassword = false;
 }
-
-#endif //WITH_FANCY_TEXT
diff --git a/Engine/Source/Runtime/Slate/Private/Framework/Text/SlateTextHighlightRunRenderer.cpp b/Engine/Source/Runtime/Slate/Private/Framework/Text/SlateTextHighlightRunRenderer.cpp
--- a/Engine/Source/Runtime/Slate/Private/Framework/Text/SlateTextHighlightRunRenderer.cpp
+++ b/Engine/Source/Runtime/Slate/Private/Framework/Text/SlateTextHighlightRunRenderer.cpp
@@ -1,9 +1,6 @@
 // Copyright 1998-2016 Epic Games, Inc. All Rights Reserved.
 
 #include "SlatePrivatePCH.h"
-
-#if WITH_FANCY_TEXT
-
 #include "SlateTextHighlightRunRenderer.h"
 
 FSlateTextHighlightRunRenderer::FSlateTextHighlightRunRenderer()
@@ -43,7 +40,3 @@ TSharedRef< FSlateTextHighlightRunRenderer > FSlateTextHighlightRunRenderer::Cre
 {
 	return MakeShareable( new FSlateTextHighlightRunRenderer() );
 }
-
-
-
-#endif //WITH_FANCY_TEXT
\ No newline at end of file
diff --git a/Engine/Source/Runtime/Slate/Private/Framework/Text/SlateTextLayout.cpp b/Engine/Source/Runtime/Slate/Private/Framework/Text/SlateTextLayout.cpp
--- a/Engine/Source/Runtime/Slate/Private/Framework/Text/SlateTextLayout.cpp
+++ b/Engine/Source/Runtime/Slate/Private/Framework/Text/SlateTextLayout.cpp
@@ -1,10 +1,9 @@
 // Copyright 1998-2016 Epic Games, Inc. All Rights Reserved.
 
 #include "SlatePrivatePCH.h"
-
-#if WITH_FANCY_TEXT
-
 #include "SlateTextLayout.h"
+#include "SlateTextRun.h"
+#include "SlatePasswordRun.h"
 
 TSharedRef< FSlateTextLayout > FSlateTextLayout::Create(FTextBlockStyle InDefaultTextStyle)
 {
@@ -17,6 +16,7 @@ TSharedRef< FSlateTextLayout > FSlateTextLayout::Create(FTextBlockStyle InDefaul
 FSlateTextLayout::FSlateTextLayout(FTextBlockStyle InDefaultTextStyle)
 	: Children()
 	, DefaultTextStyle(MoveTemp(InDefaultTextStyle))
+	, bIsPassword(false)
 	, LocalizedFallbackFontRevision(INDEX_NONE)
 {
 
@@ -163,6 +163,11 @@ const FTextBlockStyle& FSlateTextLayout::GetDefaultTextStyle() const
 	return DefaultTextStyle;
 }
 
+void FSlateTextLayout::SetIsPassword(const TAttribute<bool>& InIsPassword)
+{
+	bIsPassword = InIsPassword;
+}
+
 void FSlateTextLayout::AggregateChildren()
 {
 	Children.Empty();
@@ -187,7 +192,9 @@ void FSlateTextLayout::AggregateChildren()
 
 TSharedRef<IRun> FSlateTextLayout::CreateDefaultTextRun(const TSharedRef<FString>& NewText, const FTextRange& NewRange) const
 {
+	if (bIsPassword.Get(false))
+	{
+		return FSlatePasswordRun::Create(FRunInfo(), NewText, DefaultTextStyle, NewRange);
+	}
 	return FSlateTextRun::Create(FRunInfo(), NewText, DefaultTextStyle, NewRange);
 }
-
-#endif //WITH_FANCY_TEXT
diff --git a/Engine/Source/Runtime/Slate/Private/Framework/Text/SlateTextRun.cpp b/Engine/Source/Runtime/Slate/Private/Framework/Text/SlateTextRun.cpp
--- a/Engine/Source/Runtime/Slate/Private/Framework/Text/SlateTextRun.cpp
+++ b/Engine/Source/Runtime/Slate/Private/Framework/Text/SlateTextRun.cpp
@@ -1,9 +1,6 @@
 // Copyright 1998-2016 Epic Games, Inc. All Rights Reserved.
 
 #include "SlatePrivatePCH.h"
-
-#if WITH_FANCY_TEXT
-
 #include "SlateTextRun.h"
 #include "ShapedTextCache.h"
 #include "RunUtils.h"
@@ -258,5 +255,3 @@ FSlateTextRun::FSlateTextRun( const FSlateTextRun& Run )
 {
 
 }
-
-#endif //WITH_FANCY_TEXT
diff --git a/Engine/Source/Runtime/Slate/Private/Framework/Text/TextEditHelper.cpp b/Engine/Source/Runtime/Slate/Private/Framework/Text/TextEditHelper.cpp
--- a/Engine/Source/Runtime/Slate/Private/Framework/Text/TextEditHelper.cpp
+++ b/Engine/Source/Runtime/Slate/Private/Framework/Text/TextEditHelper.cpp
@@ -3,463 +3,6 @@
 #include "SlatePrivatePCH.h"
 #include "TextEditHelper.h"
 
-/**
- * Ensure that text transactions are always completed.
- * i.e. never forget to call FinishChangingText();
- */
-struct FScopedTextTransaction
-{
-	explicit FScopedTextTransaction( const TSharedRef< ITextEditorWidget >& InTextEditor )
-	: TextEditor( InTextEditor )
-	{
-		TextEditor->StartChangingText();
-	}
-
-	~FScopedTextTransaction()
-	{
-		TextEditor->FinishChangingText();
-	}
-
-	const TSharedRef< ITextEditorWidget > TextEditor;
-};
-
-FReply FTextEditHelper::OnKeyChar( const FCharacterEvent& InCharacterEvent, const TSharedRef< ITextEditorWidget >& TextEditor )
-{
-	FReply Reply = FReply::Unhandled();
-
-	// Check for special characters
-	const TCHAR Character = InCharacterEvent.GetCharacter();
-
-	switch( Character )
-	{
-		// Backspace
-	case TCHAR( 8 ):
-		{
-			if( !TextEditor->GetIsReadOnly() )
-			{
-				FScopedTextTransaction TextTransaction(TextEditor);
-
-				TextEditor->BackspaceChar();
-				
-				Reply = FReply::Handled();
-			}
-		}
-		break;
-
-
-		// Tab
-	case TCHAR( '\t' ):
-		{
-			Reply = FReply::Handled();
-		}
-		break;
-
-
-		// Newline (Ctrl+Enter), we handle adding new lines via SMultiLineEditableText::OnEnter rather than processing \n characters
-	case TCHAR( '\n' ):
-		{
-			Reply = FReply::Handled();
-		}
-		break;
-
-
-		// Swallow OnKeyChar keys that we don't want to be entered into the buffer
-	case 1:		// Swallow Ctrl+A, we handle that through OnKeyDown
-	case 3:		// Swallow Ctrl+C, we handle that through OnKeyDown
-	case 13:	// Swallow Enter, we handle that through OnKeyDown
-	case 22:	// Swallow Ctrl+V, we handle that through OnKeyDown
-	case 24:	// Swallow Ctrl+X, we handle that through OnKeyDown
-	case 25:	// Swallow Ctrl+Y, we handle that through OnKeyDown
-	case 26:	// Swallow Ctrl+Z, we handle that through OnKeyDown
-	case 27:	// Swallow ESC, we handle that through OnKeyDown
-	case 127:	// Swallow CTRL+Backspace, we handle that through OnKeyDown
-		Reply = FReply::Handled();
-		break;
-
-
-		// Any other character!
-	default:
-		{
-			// Type the character, but only if it is allowed.
-			if (!TextEditor->GetIsReadOnly() && TextEditor->CanTypeCharacter(Character))
-			{
-				FScopedTextTransaction TextTransaction(TextEditor);
-				TextEditor->TypeChar( Character );
-				Reply = FReply::Handled();
-			}
-			else
-			{
-				Reply = FReply::Unhandled();
-			}
-		}
-		break;
-
-	}
-
-
-	return Reply;
-}
-
-FReply FTextEditHelper::OnKeyDown( const FKeyEvent& InKeyEvent, const TSharedRef< ITextEditorWidget >& TextEditor )
-{
-	const FKey Key = InKeyEvent.GetKey();
-
-	if( Key == EKeys::Left )
-	{
-		return TextEditor->MoveCursor( FMoveCursor::Cardinal(
-			// Ctrl moves a whole word instead of one character.	
-			InKeyEvent.IsControlDown( ) ? ECursorMoveGranularity::Word : ECursorMoveGranularity::Character,
-			// Move left
-			FIntPoint( -1, 0 ),
-			// Shift selects text.	
-			InKeyEvent.IsShiftDown() ? ECursorAction::SelectText : ECursorAction::MoveCursor
-		));
-	}
-	else if( Key == EKeys::Right )
-	{
-		return TextEditor->MoveCursor( FMoveCursor::Cardinal(
-			// Ctrl moves a whole word instead of one character.	
-			InKeyEvent.IsControlDown( ) ? ECursorMoveGranularity::Word : ECursorMoveGranularity::Character,
-			// Move right
-			FIntPoint( +1, 0 ),
-			// Shift selects text.	
-			InKeyEvent.IsShiftDown() ? ECursorAction::SelectText : ECursorAction::MoveCursor
-		));
-	}
-	else if ( Key == EKeys::Up )
-	{
-		return TextEditor->MoveCursor( FMoveCursor::Cardinal(
-			ECursorMoveGranularity::Character,
-			// Move up
-			FIntPoint( 0, -1 ),
-			// Shift selects text.	
-			InKeyEvent.IsShiftDown() ? ECursorAction::SelectText : ECursorAction::MoveCursor
-		));
-	}
-	else if ( Key == EKeys::Down )
-	{
-		return TextEditor->MoveCursor( FMoveCursor::Cardinal(
-			ECursorMoveGranularity::Character,
-			// Move down
-			FIntPoint( 0, +1 ),
-			// Shift selects text.	
-			InKeyEvent.IsShiftDown() ? ECursorAction::SelectText : ECursorAction::MoveCursor
-		));
-	}
-	else if( Key == EKeys::Home )
-	{
-		// Go to the beginning of the document; select text if Shift is down.
-		TextEditor->JumpTo(
-			(InKeyEvent.IsControlDown() ) ? ETextLocation::BeginningOfDocument : ETextLocation::BeginningOfLine,
-			(InKeyEvent.IsShiftDown()) ? ECursorAction::SelectText : ECursorAction::MoveCursor );
-
-		return FReply::Handled();
-	}
-	else if ( Key == EKeys::End )
-	{
-		// Go to the end of the document; select text if Shift is down.
-		TextEditor->JumpTo(
-			(InKeyEvent.IsControlDown() ) ? ETextLocation::EndOfDocument : ETextLocation::EndOfLine,
-			(InKeyEvent.IsShiftDown()) ? ECursorAction::SelectText : ECursorAction::MoveCursor );
-
-		return FReply::Handled();
-	}
-	else if( Key == EKeys::PageUp )
-	{
-		// Go to the previous page of the document document; select text if Shift is down.
-		TextEditor->JumpTo(
-			ETextLocation::PreviousPage,
-			(InKeyEvent.IsShiftDown()) ? ECursorAction::SelectText : ECursorAction::MoveCursor );
-
-		return FReply::Handled();
-	}
-	else if ( Key == EKeys::PageDown )
-	{
-		// Go to the next page of the document document; select text if Shift is down.
-		TextEditor->JumpTo(
-			ETextLocation::NextPage,
-			(InKeyEvent.IsShiftDown()) ? ECursorAction::SelectText : ECursorAction::MoveCursor );
-
-		return FReply::Handled();
-	}
-	else if( Key == EKeys::Enter && !TextEditor->GetIsReadOnly() )
-	{
-		FScopedTextTransaction TextTransaction(TextEditor);
-
-		TextEditor->OnEnter();
-
-		return FReply::Handled();
-	}
-	else if( Key == EKeys::Delete && !TextEditor->GetIsReadOnly() )
-	{
-		// @Todo: Slate keybindings support more than one set of keys. 
-		// Delete to next word boundary (Ctrl+Delete)
-		if (InKeyEvent.IsControlDown() && !InKeyEvent.IsAltDown() && !InKeyEvent.IsShiftDown())
-		{
-			TextEditor->MoveCursor( FMoveCursor::Cardinal(
-				ECursorMoveGranularity::Word, 
-				// Move right
-				FIntPoint(+1, 0),
-				// selects text.	
-				ECursorAction::SelectText
-			));
-		}
-
-		FScopedTextTransaction TextTransaction(TextEditor);
-
-		// Delete selected text
-		TextEditor->DeleteChar();
-		
-		return FReply::Handled();
-	}	
-	else if( Key == EKeys::Escape )
-	{
-		return TextEditor->OnEscape();
-	}
-
-	// @Todo: Slate keybindings support more than one set of keys. 
-	//Alternate key for cut (Shift+Delete)
-	else if( Key == EKeys::Delete && InKeyEvent.IsShiftDown() && TextEditor->CanExecuteCut() )
-	{
-		// Cut text to clipboard
-		TextEditor->CutSelectedTextToClipboard();
-		
-		return FReply::Handled();
-	}
-
-	// @Todo: Slate keybindings support more than one set of keys. 
-	// Alternate key for copy (Ctrl+Insert) 
-	else if( Key == EKeys::Insert && InKeyEvent.IsControlDown() && TextEditor->CanExecuteCopy() ) 
-	{
-		// Copy text to clipboard
-		TextEditor->CopySelectedTextToClipboard();
-		
-		return FReply::Handled();
-	}
-
-
-	// @Todo: Slate keybindings support more than one set of keys. 
-	// Alternate key for paste (Shift+Insert) 
-	else if( Key == EKeys::Insert && InKeyEvent.IsShiftDown() && TextEditor->CanExecutePaste() )
-	{
-		// Paste text from clipboard
-		TextEditor->PasteTextFromClipboard();
-		
-		return FReply::Handled();
-	}
-
-	// @Todo: Slate keybindings support more than one set of keys. 
-	//Alternate key for undo (Alt+Backspace)
-	else if( Key == EKeys::BackSpace && InKeyEvent.IsAltDown() && !InKeyEvent.IsShiftDown() && TextEditor->CanExecuteUndo() )
-	{
-		// Undo
-		TextEditor->Undo();
-		
-		return FReply::Handled();
-	}
-
-	// @Todo: Slate keybindings support more than one set of keys. 
-	// Delete to previous word boundary (Ctrl+Backspace)
-	else if( Key == EKeys::BackSpace && InKeyEvent.IsControlDown() && !InKeyEvent.IsAltDown() && !InKeyEvent.IsShiftDown() && !TextEditor->GetIsReadOnly() )
-	{
-		FScopedTextTransaction TextTransaction(TextEditor);
-
-		TextEditor->MoveCursor( FMoveCursor::Cardinal(
-			ECursorMoveGranularity::Word, 
-			// Move left
-			FIntPoint(-1, 0), 
-			ECursorAction::SelectText
-		));
-		TextEditor->BackspaceChar();
-
-		return FReply::Handled();
-	}
-
-
-	// Ctrl+Y (or Ctrl+Shift+Z, or Alt+Shift+Backspace) to redo
-	else if( !TextEditor->GetIsReadOnly() && ( ( Key == EKeys::Y && InKeyEvent.IsControlDown() ) ||
-		( Key == EKeys::Z && InKeyEvent.IsControlDown() && InKeyEvent.IsShiftDown() ) ||
-		( Key == EKeys::BackSpace && InKeyEvent.IsAltDown() && InKeyEvent.IsShiftDown() ) ) )
-	{
-		// Redo
-		TextEditor->Redo();
-		
-		return FReply::Handled();
-	}
-	else if( !InKeyEvent.IsAltDown() && !InKeyEvent.IsControlDown() && InKeyEvent.GetKey() != EKeys::Tab && InKeyEvent.GetCharacter() != 0 )
-	{
-		// Shift and a character was pressed or a single character was pressed.  We will type something in an upcoming OnKeyChar event.  
-		// Absorb this event so it is not bubbled and handled by other widgets that could have something bound to the key press.
-
-		//Note: Tab can generate a character code but not result in a typed character in this box so we ignore it.
-		return FReply::Handled();
-	}
-	else
-	{
-		return FReply::Unhandled();
-	}
-}
-FReply FTextEditHelper::OnMouseButtonDown( const FGeometry& MyGeometry, const FPointerEvent& InMouseEvent, const TSharedRef< ITextEditorWidget >& TextEditor )
-{
-	FReply Reply = FReply::Unhandled();
-
-	// If the mouse is already captured, then don't allow a new action to be taken
-	if( !TextEditor->GetWidget()->HasMouseCapture() )
-	{
-		if( InMouseEvent.GetEffectingButton() == EKeys::LeftMouseButton ||
-			InMouseEvent.GetEffectingButton() == EKeys::RightMouseButton )
-		{
-			// Am I getting focus right now?
-			const bool bIsGettingFocus = !TextEditor->GetWidget()->HasKeyboardFocus();
-			if( bIsGettingFocus )
-			{
-				// We might be receiving keyboard focus due to this event.  Because the keyboard focus received callback
-				// won't fire until after this function exits, we need to make sure our widget's state is in order early
-
-				// Assume we'll be given keyboard focus, so load text for editing
-				TextEditor->LoadText();
-
-				// Reset 'mouse has moved' state.  We'll use this in OnMouseMove to determine whether we
-				// should reset the selection range to the caret's position.
-				TextEditor->SetWasFocusedByLastMouseDown( true );
-			}
-			else
-			{
-				// On platforms using a virtual keyboard open the virtual keyboard again 
-				FSlateApplication& SlateApplication = FSlateApplication::Get();
-				if (FPlatformMisc::GetRequiresVirtualKeyboard())
-				{
-					if (!TextEditor->GetIsReadOnly())
-					{
-						SlateApplication.ShowVirtualKeyboard(true, InMouseEvent.GetUserIndex(), TextEditor->GetVirtualKeyboardEntry());
-					}
-				}
-			}
-
-			if( InMouseEvent.GetEffectingButton() == EKeys::LeftMouseButton )
-			{
-				if( InMouseEvent.IsShiftDown() )
-				{
-					TextEditor->MoveCursor( FMoveCursor::ViaScreenPointer( MyGeometry.AbsoluteToLocal( InMouseEvent.GetScreenSpacePosition( ) ), MyGeometry.Scale, ECursorAction::SelectText ) );
-				}
-				else
-				{
-					// Deselect any text that was selected
-					TextEditor->ClearSelection();
-					TextEditor->MoveCursor( FMoveCursor::ViaScreenPointer( MyGeometry.AbsoluteToLocal( InMouseEvent.GetScreenSpacePosition( ) ), MyGeometry.Scale, ECursorAction::MoveCursor ) );
-				}
-
-				// Start drag selection
-				TextEditor->BeginDragSelection();
-			}
-			else if( InMouseEvent.GetEffectingButton() == EKeys::RightMouseButton )
-			{
-				// If the user right clicked on a character that wasn't already selected, we'll clear
-				// the selection
-				if ( TextEditor->AnyTextSelected( ) && !TextEditor->IsTextSelectedAt( MyGeometry, InMouseEvent.GetScreenSpacePosition() ) )
-				{
-					// Deselect any text that was selected
-					TextEditor->ClearSelection();
-				}
-			}
-
-			// Right clicking to summon context menu, but we'll do that on mouse-up.
-			Reply = FReply::Handled();
-			Reply.CaptureMouse( TextEditor->GetWidget() );
-			Reply.SetUserFocus(TextEditor->GetWidget(), EFocusCause::Mouse);
-		}
-	}
-
-	return Reply;
-}
-
-FReply FTextEditHelper::OnMouseMove( const FGeometry& InMyGeometry, const FPointerEvent& InMouseEvent, const TSharedRef< ITextEditorWidget >& TextEditor )
-{
-	FReply Reply = FReply::Unhandled();
-
-	if( TextEditor->IsDragSelecting() && TextEditor->GetWidget()->HasMouseCapture() && InMouseEvent.GetCursorDelta() != FVector2D::ZeroVector )
-	{
-		TextEditor->MoveCursor( FMoveCursor::ViaScreenPointer( InMyGeometry.AbsoluteToLocal( InMouseEvent.GetScreenSpacePosition( ) ), InMyGeometry.Scale, ECursorAction::SelectText ) );
-		TextEditor->SetHasDragSelectedSinceFocused( true );
-		Reply = FReply::Handled();
-	}
-
-	return Reply;
-}
-
-FReply FTextEditHelper::OnMouseButtonUp( const FGeometry& MyGeometry, const FPointerEvent& InMouseEvent, const TSharedRef< ITextEditorWidget >& TextEditor )
-{
-	FReply Reply = FReply::Unhandled();
-
-	// The mouse must have been captured by either left or right button down before we'll process mouse ups
-	if( TextEditor->GetWidget()->HasMouseCapture() )
-	{
-		if( InMouseEvent.GetEffectingButton() == EKeys::LeftMouseButton &&
-			TextEditor->IsDragSelecting() )
-		{
-			// No longer drag-selecting
-			TextEditor->EndDragSelection();
-
-			// If we received keyboard focus on this click, then we'll want to select all of the text
-			// when the user releases the mouse button, unless the user actually dragged the mouse
-			// while holding the button down, in which case they've already selected some text and
-			// we'll leave things alone!
-			if( TextEditor->WasFocusedByLastMouseDown() )
-			{
-				if( !TextEditor->HasDragSelectedSinceFocused() )
-				{
-					if( TextEditor->SelectAllTextWhenFocused() )
-					{
-						// Move the cursor to the end of the string
-						TextEditor->JumpTo(ETextLocation::EndOfDocument, ECursorAction::MoveCursor);
-
-						// User wasn't dragging the mouse, so go ahead and select all of the text now
-						// that we've become focused
-						TextEditor->SelectAllText();
-
-						// @todo Slate: In this state, the caret should actually stay hidden (until the user interacts again), and we should not move the caret
-					}
-				}
-
-				TextEditor->SetWasFocusedByLastMouseDown( false );
-			}
-
-			// Release mouse capture
-			Reply = FReply::Handled();
-			Reply.ReleaseMouseCapture();
-		}
-		else if( InMouseEvent.GetEffectingButton() == EKeys::RightMouseButton )
-		{
-			if ( MyGeometry.IsUnderLocation( InMouseEvent.GetScreenSpacePosition() ) )
-			{
-				// Right clicked, so summon a context menu if the cursor is within the widget
-				FWidgetPath WidgetPath = InMouseEvent.GetEventPath() != nullptr ? *InMouseEvent.GetEventPath() : FWidgetPath();
-				TextEditor->SummonContextMenu(InMouseEvent.GetScreenSpacePosition(), InMouseEvent.GetWindow(), WidgetPath);
-			}
-
-			// Release mouse capture
-			Reply = FReply::Handled();
-			Reply.ReleaseMouseCapture();
-		}
-	}
-
-	return Reply;
-}
-
-FReply FTextEditHelper::OnMouseButtonDoubleClick(const FGeometry& InMyGeometry, const FPointerEvent& InMouseEvent, const TSharedRef< ITextEditorWidget >& TextEditor)
-{
-	FReply Reply = FReply::Unhandled();
-
-	if (InMouseEvent.GetEffectingButton() == EKeys::LeftMouseButton)
-	{
-		TextEditor->SelectWordAt(InMyGeometry, InMouseEvent.GetScreenSpacePosition());
-
-		Reply = FReply::Handled();
-	}
-
-	return Reply;
-}
-
 float FTextEditHelper::GetFontHeight(const FSlateFontInfo& FontInfo)
 {
 	const TSharedRef< FSlateFontMeasure > FontMeasure = FSlateApplication::Get().GetRenderer()->GetFontMeasureService();
@@ -468,9 +11,7 @@ float FTextEditHelper::GetFontHeight(const FSlateFontInfo& FontInfo)
 
 float FTextEditHelper::CalculateCaretWidth(const float FontMaxCharHeight)
 {
-	// We adjust the width of the cursor because on smaller fonts the cursor would overlap too much of the characters it is next too.
-	// Other programs get around this by inverting the colors under the cursor. This causes a single character to be rendered as two characters.
-
-	// The caret with should be at least one pixel. For very small fonts the width could be < 1 which makes it not visible
-	return FMath::Max(1.0f, EditableTextDefs::CaretWidthPercent * FontMaxCharHeight);
+	// We adjust the width of the caret to avoid it becoming too wide on smaller or larger fonts and overlapping the characters it's next to.
+	// We clamp the lower limit to 1 to avoid it being invisible, and the upper limit to 2 to avoid tall fonts having very wide carets.
+	return FMath::Clamp(EditableTextDefs::CaretWidthPercent * FontMaxCharHeight, 1.0f, 2.0f);
 }
diff --git a/Engine/Source/Runtime/Slate/Private/Framework/Text/TextEditHelper.h b/Engine/Source/Runtime/Slate/Private/Framework/Text/TextEditHelper.h
--- a/Engine/Source/Runtime/Slate/Private/Framework/Text/TextEditHelper.h
+++ b/Engine/Source/Runtime/Slate/Private/Framework/Text/TextEditHelper.h
@@ -10,71 +10,22 @@ namespace EditableTextDefs
 	/** Maximum number of undo levels to store */
 	static const int32 MaxUndoLevels = 99;
 
-	/** Vertical text offset from top of widget, as a scalar percentage of maximum character height */
-	static const float TextVertOffsetPercent = 0.0f;
-
-	/** How quickly the cursor's spring will pull the cursor around.  Larger is faster. */
-	static const float CaretSpringConstant = 65.0f;
-
-	/** Vertical offset of caret from text location, in pixels */
-	static const float CaretVertOffset = 1.0f;
-
 	/** Width of the caret, as a scalar percentage of the font's maximum character height */
 	static const float CaretWidthPercent = 0.08f;
 
-	/** Height of the caret, as a scalar percentage of the font's maximum character height */
-	static const float CaretHeightPercent = 0.95f;
-
 	/** How long after the user last interacted with the keyboard should we keep the caret at full opacity? */
 	static const float CaretBlinkPauseTime = 0.1f;
 
 	/** How many times should the caret blink per second (full on/off cycles) */
 	static const float BlinksPerSecond = 1.0f;
-
-	/** How many pixels to extend the selection rectangle's left side horizontally */
-	static const float SelectionRectLeftOffset = 0.0f;
-
-	/** How many pixels to extend the selection rectangle's right side horizontally */
-	static const float SelectionRectRightOffset = 0.0f;
-
-	/** How quickly the selection 'targeting' rectangle will slide around.  Larger is faster. */
-	static const float SelectionTargetSpringConstant = 25.0f;
-
-	/** Duration of animation selection target effects */
-	static const float SelectionTargetEffectDuration = 0.25f;
-
-	/** Opacity of the selection target effect overlay */
-	static const float SelectionTargetOpacity = 0.8f;
-
-	/** How large the selection target effect will be when selecting text, as a scalar percentage of font height */
-	static const float SelectingAnimOffsetPercent = 0.15f;
-
-	/** How large the selection target effect will be when deselecting text, as a scalar percentage of font height */
-	static const float DeselectingAnimOffsetPercent = 0.25f;
 }
 
 /**
- * Contains helper routines that translate input into actions for an ITextEditorWidget.
- * Used to share code between SEditableText and SMultiLineEditableText
+ * Static helper functions
  */
 class FTextEditHelper
 {
 public:
-
-	/** Someone typed a character, translate the input event into a call to a ITextEditorWidget action */
-	static FReply OnKeyChar( const FCharacterEvent& InCharacterEvent, const TSharedRef< class ITextEditorWidget >& TextEditor );
-
-	/** Someone pressed a key, translate the input event into a call to a ITextEditorWidget action */
-	static FReply OnKeyDown( const FKeyEvent& InKeyEvent, const TSharedRef< ITextEditorWidget >& TextEditor );
-
-	static FReply OnMouseButtonDown( const FGeometry& MyGeometry, const FPointerEvent& InMouseEvent, const TSharedRef< ITextEditorWidget >& TextEditor );
-
-	static FReply OnMouseMove( const FGeometry& InMyGeometry, const FPointerEvent& InMouseEvent, const TSharedRef< ITextEditorWidget >& TextEditor );
-
-	static FReply OnMouseButtonUp( const FGeometry& MyGeometry, const FPointerEvent& InMouseEvent, const TSharedRef< ITextEditorWidget >& TextEditor );
-
-	static FReply OnMouseButtonDoubleClick(const FGeometry& InMyGeometry, const FPointerEvent& InMouseEvent, const TSharedRef< ITextEditorWidget >& TextEditor);
-
 	/**
 	 * Gets the height of the largest character in the font
 	 *
diff --git a/Engine/Source/Runtime/Slate/Private/Widgets/Input/SCheckBox.cpp b/Engine/Source/Runtime/Slate/Private/Widgets/Input/SCheckBox.cpp
--- a/Engine/Source/Runtime/Slate/Private/Widgets/Input/SCheckBox.cpp
+++ b/Engine/Source/Runtime/Slate/Private/Widgets/Input/SCheckBox.cpp
@@ -321,6 +321,14 @@ void SCheckBox::SetStyle(const FCheckBoxStyle* InStyle)
 {
 	Style = InStyle;
 
+	if (Style == nullptr)
+	{
+		FArguments Defaults;
+		Style = Defaults._Style;
+	}
+
+	check(Style);
+
 	BuildCheckBox(ContentContainer->GetContent());
 }
 
