Copying //UE4/Dev-Editor to Dev-Main (//UE4/Dev-Main)

diff --git a/Engine/Source/Runtime/Slate/Private/Widgets/Input/SEditableText.cpp b/Engine/Source/Runtime/Slate/Private/Widgets/Input/SEditableText.cpp
--- a/Engine/Source/Runtime/Slate/Private/Widgets/Input/SEditableText.cpp
+++ b/Engine/Source/Runtime/Slate/Private/Widgets/Input/SEditableText.cpp
@@ -3,205 +3,78 @@
 #include "SlatePrivatePCH.h"
 #include "BreakIterator.h"
 #include "TextEditHelper.h"
-#include "GenericCommands.h"
+#include "PlainTextLayoutMarshaller.h"
+#include "SlateEditableTextLayout.h"
+#include "ReflectionMetadata.h"
 #include "IMenu.h"
 
-/** A pointer to the editable text currently under the mouse cursor. nullptr when there isn't one. */
-SEditableText* SEditableText::EditableTextUnderCursor = nullptr;
-
-
-/** Constructor */
 SEditableText::SEditableText()
-	: ScrollHelper()
-	, CaretPosition( 0 )
-	, LastCaretInteractionTime( -1000.0 )
-	, bIsDragSelecting( false )
-	, bWasFocusedByLastMouseDown( false )
-	, bHasDragSelectedSinceFocused( false )
-	, bHasDragSelectedSinceMouseDown( false )
-	, CurrentUndoLevel( INDEX_NONE )
-	, bIsChangingText( false )
-	, UICommandList( new FUICommandList() )
-	, bTextChangedByVirtualKeyboard( false )
 {
 }
 
-/** Destructor */
 SEditableText::~SEditableText()
 {
-	ITextInputMethodSystem* const TextInputMethodSystem = FSlateApplication::IsInitialized() ? FSlateApplication::Get().GetTextInputMethodSystem() : nullptr;
-	if(TextInputMethodSystem)
-	{
-		TSharedRef<FTextInputMethodContext> TextInputMethodContextRef = TextInputMethodContext.ToSharedRef();
-		if(TextInputMethodSystem->IsActiveContext(TextInputMethodContextRef))
-		{
-			// This can happen if an entire tree of widgets is culled, as Slate isn't notified of the focus loss, the widget is just deleted
-			TextInputMethodSystem->DeactivateContext(TextInputMethodContextRef);
-		}
-
-		TextInputMethodSystem->UnregisterContext(TextInputMethodContextRef);
-	}
+	// Needed to avoid "deletion of pointer to incomplete type 'FSlateEditableTextLayout'; no destructor called" error when using TUniquePtr
 }
 
-/**
- * Construct this widget
- *
- * @param	InArgs	The declaration data for this widget
- */
 void SEditableText::Construct( const FArguments& InArgs )
 {
-	Text = InArgs._Text;
-	HintText = InArgs._HintText;
-
-	// Set colors and font using override attributes Font and ColorAndOpacity if set, otherwise use the Style argument.
-	check (InArgs._Style);
-	Font = InArgs._Font.IsSet() ? InArgs._Font : InArgs._Style->Font;
-	ColorAndOpacity = InArgs._ColorAndOpacity.IsSet() ? InArgs._ColorAndOpacity : InArgs._Style->ColorAndOpacity;
-	BackgroundImageSelected = InArgs._BackgroundImageSelected.IsSet() ? InArgs._BackgroundImageSelected : &InArgs._Style->BackgroundImageSelected;
-	BackgroundImageComposing = InArgs._BackgroundImageComposing.IsSet() ? InArgs._BackgroundImageComposing : &InArgs._Style->BackgroundImageComposing;	
-	CaretImage = InArgs._CaretImage.IsSet() ? InArgs._CaretImage : &InArgs._Style->CaretImage;
+	bIsReadOnly = InArgs._IsReadOnly;
+	bIsPassword = InArgs._IsPassword;
 
-	IsReadOnly = InArgs._IsReadOnly;
-	IsPassword = InArgs._IsPassword;
-	IsCaretMovedWhenGainFocus = InArgs._IsCaretMovedWhenGainFocus;
+	bIsCaretMovedWhenGainFocus = InArgs._IsCaretMovedWhenGainFocus;
 	bSelectAllTextWhenFocused = InArgs._SelectAllTextWhenFocused;
-	RevertTextOnEscape = InArgs._RevertTextOnEscape;
-	ClearKeyboardFocusOnCommit = InArgs._ClearKeyboardFocusOnCommit;
-	AllowContextMenu = InArgs._AllowContextMenu;
+	bRevertTextOnEscape = InArgs._RevertTextOnEscape;
+	bClearKeyboardFocusOnCommit = InArgs._ClearKeyboardFocusOnCommit;
+	bAllowContextMenu = InArgs._AllowContextMenu;
 	OnContextMenuOpening = InArgs._OnContextMenuOpening;
 	OnIsTypedCharValid = InArgs._OnIsTypedCharValid;
-	OnTextChanged = InArgs._OnTextChanged;
-	OnTextCommitted = InArgs._OnTextCommitted;
+	OnTextChangedCallback = InArgs._OnTextChanged;
+	OnTextCommittedCallback = InArgs._OnTextCommitted;
 	MinDesiredWidth = InArgs._MinDesiredWidth;
-	SelectAllTextOnCommit = InArgs._SelectAllTextOnCommit;
-	VirtualKeyboardType = IsPassword.Get() ? EKeyboardType::Keyboard_Password : InArgs._VirtualKeyboardType;
+	bSelectAllTextOnCommit = InArgs._SelectAllTextOnCommit;
+	VirtualKeyboardType = InArgs._VirtualKeyboardType;
 	OnKeyDownHandler = InArgs._OnKeyDownHandler;
-	bShowingVirtualKeyboard = false;
 
-	// Map UI commands to delegates which are called when the command should be executed
-	UICommandList->MapAction( FGenericCommands::Get().Undo,
-		FExecuteAction::CreateSP( this, &SEditableText::Undo ),
-		FCanExecuteAction::CreateSP( this, &SEditableText::CanExecuteUndo ),
-		EUIActionRepeatMode::RepeatEnabled );
+	Font = InArgs._Font;
+	ColorAndOpacity = InArgs._ColorAndOpacity;
+	BackgroundImageSelected = InArgs._BackgroundImageSelected;
 
-	UICommandList->MapAction( FGenericCommands::Get().Cut,
-		FExecuteAction::CreateSP( this, &SEditableText::CutSelectedTextToClipboard ),
-		FCanExecuteAction::CreateSP( this, &SEditableText::CanExecuteCut ) );
+	// We use the given style when creating the text layout as it may not be safe to call the override delegates until we've finished being constructed
+	// The first call to SyncronizeTextStyle will apply the correct overrides, and that will happen before the first paint
+	check(InArgs._Style);
+	FTextBlockStyle TextStyle = FCoreStyle::Get().GetWidgetStyle<FTextBlockStyle>("NormalText");
+	TextStyle.Font = InArgs._Style->Font;
+	TextStyle.ColorAndOpacity = InArgs._Style->ColorAndOpacity;
+	TextStyle.HighlightShape = InArgs._Style->BackgroundImageSelected;
 
-	UICommandList->MapAction( FGenericCommands::Get().Paste,
-		FExecuteAction::CreateSP( this, &SEditableText::PasteTextFromClipboard ),
-		FCanExecuteAction::CreateSP( this, &SEditableText::CanExecutePaste ),
-		EUIActionRepeatMode::RepeatEnabled );
+	PlainTextMarshaller = FPlainTextLayoutMarshaller::Create();
+	PlainTextMarshaller->SetIsPassword(bIsPassword);
 
-	UICommandList->MapAction( FGenericCommands::Get().Copy,
-		FExecuteAction::CreateSP( this, &SEditableText::CopySelectedTextToClipboard ),
-		FCanExecuteAction::CreateSP( this, &SEditableText::CanExecuteCopy ) );
+	// We use a separate marshaller for the hint text, as that should never be displayed as a password
+	TSharedRef<FPlainTextLayoutMarshaller> HintTextMarshaller = FPlainTextLayoutMarshaller::Create();
 
-	UICommandList->MapAction( FGenericCommands::Get().Delete,
-		FExecuteAction::CreateSP( this, &SEditableText::ExecuteDeleteAction ),
-		FCanExecuteAction::CreateSP( this, &SEditableText::CanExecuteDelete ) );
-
-	UICommandList->MapAction( FGenericCommands::Get().SelectAll,
-		FExecuteAction::CreateSP( this, &SEditableText::SelectAllText ),
-		FCanExecuteAction::CreateSP( this, &SEditableText::CanExecuteSelectAll ) );
+	EditableTextLayout = MakeUnique<FSlateEditableTextLayout>(*this, InArgs._Text, TextStyle, InArgs._TextShapingMethod, InArgs._TextFlowDirection, PlainTextMarshaller.ToSharedRef(), HintTextMarshaller);
+	EditableTextLayout->SetHintText(InArgs._HintText);
+	EditableTextLayout->SetCursorBrush(InArgs._CaretImage.IsSet() ? InArgs._CaretImage : &InArgs._Style->CaretImage);
+	EditableTextLayout->SetCompositionBrush(InArgs._BackgroundImageComposing.IsSet() ? InArgs._BackgroundImageComposing : &InArgs._Style->BackgroundImageComposing);
+	EditableTextLayout->SetDebugSourceInfo(TAttribute<FString>::Create(TAttribute<FString>::FGetter::CreateLambda([this]{ return FReflectionMetaData::GetWidgetDebugInfo(this); })));
 
 	// build context menu extender
-	MenuExtender = MakeShareable(new FExtender);
+	MenuExtender = MakeShareable(new FExtender());
 	MenuExtender->AddMenuExtension("EditText", EExtensionHook::Before, TSharedPtr<FUICommandList>(), InArgs._ContextMenuExtender);
-
-	TextInputMethodContext = MakeShareable( new FTextInputMethodContext( StaticCastSharedRef<SEditableText>(AsShared()) ) );
-	ITextInputMethodSystem* const TextInputMethodSystem = FSlateApplication::Get().GetTextInputMethodSystem();
-	if(TextInputMethodSystem)
-	{
-		TextInputMethodChangeNotifier = TextInputMethodSystem->RegisterContext(TextInputMethodContext.ToSharedRef());
-	}
-	if(TextInputMethodChangeNotifier.IsValid())
-	{
-		TextInputMethodChangeNotifier->NotifyLayoutChanged(ITextInputMethodChangeNotifier::ELayoutChangeType::Created);
-	}
 }
 
-/**
- * Sets the text string currently being edited 
- *
- * @param  InNewText  The new text string
- */
 void SEditableText::SetText( const TAttribute< FText >& InNewText )
 {
-	const bool bHasTextChanged = HasKeyboardFocus() ? !InNewText.Get().EqualTo(EditedText) : !InNewText.Get().EqualTo(Text.Get());
-
-	// NOTE: This will unbind any getter that is currently assigned to the Text attribute!  You should only be calling
-	//       SetText() if you know what you're doing.
-	Text = InNewText;
-
-	if (bHasTextChanged)
-	{
-		// Update the edited text, if the user isn't already editing it
-		if (HasKeyboardFocus())
-		{
-			// Update the edited text, if the user isn't already editing it
-			ClearSelection();
-			EditedText = FText();
-
-			// Load the new text
-			LoadText();
-
-			// Move the cursor to the end of the string
-			SetCaretPosition(EditedText.ToString().Len());
-
-			// Let outsiders know that the text being edited has been changed
-			OnTextChanged.ExecuteIfBound(EditedText);
-		}
-		else
-		{
-			// Let outsiders know that the text content has been changed
-			OnTextChanged.ExecuteIfBound(Text.Get());
-		}
-	}
-}
-
-void SEditableText::SetTextFromVirtualKeyboard(const FText& InNewText, ESetTextType SetTextType, ETextCommit::Type CommitType)
-{
-	// Only set the text if the text attribute doesn't have a getter binding (otherwise it would be blown away).
-	// If it is bound, we'll assume that OnTextCommitted will handle the update.
-	if (!Text.IsBound())
-	{
-		Text.Set(InNewText);
-	}
-
-	if (!InNewText.EqualTo(EditedText))
-	{
-		EditedText = InNewText;
-
-		// This method is called from the main thread (i.e. not the game thread) of the device with the virtual keyboard
-		// This causes the app to crash on those devices, so we're using polling here to ensure delegates are
-		// fired on the game thread in Tick.
-		bTextChangedByVirtualKeyboard = true;
-	}
-}
-
-void SEditableText::RestoreOriginalText()
-{
-	if( HasTextChangedFromOriginal() )
-	{
-		EditedText = OriginalText;
-
-		// Reset the cursor position
-		SetCaretPosition(0);
-
-		SaveText();
-
-		// Let outsiders know that the text content has been changed
-		OnTextChanged.ExecuteIfBound( EditedText );
-	}
+	EditableTextLayout->SetText(InNewText);
 }
 
-bool SEditableText::HasTextChangedFromOriginal() const
+FText SEditableText::GetText() const
 {
-	return ( !IsReadOnly.Get() && !EditedText.EqualTo(OriginalText)  );
+	return EditableTextLayout->GetText();
 }
 
-
 void SEditableText::SetFont( const TAttribute< FSlateFontInfo >& InNewFont )
 {
 	Font = InNewFont;
@@ -209,1141 +82,39 @@ void SEditableText::SetFont( const TAttribute< FSlateFontInfo >& InNewFont )
 
 void SEditableText::Tick( const FGeometry& AllottedGeometry, const double InCurrentTime, const float InDeltaTime )
 {
-	// Poll to see if the text has been updated by a virtual keyboard
-	if (bTextChangedByVirtualKeyboard)
-	{
-		OnEnter();
-
-		// Move the cursor to the end of the string
-		SetCaretPosition(EditedText.ToString().Len());
-
-		bTextChangedByVirtualKeyboard = false;
-	}
-
-	if(TextInputMethodChangeNotifier.IsValid() && TextInputMethodContext.IsValid() && TextInputMethodContext->CachedGeometry != AllottedGeometry)
-	{
-		TextInputMethodContext->CachedGeometry = AllottedGeometry;
-		TextInputMethodChangeNotifier->NotifyLayoutChanged(ITextInputMethodChangeNotifier::ELayoutChangeType::Changed);
-	}
-	
-	const FString VisibleText = GetStringToRender();
-	const FSlateFontInfo& FontInfo = Font.Get();
-	const TSharedRef< FSlateFontMeasure > FontMeasureService = FSlateApplication::Get().GetRenderer()->GetFontMeasureService();
-	const float FontMaxCharHeight = FontMeasureService->GetMaxCharacterHeight( FontInfo );
-
-	const bool bShouldAppearFocused = ShouldAppearFocused();
-	
-	float CaretVisualPosition = 0.0f;
-	if( bShouldAppearFocused )
-	{
-		const FVector2D TextBeforeCaretSize = FontMeasureService->Measure( VisibleText, 0, CaretPosition, FontInfo, false, AllottedGeometry.Scale );
-		CaretVisualPosition = TextBeforeCaretSize.X;
-	}
-
-	// Make sure the caret is scrolled into view.  Note that even if we don't have keyboard focus (and the caret
-	// is hidden), we still do this to make sure the beginning of the string is in view.
-	{
-		const float CaretWidth = FTextEditHelper::CalculateCaretWidth(FontMaxCharHeight);
-
-		// Grab the caret position in the scrolled area space
-		const float CaretLeft = CaretVisualPosition;
-		const float CaretRight = CaretVisualPosition + CaretWidth;
-
-		// Figure out where the caret is in widget space
-		const float WidgetSpaceCaretLeft = ScrollHelper.FromScrollerSpace( FVector2D( CaretLeft, ScrollHelper.GetPosition().Y ) ).X;
-		const float WidgetSpaceCaretRight = ScrollHelper.FromScrollerSpace( FVector2D( CaretRight, ScrollHelper.GetPosition().Y ) ).X;
-
-		const float AllotedGeometryDrawX = AllottedGeometry.Size.X*AllottedGeometry.Scale;
-
-		// If caret is to the left of the scrolled area, start scrolling left
-		if( WidgetSpaceCaretLeft < 0.0f )
-		{
-			ScrollHelper.SetPosition( FVector2D( CaretLeft, ScrollHelper.GetPosition().Y ) );
-		}
-		// If caret is to the right of the scrolled area, start scrolling right
-		else if ( WidgetSpaceCaretRight > AllotedGeometryDrawX )
-		{
-			ScrollHelper.SetPosition( FVector2D( CaretRight - AllotedGeometryDrawX, ScrollHelper.GetPosition( ).Y ) );
-		}
-
-		// Make sure text is never scrolled out of view when it doesn't need to be!
-		{
-			// Measure the total widget of text string, plus the caret's width!
-			const float TotalTextWidth = FontMeasureService->Measure( VisibleText, FontInfo, AllottedGeometry.Scale ).X + CaretWidth;
-			if ( ScrollHelper.GetPosition( ).X > TotalTextWidth - AllotedGeometryDrawX )
-			{
-				ScrollHelper.SetPosition( FVector2D( TotalTextWidth - AllotedGeometryDrawX, ScrollHelper.GetPosition( ).Y ) );
-			}
-			if( ScrollHelper.GetPosition().X < 0.0f )
-			{
-				ScrollHelper.SetPosition( FVector2D( 0.0f, ScrollHelper.GetPosition().Y ) );
-			}
-		}
-	}
-}
-
-bool SEditableText::FTextInputMethodContext::IsReadOnly()
-{
-	bool Return = true;
-	const TSharedPtr<SEditableText> OwningWidgetPtr = OwningWidget.Pin();
-	if(OwningWidgetPtr.IsValid())
-	{
-		Return = OwningWidgetPtr->GetIsReadOnly();
-	}
-	return Return;
-}
-
-uint32 SEditableText::FTextInputMethodContext::GetTextLength()
-{
-	uint32 TextLength = 0;
-	const TSharedPtr<SEditableText> OwningWidgetPtr = OwningWidget.Pin();
-	if(OwningWidgetPtr.IsValid())
-	{
-		const FString& WidgetEditedText = OwningWidgetPtr->EditedText.ToString();
-		TextLength = WidgetEditedText.Len();
-	}
-	return TextLength;
-}
-
-void SEditableText::FTextInputMethodContext::GetSelectionRange(uint32& BeginIndex, uint32& Length, ECaretPosition& OutCaretPosition)
-{
-	const TSharedPtr<SEditableText> OwningWidgetPtr = OwningWidget.Pin();
-	if(OwningWidgetPtr.IsValid())
-	{
-		if(OwningWidgetPtr->Selection.StartIndex != OwningWidgetPtr->Selection.FinishIndex)
-		{
-			BeginIndex = OwningWidgetPtr->Selection.GetMinIndex();
-			Length = OwningWidgetPtr->Selection.GetMaxIndex() - BeginIndex;
-		}
-		else
-		{
-			BeginIndex = OwningWidgetPtr->CaretPosition;
-			Length = 0;
-		}
-
-		if(OwningWidgetPtr->CaretPosition == BeginIndex + Length)
-		{
-			OutCaretPosition = ITextInputMethodContext::ECaretPosition::Ending;
-		}
-		else if(OwningWidgetPtr->CaretPosition == BeginIndex)
-		{
-			OutCaretPosition = ITextInputMethodContext::ECaretPosition::Beginning;
-		}
-	}
-}
-
-void SEditableText::FTextInputMethodContext::SetSelectionRange(const uint32 BeginIndex, const uint32 Length, const ECaretPosition InCaretPosition)
-{
-	const TSharedPtr<SEditableText> OwningWidgetPtr = OwningWidget.Pin();
-	if(OwningWidgetPtr.IsValid())
-	{
-		const uint32 MinIndex = BeginIndex;
-		const uint32 MaxIndex = MinIndex + Length;
-
-		uint32 SelectionBeginIndex = 0;
-		uint32 SelectionEndIndex = 0;
-
-		switch (InCaretPosition)
-		{
-		case ITextInputMethodContext::ECaretPosition::Beginning:
-			{
-				SelectionBeginIndex = MaxIndex;
-				SelectionEndIndex = MinIndex;
-			}
-			break;
-		case ITextInputMethodContext::ECaretPosition::Ending:
-			{
-				SelectionBeginIndex = MinIndex;
-				SelectionEndIndex = MaxIndex;
-			}
-			break;
-		}
-
-		OwningWidgetPtr->SetCaretPosition(SelectionEndIndex);
-		OwningWidgetPtr->ClearSelection();
-		OwningWidgetPtr->SelectText(SelectionBeginIndex);
-	}
-}
-
-void SEditableText::FTextInputMethodContext::GetTextInRange(const uint32 BeginIndex, const uint32 Length, FString& OutString)
-{
-	const TSharedPtr<SEditableText> OwningWidgetPtr = OwningWidget.Pin();
-	if(OwningWidgetPtr.IsValid())
-	{
-		const FString& WidgetEditedText = OwningWidgetPtr->EditedText.ToString();
-		OutString = WidgetEditedText.Mid(BeginIndex, Length);
-	}
-}
-
-void SEditableText::FTextInputMethodContext::SetTextInRange(const uint32 BeginIndex, const uint32 Length, const FString& InString)
-{
-	const TSharedPtr<SEditableText> OwningWidgetPtr = OwningWidget.Pin();
-	if(OwningWidgetPtr.IsValid())
-	{
-		const FString& OldText = OwningWidgetPtr->EditedText.ToString();
-		
-		// We don't use Start/FinishEditing text here because the whole IME operation handles that.
-		// Also, we don't want to support undo for individual characters added during an IME context
-		OwningWidgetPtr->EditedText = FText::FromString(OldText.Left( BeginIndex ) + InString + OldText.Mid( BeginIndex + Length ));
-		OwningWidgetPtr->SaveText();
-
-		OwningWidgetPtr->OnTextChanged.ExecuteIfBound(OwningWidgetPtr->EditedText);
-	}
-}
-
-int32 SEditableText::FTextInputMethodContext::GetCharacterIndexFromPoint(const FVector2D& Point)
-{
-	int32 CharacterIndex = INDEX_NONE;
-	const TSharedPtr<SEditableText> OwningWidgetPtr = OwningWidget.Pin();
-	if(OwningWidgetPtr.IsValid())
-	{
-		CharacterIndex = OwningWidgetPtr->FindClickedCharacterIndex(CachedGeometry, Point);
-	}
-	return CharacterIndex;
-}
-
-bool SEditableText::FTextInputMethodContext::GetTextBounds(const uint32 BeginIndex, const uint32 Length, FVector2D& Position, FVector2D& Size)
-{
-	bool IsClipped = false;
-	const TSharedPtr<SEditableText> OwningWidgetPtr = OwningWidget.Pin();
-	if(OwningWidgetPtr.IsValid())
-	{
-		const FSlateFontInfo& FontInfo = OwningWidgetPtr->Font.Get();
-		const FString VisibleText = OwningWidgetPtr->GetStringToRender();
-		const TSharedRef< FSlateFontMeasure > FontMeasureService = FSlateApplication::Get().GetRenderer()->GetFontMeasureService();
-
-		const float TextLeft = FontMeasureService->Measure( VisibleText, 0, BeginIndex, FontInfo, false, CachedGeometry.Scale ).X;
-		const float TextRight = FontMeasureService->Measure( VisibleText, 0, BeginIndex + Length, FontInfo, false, CachedGeometry.Scale ).X;
-
-		const float ScrollAreaLeft = OwningWidgetPtr->ScrollHelper.ToScrollerSpace( FVector2D::ZeroVector ).X;
-		const float ScrollAreaRight = OwningWidgetPtr->ScrollHelper.ToScrollerSpace( CachedGeometry.Size ).X;
-
-		const int32 FirstVisibleCharacter = FontMeasureService->FindCharacterIndexAtOffset( VisibleText, FontInfo, ScrollAreaLeft );
-		const int32 LastVisibleCharacter = FontMeasureService->FindCharacterIndexAtOffset( VisibleText, FontInfo, ScrollAreaRight );
-
-		if(FirstVisibleCharacter == INDEX_NONE || LastVisibleCharacter == INDEX_NONE || static_cast<uint32>(FirstVisibleCharacter) > BeginIndex || static_cast<uint32>(LastVisibleCharacter) <= BeginIndex + Length)
-		{
-			IsClipped = true;
-		}
-
-		const float FontHeight = FTextEditHelper::GetFontHeight(FontInfo);
-		const float DrawPositionY = ( CachedGeometry.Size.Y / 2 ) - ( FontHeight / 2 );
-
-		const float TextVertOffset = EditableTextDefs::TextVertOffsetPercent * FontHeight;
-
-		Position = CachedGeometry.LocalToAbsolute( OwningWidgetPtr->ScrollHelper.FromScrollerSpace( FVector2D( TextLeft, DrawPositionY + TextVertOffset ) ) );
-		Size = CachedGeometry.Scale * OwningWidgetPtr->ScrollHelper.SizeFromScrollerSpace( FVector2D( TextRight - TextLeft, CachedGeometry.Size.Y ) );
-	}
-	return IsClipped;
-}
-
-void SEditableText::FTextInputMethodContext::GetScreenBounds(FVector2D& Position, FVector2D& Size)
-{
-	const TSharedPtr<SEditableText> OwningWidgetPtr = OwningWidget.Pin();
-	if(OwningWidgetPtr.IsValid())
-	{
-		Position = CachedGeometry.AbsolutePosition;
-		Size = CachedGeometry.GetDrawSize();
-	}
-}
-
-TSharedPtr<FGenericWindow> SEditableText::FTextInputMethodContext::GetWindow()
-{
-	TSharedPtr<FGenericWindow> Window;
-	const TSharedPtr<SEditableText> OwningWidgetPtr = OwningWidget.Pin();
-	if(OwningWidgetPtr.IsValid())
-	{
-		const TSharedPtr<SWindow> SlateWindow = FSlateApplication::Get().FindWidgetWindow( OwningWidgetPtr.ToSharedRef() );
-		Window = SlateWindow.IsValid() ? SlateWindow->GetNativeWindow() : nullptr;
-	}
-	return Window;
-}
-
-void SEditableText::FTextInputMethodContext::BeginComposition()
-{
-	if(!IsComposing)
-	{
-		IsComposing = true;
-
-		const TSharedPtr<SEditableText> OwningWidgetPtr = OwningWidget.Pin();
-		if(OwningWidgetPtr.IsValid())
-		{
-			OwningWidgetPtr->StartChangingText();
-		}
-	}
-}
-
-void SEditableText::FTextInputMethodContext::UpdateCompositionRange(const int32 InBeginIndex, const uint32 InLength)
-{
-	check(IsComposing);
-
-	CompositionBeginIndex = InBeginIndex;
-	CompositionLength = InLength;
-}
-
-void SEditableText::FTextInputMethodContext::EndComposition()
-{
-	if(IsComposing)
-	{
-		const TSharedPtr<SEditableText> OwningWidgetPtr = OwningWidget.Pin();
-		if(OwningWidgetPtr.IsValid())
-		{
-			OwningWidgetPtr->FinishChangingText();
-		}
-
-		IsComposing = false;
-	}
-}
-
-void SEditableText::StartChangingText()
-{
-	// Never change text on read only controls! 
-	check( !IsReadOnly.Get() );
-
-	// We're starting to (potentially) change text
-	check( !bIsChangingText );
-	bIsChangingText = true;
-
-	// Save off an undo state in case we actually change the text
-	MakeUndoState( StateBeforeChangingText );
-}
-
-
-void SEditableText::FinishChangingText()
-{
-	// We're no longer changing text
-	check( bIsChangingText );
-	bIsChangingText = false;
-
-	// Has the text changed?
-	if( !EditedText.EqualTo(StateBeforeChangingText.Text) )
-	{
-		// Save text state
-		SaveText();
-
-		// Text was actually changed, so push the undo state we previously saved off
-		PushUndoState( StateBeforeChangingText );
-		
-		// We're done with this state data now.  Clear out any old data.
-		StateBeforeChangingText = FUndoState();
-
-		// Let outsiders know that the text content has been changed
-		OnTextChanged.ExecuteIfBound( EditedText );
-	}
-}
-
-
-bool SEditableText::GetIsReadOnly() const
-{
-	return IsReadOnly.Get();
-}
-
-
-void SEditableText::BackspaceChar()
-{
-	if( AnyTextSelected() )
-	{
-		// Delete selected text
-		DeleteSelectedText();
-	}
-	else
-	{
-		// Delete character to the left of the caret
-		if( CaretPosition > 0 )
-		{
-			const FString& OldText = EditedText.ToString();
-			EditedText = FText::FromString(OldText.Left( CaretPosition - 1 ) + OldText.Mid( CaretPosition ));
-
-			// Move the caret back too
-			SetCaretPosition( CaretPosition - 1 );
-		}
-	}
-}
-
-
-void SEditableText::DeleteChar()
-{
-	if( !IsReadOnly.Get() )
-	{
-		if( AnyTextSelected() )
-		{
-			// Delete selected text
-			DeleteSelectedText();
-		}
-		else
-		{
-			// Delete character to the right of the caret
-			const FString& OldText = EditedText.ToString();
-			if( CaretPosition < OldText.Len() )
-			{
-				EditedText = FText::FromString(OldText.Left( CaretPosition ) + OldText.Mid( CaretPosition + 1 ));
-			}
-		}
-	}
-}
-
-
-bool SEditableText::CanTypeCharacter(const TCHAR CharInQuestion) const
-{
-	return !OnIsTypedCharValid.IsBound() || OnIsTypedCharValid.Execute(CharInQuestion);
-}
-
-
-void SEditableText::TypeChar( const int32 InChar )
-{
-	if( AnyTextSelected() )
-	{
-		// Delete selected text
-		DeleteSelectedText();
-	}
-
-	// Certain characters are not allowed
-	bool bIsCharAllowed = true;
-	{
-		if( InChar == TEXT('\t') )
-		{
-			bIsCharAllowed = true;
-		}
-		else if( InChar <= 0x1F )
-		{
-			bIsCharAllowed = false;
-		}
-	}
-
-	if( bIsCharAllowed )
-	{
-		// Insert character at caret position
-		const FString& OldText = EditedText.ToString();
-		EditedText = FText::FromString(OldText.Left( CaretPosition ) + InChar + OldText.Mid( CaretPosition ));
-
-		// Advance caret position
-		SetCaretPosition( CaretPosition + 1 );
-	}
-}
-
-/** When scanning forward we look at the location of the cursor. When scanning backwards we look at the location before the cursor. */
-static int32 OffsetToTestBasedOnDirection( int8 Direction )
-{
-	return (Direction >= 0) ? 0 : -1;
-}
-
-FReply SEditableText::MoveCursor( FMoveCursor Args )
-{
-	// We can't use the keyboard to move the cursor while composing, as the IME has control over it
-	if ( TextInputMethodContext->IsComposing && Args.GetMoveMethod() != ECursorMoveMethod::ScreenPosition )
-	{
-		return FReply::Handled();
-	}
-
-	bool bAllowMoveCursor = true;
-	int32 OldCaretPosition = CaretPosition;
-	int32 NewCaretPosition = OldCaretPosition;
-
-	// If we have selected text, the cursor needs to:
-	// a) Jump to the start of the selection if moving the cursor Left
-	// b) Jump to the end of the selection if moving the cursor Right
-	// This is done regardless of whether the selection was made left-to-right, or right-to-left
-	// This also needs to be done *before* moving to word boundaries, as the start point needs to be 
-	// the start or end of the selection depending on the above rules
-	if ( Args.GetAction( ) == ECursorAction::MoveCursor && Args.GetMoveMethod() != ECursorMoveMethod::ScreenPosition && AnyTextSelected( ) )
-	{
-		if (Args.IsHorizontalMovement())
-		{
-			// If we're moving the cursor horizontally, we just snap to the start or end of the selection rather than 
-			// move the cursor by the normal movement rules
-			bAllowMoveCursor = false;
-		}
-
-		// Work out which edge of the selection we need to start at
-		bool bSnapToSelectionStart = Args.IsHorizontalMovement() && Args.GetMoveDirection().X < 0;
-
-		// Adjust the current cursor position - also set the new cursor position so that the bAllowMoveCursor == false case is handled
-		OldCaretPosition = bSnapToSelectionStart ? Selection.GetMinIndex() : Selection.GetMaxIndex();
-		NewCaretPosition = OldCaretPosition;
-
-		// If we're snapping to a word boundary, but the selection was already at a word boundary, don't let the cursor move any more
-		if (Args.GetGranularity() == ECursorMoveGranularity::Word && IsAtWordStart(NewCaretPosition))
-		{
-			bAllowMoveCursor = false;
-		}
-	}
-
-	if (bAllowMoveCursor)
-	{
-		const int32 StringLength = EditedText.ToString().Len();
-
-		// If control is held down, jump to beginning of the current word (or, the previous word, if
-		// the caret is currently on a whitespace character
-		if ( Args.GetMoveMethod() == ECursorMoveMethod::Cardinal )
-		{
-			if ( Args.GetGranularity() == ECursorMoveGranularity::Word )
-			{
-				NewCaretPosition = ScanForWordBoundary( CaretPosition, Args.GetMoveDirection().X );
-			}
-			else if ( Args.IsHorizontalMovement() )
-			{
-				checkSlow( Args.GetGranularity() == ECursorMoveGranularity::Character );
-				NewCaretPosition = FMath::Clamp<int32>( CaretPosition + Args.GetMoveDirection( ).X, 0, StringLength );
-			}
-			else
-			{
-				checkSlow( Args.IsVerticalMovement() );
-				return FReply::Unhandled( );
-			}			
-		}
-		else if ( Args.GetMoveMethod( ) == ECursorMoveMethod::ScreenPosition )
-		{			
-			NewCaretPosition = FindClickedCharacterIndex( Args.GetLocalPosition(), Args.GetGeometryScale() );
-		}
-		else
-		{
-			checkfSlow(false, TEXT("Unknown ECursorMoveMethod value"));
-		}
-	}
-
-	SetCaretPosition( NewCaretPosition );
-
-	if( Args.GetAction() == ECursorAction::SelectText )
-	{
-		SelectText( OldCaretPosition );
-	}
-	else
-	{
-		ClearSelection();
-	}
-
-	// If we've moved the cursor while composing, we need to end the current composition session
-	// Note: You should only be able to do this via the mouse due to the check at the top of this function
-	if(TextInputMethodContext->IsComposing)
-	{
-		ITextInputMethodSystem* const TextInputMethodSystem = FSlateApplication::Get().GetTextInputMethodSystem();
-		if(TextInputMethodSystem)
-		{
-			TextInputMethodSystem->DeactivateContext(TextInputMethodContext.ToSharedRef());
-			TextInputMethodSystem->ActivateContext(TextInputMethodContext.ToSharedRef());
-		}
-	}
-
-	return FReply::Handled();
-}
-
-
-void SEditableText::JumpTo(ETextLocation JumpLocation, ECursorAction Action)
-{
-	switch( JumpLocation )
-	{
-		case ETextLocation::BeginningOfLine:
-		case ETextLocation::BeginningOfDocument:
-		case ETextLocation::PreviousPage:
-		{
-			const int32 OldCaretPosition = CaretPosition;
-			SetCaretPosition( 0 );
-
-			if( Action == ECursorAction::SelectText )
-			{
-				SelectText( OldCaretPosition );
-			}
-			else
-			{
-				ClearSelection();
-			}
-		}
-		break;
-
-		case ETextLocation::EndOfLine:
-		case ETextLocation::EndOfDocument:
-		case ETextLocation::NextPage:
-		{
-			const int32 OldCaretPosition = CaretPosition;
-			SetCaretPosition( EditedText.ToString().Len() );
-
-			if( Action == ECursorAction::SelectText )
-			{
-				SelectText( OldCaretPosition );
-			}
-			else
-			{
-				ClearSelection();
-			}
-		}
-		break;
-	}
-}
-
-
-void SEditableText::ClearSelection()
-{
-	// Explicitly invalidate the selection range indices.  The user no longer has a selection
-	// start and finish point
-	Selection.Clear();
-}
-
-void SEditableText::SelectAllText()
-{
-	// NOTE: Caret position is left unchanged
-
-	// NOTE: The selection anchor will always be after the last character in the string
-	
-	Selection.StartIndex = EditedText.ToString().Len();
-	Selection.FinishIndex = 0;
-}
-
-bool SEditableText::SelectAllTextWhenFocused()
-{
-	return bSelectAllTextWhenFocused.Get();
-}
-
-void SEditableText::SelectWordAt( const FGeometry& MyGeometry, const FVector2D& ScreenSpacePosition )
-{
-	const FVector2D LocalPosition = MyGeometry.AbsoluteToLocal( ScreenSpacePosition );
-	// Figure out where in the string the user clicked
-	const int32 ClickedCharIndex = FindClickedCharacterIndex( LocalPosition, MyGeometry.Scale );
-
-	// Deselect any text that was selected
-	ClearSelection();
-
-	// Select the word that the user double-clicked on, and move the caret to the end of that word
-	{
-		const FString& EditedTextString = EditedText.ToString();
-
-		// Find beginning of the word
-		int32 FirstWordCharIndex = ClickedCharIndex;
-		while( FirstWordCharIndex > 0 && !FText::IsWhitespace( EditedTextString[ FirstWordCharIndex - 1 ] ) )
-		{
-			--FirstWordCharIndex;
-		}
-
-		// Find end of the word
-		int32 LastWordCharIndex = ClickedCharIndex;
-		while( LastWordCharIndex < EditedTextString.Len() && !FText::IsWhitespace( EditedTextString[ LastWordCharIndex ] ) )
-		{
-			++LastWordCharIndex;
-		}
-
-		// Move the caret to the end of the selected word
-		SetCaretPosition( ClickedCharIndex );
-
-		// Select the word!
-		Selection.StartIndex = FirstWordCharIndex;
-		Selection.FinishIndex = LastWordCharIndex;
-	}
-}
-
-void SEditableText::BeginDragSelection() 
-{
-	bIsDragSelecting = true;
-}
-
-bool SEditableText::IsDragSelecting() const
-{
-	return bIsDragSelecting;
-}
-
-void SEditableText::EndDragSelection() 
-{
-	bIsDragSelecting = false;
-}
-
-bool SEditableText::AnyTextSelected() const
-{
-	return Selection.StartIndex != INDEX_NONE && Selection.StartIndex != Selection.FinishIndex;
-}
-
-bool SEditableText::IsTextSelectedAt( const FGeometry& MyGeometry, const FVector2D& ScreenSpacePosition ) const
-{
-	const int32 ClickedCharIndex = FindClickedCharacterIndex(MyGeometry, ScreenSpacePosition);
-	return ClickedCharIndex >= Selection.GetMinIndex() && ClickedCharIndex < Selection.GetMaxIndex();
-}
-
-void SEditableText::SetWasFocusedByLastMouseDown( bool Value )
-{
-	bWasFocusedByLastMouseDown = Value;
-}
-
-bool SEditableText::WasFocusedByLastMouseDown() const
-{
-	return bWasFocusedByLastMouseDown;
-}
-
-bool SEditableText::HasDragSelectedSinceFocused() const
-{
-	return bHasDragSelectedSinceFocused;
-}
-
-void SEditableText::SetHasDragSelectedSinceFocused( bool Value )
-{
-	bHasDragSelectedSinceFocused = Value;
-}
-
-
-FReply SEditableText::OnEscape()
-{
-	FReply MyReply = FReply::Unhandled();
-
-	if ( AnyTextSelected() )
-	{
-		// Clear selection
-		ClearSelection();
-		MyReply = FReply::Handled();
-	}
-
-	if ( !GetIsReadOnly() )
-	{
-		// Restore the text if the revert flag is set
-		if ( RevertTextOnEscape.Get() && HasTextChangedFromOriginal() )
-		{
-			RestoreOriginalText();
-			MyReply = FReply::Handled();
-		}
-	}
-
-	return MyReply;
-}
-
-
-void SEditableText::OnEnter()
-{
-	//Always clear the local undo chain on commit.
-	ClearUndoStates();
-
-	// When enter is pressed text is committed.  Let anyone interested know about it.
-	OnTextCommitted.ExecuteIfBound( EditedText, ETextCommit::OnEnter );
-
-	// Clear keyboard focus if we were configured to do that
-	if( HasKeyboardFocus() )
-	{
-		// Reload underlying value now it is committed  (commit may alter the value) 
-		// so it can be re-displayed in the edit box if it retains focus
-		LoadText();
-
-		if ( ClearKeyboardFocusOnCommit.Get() )
-		{
-			FSlateApplication::Get().ClearKeyboardFocus( EFocusCause::Cleared );
-		}
-		else
-		{
-			// If we aren't clearing keyboard focus, make it easy to type text again
-			SelectAllText();
-		}
-	}
-		
-	if( SelectAllTextOnCommit.Get() )
-	{
-		SelectAllText();
-	}	
-}
-
-
-bool SEditableText::CanExecuteCut() const
-{
-	bool bCanExecute = true;
-
-	// Can't execute if this is a read-only control
-	if( IsReadOnly.Get() )
-	{
-		bCanExecute = false;
-	}
-
-	// Can't execute if it's a password or there is no text selected
-	if( IsPassword.Get() || !AnyTextSelected() )
-	{
-		bCanExecute = false;
-	}
-
-	return bCanExecute;
-}
-
-
-void SEditableText::CutSelectedTextToClipboard()
-{
-	if( !IsPassword.Get() && AnyTextSelected() )
-	{
-		// Copy the text to the clipboard...
-		CopySelectedTextToClipboard();
-
-		if( !IsReadOnly.Get() )
-		{
-			StartChangingText();
-
-			// Delete the text!
-			DeleteSelectedText();
-
-			FinishChangingText();
-		}
-		else
-		{
-			// When read only, cut will simply copy the text
-		}
-	}
-}
-
-
-bool SEditableText::CanExecuteCopy() const
-{
-	bool bCanExecute = true;
-
-	// Can't execute if it's a password or there is no text selected
-	if( IsPassword.Get() || !AnyTextSelected() )
-	{
-		bCanExecute = false;
-	}
-
-	return bCanExecute;
-}
-
-
-void SEditableText::CopySelectedTextToClipboard()
-{
-	if( !IsPassword.Get() && AnyTextSelected() )
-	{
-		// Grab the selected substring
-		const FString SelectedText = EditedText.ToString().Mid( Selection.GetMinIndex(), Selection.GetMaxIndex() - Selection.GetMinIndex() );
-
-		// Copy text to clipboard
-		FPlatformMisc::ClipboardCopy( *SelectedText );
-	}
-}
-
-
-bool SEditableText::CanExecutePaste() const
-{
-	bool bCanExecute = true;
-
-	// Can't execute if this is a read-only control
-	if( IsReadOnly.Get() )
-	{
-		bCanExecute = false;
-	}
-
-	// Can't paste unless the clipboard has a string in it
-	if( !DoesClipboardHaveAnyText() )
-	{
-		bCanExecute = false;
-	}
-
-	return bCanExecute;
+	EditableTextLayout->Tick(AllottedGeometry, InCurrentTime, InDeltaTime);
 }
 
-
-void SEditableText::PasteTextFromClipboard()
-{
-	if( !IsReadOnly.Get() )
-	{
-		StartChangingText();
-
-		// Delete currently selected text, if any
-		DeleteSelectedText();
-
-		// Paste from the clipboard
-		FString PastedText;
-		FPlatformMisc::ClipboardPaste(PastedText);
-
-		for( int32 CurCharIndex = 0; CurCharIndex < PastedText.Len(); ++CurCharIndex )
-		{
-			TypeChar( PastedText.GetCharArray()[ CurCharIndex ] );
-		}
-
-		FinishChangingText();
-	}
-}
-
-
-bool SEditableText::CanExecuteUndo() const
-{
-	bool bCanExecute = true;
-
-	// Can't execute if this is a read-only control
-	if( IsReadOnly.Get() )
-	{
-		bCanExecute = false;
-	}
-
-	// Can't execute unless we have some undo history
-	if( UndoStates.Num() == 0 )
-	{
-		bCanExecute = false;
-	}
-
-	if(TextInputMethodContext->IsComposing)
-	{
-		bCanExecute = false;
-	}
-
-	return bCanExecute;
-}
-
-
-void SEditableText::Undo()
+int32 SEditableText::OnPaint( const FPaintArgs& Args, const FGeometry& AllottedGeometry, const FSlateRect& MyClippingRect, FSlateWindowElementList& OutDrawElements, int32 LayerId, const FWidgetStyle& InWidgetStyle, bool bParentEnabled ) const
 {
-	if( !IsReadOnly.Get() && UndoStates.Num() > 0 && !TextInputMethodContext->IsComposing )
-		{
-			// Restore from undo state
-			int32 UndoStateIndex;
-			if( CurrentUndoLevel == INDEX_NONE )
-			{
-				// We haven't undone anything since the last time a new undo state was added
-				UndoStateIndex = UndoStates.Num() - 1;
-
-				// Store an undo state for the current state (before undo was pressed)
-				FUndoState NewUndoState;
-				MakeUndoState( NewUndoState );
-				PushUndoState( NewUndoState );
-			}
-			else
-			{
-				// Move down to the next undo level
-				UndoStateIndex = CurrentUndoLevel - 1;
-			}
-
-			// Is there anything else to undo?
-			if( UndoStateIndex >= 0 )
-			{
-				{
-					// NOTE: It's important the no code called here creates or destroys undo states!
-					const FUndoState& UndoState = UndoStates[ UndoStateIndex ];
-		
-					EditedText = UndoState.Text;
-					CaretPosition = UndoState.CaretPosition;
-					Selection = UndoState.Selection;
-
-					SaveText();
-				}
-
-				CurrentUndoLevel = UndoStateIndex;
-
-				// Let outsiders know that the text content has been changed
-				OnTextChanged.ExecuteIfBound( EditedText );
-			}
-		}
-	}
+	LayerId = EditableTextLayout->OnPaint(Args, AllottedGeometry, MyClippingRect, OutDrawElements, LayerId, InWidgetStyle, ShouldBeEnabled(bParentEnabled));
 
-
-void SEditableText::Redo()
-{
-	// Is there anything to redo?  If we've haven't tried to undo since the last time new
-	// undo state was added, then CurrentUndoLevel will be INDEX_NONE
-	if( !IsReadOnly.Get() && CurrentUndoLevel != INDEX_NONE && !(TextInputMethodContext->IsComposing))
-	{
-		const int32 NextUndoLevel = CurrentUndoLevel + 1;
-		if( UndoStates.Num() > NextUndoLevel )
-		{
-			// Restore from undo state
-			{
-				// NOTE: It's important the no code called here creates or destroys undo states!
-				const FUndoState& UndoState = UndoStates[ NextUndoLevel ];
-		
-				EditedText = UndoState.Text;
-				CaretPosition = UndoState.CaretPosition;
-				Selection = UndoState.Selection;
-
-				SaveText();
-			}
-
-			CurrentUndoLevel = NextUndoLevel;
-
-			if( UndoStates.Num() <= CurrentUndoLevel + 1 )
-			{
-				// We've redone all available undo states
-				CurrentUndoLevel = INDEX_NONE;
-
-				// Pop last undo state that we created on initial undo
-				UndoStates.RemoveAt( UndoStates.Num() - 1 );
-			}
-
-			// Let outsiders know that the text content has been changed
-			OnTextChanged.ExecuteIfBound( EditedText );
-		}
-	}
+	return LayerId;
 }
 
-TSharedRef< SWidget > SEditableText::GetWidget()
+void SEditableText::CacheDesiredSize(float LayoutScaleMultiplier)
 {
-	return SharedThis( this );
-}
-
-TSharedRef<IVirtualKeyboardEntry> SEditableText::GetVirtualKeyboardEntry()
-{
-	return SharedThis(this);
+	SynchronizeTextStyle();
+	EditableTextLayout->CacheDesiredSize(LayoutScaleMultiplier);
+	SWidget::CacheDesiredSize(LayoutScaleMultiplier);
 }
 
 FVector2D SEditableText::ComputeDesiredSize(float LayoutScaleMultiplier) const
 {
-	const TSharedRef< FSlateFontMeasure > FontMeasureService = FSlateApplication::Get().GetRenderer()->GetFontMeasureService();
-	const FSlateFontInfo& FontInfo = Font.Get();
-	
-	const float FontMaxCharHeight = FTextEditHelper::GetFontHeight(FontInfo);
-	const float CaretWidth = FTextEditHelper::CalculateCaretWidth(FontMaxCharHeight);
-
-	FVector2D TextSize;
-
-	const FString StringToRender = GetStringToRender();
-	if( !StringToRender.IsEmpty() )
-	{
-		TextSize = FontMeasureService->Measure(StringToRender, FontInfo, LayoutScaleMultiplier) / LayoutScaleMultiplier;
-	}
-	else
-	{
-		TextSize = FontMeasureService->Measure(HintText.Get().ToString(), FontInfo, LayoutScaleMultiplier) / LayoutScaleMultiplier;
-	}
-	
-	const float DesiredWidth = FMath::Max(TextSize.X, MinDesiredWidth.Get()) + CaretWidth;
-	const float DesiredHeight = FMath::Max(TextSize.Y, FontMaxCharHeight);
-	return FVector2D( DesiredWidth, DesiredHeight );
-}
-
-
-int32 SEditableText::OnPaint( const FPaintArgs& Args, const FGeometry& AllottedGeometry, const FSlateRect& MyClippingRect, FSlateWindowElementList& OutDrawElements, int32 LayerId, const FWidgetStyle& InWidgetStyle, bool bParentEnabled ) const
-{
-	// The text and some effects draws in front of the widget's background and selection.
-	const int32 SelectionLayer = 0;
-	const int32 TextLayer = 1;
-	const TSharedRef< FSlateFontMeasure > FontMeasureService = FSlateApplication::Get().GetRenderer()->GetFontMeasureService();
-	const float ScaleInverse = 1.0f / AllottedGeometry.Scale;
-	const FVector2D DrawGeometrySize = AllottedGeometry.Size * AllottedGeometry.Scale;
-
-	// See if a disabled effect should be used
-	bool bEnabled = ShouldBeEnabled( bParentEnabled );
-	bool bIsReadonly = IsReadOnly.Get();
-	ESlateDrawEffect::Type DrawEffects = (bEnabled) ? ESlateDrawEffect::None : ESlateDrawEffect::DisabledEffect;
-
-	const FSlateFontInfo& FontInfo = Font.Get();
-	const FString VisibleText = GetStringToRender();
-	const FLinearColor ThisColorAndOpacity = ColorAndOpacity.Get().GetColor(InWidgetStyle);
-	const FLinearColor ColorAndOpacitySRGB = ThisColorAndOpacity * InWidgetStyle.GetColorAndOpacityTint();
-	const float FontMaxCharHeight = FTextEditHelper::GetFontHeight(FontInfo) * AllottedGeometry.Scale;
-	const double CurrentTime = FSlateApplication::Get().GetCurrentTime();
-
-	const bool bShouldAppearFocused = ShouldAppearFocused();
-
-	// Draw selection background
-	if( AnyTextSelected() && ( bShouldAppearFocused || bIsReadonly ) )
-	{
-		// Figure out a universally visible selection color.
-		const FLinearColor SelectionBackgroundColorAndOpacity = ( (FLinearColor::White - ThisColorAndOpacity)*0.5f + FLinearColor(-0.2f, -0.05f, 0.15f)) * InWidgetStyle.GetColorAndOpacityTint();
-
-		const float SelectionLeftX = EditableTextDefs::SelectionRectLeftOffset + FontMeasureService->Measure( VisibleText, 0, Selection.GetMinIndex(), FontInfo, false, AllottedGeometry.Scale ).X;
-		const float SelectionRightX = EditableTextDefs::SelectionRectRightOffset + FontMeasureService->Measure( VisibleText, 0, Selection.GetMaxIndex(), FontInfo, false, AllottedGeometry.Scale ).X;
-		const float SelectionTopY = 0.0f;
-		const float SelectionBottomY = DrawGeometrySize.Y;
-
-		const FVector2D DrawPosition = ScrollHelper.FromScrollerSpace( FVector2D( SelectionLeftX, SelectionTopY ) );
-		const FVector2D DrawSize = ScrollHelper.SizeFromScrollerSpace( FVector2D( SelectionRightX - SelectionLeftX, SelectionBottomY - SelectionTopY ) );
-
-		FSlateDrawElement::MakeBox(
-			OutDrawElements,
-			LayerId + SelectionLayer,
-			AllottedGeometry.ToPaintGeometry( DrawPosition, DrawSize, ScaleInverse ),	// Position, Size, Scale
-			BackgroundImageSelected.Get(),								// Image
-			MyClippingRect,												// Clipping rect
-			DrawEffects,												// Effects to use
-			SelectionBackgroundColorAndOpacity );						// Color
-	}
-
-	// Draw composition background
-	// We only draw the composition highlight if the cursor is within the composition range
-	const FTextRange CompositionRange(TextInputMethodContext->CompositionBeginIndex, TextInputMethodContext->CompositionBeginIndex + TextInputMethodContext->CompositionLength);
-	if( TextInputMethodContext->IsComposing && CompositionRange.InclusiveContains(CaretPosition) )
-	{
-		const float CompositionLeft = FontMeasureService->Measure( VisibleText, 0, CompositionRange.BeginIndex, FontInfo, false, AllottedGeometry.Scale ).X;
-		const float CompositionRight = FontMeasureService->Measure( VisibleText, 0, CompositionRange.EndIndex, FontInfo, false, AllottedGeometry.Scale ).X;
-		const float CompositionTop = 0.0f;
-		const float CompositionBottom = FontMeasureService->Measure( VisibleText, TextInputMethodContext->CompositionBeginIndex, TextInputMethodContext->CompositionLength, FontInfo, false, AllottedGeometry.Scale ).Y;
-
-		const FVector2D DrawPosition = ScrollHelper.FromScrollerSpace( FVector2D( CompositionLeft, CompositionTop ) );
-		const FVector2D DrawSize = ScrollHelper.SizeFromScrollerSpace( FVector2D( CompositionRight - CompositionLeft, CompositionBottom - CompositionTop ) );
-
-		FSlateDrawElement::MakeBox(
-			OutDrawElements,
-			LayerId + SelectionLayer,
-			AllottedGeometry.ToPaintGeometry( DrawPosition, DrawSize, ScaleInverse ),	// Position, Size, Scale
-			BackgroundImageComposing.Get(),								// Image
-			MyClippingRect,												// Clipping rect
-			DrawEffects,												// Effects to use
-			ColorAndOpacitySRGB );						// Color
-	}
-
-	const float DrawPositionY = ( DrawGeometrySize.Y / 2) - (FontMaxCharHeight / 2);
-	if (VisibleText.Len() == 0)
-	{
-		// Draw the hint text.
-		const FLinearColor HintTextColor = FLinearColor(ColorAndOpacitySRGB.R, ColorAndOpacitySRGB.G, ColorAndOpacitySRGB.B, 0.35f) * InWidgetStyle.GetColorAndOpacityTint();
-		const FString ThisHintText = this->HintText.Get().ToString();
-		FSlateDrawElement::MakeText(
-			OutDrawElements,
-			LayerId + TextLayer,
-			AllottedGeometry.ToPaintGeometry(FVector2D(0, DrawPositionY * ScaleInverse), AllottedGeometry.Size),
-			ThisHintText,          // Text
-			FontInfo,              // Font information (font name, size)
-			MyClippingRect,        // Clipping rect
-			DrawEffects,           // Effects to use
-			HintTextColor          // Color
-		);
-	}
-	else
-	{
-		// Draw the text
-
-		// Only draw the text that's potentially visible
-		const float ScrollAreaLeft = ScrollHelper.ToScrollerSpace( FVector2D::ZeroVector ).X;
-		const float ScrollAreaRight = ScrollHelper.ToScrollerSpace( DrawGeometrySize ).X;
-		const int32 FirstVisibleCharacter = FontMeasureService->FindCharacterIndexAtOffset( VisibleText, FontInfo, ScrollAreaLeft, AllottedGeometry.Scale );
-		const int32 LastVisibleCharacter = FontMeasureService->FindCharacterIndexAtOffset( VisibleText, FontInfo, ScrollAreaRight, AllottedGeometry.Scale );
-		const FString PotentiallyVisibleText( VisibleText.Mid( FirstVisibleCharacter, ( LastVisibleCharacter - FirstVisibleCharacter ) + 1 ) );
-		const float FirstVisibleCharacterPosition = FontMeasureService->Measure( VisibleText, 0, FirstVisibleCharacter, FontInfo, false, AllottedGeometry.Scale ).X;
-
-		const float TextVertOffset = EditableTextDefs::TextVertOffsetPercent * FontMaxCharHeight;
-		const FVector2D DrawPosition = ScrollHelper.FromScrollerSpace( FVector2D( FirstVisibleCharacterPosition, DrawPositionY + TextVertOffset ) );
-		const FVector2D DrawSize = ScrollHelper.SizeFromScrollerSpace( AllottedGeometry.Size );
-
-		FSlateDrawElement::MakeText(
-			OutDrawElements,
-			LayerId + TextLayer,
-			AllottedGeometry.ToPaintGeometry( DrawPosition, DrawSize ),
-			PotentiallyVisibleText,          // Text
-			FontInfo,                        // Font information (font name, size)
-			MyClippingRect,                  // Clipping rect
-			DrawEffects,                     // Effects to use
-			ColorAndOpacitySRGB              // Color
-		);
-	}
-
-
-	// Draw the caret!
-	if( bShouldAppearFocused )
-	{
-		// The caret is always visible (i.e. not blinking) when we're interacting with it; otherwise it might get lost.
-		const bool bForceCaretVisible = (CurrentTime - LastCaretInteractionTime) < EditableTextDefs::CaretBlinkPauseTime;
-		float CaretOpacity = (bForceCaretVisible)
-			? 1.0f
-			: FMath::RoundToFloat( FMath::MakePulsatingValue( CurrentTime, EditableTextDefs::BlinksPerSecond ));
-		CaretOpacity *= CaretOpacity;	// Squared falloff, because it looks more interesting
-
-		FLinearColor CursorColorAndOpacity = ThisColorAndOpacity;
-		CursorColorAndOpacity.A = CaretOpacity;
-
-		const FVector2D TextBeforeCaretSize = FontMeasureService->Measure( VisibleText, 0, CaretPosition, FontInfo, false, AllottedGeometry.Scale );
-		const float CaretVisualPosition = TextBeforeCaretSize.X;
-
-		const float CaretX = CaretVisualPosition;
-		const float CaretY = DrawPositionY + EditableTextDefs::CaretVertOffset;
-		const float CaretWidth = FTextEditHelper::CalculateCaretWidth(FontMaxCharHeight);
-		const float CaretHeight = EditableTextDefs::CaretHeightPercent * FontMaxCharHeight;
-
-		const FVector2D CaretDrawPosition = ScrollHelper.FromScrollerSpace( FVector2D( CaretX, CaretY ) );
-		const FVector2D CaretDrawSize = ScrollHelper.SizeFromScrollerSpace( FVector2D( CaretWidth, CaretHeight ) );
-		FSlateDrawElement::MakeBox(
-			OutDrawElements,
-			LayerId + TextLayer,
-			AllottedGeometry.ToPaintGeometry( CaretDrawPosition, CaretDrawSize, ScaleInverse ),		// Position, Size, Scale
-			CaretImage.Get(),															// Image
-			MyClippingRect,																// Clipping rect
-			DrawEffects,																// Effects to use
-			CursorColorAndOpacity*InWidgetStyle.GetColorAndOpacityTint() );				// Color
-	}
-	
-	return LayerId + TextLayer;
+	FVector2D TextLayoutSize = EditableTextLayout->ComputeDesiredSize(LayoutScaleMultiplier);
+	TextLayoutSize.X = FMath::Max(TextLayoutSize.X, MinDesiredWidth.Get());
+	return TextLayoutSize;
 }
 
+FChildren* SEditableText::GetChildren()
+{
+	return EditableTextLayout->GetChildren();
+}
+
+void SEditableText::OnArrangeChildren(const FGeometry& AllottedGeometry, FArrangedChildren& ArrangedChildren) const
+{
+	EditableTextLayout->OnArrangeChildren(AllottedGeometry, ArrangedChildren);
+}
 
 FReply SEditableText::OnDragOver( const FGeometry& MyGeometry, const FDragDropEvent& DragDropEvent )
 {
@@ -1359,7 +130,6 @@ FReply SEditableText::OnDragOver( const FGeometry& MyGeometry, const FDragDropEv
 	return FReply::Unhandled();
 }
 
-
 FReply SEditableText::OnDrop( const FGeometry& MyGeometry, const FDragDropEvent& DragDropEvent )
 {
 	TSharedPtr<FExternalDragOperation> DragDropOp = DragDropEvent.GetOperationAs<FExternalDragOperation>();
@@ -1367,7 +137,7 @@ FReply SEditableText::OnDrop( const FGeometry& MyGeometry, const FDragDropEvent&
 	{
 		if ( DragDropOp->HasText() )
 		{
-			this->Text = FText::FromString(DragDropOp->GetText());
+			EditableTextLayout->SetText(FText::FromString(DragDropOp->GetText()));
 			return FReply::Handled();
 		}
 	}
@@ -1375,7 +145,6 @@ FReply SEditableText::OnDrop( const FGeometry& MyGeometry, const FDragDropEvent&
 	return FReply::Unhandled();
 }
 
-
 bool SEditableText::SupportsKeyboardFocus() const
 {
 	return true;
@@ -1383,133 +152,29 @@ bool SEditableText::SupportsKeyboardFocus() const
 
 FReply SEditableText::OnFocusReceived( const FGeometry& MyGeometry, const FFocusEvent& InFocusEvent )
 {
-	// Skip the rest of the focus received code if it's due to the context menu closing
-	if ( !ActiveContextMenu.IsValid() )
+	if (EditableTextLayout->HandleFocusReceived(InFocusEvent))
 	{
-		EnsureActiveTick();
-
-		// Don't unselect text when focus is set because another widget lost focus, as that may be in response to a
-		// dismissed context menu where the user clicked an option to select text.
-		if( InFocusEvent.GetCause() != EFocusCause::OtherWidgetLostFocus )
+		if (InFocusEvent.GetCause() != EFocusCause::Mouse && InFocusEvent.GetCause() != EFocusCause::OtherWidgetLostFocus && bIsCaretMovedWhenGainFocus.Get())
 		{
-			ClearSelection();
+			EditableTextLayout->GoTo(ETextLocation::EndOfDocument);
 		}
-	 
-		bHasDragSelectedSinceFocused = false;
 
-		// The user wants to edit text. Make a copy of the observed text for the user to edit.
-		LoadText();
-
-		// For mouse-based focus, we'll handle the cursor positioning in the actual mouse down/up events
-		if( InFocusEvent.GetCause() != EFocusCause::Mouse )
-		{
-			// Don't move the cursor if we gained focus because another widget lost focus.  It may
-			// have been a dismissed context menu where the user is expecting the selection to be retained
-			if( InFocusEvent.GetCause() != EFocusCause::OtherWidgetLostFocus )
-			{
-				// Select all text on keyboard focus.  This mirrors mouse focus functionality.
-				if( bSelectAllTextWhenFocused.Get() )
-				{
-					SelectAllText();
-				}
-			
-				// Move the cursor to the end of the string
-				if( IsCaretMovedWhenGainFocus.Get() )
-				{
-					SetCaretPosition( EditedText.ToString().Len() );
-				}
-			}
-		}
-
-		FSlateApplication& SlateApplication = FSlateApplication::Get();
-		if (FPlatformMisc::GetRequiresVirtualKeyboard())
-		{
-			if (!GetIsReadOnly())
-			{
-				// @TODO: Create ITextInputMethodSystem derivations for mobile
-				SlateApplication.ShowVirtualKeyboard(true, InFocusEvent.GetUser(), SharedThis(this));
-				bShowingVirtualKeyboard = true;
-			}
-		}
-		else
-		{
-			ITextInputMethodSystem* const TextInputMethodSystem = SlateApplication.GetTextInputMethodSystem();
-			if (TextInputMethodSystem)
-			{
-				TextInputMethodSystem->ActivateContext(TextInputMethodContext.ToSharedRef());
-			}
-		}
+		return SWidget::OnFocusReceived(MyGeometry, InFocusEvent);
 	}
 
 	return FReply::Handled();
 }
 
-
 void SEditableText::OnFocusLost( const FFocusEvent& InFocusEvent )
 {
-	// Skip the focus lost code if it's due to the context menu opening
-	if ( !ActiveContextMenu.IsValid() )
-	{
-		// Place the caret at the front of the text
-		SetCaretPosition(0);
-
-		FSlateApplication& SlateApplication = FSlateApplication::Get();
-		if (FPlatformMisc::GetRequiresVirtualKeyboard())
-		{
-			if (bShowingVirtualKeyboard)
-			{
-				SlateApplication.ShowVirtualKeyboard(false, InFocusEvent.GetUser());
-				bShowingVirtualKeyboard = false;
-			}
-		}
-		else
-		{
-			ITextInputMethodSystem* const TextInputMethodSystem = SlateApplication.GetTextInputMethodSystem();
-			if (TextInputMethodSystem)
-			{
-				TextInputMethodSystem->DeactivateContext(TextInputMethodContext.ToSharedRef());
-			}
-		}
-
-		// Clear selection unless activating a new window (otherwise can't copy and past on right click)
-		if (InFocusEvent.GetCause() != EFocusCause::WindowActivate)
-		{
-			ClearSelection();
-		}
-
-		//Always clear the local undo chain on commit.
-		ClearUndoStates();
-		
-		// See if user explicitly tabbed away or moved focus
-		ETextCommit::Type TextAction;
-		switch ( InFocusEvent.GetCause() )
-		{
-		case EFocusCause::Navigation:
-		case EFocusCause::Mouse:
-			TextAction = ETextCommit::OnUserMovedFocus;
-			break;
-
-		case EFocusCause::Cleared:
-			TextAction = ETextCommit::OnCleared;
-			break;
-
-		default:
-			TextAction = ETextCommit::Default;
-			break;
-		}
-
-		// When focus is lost let anyone who is interested know that text was committed
-		OnTextCommitted.ExecuteIfBound( EditedText, TextAction);
-	}
+	EditableTextLayout->HandleFocusLost(InFocusEvent);
 }
 
-
 FReply SEditableText::OnKeyChar( const FGeometry& MyGeometry, const FCharacterEvent& InCharacterEvent )
 {
-	return FTextEditHelper::OnKeyChar( InCharacterEvent, SharedThis( this ) );
+	return EditableTextLayout->HandleKeyChar(InCharacterEvent);
 }
 
-
 FReply SEditableText::OnKeyDown( const FGeometry& MyGeometry, const FKeyEvent& InKeyEvent )
 {
 	FReply Reply = FReply::Unhandled();
@@ -1520,587 +185,305 @@ FReply SEditableText::OnKeyDown( const FGeometry& MyGeometry, const FKeyEvent& I
 		Reply = OnKeyDownHandler.Execute(MyGeometry, InKeyEvent);
 	}
 
-	if( !Reply.IsEventHandled() )
+	if (!Reply.IsEventHandled())
 	{
-		Reply = FTextEditHelper::OnKeyDown( InKeyEvent, SharedThis( this ) );
-		
+		Reply = EditableTextLayout->HandleKeyDown(InKeyEvent);
+
 		if (!Reply.IsEventHandled())
 		{
-			// Process keybindings if the event wasn't already handled
-			if ( UICommandList->ProcessCommandBindings( InKeyEvent ) )
-			{
-				Reply = FReply::Handled();
-			}
-			else
-			{
-				Reply = SLeafWidget::OnKeyDown( MyGeometry, InKeyEvent );
-			}
+			Reply = SWidget::OnKeyDown(MyGeometry, InKeyEvent);
 		}
-
 	}
 
 	return Reply;
 }
 
-
 FReply SEditableText::OnKeyUp( const FGeometry& MyGeometry, const FKeyEvent& InKeyEvent )
 {
-	FReply Reply = FReply::Unhandled();
-
-	// ...
-
-	return Reply;
+	return FReply::Unhandled();
 }
 
-
 FReply SEditableText::OnMouseButtonDown( const FGeometry& InMyGeometry, const FPointerEvent& InMouseEvent )
 {
-	if( !HasMouseCapture() )
-	{
-		if( InMouseEvent.GetEffectingButton() == EKeys::LeftMouseButton ||
-			InMouseEvent.GetEffectingButton() == EKeys::RightMouseButton )
-		{
-			// Reset 'is drag selecting' state
-			bHasDragSelectedSinceMouseDown = false;
-		}
-	}
-
-	return FTextEditHelper::OnMouseButtonDown( InMyGeometry, InMouseEvent, SharedThis( this ) );
+	return EditableTextLayout->HandleMouseButtonDown(InMyGeometry, InMouseEvent);
 }
 
-
 FReply SEditableText::OnMouseButtonUp( const FGeometry& InMyGeometry, const FPointerEvent& InMouseEvent )
 {
-	return FTextEditHelper::OnMouseButtonUp( InMyGeometry, InMouseEvent, SharedThis( this ) );
+	return EditableTextLayout->HandleMouseButtonUp(InMyGeometry, InMouseEvent);
 }
 
-
 FReply SEditableText::OnMouseMove( const FGeometry& InMyGeometry, const FPointerEvent& InMouseEvent )
 {
-	if( bIsDragSelecting && HasMouseCapture() )
-	{
-		bHasDragSelectedSinceMouseDown = true;
-	}
-
-	return FTextEditHelper::OnMouseMove( InMyGeometry, InMouseEvent, SharedThis( this ) );;
+	return EditableTextLayout->HandleMouseMove(InMyGeometry, InMouseEvent);
 }
 
-
 FReply SEditableText::OnMouseButtonDoubleClick( const FGeometry& InMyGeometry, const FPointerEvent& InMouseEvent )
 {
-	return FTextEditHelper::OnMouseButtonDoubleClick( InMyGeometry, InMouseEvent, SharedThis( this ) );
+	return EditableTextLayout->HandleMouseButtonDoubleClick(InMyGeometry, InMouseEvent);
 }
 
-
 FCursorReply SEditableText::OnCursorQuery( const FGeometry& MyGeometry, const FPointerEvent& CursorEvent ) const
 {
 	return FCursorReply::Cursor( EMouseCursor::TextEditBeam );
 }
 
-
 const FSlateBrush* SEditableText::GetFocusBrush() const
 {
 	return nullptr;
 }
 
-
 bool SEditableText::IsInteractable() const
 {
 	return IsEnabled();
 }
 
-
-void SEditableText::SelectText( const int32 InOldCaretPosition )
+void SEditableText::SetHintText( const TAttribute< FText >& InHintText )
 {
-	if( InOldCaretPosition != CaretPosition )
-	{
-		if( !AnyTextSelected() )
-		{
-			// Start selecting
-			Selection.StartIndex = Selection.FinishIndex = InOldCaretPosition;
-		}
-		else
-		{
-			// NOTE: We intentionally don't call RestartSelectionTargetAnimation() when text was already selected,
-			//		 as the effect is often distracting when shrinking or growing an existing selection.
-		}
-		Selection.FinishIndex = CaretPosition;
-	}
+	EditableTextLayout->SetHintText(InHintText);
 }
 
-
-void SEditableText::DeleteSelectedText()
+FText SEditableText::GetHintText() const
 {
-	if( AnyTextSelected() )
-	{
-		// Move the caret to the location of the first selected character
-		CaretPosition = Selection.GetMinIndex();
-
-		// Delete the selected text
-		const FString& OldText = EditedText.ToString();
-		const FString NewText = OldText.Left( Selection.GetMinIndex() ) + OldText.Mid( Selection.GetMaxIndex() );
-		EditedText = FText::FromString(NewText);
-
-		// Clear selection
-		ClearSelection();
-	}
+	return EditableTextLayout->GetHintText();
 }
 
-
-int32 SEditableText::FindClickedCharacterIndex( const FGeometry& InMyGeometry, const FVector2D& InScreenCursorPosition ) const
+void SEditableText::SetIsReadOnly( TAttribute< bool > InIsReadOnly )
 {
-	// Find the click position in this widget's local space
-	const FVector2D& ScreenCursorPosition = InScreenCursorPosition;
-	const FVector2D& LocalCursorPosition = InMyGeometry.AbsoluteToLocal( ScreenCursorPosition );
-
-	return FindClickedCharacterIndex( LocalCursorPosition, InMyGeometry.Scale );
+	bIsReadOnly = InIsReadOnly;
 }
 
-
-int32 SEditableText::FindClickedCharacterIndex( const FVector2D& InLocalCursorPosition, float GeometryScale ) const
+void SEditableText::SetIsPassword( TAttribute< bool > InIsPassword )
 {
-	// Convert the click position from the local widget space to the scrolled area
-	const FVector2D& PositionInScrollableArea = ScrollHelper.ToScrollerSpace( InLocalCursorPosition * GeometryScale );
-	const float ClampedCursorX = FMath::Max( PositionInScrollableArea.X, 0.0f );
-
-	// Figure out where in the string the user clicked
-	const TSharedRef< FSlateFontMeasure > FontMeasureService = FSlateApplication::Get().GetRenderer()->GetFontMeasureService();
-	const int32 ClickedCharIndex = FontMeasureService->FindCharacterIndexAtOffset( EditedText.ToString( ), Font.Get( ), ClampedCursorX, GeometryScale );
-
-	return ClickedCharIndex;
+	bIsPassword = InIsPassword;
+	PlainTextMarshaller->SetIsPassword(bIsPassword);
 }
 
-
-int32 SEditableText::ScanForWordBoundary( const int32 Location, int8 Direction ) const
+void SEditableText::SetColorAndOpacity(TAttribute<FSlateColor> Color)
 {
-	// We delay creating the iterator until we actually need it
-	if (!WordBreakIterator.IsValid())
-	{
-		WordBreakIterator = FBreakIterator::CreateWordBreakIterator();
-	}
-
-	const FString& EditedTextString = EditedText.ToString();
-	WordBreakIterator->SetString(EditedTextString);
-
-	int32 NewLocation = Location;
-	if (Direction > 0)
-	{
-		// First move right to the next break candidate
-		NewLocation = WordBreakIterator->MoveToCandidateAfter(NewLocation);
-
-		// Then step over any whitespace
-		while (EditedTextString.IsValidIndex(NewLocation) && FText::IsWhitespace(EditedTextString[NewLocation]))
-		{
-			++NewLocation;
-		}
-
-		// No subsequent break candidate, just set to the upper-range
-		if (NewLocation == INDEX_NONE)
-		{
-			NewLocation = EditedTextString.Len();
-		}
-	}
-	else
-	{
-		// First step over any whitespace
-		while (EditedTextString.IsValidIndex(NewLocation-1) && FText::IsWhitespace(EditedTextString[NewLocation-1]))
-		{
-			--NewLocation;
-		}
-
-		// Then move left to the next previous candidate
-		NewLocation = WordBreakIterator->MoveToCandidateBefore(NewLocation);
-
-		// No previous break candidate, just set to the lower-range
-		if (NewLocation == INDEX_NONE)
-		{
-			NewLocation = 0;
-		}
-	}
-
-	WordBreakIterator->ClearString();
-
-	return NewLocation;
+	ColorAndOpacity = Color;
 }
 
-
-bool SEditableText::IsAtWordStart( const int32 Location ) const
+void SEditableText::SetMinDesiredWidth(const TAttribute<float>& InMinDesiredWidth)
 {
-	const FString& EditedTextString = EditedText.ToString();
-
-	if (Location > 0 && Location < EditedTextString.Len())
-	{
-		const TCHAR CharBeforeCursor = EditedTextString[Location - 1];
-		const TCHAR CharAtCursor = EditedTextString[Location];
-		return FText::IsWhitespace(CharBeforeCursor) && !FText::IsWhitespace(CharAtCursor);
-	}
-
-	return false;
+	MinDesiredWidth = InMinDesiredWidth;
 }
 
-
-void SEditableText::PushUndoState( const SEditableText::FUndoState& InUndoState )
+void SEditableText::SetIsCaretMovedWhenGainFocus(const TAttribute<bool>& InIsCaretMovedWhenGainFocus)
 {
-	// If we've already undone some state, then we'll remove any undo state beyond the level that
-	// we've already undone up to.
-	if( CurrentUndoLevel != INDEX_NONE )
-	{
-		UndoStates.RemoveAt( CurrentUndoLevel, UndoStates.Num() - CurrentUndoLevel );
-
-		// Reset undo level as we haven't undone anything since this latest undo
-		CurrentUndoLevel = INDEX_NONE;
-	}
-
-	// Cache new undo state
-	UndoStates.Add( InUndoState );
-
-	// If we've reached the maximum number of undo levels, then trim our array
-	if( UndoStates.Num() > EditableTextDefs::MaxUndoLevels )
-	{
-		UndoStates.RemoveAt( 0 );
-	}
+	bIsCaretMovedWhenGainFocus = InIsCaretMovedWhenGainFocus;
 }
 
-
-void SEditableText::ClearUndoStates()
+void SEditableText::SetSelectAllTextWhenFocused(const TAttribute<bool>& InSelectAllTextWhenFocused)
 {
-	CurrentUndoLevel = INDEX_NONE;
-	UndoStates.Empty();
+	bSelectAllTextWhenFocused = InSelectAllTextWhenFocused;
 }
 
-
-void SEditableText::MakeUndoState( SEditableText::FUndoState& OutUndoState )
+void SEditableText::SetRevertTextOnEscape(const TAttribute<bool>& InRevertTextOnEscape)
 {
-	OutUndoState.Text = EditedText;
-	OutUndoState.CaretPosition = CaretPosition;
-	OutUndoState.Selection = Selection;
+	bRevertTextOnEscape = InRevertTextOnEscape;
 }
 
+void SEditableText::SetClearKeyboardFocusOnCommit(const TAttribute<bool>& InClearKeyboardFocusOnCommit)
+{
+	bClearKeyboardFocusOnCommit = InClearKeyboardFocusOnCommit;
+}
 
-void SEditableText::LoadText()
+void SEditableText::SetSelectAllTextOnCommit(const TAttribute<bool>& InSelectAllTextOnCommit)
 {
-	// Check to see if the text is actually different.  This might often be the case the first time we
-	// check it after this control is created, but otherwise should only be different if the text attribute
-	// was set (or driven) by an external source
-	const FText NewText = Text.Get();
-	
-	// Store the original text.  This functionality is off by default so we only do it when the option is enabled
-	if( RevertTextOnEscape.Get() && !NewText.EqualTo(OriginalText))
-	{
-			OriginalText = NewText;
-	}
+	bSelectAllTextOnCommit = InSelectAllTextOnCommit;
+}
 
-	if( !NewText.EqualTo(EditedText) )
-	{
-		// @todo Slate: We should run any sort of format validation/sanitation here before accepting the new text
-		EditedText = NewText;
+void SEditableText::SetAllowContextMenu(const TAttribute< bool >& InAllowContextMenu)
+{
+	bAllowContextMenu = InAllowContextMenu;
+}
 
-		int32 StringLength = EditedText.ToString().Len();
+void SEditableText::SetTextShapingMethod(const TOptional<ETextShapingMethod>& InTextShapingMethod)
+{
+	EditableTextLayout->SetTextShapingMethod(InTextShapingMethod);
+}
 
-		// Validate cursor positioning
-		if( CaretPosition > StringLength )
-		{
-			CaretPosition = StringLength;
-		}
+void SEditableText::SetTextFlowDirection(const TOptional<ETextFlowDirection>& InTextFlowDirection)
+{
+	EditableTextLayout->SetTextFlowDirection(InTextFlowDirection);
+}
 
+bool SEditableText::AnyTextSelected() const
+{
+	return EditableTextLayout->AnyTextSelected();
+}
 
-		// Validate selection
-		if( Selection.StartIndex != INDEX_NONE && Selection.StartIndex > StringLength )
-		{
-			Selection.StartIndex = StringLength;
-		}
-		if( Selection.FinishIndex != INDEX_NONE && Selection.FinishIndex > StringLength )
-		{
-			Selection.FinishIndex = StringLength;
-		}
-	}
+void SEditableText::SelectAllText()
+{
+	EditableTextLayout->SelectAllText();
 }
 
+void SEditableText::ClearSelection()
+{
+	EditableTextLayout->ClearSelection();
+}
 
-void SEditableText::SaveText()
+FText SEditableText::GetSelectedText() const
 {
-	// Don't set text if the text attribute has a 'getter' binding on it, otherwise we'd blow away
-	// that binding.  If there is a getter binding, then we'll assume it will provide us with
-	// updated text after we've fired our 'text changed' callbacks
-	if( !Text.IsBound() )
-	{
-		Text.Set ( EditedText );
-	}
+	return EditableTextLayout->GetSelectedText();
 }
 
-void SEditableText::SummonContextMenu(const FVector2D& InLocation, TSharedPtr<SWindow> ParentWindow, const FWidgetPath& EventPath)
+void SEditableText::SynchronizeTextStyle()
 {
-	if (!AllowContextMenu.Get())
-	{
-		return;
-	}
+	// Has the style used for this editable text changed?
+	bool bTextStyleChanged = false;
+	FTextBlockStyle NewTextStyle = EditableTextLayout->GetTextStyle();
 
-	TSharedPtr<SWidget> MenuContentWidget;
-	if (OnContextMenuOpening.IsBound())
+	// Sync from the font override
+	if (Font.IsSet())
 	{
-		MenuContentWidget = OnContextMenuOpening.Execute();
-	}
-	else
-	{
-		// Set the menu to automatically close when the user commits to a choice
-		const bool bShouldCloseWindowAfterMenuSelection = true;
-
-#define LOCTEXT_NAMESPACE "EditableTextContextMenu"
-
-		// This is a context menu which could be summoned from within another menu if this text block is in a menu
-		// it should not close the menu it is inside
-		bool bCloseSelfOnly = true;
-		FMenuBuilder MenuBuilder(bShouldCloseWindowAfterMenuSelection, UICommandList, MenuExtender, bCloseSelfOnly, &FCoreStyle::Get());
+		const FSlateFontInfo& NewFontInfo = Font.Get();
+		if (NewTextStyle.Font != NewFontInfo)
 		{
-			MenuBuilder.BeginSection("EditText", LOCTEXT("Heading", "Modify Text"));
-			{
-				// Undo
-				MenuBuilder.AddMenuEntry(FGenericCommands::Get().Undo);
-			}
-			MenuBuilder.EndSection();
-
-			MenuBuilder.BeginSection("EditableTextModify2");
-			{
-				// Cut
-				MenuBuilder.AddMenuEntry(FGenericCommands::Get().Cut);
-
-				// Copy
-				MenuBuilder.AddMenuEntry(FGenericCommands::Get().Copy);
-
-				// Paste
-				MenuBuilder.AddMenuEntry(FGenericCommands::Get().Paste);
-
-				// Delete
-				MenuBuilder.AddMenuEntry(FGenericCommands::Get().Delete);
-			}
-			MenuBuilder.EndSection();
-
-			MenuBuilder.BeginSection("EditableTextModify3");
-			{
-				// Select All
-				MenuBuilder.AddMenuEntry(FGenericCommands::Get().SelectAll);
-			}
-			MenuBuilder.EndSection();
+			NewTextStyle.Font = NewFontInfo;
+			bTextStyleChanged = true;
 		}
-
-		MenuContentWidget = MenuBuilder.MakeWidget();
-
-#undef LOCTEXT_NAMESPACE
 	}
 
-	if (MenuContentWidget.IsValid())
+	// Sync from the color override
+	if (ColorAndOpacity.IsSet())
 	{
-		ActiveContextMenu.PrepareToSummon();
-
-		const bool bFocusImmediately = true;
-		TSharedRef<SWidget> MenuParent = SharedThis(this);
-		if (ParentWindow.IsValid())
+		const FSlateColor& NewColorAndOpacity = ColorAndOpacity.Get();
+		if (NewTextStyle.ColorAndOpacity != NewColorAndOpacity)
 		{
-			MenuParent = StaticCastSharedRef<SWidget>(ParentWindow.ToSharedRef());
+			NewTextStyle.ColorAndOpacity = NewColorAndOpacity;
+			bTextStyleChanged = true;
 		}
-		TSharedPtr< IMenu > ContextMenu = FSlateApplication::Get().PushMenu(MenuParent, EventPath, MenuContentWidget.ToSharedRef(), InLocation, FPopupTransitionEffect(FPopupTransitionEffect::ContextMenu), bFocusImmediately);
+	}
 
-		// Make sure the menu is valid.  It's possible for the parent to already be in the destroy queue, for example if the editable text was configured to dismiss it's window during OnTextCommitted.
-		if (ContextMenu.IsValid())
-		{
-			ContextMenu->GetOnMenuDismissed().AddSP(this, &SEditableText::OnContextMenuClosed);
-			ActiveContextMenu.SummonSucceeded(ContextMenu.ToSharedRef());
-		}
-		else
+	// Sync from the highlight shape override
+	if (BackgroundImageSelected.IsSet())
+	{
+		const FSlateBrush* NewSelectionBrush = BackgroundImageSelected.Get();
+		if (NewSelectionBrush && NewTextStyle.HighlightShape != *NewSelectionBrush)
 		{
-			ActiveContextMenu.SummonFailed();
+			NewTextStyle.HighlightShape = *NewSelectionBrush;
+			bTextStyleChanged = true;
 		}
 	}
-}
 
-
-void SEditableText::ExecuteDeleteAction()
-{
-	if( !IsReadOnly.Get() && AnyTextSelected() )
+	if (bTextStyleChanged)
 	{
-		StartChangingText();
-
-		// Delete selected text
-		DeleteSelectedText();
-
-		FinishChangingText();
+		EditableTextLayout->SetTextStyle(NewTextStyle);
+		EditableTextLayout->Refresh();
 	}
 }
 
-
-bool SEditableText::CanExecuteSelectAll() const
+bool SEditableText::IsTextReadOnly() const
 {
-	bool bCanExecute = true;
-
-	// Can't select all if string is empty
-	if( EditedText.IsEmpty() )
-	{
-		bCanExecute = false;
-	}
-
-	// Can't select all if entire string is already selected
-	if( Selection.GetMinIndex() == 0 &&
-		Selection.GetMaxIndex() == EditedText.ToString().Len() - 1 )
-	{
-		bCanExecute = false;
-	}
-
-	return bCanExecute;
+	return bIsReadOnly.Get(false);
 }
 
-
-bool SEditableText::CanExecuteDelete() const
+bool SEditableText::IsTextPassword() const
 {
-	bool bCanExecute = true;
-
-	// Can't execute if this is a read-only control
-	if( IsReadOnly.Get() )
-	{
-		bCanExecute = false;
-	}
-
-	// Can't execute unless there is some text selected
-	if( !AnyTextSelected() )
-	{
-		bCanExecute = false;
-	}
-
-	return bCanExecute;
+	return bIsPassword.Get(false);
 }
 
-
-void SEditableText::SetCaretPosition( int32 Position )
+bool SEditableText::IsMultiLineTextEdit() const
 {
-	CaretPosition = Position;
-	LastCaretInteractionTime = FSlateApplication::Get().GetCurrentTime();
+	return false;
 }
 
-bool SEditableText::ShouldAppearFocused() const
+bool SEditableText::ShouldSelectAllTextWhenFocused() const
 {
-	// We'll draw with the 'focused' look if we're either focused or we have a context menu summoned
-	return HasKeyboardFocus() || ActiveContextMenu.IsValid();
+	return bSelectAllTextWhenFocused.Get(false);
 }
 
-bool SEditableText::DoesClipboardHaveAnyText() const 
+bool SEditableText::ShouldClearTextSelectionOnFocusLoss() const
 {
-	FString ClipboardContent;
-	FPlatformMisc::ClipboardPaste(ClipboardContent);
-	return ClipboardContent.Len() > 0;
+	return true;
 }
 
-
-FString SEditableText::GetStringToRender() const
+bool SEditableText::ShouldRevertTextOnEscape() const
 {
-	FString VisibleText;
-
-	const bool bShouldAppearFocused = ShouldAppearFocused();
-	if(bShouldAppearFocused)
-	{
-		VisibleText = EditedText.ToString();
-	}
-	else
-	{
-		VisibleText = Text.Get().ToString();
-	}
-
-	// If this string is a password we need to replace all the characters in the string with something else
-	if ( IsPassword.Get() )
-	{
-#if PLATFORM_TCHAR_IS_1_BYTE
-		const TCHAR Dot = '*';
-#else
-		const TCHAR Dot = 0x2022;
-#endif
-		int32 VisibleTextLen = VisibleText.Len();
-		VisibleText.Empty();
-		while ( VisibleTextLen )
-		{
-			VisibleText += Dot;
-			VisibleTextLen--;
-		}
-	}
-
-	return VisibleText;
+	return bRevertTextOnEscape.Get(false);
 }
 
-void SEditableText::OnContextMenuClosed(TSharedRef<IMenu> Menu)
+bool SEditableText::ShouldClearKeyboardFocusOnCommit() const
 {
-	// Note: We don't reset the ActiveContextMenu here, as Slate hasn't yet finished processing window focus events, and we need 
-	// to know that the window is still available for OnFocusReceived and OnFocusLost even though it's about to be destroyed
-
-	// Give this focus when the context menu has been dismissed
-	FSlateApplication::Get().SetKeyboardFocus( AsShared(), EFocusCause::OtherWidgetLostFocus );
+	return bClearKeyboardFocusOnCommit.Get(false);
 }
 
-void SEditableText::SetHintText( const TAttribute< FText >& InHintText )
+bool SEditableText::CanInsertCarriageReturn() const
 {
-	HintText = InHintText;
+	return false;
 }
 
-void SEditableText::SetIsReadOnly( TAttribute< bool > InIsReadOnly )
+bool SEditableText::CanTypeCharacter(const TCHAR InChar) const
 {
-	IsReadOnly = InIsReadOnly;
+	return !OnIsTypedCharValid.IsBound() || OnIsTypedCharValid.Execute(InChar);
 }
 
-void SEditableText::SetIsPassword( TAttribute< bool > InIsPassword )
+void SEditableText::EnsureActiveTick()
 {
-	IsPassword = InIsPassword;
-}
+	TSharedPtr<FActiveTimerHandle> ActiveTickTimerPin = ActiveTickTimer.Pin();
+	if (ActiveTickTimerPin.IsValid())
+	{
+		return;
+	}
 
-void SEditableText::SetColorAndOpacity(TAttribute<FSlateColor> Color)
-{
-	ColorAndOpacity = Color;
+	auto DoActiveTick = [this](double InCurrentTime, float InDeltaTime) -> EActiveTimerReturnType
+	{
+		// Continue if we still have focus, otherwise treat as a fire-and-forget Tick() request
+		const bool bShouldAppearFocused = HasKeyboardFocus() || EditableTextLayout->HasActiveContextMenu();
+		return (bShouldAppearFocused) ? EActiveTimerReturnType::Continue : EActiveTimerReturnType::Stop;
+	};
+
+	const float TickPeriod = EditableTextDefs::BlinksPerSecond * 0.5f;
+	ActiveTickTimer = RegisterActiveTimer(TickPeriod, FWidgetActiveTimerDelegate::CreateLambda(DoActiveTick));
 }
 
-void SEditableText::SetMinDesiredWidth(const TAttribute<float>& InMinDesiredWidth)
+EKeyboardType SEditableText::GetVirtualKeyboardType() const
 {
-	MinDesiredWidth = InMinDesiredWidth;
+	return VirtualKeyboardType.Get();
 }
 
-void SEditableText::SetIsCaretMovedWhenGainFocus(const TAttribute<bool>& InIsCaretMovedWhenGainFocus)
+TSharedRef<SWidget> SEditableText::GetSlateWidget()
 {
-	IsCaretMovedWhenGainFocus = InIsCaretMovedWhenGainFocus;
+	return SharedThis(this);
 }
 
-void SEditableText::SetSelectAllTextWhenFocused(const TAttribute<bool>& InSelectAllTextWhenFocused)
+TSharedPtr<SWidget> SEditableText::BuildContextMenuContent() const
 {
-	bSelectAllTextWhenFocused = InSelectAllTextWhenFocused;
+	if (!bAllowContextMenu.Get())
+	{
+		return nullptr;
+	}
+
+	if (OnContextMenuOpening.IsBound())
+	{
+		return OnContextMenuOpening.Execute();
+	}
+	
+	return EditableTextLayout->BuildDefaultContextMenu(MenuExtender);
 }
 
-void SEditableText::SetRevertTextOnEscape(const TAttribute<bool>& InRevertTextOnEscape)
+void SEditableText::OnTextChanged(const FText& InText)
 {
-	RevertTextOnEscape = InRevertTextOnEscape;
+	OnTextChangedCallback.ExecuteIfBound(InText);
 }
 
-void SEditableText::SetClearKeyboardFocusOnCommit(const TAttribute<bool>& InClearKeyboardFocusOnCommit)
+void SEditableText::OnTextCommitted(const FText& InText, const ETextCommit::Type InTextAction)
 {
-	ClearKeyboardFocusOnCommit = InClearKeyboardFocusOnCommit;
+	OnTextCommittedCallback.ExecuteIfBound(InText, InTextAction);
 }
 
-void SEditableText::SetSelectAllTextOnCommit(const TAttribute<bool>& InSelectAllTextOnCommit)
+void SEditableText::OnCursorMoved(const FTextLocation& InLocation)
 {
-	SelectAllTextOnCommit = InSelectAllTextOnCommit;
 }
 
-void SEditableText::SetAllowContextMenu(const TAttribute< bool >& InAllowContextMenu)
+float SEditableText::UpdateAndClampHorizontalScrollBar(const float InViewOffset, const float InViewFraction, const EVisibility InVisiblityOverride)
 {
-	AllowContextMenu = InAllowContextMenu;
+	return EditableTextLayout->GetScrollOffset().X;
 }
 
-void SEditableText::EnsureActiveTick()
+float SEditableText::UpdateAndClampVerticalScrollBar(const float InViewOffset, const float InViewFraction, const EVisibility InVisiblityOverride)
 {
-	TSharedPtr<FActiveTimerHandle> ActiveTickTimerPin = ActiveTickTimer.Pin();
-	if(ActiveTickTimerPin.IsValid())
-	{
-		return;
-	}
-
-	auto DoActiveTick = [this](double InCurrentTime, float InDeltaTime) -> EActiveTimerReturnType
-	{
-		// Continue if we still have focus, otherwise treat as a fire-and-forget Tick() request
-		const bool bShouldAppearFocused = HasKeyboardFocus() || ActiveContextMenu.IsValid();
-		return (bShouldAppearFocused) ? EActiveTimerReturnType::Continue : EActiveTimerReturnType::Stop;
-	};
-
-	ActiveTickTimer = RegisterActiveTimer(0.5f, FWidgetActiveTimerDelegate::CreateLambda(DoActiveTick));
+	return 0.0f;
 }
diff --git a/Engine/Source/Runtime/Slate/Private/Widgets/Input/SEditableTextBox.cpp b/Engine/Source/Runtime/Slate/Private/Widgets/Input/SEditableTextBox.cpp
--- a/Engine/Source/Runtime/Slate/Private/Widgets/Input/SEditableTextBox.cpp
+++ b/Engine/Source/Runtime/Slate/Private/Widgets/Input/SEditableTextBox.cpp
@@ -59,6 +59,8 @@ void SEditableTextBox::Construct( const FArguments& InArgs )
 					.SelectAllTextOnCommit( InArgs._SelectAllTextOnCommit )
 					.OnKeyDownHandler( InArgs._OnKeyDownHandler )
 					.VirtualKeyboardType( InArgs._VirtualKeyboardType )
+					.TextShapingMethod( InArgs._TextShapingMethod )
+					.TextFlowDirection( InArgs._TextFlowDirection )
 				]
 			]
 		]
@@ -88,6 +90,8 @@ void SEditableTextBox::SetStyle(const FEditableTextBoxStyle* InStyle)
 		Style = Defaults._Style;
 	}
 
+	check(Style);
+
 	if (!PaddingOverride.IsSet() && PaddingBox.IsValid())
 	{
 		PaddingBox->SetPadding(Style->Padding);
@@ -152,15 +156,51 @@ void SEditableTextBox::SetError( const FString& InError )
 	ErrorReporting->SetError( InError );
 }
 
+
 void SEditableTextBox::SetOnKeyDownHandler(FOnKeyDown InOnKeyDownHandler)
 {
 	EditableText->SetOnKeyDownHandler(InOnKeyDownHandler);
 }
 
-/** @return Border image for the text box based on the hovered and focused state */
+
+void SEditableTextBox::SetTextShapingMethod(const TOptional<ETextShapingMethod>& InTextShapingMethod)
+{
+	EditableText->SetTextShapingMethod(InTextShapingMethod);
+}
+
+
+void SEditableTextBox::SetTextFlowDirection(const TOptional<ETextFlowDirection>& InTextFlowDirection)
+{
+	EditableText->SetTextFlowDirection(InTextFlowDirection);
+}
+
+
+bool SEditableTextBox::AnyTextSelected() const
+{
+	return EditableText->AnyTextSelected();
+}
+
+
+void SEditableTextBox::SelectAllText()
+{
+	EditableText->SelectAllText();
+}
+
+
+void SEditableTextBox::ClearSelection()
+{
+	EditableText->ClearSelection();
+}
+
+
+FText SEditableTextBox::GetSelectedText() const
+{
+	return EditableText->GetSelectedText();
+}
+
 const FSlateBrush* SEditableTextBox::GetBorderImage() const
 {
-	if ( EditableText->GetIsReadOnly() )
+	if ( EditableText->IsTextReadOnly() )
 	{
 		return BorderImageReadOnly;
 	}
@@ -209,14 +249,13 @@ FReply SEditableTextBox::OnFocusReceived( const FGeometry& MyGeometry, const FFo
 }
 
 
-FReply SEditableTextBox::OnKeyDown( const FGeometry& MyGeometry, const FKeyEvent& InKeyEvent )
+FReply SEditableTextBox::OnKeyDown(const FGeometry& MyGeometry, const FKeyEvent& InKeyEvent)
 {
 	FKey Key = InKeyEvent.GetKey();
 
-	if( Key == EKeys::Escape && EditableText->HasKeyboardFocus() )
+	if (Key == EKeys::Escape && EditableText->HasKeyboardFocus())
 	{
-		// Clear selection
-		EditableText->ClearSelection();
+		// Clear focus
 		return FReply::Handled().SetUserFocus(SharedThis(this), EFocusCause::Cleared);
 	}
 
diff --git a/Engine/Source/Runtime/Slate/Private/Widgets/Input/SMultiLineEditableTextBox.cpp b/Engine/Source/Runtime/Slate/Private/Widgets/Input/SMultiLineEditableTextBox.cpp
--- a/Engine/Source/Runtime/Slate/Private/Widgets/Input/SMultiLineEditableTextBox.cpp
+++ b/Engine/Source/Runtime/Slate/Private/Widgets/Input/SMultiLineEditableTextBox.cpp
@@ -166,8 +166,17 @@ void SMultiLineEditableTextBox::Construct( const FArguments& InArgs )
 
 void SMultiLineEditableTextBox::SetStyle(const FEditableTextBoxStyle* InStyle)
 {
-	check(InStyle);
-	Style = InStyle;
+	if (InStyle)
+	{
+		Style = InStyle;
+	}
+	else
+	{
+		FArguments Defaults;
+		Style = Defaults._Style;
+	}
+
+	check(Style);
 
 	if (!PaddingOverride.IsSet())
 	{
@@ -319,7 +328,7 @@ void SMultiLineEditableTextBox::SetError( const FString& InError )
 /** @return Border image for the text box based on the hovered and focused state */
 const FSlateBrush* SMultiLineEditableTextBox::GetBorderImage() const
 {
-	if ( EditableText->GetIsReadOnly() )
+	if ( EditableText->IsTextReadOnly() )
 	{
 		return BorderImageReadOnly;
 	}
@@ -364,6 +373,21 @@ FReply SMultiLineEditableTextBox::OnFocusReceived( const FGeometry& MyGeometry,
 	return Reply;
 }
 
+bool SMultiLineEditableTextBox::AnyTextSelected() const
+{
+	return EditableText->AnyTextSelected();
+}
+
+void SMultiLineEditableTextBox::SelectAllText()
+{
+	EditableText->SelectAllText();
+}
+
+void SMultiLineEditableTextBox::ClearSelection()
+{
+	EditableText->ClearSelection();
+}
+
 FText SMultiLineEditableTextBox::GetSelectedText() const
 {
 	return EditableText->GetSelectedText();
@@ -404,7 +428,7 @@ TSharedPtr<const IRun> SMultiLineEditableTextBox::GetRunUnderCursor() const
 	return EditableText->GetRunUnderCursor();
 }
 
-const TArray<TSharedRef<const IRun>> SMultiLineEditableTextBox::GetSelectedRuns() const
+TArray<TSharedRef<const IRun>> SMultiLineEditableTextBox::GetSelectedRuns() const
 {
 	return EditableText->GetSelectedRuns();
 }
diff --git a/Engine/Source/Runtime/Slate/Private/Widgets/Layout/SScrollBar.cpp b/Engine/Source/Runtime/Slate/Private/Widgets/Layout/SScrollBar.cpp
--- a/Engine/Source/Runtime/Slate/Private/Widgets/Layout/SScrollBar.cpp
+++ b/Engine/Source/Runtime/Slate/Private/Widgets/Layout/SScrollBar.cpp
@@ -517,23 +517,31 @@ bool SScrollBar::IsScrolling() const
 
 void SScrollBar::SetStyle(const FScrollBarStyle* InStyle)
 {
-	check(InStyle);
+	const FScrollBarStyle* Style = InStyle;
 
-	NormalThumbImage = &InStyle->NormalThumbImage;
-	HoveredThumbImage = &InStyle->HoveredThumbImage;
-	DraggedThumbImage = &InStyle->DraggedThumbImage;
+	if (Style == nullptr)
+	{
+		FArguments Defaults;
+		Style = Defaults._Style;
+	}
+
+	check(Style);
+
+	NormalThumbImage = &Style->NormalThumbImage;
+	HoveredThumbImage = &Style->HoveredThumbImage;
+	DraggedThumbImage = &Style->DraggedThumbImage;
 
 	if (Orientation == Orient_Vertical)
 	{
-		BackgroundBrush = &InStyle->VerticalBackgroundImage;
-		TopBrush = &InStyle->VerticalTopSlotImage;
-		BottomBrush = &InStyle->VerticalBottomSlotImage;
+		BackgroundBrush = &Style->VerticalBackgroundImage;
+		TopBrush = &Style->VerticalTopSlotImage;
+		BottomBrush = &Style->VerticalBottomSlotImage;
 	}
 	else
 	{
-		BackgroundBrush = &InStyle->HorizontalBackgroundImage;
-		TopBrush = &InStyle->HorizontalTopSlotImage;
-		BottomBrush = &InStyle->HorizontalBottomSlotImage;
+		BackgroundBrush = &Style->HorizontalBackgroundImage;
+		TopBrush = &Style->HorizontalTopSlotImage;
+		BottomBrush = &Style->HorizontalBottomSlotImage;
 	}
 }
 
diff --git a/Engine/Source/Runtime/Slate/Private/Widgets/Notifications/SProgressBar.cpp b/Engine/Source/Runtime/Slate/Private/Widgets/Notifications/SProgressBar.cpp
--- a/Engine/Source/Runtime/Slate/Private/Widgets/Notifications/SProgressBar.cpp
+++ b/Engine/Source/Runtime/Slate/Private/Widgets/Notifications/SProgressBar.cpp
@@ -39,6 +39,15 @@ void SProgressBar::SetPercent(TAttribute< TOptional<float> > InPercent)
 void SProgressBar::SetStyle(const FProgressBarStyle* InStyle)
 {
 	Style = InStyle;
+
+	if (Style == nullptr)
+	{
+		FArguments Defaults;
+		Style = Defaults._Style;
+	}
+
+	check(Style);
+
 	Invalidate(EInvalidateWidget::Layout);
 }
 
diff --git a/Engine/Source/Runtime/Slate/Private/Widgets/SViewport.cpp b/Engine/Source/Runtime/Slate/Private/Widgets/SViewport.cpp
--- a/Engine/Source/Runtime/Slate/Private/Widgets/SViewport.cpp
+++ b/Engine/Source/Runtime/Slate/Private/Widgets/SViewport.cpp
@@ -227,11 +227,6 @@ FReply SViewport::OnKeyChar( const FGeometry& MyGeometry, const FCharacterEvent&
 
 FReply SViewport::OnFocusReceived( const FGeometry& MyGeometry, const FFocusEvent& InFocusEvent )
 {
-	if (WidgetToFocusOnActivate.IsValid())
-	{
-		return FReply::Handled().SetUserFocus(WidgetToFocusOnActivate.Pin().ToSharedRef(), InFocusEvent.GetCause());
-	}
-
 	return ViewportInterface.IsValid() ? ViewportInterface.Pin()->OnFocusReceived(InFocusEvent) : FReply::Unhandled();
 }
 
@@ -269,6 +264,19 @@ void SViewport::OnWindowClosed( const TSharedRef<SWindow>& WindowBeingClosed )
 	}
 }
 
+FReply SViewport::OnViewportActivated(const FWindowActivateEvent& InActivateEvent)
+{
+	return ViewportInterface.IsValid() ? ViewportInterface.Pin()->OnViewportActivated(InActivateEvent) : FReply::Unhandled();
+}
+
+void SViewport::OnViewportDeactivated(const FWindowActivateEvent& InActivateEvent)
+{
+	if (ViewportInterface.IsValid())
+	{
+		ViewportInterface.Pin()->OnViewportDeactivated(InActivateEvent);
+	}
+}
+
 FReply SViewport::OnTouchStarted( const FGeometry& MyGeometry, const FPointerEvent& InTouchEvent )
 {
 	return ViewportInterface.IsValid() ? ViewportInterface.Pin()->OnTouchStarted(MyGeometry, InTouchEvent) : FReply::Unhandled();
