Copying //UE4/Dev-Editor to Dev-Main (//UE4/Dev-Main)

diff --git a/Engine/Source/Runtime/Slate/Private/Widgets/Text/SMultiLineEditableText.cpp b/Engine/Source/Runtime/Slate/Private/Widgets/Text/SMultiLineEditableText.cpp
--- a/Engine/Source/Runtime/Slate/Private/Widgets/Text/SMultiLineEditableText.cpp
+++ b/Engine/Source/Runtime/Slate/Private/Widgets/Text/SMultiLineEditableText.cpp
@@ -6,280 +6,41 @@
 
 #include "TextEditHelper.h"
 #include "PlainTextLayoutMarshaller.h"
-#include "TextBlockLayout.h"
+#include "SlateEditableTextLayout.h"
 #include "ReflectionMetadata.h"
-#include "GenericCommands.h"
 #include "IMenu.h"
 
-void SMultiLineEditableText::FCursorInfo::SetCursorLocationAndCalculateAlignment(const FTextLayout& InTextLayout, const FTextLocation& InCursorPosition)
-{
-	FTextLocation NewCursorPosition = InCursorPosition;
-	ECursorAlignment NewAlignment = ECursorAlignment::Left;
-
-	{
-		const int32 CursorLineIndex = InCursorPosition.GetLineIndex();
-		const int32 CursorOffset = InCursorPosition.GetOffset();
-
-		// A CursorOffset of zero could mark the end of an empty line, but we don't need to adjust the cursor for an empty line
-		if (CursorOffset > 0)
-		{
-			const TArray<FTextLayout::FLineModel>& Lines = InTextLayout.GetLineModels();
-			const FTextLayout::FLineModel& Line = Lines[CursorLineIndex];
-			if (Line.Text->Len() == CursorOffset)
-			{
-				// We need to move the cursor back one from where it currently is; this keeps the interaction point the same 
-				// (since the cursor is aligned to the right), but visually draws the cursor in the correct place
-				NewCursorPosition = FTextLocation(NewCursorPosition, -1);
-				NewAlignment = ECursorAlignment::Right;
-			}
-		}
-	}
-
-	SetCursorLocationAndAlignment(InTextLayout, NewCursorPosition, NewAlignment);
-}
-
-void SMultiLineEditableText::FCursorInfo::SetCursorLocationAndAlignment(const FTextLayout& InTextLayout, const FTextLocation& InCursorPosition, const ECursorAlignment InCursorAlignment)
-{
-	CursorPosition = InCursorPosition;
-	CursorAlignment = InCursorAlignment;
-	CursorTextDirection = TextBiDi::ETextDirection::LeftToRight;
-
-	// Get the text direction for the block under the cursor
-	{
-		const TArray<FTextLayout::FLineView>& LineViews = InTextLayout.GetLineViews();
-		const int32 LineViewIndex = InTextLayout.GetLineViewIndexForTextLocation(LineViews, InCursorPosition, InCursorAlignment == ECursorAlignment::Right);
-		if (LineViews.IsValidIndex(LineViewIndex))
-		{
-			const FTextLayout::FLineView& LineView = LineViews[LineViewIndex];
-			const int32 CursorOffset = InCursorPosition.GetOffset();
-
-			for (const auto& Block : LineView.Blocks)
-			{
-				const FTextRange BlockTextRange = Block->GetTextRange();
-				if (BlockTextRange.Contains(CursorOffset))
-				{
-					CursorTextDirection = Block->GetTextContext().TextDirection;
-					break;
-				}
-			}
-		}
-	}
-
-	LastCursorInteractionTime = FSlateApplication::Get().GetCurrentTime();
-}
-
-SMultiLineEditableText::FCursorInfo SMultiLineEditableText::FCursorInfo::CreateUndo() const
-{
-	FCursorInfo UndoData;
-	UndoData.CursorPosition = CursorPosition;
-	UndoData.CursorAlignment = CursorAlignment;
-	UndoData.CursorTextDirection = CursorTextDirection;
-	return UndoData;
-}
-
-void SMultiLineEditableText::FCursorInfo::RestoreFromUndo(const FCursorInfo& UndoData)
-{
-	CursorPosition = UndoData.CursorPosition;
-	CursorAlignment = UndoData.CursorAlignment;
-	CursorTextDirection = UndoData.CursorTextDirection;
-	LastCursorInteractionTime = FSlateApplication::Get().GetCurrentTime();
-}
-
-SMultiLineEditableText::FCursorLineHighlighter::FCursorLineHighlighter(const FCursorInfo* InCursorInfo)
-	: CursorInfo(InCursorInfo)
-{
-	check(CursorInfo);
-}
-
-int32 SMultiLineEditableText::FCursorLineHighlighter::OnPaint(const FPaintArgs& Args, const FTextLayout::FLineView& Line, const float OffsetX, const float Width, const FTextBlockStyle& DefaultStyle, const FGeometry& AllottedGeometry, const FSlateRect& MyClippingRect, FSlateWindowElementList& OutDrawElements, int32 LayerId, const FWidgetStyle& InWidgetStyle, bool bParentEnabled) const
-{
-	const FVector2D Location(Line.Offset.X + OffsetX, Line.Offset.Y);
-	const FVector2D Size(Width, Line.TextSize.Y);
-
-	FLinearColor CursorColorAndOpacity = InWidgetStyle.GetForegroundColor();
-
-	const float FontMaxCharHeight = FTextEditHelper::GetFontHeight(DefaultStyle.Font);
-	const float CursorWidth = FTextEditHelper::CalculateCaretWidth(FontMaxCharHeight);
-	const double CurrentTime = FSlateApplication::Get().GetCurrentTime();
-
-	// The block size and offset values are pre-scaled, so we need to account for that when converting the block offsets into paint geometry
-	const float InverseScale = Inverse(AllottedGeometry.Scale);
-
-	// The cursor is always visible (i.e. not blinking) when we're interacting with it; otherwise it might get lost.
-	const bool bForceCursorVisible = (CurrentTime - CursorInfo->GetLastCursorInteractionTime()) < EditableTextDefs::CaretBlinkPauseTime;
-	float CursorOpacity = (bForceCursorVisible)
-		? 1.0f
-		: FMath::RoundToFloat( FMath::MakePulsatingValue( CurrentTime, EditableTextDefs::BlinksPerSecond ));
-
-	CursorOpacity *= CursorOpacity;	// Squared falloff, because it looks more interesting
-	CursorColorAndOpacity.A = CursorOpacity;
-
-	// @todo: Slate Styles - make this brush part of the widget style
-	const FSlateBrush* CursorBrush = FCoreStyle::Get().GetBrush("EditableText.SelectionBackground");
-
-	ECursorAlignment VisualCursorAlignment = CursorInfo->GetCursorAlignment();
-	if (CursorInfo->GetCursorTextDirection() == TextBiDi::ETextDirection::RightToLeft)
-	{
-		// We need to flip the visual alignment for right-to-left text, as the start of the glyph is one the right hand side of the highlight
-		VisualCursorAlignment = VisualCursorAlignment == ECursorAlignment::Left ? ECursorAlignment::Right : ECursorAlignment::Left;
-	}
-	const FVector2D OptionalWidth = VisualCursorAlignment == ECursorAlignment::Right ? FVector2D(Size.X, 0) : FVector2D::ZeroVector;
-
-	FSlateDrawElement::MakeBox(
-		OutDrawElements,
-		LayerId,
-		AllottedGeometry.ToPaintGeometry(TransformVector(InverseScale, FVector2D(CursorWidth * AllottedGeometry.Scale, Size.Y)), FSlateLayoutTransform(TransformPoint(InverseScale, Location + OptionalWidth))),
-		CursorBrush,
-		MyClippingRect,
-		bParentEnabled ? ESlateDrawEffect::None : ESlateDrawEffect::DisabledEffect,
-		CursorColorAndOpacity*InWidgetStyle.GetColorAndOpacityTint()
-		);
-
-	return LayerId;
-}
-
-TSharedRef< SMultiLineEditableText::FCursorLineHighlighter > SMultiLineEditableText::FCursorLineHighlighter::Create(const FCursorInfo* InCursorInfo)
-{
-	return MakeShareable( new SMultiLineEditableText::FCursorLineHighlighter(InCursorInfo) );
-}
-
-SMultiLineEditableText::FTextCompositionHighlighter::FTextCompositionHighlighter()
-{
-}
-
-int32 SMultiLineEditableText::FTextCompositionHighlighter::OnPaint( const FPaintArgs& Args, const FTextLayout::FLineView& Line, const float OffsetX, const float Width, const FTextBlockStyle& DefaultStyle, const FGeometry& AllottedGeometry, const FSlateRect& MyClippingRect, FSlateWindowElementList& OutDrawElements, int32 LayerId, const FWidgetStyle& InWidgetStyle, bool bParentEnabled ) const
-{
-	const FVector2D Location(Line.Offset.X + OffsetX, Line.Offset.Y);
-	const FVector2D Size(Width, Line.TextSize.Y);
-
-	// The block size and offset values are pre-scaled, so we need to account for that when converting the block offsets into paint geometry
-	const float InverseScale = Inverse(AllottedGeometry.Scale);
-
-	if (Size.X)
-	{
-		const FLinearColor LineColorAndOpacity = InWidgetStyle.GetForegroundColor();
-
-		// @todo: Slate Styles - make this brush part of the widget style
-		const FSlateBrush* CompositionBrush = FCoreStyle::Get().GetBrush("EditableText.CompositionBackground");
-
-		// Draw composition background
-		FSlateDrawElement::MakeBox(
-			OutDrawElements,
-			++LayerId,
-			AllottedGeometry.ToPaintGeometry(TransformVector(InverseScale, Size), FSlateLayoutTransform(TransformPoint(InverseScale, Location))),
-			CompositionBrush,
-			MyClippingRect,
-			bParentEnabled ? ESlateDrawEffect::None : ESlateDrawEffect::DisabledEffect,
-			LineColorAndOpacity * InWidgetStyle.GetColorAndOpacityTint()
-		);
-	}
-
-	return LayerId;
-}
-
-TSharedRef< SMultiLineEditableText::FTextCompositionHighlighter > SMultiLineEditableText::FTextCompositionHighlighter::Create()
-{
-	return MakeShareable( new SMultiLineEditableText::FTextCompositionHighlighter() );
-}
-
-SMultiLineEditableText::FTextSelectionRunRenderer::FTextSelectionRunRenderer()
-{
-}
-
-int32 SMultiLineEditableText::FTextSelectionRunRenderer::OnPaint( const FPaintArgs& Args, const FTextLayout::FLineView& Line, const TSharedRef< ISlateRun >& Run, const TSharedRef< ILayoutBlock >& Block, const FTextBlockStyle& DefaultStyle, const FGeometry& AllottedGeometry, const FSlateRect& MyClippingRect, FSlateWindowElementList& OutDrawElements, int32 LayerId, const FWidgetStyle& InWidgetStyle, bool bParentEnabled ) const 
-{
-	FVector2D Location( Block->GetLocationOffset() );
-	Location.Y = Line.Offset.Y;
-
-	// If we've not been set to an explicit color, calculate a suitable one from the linked color
-	const FLinearColor SelectionBackgroundColorAndOpacity = DefaultStyle.SelectedBackgroundColor.IsColorSpecified() 
-		? DefaultStyle.SelectedBackgroundColor.GetSpecifiedColor() * InWidgetStyle.GetColorAndOpacityTint()
-		: ((FLinearColor::White - DefaultStyle.SelectedBackgroundColor.GetColor(InWidgetStyle))*0.5f + FLinearColor(-0.2f, -0.05f, 0.15f)) * InWidgetStyle.GetColorAndOpacityTint();
-
-	// The block size and offset values are pre-scaled, so we need to account for that when converting the block offsets into paint geometry
-	const float InverseScale = Inverse(AllottedGeometry.Scale);
-
-	// We still want to show a small selection outline on empty lines to make it clear that the line itself is selected despite being empty
-	const float MinHighlightWidth = (Line.Range.IsEmpty()) ? 4.0f * AllottedGeometry.Scale : 0.0f;
-	const float HighlightWidth = FMath::Max(Block->GetSize().X, MinHighlightWidth);
-	if (HighlightWidth > 0.0f)
-	{
-		// Draw the actual highlight rectangle
-		FSlateDrawElement::MakeBox(
-			OutDrawElements,
-			++LayerId,
-			AllottedGeometry.ToPaintGeometry(TransformVector(InverseScale, FVector2D(HighlightWidth, Line.Size.Y)), FSlateLayoutTransform(TransformPoint(InverseScale, Location))),
-			&DefaultStyle.HighlightShape,
-			MyClippingRect,
-			bParentEnabled && bHasKeyboardFocus ? ESlateDrawEffect::None : ESlateDrawEffect::DisabledEffect,
-			SelectionBackgroundColorAndOpacity
-		);
-	}
-
-	/*
-	FLinearColor InvertedForeground = FLinearColor::White - InWidgetStyle.GetForegroundColor();
-	InvertedForeground.A = InWidgetStyle.GetForegroundColor().A;
-
-	FWidgetStyle WidgetStyle( InWidgetStyle );
-	WidgetStyle.SetForegroundColor( InvertedForeground );
-	*/
-
-	return Run->OnPaint( Args, Line, Block, DefaultStyle, AllottedGeometry, MyClippingRect, OutDrawElements, LayerId, InWidgetStyle, bParentEnabled );
-}
-
-TSharedRef< SMultiLineEditableText::FTextSelectionRunRenderer > SMultiLineEditableText::FTextSelectionRunRenderer::Create()
-{
-	return MakeShareable( new SMultiLineEditableText::FTextSelectionRunRenderer() );
-}
-
 SMultiLineEditableText::SMultiLineEditableText()
-	: PreferredCursorScreenOffsetInLine(0)
-	, bSelectAllTextWhenFocused(false)
-	, bIsDragSelecting(false)
-	, bWasFocusedByLastMouseDown(false)
-	, bHasDragSelectedSinceFocused(false)
-	, CurrentUndoLevel(INDEX_NONE)
-	, bIsChangingText(false)
-	, IsReadOnly(false)
-	, UICommandList(new FUICommandList())
-	, bTextChangedByVirtualKeyboard(false)
-	, bTextCommittedByVirtualKeyboard(false)
+	: bSelectAllTextWhenFocused(false)
+	, bIsReadOnly(false)
 	, AmountScrolledWhileRightMouseDown(0.0f)
 	, bIsSoftwareCursor(false)
 {
-	
 }
 
 SMultiLineEditableText::~SMultiLineEditableText()
 {
-	ITextInputMethodSystem* const TextInputMethodSystem = FSlateApplication::IsInitialized() ? FSlateApplication::Get().GetTextInputMethodSystem() : nullptr;
-	if(TextInputMethodSystem)
-	{
-		TSharedRef<FTextInputMethodContext> TextInputMethodContextRef = TextInputMethodContext.ToSharedRef();
-		if(TextInputMethodSystem->IsActiveContext(TextInputMethodContextRef))
-		{
-			// This can happen if an entire tree of widgets is culled, as Slate isn't notified of the focus loss, the widget is just deleted
-			TextInputMethodSystem->DeactivateContext(TextInputMethodContextRef);
-		}
-
-		TextInputMethodSystem->UnregisterContext(TextInputMethodContextRef);
-	}
+	// Needed to avoid "deletion of pointer to incomplete type 'FSlateEditableTextLayout'; no destructor called" error when using TUniquePtr
 }
 
 void SMultiLineEditableText::Construct( const FArguments& InArgs )
 {
-	CursorInfo = FCursorInfo();
-	SelectionStart = TOptional<FTextLocation>();
-	PreferredCursorScreenOffsetInLine = 0;
+	bIsReadOnly = InArgs._IsReadOnly;
 
+	OnTextChangedCallback = InArgs._OnTextChanged;
+	OnTextCommittedCallback = InArgs._OnTextCommitted;
+	OnCursorMovedCallback = InArgs._OnCursorMoved;
+	bSelectAllTextWhenFocused = InArgs._SelectAllTextWhenFocused;
+	bClearTextSelectionOnFocusLoss = InArgs._ClearTextSelectionOnFocusLoss;
+	bClearKeyboardFocusOnCommit = InArgs._ClearKeyboardFocusOnCommit;
+	bAllowContextMenu = InArgs._AllowContextMenu;
+	OnContextMenuOpening = InArgs._OnContextMenuOpening;
+	bRevertTextOnEscape = InArgs._RevertTextOnEscape;
+	OnHScrollBarUserScrolled = InArgs._OnHScrollBarUserScrolled;
+	OnVScrollBarUserScrolled = InArgs._OnVScrollBarUserScrolled;
+	OnKeyDownHandler = InArgs._OnKeyDownHandler;
 	ModiferKeyForNewLine = InArgs._ModiferKeyForNewLine;
 
-	WrapTextAt = InArgs._WrapTextAt;
-	AutoWrapText = InArgs._AutoWrapText;
-	CachedSize = FVector2D::ZeroVector;
-
-	ScrollOffset = FVector2D::ZeroVector;
-
 	HScrollBar = InArgs._HScrollBar;
 	if (HScrollBar.IsValid())
 	{
@@ -294,344 +55,178 @@ void SMultiLineEditableText::Construct( const FArguments& InArgs )
 		VScrollBar->SetOnUserScrolled(FOnUserScrolled::CreateSP(this, &SMultiLineEditableText::OnVScrollBarMoved));
 	}
 
-	Margin = InArgs._Margin;
-	LineHeightPercentage = InArgs._LineHeightPercentage;
-	Justification = InArgs._Justification;
-	TextStyle = *InArgs._TextStyle;
-	
+	FTextBlockStyle TextStyle = *InArgs._TextStyle;
 	if (InArgs._Font.IsSet() || InArgs._Font.IsBound())
 	{
 		TextStyle.SetFont(InArgs._Font.Get());
 	}
 
-	IsReadOnly = InArgs._IsReadOnly;
-
-	OnTextChanged = InArgs._OnTextChanged;
-	OnTextCommitted = InArgs._OnTextCommitted;
-	OnCursorMoved = InArgs._OnCursorMoved;
-	bSelectAllTextWhenFocused = InArgs._SelectAllTextWhenFocused;
-	bClearTextSelectionOnFocusLoss = InArgs._ClearTextSelectionOnFocusLoss;
-	bClearKeyboardFocusOnCommit = InArgs._ClearKeyboardFocusOnCommit;
-	AllowContextMenu = InArgs._AllowContextMenu;
-	OnContextMenuOpening = InArgs._OnContextMenuOpening;
-	bRevertTextOnEscape = InArgs._RevertTextOnEscape;
-	OnHScrollBarUserScrolled = InArgs._OnHScrollBarUserScrolled;
-	OnVScrollBarUserScrolled = InArgs._OnVScrollBarUserScrolled;
-	OnKeyDownHandler = InArgs._OnKeyDownHandler;
-
-	Marshaller = InArgs._Marshaller;
+	TSharedPtr<ITextLayoutMarshaller> Marshaller = InArgs._Marshaller;
 	if (!Marshaller.IsValid())
 	{
 		Marshaller = FPlainTextLayoutMarshaller::Create();
 	}
 
-	CursorLineHighlighter = FCursorLineHighlighter::Create(&CursorInfo);
-	TextCompositionHighlighter = FTextCompositionHighlighter::Create();
-	TextSelectionRunRenderer = FTextSelectionRunRenderer::Create();
-	TextLayout = FSlateTextLayout::Create(TextStyle);
-	TextLayout->SetDebugSourceInfo(TAttribute<FString>::Create(TAttribute<FString>::FGetter::CreateLambda([this]{ return FReflectionMetaData::GetWidgetDebugInfo(this); })));
-
-	if (InArgs._TextShapingMethod.IsSet())
-	{
-		TextLayout->SetTextShapingMethod(InArgs._TextShapingMethod.GetValue());
-	}
-
-	if (InArgs._TextFlowDirection.IsSet())
-	{
-		TextLayout->SetTextFlowDirection(InArgs._TextFlowDirection.GetValue());
-	}
-
-	BoundText = InArgs._Text;
-
-	{
-		const FText& TextToSet = BoundText.Get(FText::GetEmpty());
-		SetEditableText(TextToSet, true); // force the set to ensure we create an empty run (if BoundText returns an empty string)
-
-		// Update the cached BoundText value to prevent it triggering another SetEditableText update again next Tick
-		if (BoundText.IsBound())
-		{
-			BoundTextLastTick = FTextSnapshot(TextToSet);
-		}
-	}
-
-	TextLayout->SetWrappingWidth( WrapTextAt.Get() );
-	TextLayout->SetMargin( Margin.Get() );
-	TextLayout->SetLineHeightPercentage( LineHeightPercentage.Get() );
-	TextLayout->SetJustification( Justification.Get() );
-
-	TextLayout->UpdateIfNeeded();
-
-	// If we have hint text that is either non-empty or bound to a delegate, we'll also need to make the hint text layout
-	HintText = InArgs._HintText;
-	if (HintText.IsBound() || !HintText.Get(FText::GetEmpty()).IsEmpty())
-	{
-		HintTextStyle = MakeShareable(new FTextBlockStyle(TextStyle));
-		HintTextLayout = FTextBlockLayout::Create(*HintTextStyle, TextLayout->GetTextShapingMethod(), TextLayout->GetTextFlowDirection(), Marshaller.ToSharedRef(), nullptr);
-		HintTextLayout->SetDebugSourceInfo(TAttribute<FString>::Create(TAttribute<FString>::FGetter::CreateLambda([this]{ return FReflectionMetaData::GetWidgetDebugInfo(this); })));
-	}
-
-	// Map UI commands to delegates which are called when the command should be executed
-	UICommandList->MapAction(FGenericCommands::Get().Undo,
-		FExecuteAction::CreateSP(this, &SMultiLineEditableText::Undo),
-		FCanExecuteAction::CreateSP(this, &SMultiLineEditableText::CanExecuteUndo),
-		EUIActionRepeatMode::RepeatEnabled);
-
-	UICommandList->MapAction(FGenericCommands::Get().Cut,
-		FExecuteAction::CreateSP(this, &SMultiLineEditableText::CutSelectedTextToClipboard),
-		FCanExecuteAction::CreateSP(this, &SMultiLineEditableText::CanExecuteCut));
-
-	UICommandList->MapAction(FGenericCommands::Get().Paste,
-		FExecuteAction::CreateSP(this, &SMultiLineEditableText::PasteTextFromClipboard),
-		FCanExecuteAction::CreateSP(this, &SMultiLineEditableText::CanExecutePaste),
-		EUIActionRepeatMode::RepeatEnabled);
-
-	UICommandList->MapAction(FGenericCommands::Get().Copy,
-		FExecuteAction::CreateSP(this, &SMultiLineEditableText::CopySelectedTextToClipboard),
-		FCanExecuteAction::CreateSP(this, &SMultiLineEditableText::CanExecuteCopy));
-
-	UICommandList->MapAction(FGenericCommands::Get().Delete,
-		FExecuteAction::CreateSP(this, &SMultiLineEditableText::ExecuteDeleteAction),
-		FCanExecuteAction::CreateSP(this, &SMultiLineEditableText::CanExecuteDelete));
-
-	UICommandList->MapAction(FGenericCommands::Get().SelectAll,
-		FExecuteAction::CreateSP(this, &SMultiLineEditableText::SelectAllText),
-		FCanExecuteAction::CreateSP(this, &SMultiLineEditableText::CanExecuteSelectAll));
+	EditableTextLayout = MakeUnique<FSlateEditableTextLayout>(*this, InArgs._Text, TextStyle, InArgs._TextShapingMethod, InArgs._TextFlowDirection, Marshaller.ToSharedRef(), Marshaller.ToSharedRef());
+	EditableTextLayout->SetHintText(InArgs._HintText);
+	EditableTextLayout->SetTextWrapping(InArgs._WrapTextAt, InArgs._AutoWrapText);
+	EditableTextLayout->SetMargin(InArgs._Margin);
+	EditableTextLayout->SetJustification(InArgs._Justification);
+	EditableTextLayout->SetLineHeightPercentage(InArgs._LineHeightPercentage);
+	EditableTextLayout->SetDebugSourceInfo(TAttribute<FString>::Create(TAttribute<FString>::FGetter::CreateLambda([this]{ return FReflectionMetaData::GetWidgetDebugInfo(this); })));
 
 	// build context menu extender
 	MenuExtender = MakeShareable(new FExtender);
 	MenuExtender->AddMenuExtension("EditText", EExtensionHook::Before, TSharedPtr<FUICommandList>(), InArgs._ContextMenuExtender);
-
-	TextInputMethodContext = MakeShareable( new FTextInputMethodContext( StaticCastSharedRef<SMultiLineEditableText>(AsShared()) ) );
-	ITextInputMethodSystem* const TextInputMethodSystem = FSlateApplication::Get().GetTextInputMethodSystem();
-	if(TextInputMethodSystem)
-	{
-		TextInputMethodChangeNotifier = TextInputMethodSystem->RegisterContext(TextInputMethodContext.ToSharedRef());
-	}
-	if(TextInputMethodChangeNotifier.IsValid())
-	{
-		TextInputMethodChangeNotifier->NotifyLayoutChanged(ITextInputMethodChangeNotifier::ELayoutChangeType::Created);
-	}
 }
 
 void SMultiLineEditableText::SetText(const TAttribute< FText >& InText)
 {
-	// NOTE: This will unbind any getter that is currently assigned to the Text attribute!  You should only be calling
-	//       SetText() if you know what you're doing.
-	BoundText = InText;
-
-	const FText& TextToSet = BoundText.Get(FText::GetEmpty());
-
-	// Update the cached BoundText value to prevent it triggering another SetEditableText update again next Tick
-	if (BoundText.IsBound())
-	{
-		BoundTextLastTick = FTextSnapshot(TextToSet);
-	}
+	EditableTextLayout->SetText(InText);
+}
 
-	// Update the internal editable text
-	if (SetEditableText(TextToSet))
-	{
-		// Let outsiders know that the text content has been changed
-		OnTextChanged.ExecuteIfBound(TextToSet);
-	}
+FText SMultiLineEditableText::GetText() const
+{
+	return EditableTextLayout->GetText();
 }
 
-const FText SMultiLineEditableText::GetPlainText() const
+FText SMultiLineEditableText::GetPlainText() const
 {
-	const TArray< FTextLayout::FLineModel >& Lines = TextLayout->GetLineModels();
-	const int32 NumberOfLines = Lines.Num();
-	if( NumberOfLines > 0)
-	{
-		FString SelectedText;
-		const FTextLocation StartLocation(0, 0);
-		const FTextLocation EndLocation(NumberOfLines -1, Lines[NumberOfLines -1].Text->Len());
-		FTextSelection Selection(StartLocation, EndLocation);
-		TextLayout->GetSelectionAsText(SelectedText, Selection);
-		return FText::FromString(SelectedText);
-	}
-	return FText::GetEmpty();
+	return EditableTextLayout->GetPlainText();
 }
 
 void SMultiLineEditableText::SetHintText(const TAttribute< FText >& InHintText)
 {
-	HintText = InHintText;
+	EditableTextLayout->SetHintText(InHintText);
+}
+
+FText SMultiLineEditableText::GetHintText() const
+{
+	return EditableTextLayout->GetHintText();
+}
 
-	// If we have hint text that is either non-empty or bound to a delegate, we'll also need to make the hint text layout
-	if (HintText.IsBound() || !HintText.Get(FText::GetEmpty()).IsEmpty())
+void SMultiLineEditableText::SetTextStyle(const FTextBlockStyle* InTextStyle)
+{
+	if (InTextStyle)
 	{
-		HintTextStyle = MakeShareable(new FTextBlockStyle(TextStyle));
-		HintTextLayout = FTextBlockLayout::Create(*HintTextStyle, TextLayout->GetTextShapingMethod(), TextLayout->GetTextFlowDirection(), Marshaller.ToSharedRef(), nullptr);
-		HintTextLayout->SetDebugSourceInfo(TAttribute<FString>::Create(TAttribute<FString>::FGetter::CreateLambda([this]{ return FReflectionMetaData::GetWidgetDebugInfo(this); })));
+		EditableTextLayout->SetTextStyle(*InTextStyle);
 	}
 	else
 	{
-		HintTextStyle.Reset();
-		HintTextLayout.Reset();
+		FArguments Defaults;
+		check(Defaults._TextStyle);
+		EditableTextLayout->SetTextStyle(*Defaults._TextStyle);
 	}
 }
 
 void SMultiLineEditableText::SetFont(const TAttribute< FSlateFontInfo >& InNewFont)
 {
+	FTextBlockStyle TextStyle = EditableTextLayout->GetTextStyle();
 	TextStyle.SetFont(InNewFont.Get());
-
-	// @todo: Chris.Wood - this doesn't update the font (needs the TextLayout updating?)
+	EditableTextLayout->SetTextStyle(TextStyle);
 }
 
 void SMultiLineEditableText::SetTextShapingMethod(const TOptional<ETextShapingMethod>& InTextShapingMethod)
 {
-	TextLayout->SetTextShapingMethod((InTextShapingMethod.IsSet()) ? InTextShapingMethod.GetValue() : GetDefaultTextShapingMethod());
+	EditableTextLayout->SetTextShapingMethod(InTextShapingMethod);
 }
 
 void SMultiLineEditableText::SetTextFlowDirection(const TOptional<ETextFlowDirection>& InTextFlowDirection)
 {
-	TextLayout->SetTextFlowDirection((InTextFlowDirection.IsSet()) ? InTextFlowDirection.GetValue() : GetDefaultTextFlowDirection());
+	EditableTextLayout->SetTextFlowDirection(InTextFlowDirection);
 }
 
 void SMultiLineEditableText::SetWrapTextAt(const TAttribute<float>& InWrapTextAt)
 {
-	WrapTextAt = InWrapTextAt;
+	EditableTextLayout->SetWrapTextAt(InWrapTextAt);
 }
 
 void SMultiLineEditableText::SetAutoWrapText(const TAttribute<bool>& InAutoWrapText)
 {
-	AutoWrapText = InAutoWrapText;
+	EditableTextLayout->SetAutoWrapText(InAutoWrapText);
 }
 
 void SMultiLineEditableText::SetLineHeightPercentage(const TAttribute<float>& InLineHeightPercentage)
 {
-	LineHeightPercentage = InLineHeightPercentage;
+	EditableTextLayout->SetLineHeightPercentage(InLineHeightPercentage);
 }
 
 void SMultiLineEditableText::SetMargin(const TAttribute<FMargin>& InMargin)
 {
-	Margin = InMargin;
+	EditableTextLayout->SetMargin(InMargin);
 }
 
 void SMultiLineEditableText::SetJustification(const TAttribute<ETextJustify::Type>& InJustification)
 {
-	Justification = InJustification;
+	EditableTextLayout->SetJustification(InJustification);
 }
 
 void SMultiLineEditableText::SetAllowContextMenu(const TAttribute< bool >& InAllowContextMenu)
 {
-	AllowContextMenu = InAllowContextMenu;
+	bAllowContextMenu = InAllowContextMenu;
 }
 
-void SMultiLineEditableText::SetTextFromVirtualKeyboard(const FText& InNewText, ESetTextType SetTextType, ETextCommit::Type CommitType)
+void SMultiLineEditableText::OnHScrollBarMoved(const float InScrollOffsetFraction)
 {
-	// Only set the text if the text attribute doesn't have a getter binding (otherwise it would be blown away).
-	// If it is bound, we'll assume that OnTextCommitted will handle the update.
-	if (!BoundText.IsBound())
-	{
-		BoundText.Set(InNewText);
-	}
-
-	// Update the internal editable text
-	if (SetEditableText(InNewText))
-	{
-		// This method is called from the main thread (i.e. not the game thread) of the device with the virtual keyboard
-		// This causes the app to crash on those devices, so we're using polling here to ensure delegates are
-		// fired on the game thread in Tick.		
-		if (SetTextType == ESetTextType::Changed)
-		{
-			bTextChangedByVirtualKeyboard = true;
-		}
-	}
-	if (SetTextType == ESetTextType::Commited)
-	{
-		VirtualKeyboardTextCommitType = CommitType;
-		bTextCommittedByVirtualKeyboard = true;
-	}
+	EditableTextLayout->SetHorizontalScrollFraction(InScrollOffsetFraction);
+	OnHScrollBarUserScrolled.ExecuteIfBound(InScrollOffsetFraction);
 }
 
-bool SMultiLineEditableText::SetEditableText(const FText& TextToSet, const bool bForce)
+void SMultiLineEditableText::OnVScrollBarMoved(const float InScrollOffsetFraction)
 {
-	const FText EditedText = GetEditableText();
-
-	Marshaller->ClearDirty();
-
-	const bool HasTextChanged = !EditedText.ToString().Equals(TextToSet.ToString(), ESearchCase::CaseSensitive);
-	if (HasTextChanged || bForce)
-	{
-		const FString& TextToSetString = TextToSet.ToString();
-
-		ClearSelection();
-		TextLayout->ClearLines();
-
-		Marshaller->SetText(TextToSetString, *TextLayout);
-
-		const TArray< FTextLayout::FLineModel >& Lines = TextLayout->GetLineModels();
-		if(Lines.Num() == 0)
-		{
-			TSharedRef< FString > LineText = MakeShareable(new FString());
-
-			// Create an empty run
-			TArray< TSharedRef< IRun > > Runs;
-			Runs.Add(FSlateTextRun::Create(FRunInfo(), LineText, TextStyle));
-
-			TextLayout->AddLine(FTextLayout::FNewLineData(MoveTemp(LineText), MoveTemp(Runs)));
-		}
-
-		{
-			const FTextLocation OldCursorPos = CursorInfo.GetCursorInteractionLocation();
-
-			// Make sure the cursor is still at a valid location
-			if(OldCursorPos.GetLineIndex() >= Lines.Num() || OldCursorPos.GetOffset() > Lines[OldCursorPos.GetLineIndex()].Text->Len())
-			{
-				const int32 LastLineIndex = Lines.Num() - 1;
-				const FTextLocation NewCursorPosition = FTextLocation(LastLineIndex, Lines[LastLineIndex].Text->Len());
-
-				CursorInfo.SetCursorLocationAndCalculateAlignment(*TextLayout, NewCursorPosition);
-				OnCursorMoved.ExecuteIfBound(CursorInfo.GetCursorInteractionLocation());
-				UpdatePreferredCursorScreenOffsetInLine();
-				UpdateCursorHighlight();
-			}
-		}
+	EditableTextLayout->SetVerticalScrollFraction(InScrollOffsetFraction);
+	OnVScrollBarUserScrolled.ExecuteIfBound(InScrollOffsetFraction);
+}
 
-		return true;
-	}
+bool SMultiLineEditableText::IsTextReadOnly() const
+{
+	return bIsReadOnly.Get(false);
+}
 
+bool SMultiLineEditableText::IsTextPassword() const
+{
 	return false;
 }
 
-FText SMultiLineEditableText::GetEditableText() const
+bool SMultiLineEditableText::IsMultiLineTextEdit() const
 {
-	FString EditedText;
-	Marshaller->GetText(EditedText, *TextLayout);
-	return FText::FromString(EditedText);
+	return true;
 }
 
-void SMultiLineEditableText::ForceRefreshTextLayout(const FText& CurrentText)
+bool SMultiLineEditableText::ShouldSelectAllTextWhenFocused() const
 {
-	// Marshallers shouldn't inject any visible characters into the text, but SetEditableText will clear the current selection
-	// so we need to back that up here so we don't cause the cursor to jump around
-	const TOptional<FTextLocation> OldSelectionStart = SelectionStart;
-	const FCursorInfo OldCursorInfo = CursorInfo;
+	return bSelectAllTextWhenFocused.Get(false);
+}
 
-	SetEditableText(CurrentText, true);
+bool SMultiLineEditableText::ShouldClearTextSelectionOnFocusLoss() const
+{
+	return bClearTextSelectionOnFocusLoss.Get(false);
+}
 
-	SelectionStart = OldSelectionStart;
-	CursorInfo = OldCursorInfo;
-	UpdateCursorHighlight();
+bool SMultiLineEditableText::ShouldRevertTextOnEscape() const
+{
+	return bRevertTextOnEscape.Get(false);
+}
 
-	TextLayout->UpdateIfNeeded();
+bool SMultiLineEditableText::ShouldClearKeyboardFocusOnCommit() const
+{
+	return bClearKeyboardFocusOnCommit.Get(false);
 }
 
-void SMultiLineEditableText::OnHScrollBarMoved(const float InScrollOffsetFraction)
+bool SMultiLineEditableText::CanInsertCarriageReturn() const
 {
-	ScrollOffset.X = FMath::Clamp<float>(InScrollOffsetFraction, 0.0, 1.0) * TextLayout->GetSize().X;
-	OnHScrollBarUserScrolled.ExecuteIfBound(InScrollOffsetFraction);
+	return FSlateApplication::Get().GetModifierKeys().AreModifersDown(ModiferKeyForNewLine);
 }
 
-void SMultiLineEditableText::OnVScrollBarMoved(const float InScrollOffsetFraction)
+bool SMultiLineEditableText::CanTypeCharacter(const TCHAR InChar) const
 {
-	ScrollOffset.Y = FMath::Clamp<float>(InScrollOffsetFraction, 0.0, 1.0) * TextLayout->GetSize().Y;
-	OnVScrollBarUserScrolled.ExecuteIfBound(InScrollOffsetFraction);
+	return true;
 }
 
 void SMultiLineEditableText::EnsureActiveTick()
 {
 	TSharedPtr<FActiveTimerHandle> ActiveTickTimerPin = ActiveTickTimer.Pin();
-	if(ActiveTickTimerPin.IsValid())
+	if (ActiveTickTimerPin.IsValid())
 	{
 		return;
 	}
@@ -639,2055 +234,191 @@ void SMultiLineEditableText::EnsureActiveTick()
 	auto DoActiveTick = [this](double InCurrentTime, float InDeltaTime) -> EActiveTimerReturnType
 	{
 		// Continue if we still have focus, otherwise treat as a fire-and-forget Tick() request
-		const bool bShouldAppearFocused = HasKeyboardFocus() || ActiveContextMenu.IsValid();
+		const bool bShouldAppearFocused = HasKeyboardFocus() || EditableTextLayout->HasActiveContextMenu();
 		return (bShouldAppearFocused) ? EActiveTimerReturnType::Continue : EActiveTimerReturnType::Stop;
 	};
 
-	ActiveTickTimer = RegisterActiveTimer(0.5f, FWidgetActiveTimerDelegate::CreateLambda(DoActiveTick));
+	const float TickPeriod = EditableTextDefs::BlinksPerSecond * 0.5f;
+	ActiveTickTimer = RegisterActiveTimer(TickPeriod, FWidgetActiveTimerDelegate::CreateLambda(DoActiveTick));
 }
 
-FReply SMultiLineEditableText::OnFocusReceived( const FGeometry& MyGeometry, const FFocusEvent& InFocusEvent )
+EKeyboardType SMultiLineEditableText::GetVirtualKeyboardType() const
 {
-	// Skip the focus received code if it's due to the context menu closing
-	if ( !ActiveContextMenu.IsValid() )
-	{
-		// We need to Tick() while we have focus to keep some things up-to-date
-		EnsureActiveTick();
+	return Keyboard_Default;
+}
 
-		FSlateApplication& SlateApplication = FSlateApplication::Get();
-		if (FPlatformMisc::GetRequiresVirtualKeyboard())
-		{
-			if (!GetIsReadOnly())
-			{
-				// @TODO: Create ITextInputMethodSystem derivations for mobile
-				SlateApplication.ShowVirtualKeyboard(true, InFocusEvent.GetUser(), SharedThis(this));
-			}
-		}
-		else
-		{
-			ITextInputMethodSystem* const TextInputMethodSystem = SlateApplication.GetTextInputMethodSystem();
-			if (TextInputMethodSystem)
-			{
-				TextInputMethodSystem->ActivateContext(TextInputMethodContext.ToSharedRef());
-			}
-		}
-		// Store undo state to use for escape key reverts
-		MakeUndoState(OriginalText);
+TSharedRef<SWidget> SMultiLineEditableText::GetSlateWidget()
+{
+	return SharedThis(this);
+}
 
-		// Select All Text
-		if(SelectAllTextWhenFocused())
-		{
-			SelectAllText();
-		}
+TSharedPtr<SWidget> SMultiLineEditableText::BuildContextMenuContent() const
+{
+	if (!bAllowContextMenu.Get())
+	{
+		return nullptr;
+	}
 
-		UpdateCursorHighlight();
+	if (OnContextMenuOpening.IsBound())
+	{
+		return OnContextMenuOpening.Execute();
+	}
 
-		// UpdateCursorHighlight always tries to scroll to the cursor, but we don't want that to happen when we 
-		// gain focus since it can cause the scroll position to jump unexpectedly
-		// If we gained focus via a mouse click that moved the cursor, then MoveCursor will already take care
-		// of making sure that gets scrolled into view
-		PositionToScrollIntoView.Reset();
+	return EditableTextLayout->BuildDefaultContextMenu(MenuExtender);
+}
 
-		return SWidget::OnFocusReceived( MyGeometry, InFocusEvent );
-	}
+void SMultiLineEditableText::OnTextChanged(const FText& InText)
+{
+	OnTextChangedCallback.ExecuteIfBound(InText);
+}
 
-	return FReply::Handled();
+void SMultiLineEditableText::OnTextCommitted(const FText& InText, const ETextCommit::Type InTextAction)
+{
+	OnTextCommittedCallback.ExecuteIfBound(InText, InTextAction);
 }
 
-void SMultiLineEditableText::OnFocusLost( const FFocusEvent& InFocusEvent )
+void SMultiLineEditableText::OnCursorMoved(const FTextLocation& InLocation)
 {
-	bIsSoftwareCursor = false;
+	OnCursorMovedCallback.ExecuteIfBound(InLocation);
+}
 
-	// Skip the focus lost code if it's due to the context menu opening
-	if (!ActiveContextMenu.IsValid())
+float SMultiLineEditableText::UpdateAndClampHorizontalScrollBar(const float InViewOffset, const float InViewFraction, const EVisibility InVisiblityOverride)
+{
+	if (HScrollBar.IsValid())
 	{
-		FSlateApplication& SlateApplication = FSlateApplication::Get();
-		if (FPlatformMisc::GetRequiresVirtualKeyboard())
-		{
-			SlateApplication.ShowVirtualKeyboard(false, InFocusEvent.GetUser());
-		}
-		else
+		HScrollBar->SetState(InViewOffset, InViewFraction);
+		HScrollBar->SetUserVisibility(InVisiblityOverride);
+		if (!HScrollBar->IsNeeded())
 		{
-			ITextInputMethodSystem* const TextInputMethodSystem = SlateApplication.GetTextInputMethodSystem();
-			if (TextInputMethodSystem)
-			{
-				TextInputMethodSystem->DeactivateContext(TextInputMethodContext.ToSharedRef());
-			}
+			// We cannot scroll, so ensure that there is no offset
+			return 0.0f;
 		}
+	}
 
-		// Clear selection unless activating a new window (otherwise can't copy and past on right click)
-		if (bClearTextSelectionOnFocusLoss.Get() && InFocusEvent.GetCause() != EFocusCause::WindowActivate)
-		{
-			ClearSelection();
-		}
+	return EditableTextLayout->GetScrollOffset().X;
+}
 
-		// When focus is lost let anyone who is interested that text was committed
-		// See if user explicitly tabbed away or moved focus
-		ETextCommit::Type TextAction;
-		switch (InFocusEvent.GetCause())
+float SMultiLineEditableText::UpdateAndClampVerticalScrollBar(const float InViewOffset, const float InViewFraction, const EVisibility InVisiblityOverride)
+{
+	if (VScrollBar.IsValid())
+	{
+		VScrollBar->SetState(InViewOffset, InViewFraction);
+		VScrollBar->SetUserVisibility(InVisiblityOverride);
+		if (!VScrollBar->IsNeeded())
 		{
-		case EFocusCause::Navigation:
-		case EFocusCause::Mouse:
-			TextAction = ETextCommit::OnUserMovedFocus;
-			break;
-
-		case EFocusCause::Cleared:
-			TextAction = ETextCommit::OnCleared;
-			break;
-
-		default:
-			TextAction = ETextCommit::Default;
-			break;
+			// We cannot scroll, so ensure that there is no offset
+			return 0.0f;
 		}
+	}
 
-		//Always clear the local undo chain on commit.
-		ClearUndoStates();
-
-		const FText EditedText = GetEditableText();
-
-		OnTextCommitted.ExecuteIfBound(EditedText, TextAction);
+	return EditableTextLayout->GetScrollOffset().Y;
+}
 
-		UpdateCursorHighlight();
-
-		// UpdateCursorHighlight always tries to scroll to the cursor, but we don't want that to happen when we 
-		// lose focus since it can cause the scroll position to jump unexpectedly
-		PositionToScrollIntoView.Reset();
-	}
-}
-
-void SMultiLineEditableText::StartChangingText()
-{
-	// Never change text on read only controls! 
-	check(!IsReadOnly.Get());
-
-	// We're starting to (potentially) change text
-	check(!bIsChangingText);
-	bIsChangingText = true;
-
-	// Save off an undo state in case we actually change the text
-	MakeUndoState(StateBeforeChangingText);
-}
-
-void SMultiLineEditableText::FinishChangingText()
-{
-	// We're no longer changing text
-	check(bIsChangingText);
-	bIsChangingText = false;
-
-	const FText EditedText = GetEditableText();
-
-	// Has the text changed?
-	const bool HasTextChanged = !EditedText.ToString().Equals(StateBeforeChangingText.Text.ToString(), ESearchCase::CaseSensitive);
-	if (HasTextChanged)
-	{
-		// Save text state
-		SaveText(EditedText);
-
-		// Text was actually changed, so push the undo state we previously saved off
-		PushUndoState(StateBeforeChangingText);
-
-		// We're done with this state data now.  Clear out any old data.
-		StateBeforeChangingText = FUndoState();
-
-		TextLayout->UpdateIfNeeded();
-
-		// Let outsiders know that the text content has been changed
-		OnTextChanged.ExecuteIfBound(EditedText);
-		
-		OnCursorMoved.ExecuteIfBound(CursorInfo.GetCursorInteractionLocation());
-
-		// Update the desired cursor position, since typing will have moved it
-		UpdatePreferredCursorScreenOffsetInLine();
-
-		// If the marshaller we're using requires live text updates (eg, because it injects formatting into the source text)
-		// then we need to force a SetEditableText here so that it can update the new editable text with any extra formatting
-		if (Marshaller->RequiresLiveUpdate())
-		{
-			ForceRefreshTextLayout(EditedText);
-		}
-	}
-}
-
-void SMultiLineEditableText::SaveText(const FText& TextToSave)
-{
-	// Don't set text if the text attribute has a 'getter' binding on it, otherwise we'd blow away
-	// that binding.  If there is a getter binding, then we'll assume it will provide us with
-	// updated text after we've fired our 'text changed' callbacks
-	if (!BoundText.IsBound())
-	{
-		BoundText.Set(TextToSave);
-	}
-}
-
-bool SMultiLineEditableText::GetIsReadOnly() const
-{
-	return IsReadOnly.Get();
-}
-
-void SMultiLineEditableText::BackspaceChar()
-{
-	if( AnyTextSelected() )
-	{
-		// Delete selected text
-		DeleteSelectedText();
-	}
-	else
-	{
-		const FTextLocation CursorInteractionPosition = CursorInfo.GetCursorInteractionLocation();
-		FTextLocation FinalCursorLocation = CursorInteractionPosition;
-
-		const TArray< FTextLayout::FLineModel >& Lines = TextLayout->GetLineModels();
-
-		//If we are at the very beginning of the line...
-		if (CursorInteractionPosition.GetOffset() == 0)
-		{
-			//And if the current line isn't the very first line then...
-			if (CursorInteractionPosition.GetLineIndex() > 0)
-			{
-				const int32 PreviousLineIndex = CursorInteractionPosition.GetLineIndex() - 1;
-				const int32 CachePreviousLinesCurrentLength = Lines[PreviousLineIndex].Text->Len();
-				if (TextLayout->JoinLineWithNextLine(PreviousLineIndex))
-				{
-					//Update the cursor so it appears at the end of the previous line,
-					//as we're going to delete the imaginary \n separating them
-					FinalCursorLocation = FTextLocation(PreviousLineIndex, CachePreviousLinesCurrentLength);
-				}
-			}
-			//else do nothing as the FinalCursorLocation is already correct
-		}
-		else
-		{
-			// Delete character to the left of the caret
-			if (TextLayout->RemoveAt(FTextLocation(CursorInteractionPosition, -1)))
-			{
-				// Adjust caret position one to the left
-				FinalCursorLocation = FTextLocation(CursorInteractionPosition, -1);
-			}
-		}
-
-		CursorInfo.SetCursorLocationAndCalculateAlignment(*TextLayout, FinalCursorLocation);
-
-		ClearSelection();
-		UpdateCursorHighlight();
-	}
-}
-
-/**
- * This gets executed by the context menu and should only attempt to delete any selected text
- */
-void SMultiLineEditableText::ExecuteDeleteAction()
-{
-	if (!IsReadOnly.Get() && AnyTextSelected())
-	{
-		StartChangingText();
-
-		// Delete selected text
-		DeleteSelectedText();
-
-		FinishChangingText();
-	}
-}
-
-bool SMultiLineEditableText::CanExecuteDelete() const
-{
-	bool bCanExecute = true;
-
-	// Can't execute if this is a read-only control
-	if (IsReadOnly.Get())
-	{
-		bCanExecute = false;
-	}
-
-	// Can't execute unless there is some text selected
-	if (!AnyTextSelected())
-	{
-		bCanExecute = false;
-	}
-
-	return bCanExecute;
-}
-
-void SMultiLineEditableText::DeleteChar()
-{
-	if( !IsReadOnly.Get() )
-	{
-		if( AnyTextSelected() )
-		{
-			// Delete selected text
-			DeleteSelectedText();
-		}
-		else
-		{
-			const FTextLocation CursorInteractionPosition = CursorInfo.GetCursorInteractionLocation();
-			FTextLocation FinalCursorLocation = CursorInteractionPosition;
-
-			const TArray< FTextLayout::FLineModel >& Lines = TextLayout->GetLineModels();
-			const FTextLayout::FLineModel& Line = Lines[CursorInteractionPosition.GetLineIndex()];
-
-			//If we are at the very beginning of the line...
-			if (Line.Text->Len() == 0)
-			{
-				//And if the current line isn't the very last line then...
-				if (Lines.IsValidIndex(CursorInteractionPosition.GetLineIndex() + 1))
-				{
-					TextLayout->RemoveLine(CursorInteractionPosition.GetLineIndex());
-				}
-				//else do nothing as the FinalCursorLocation is already correct
-			}
-			else if (CursorInteractionPosition.GetOffset() >= Line.Text->Len())
-			{
-				//And if the current line isn't the very last line then...
-				if (Lines.IsValidIndex(CursorInteractionPosition.GetLineIndex() + 1))
-				{
-					if (TextLayout->JoinLineWithNextLine(CursorInteractionPosition.GetLineIndex()))
-					{
-						//else do nothing as the FinalCursorLocation is already correct
-					}
-				}
-				//else do nothing as the FinalCursorLocation is already correct
-			}
-			else
-			{
-				// Delete character to the right of the caret
-				TextLayout->RemoveAt(CursorInteractionPosition);
-				//do nothing to the cursor as the FinalCursorLocation is already correct
-			}
-
-			CursorInfo.SetCursorLocationAndCalculateAlignment(*TextLayout, FinalCursorLocation);
-
-			ClearSelection();
-			UpdateCursorHighlight();
-		}
-	}
-}
-
-bool SMultiLineEditableText::CanTypeCharacter(const TCHAR CharInQuestion) const
-{
-	return true;
-}
-
-void SMultiLineEditableText::TypeChar( const int32 Character )
-{
-	if( AnyTextSelected() )
-	{
-		// Delete selected text
-		DeleteSelectedText();
-	}
-
-	// Certain characters are not allowed
-	bool bIsCharAllowed = true;
-	{
-		if( Character == TEXT('\t') )
-		{
-			bIsCharAllowed = true;
-		}
-		else if( Character <= 0x1F )
-		{
-			bIsCharAllowed = false;
-		}
-	}
-
-	if( bIsCharAllowed )
-	{
-		const FTextLocation CursorInteractionPosition = CursorInfo.GetCursorInteractionLocation();
-		const TArray< FTextLayout::FLineModel >& Lines = TextLayout->GetLineModels();
-		const FTextLayout::FLineModel& Line = Lines[ CursorInteractionPosition.GetLineIndex() ];
-
-		// Insert character at caret position
-		TextLayout->InsertAt( CursorInteractionPosition, Character );
-
-		// Advance caret position
-		ClearSelection();
-		const FTextLocation FinalCursorLocation = FTextLocation( CursorInteractionPosition.GetLineIndex(), FMath::Min( CursorInteractionPosition.GetOffset() + 1, Line.Text->Len() ) );
-
-		CursorInfo.SetCursorLocationAndCalculateAlignment(*TextLayout, FinalCursorLocation);
-		UpdateCursorHighlight();
-	}
-}
-
-void SMultiLineEditableText::DeleteSelectedText()
-{
-	if( AnyTextSelected() )
-	{
-		const FTextLocation CursorInteractionPosition = CursorInfo.GetCursorInteractionLocation();
-		FTextLocation SelectionLocation = SelectionStart.Get(CursorInteractionPosition);
-		FTextSelection Selection(SelectionLocation, CursorInteractionPosition);
-
-		int32 SelectionBeginningLineIndex = Selection.GetBeginning().GetLineIndex();
-		int32 SelectionBeginningLineOffset = Selection.GetBeginning().GetOffset();
-
-		int32 SelectionEndLineIndex = Selection.GetEnd().GetLineIndex();
-		int32 SelectionEndLineOffset = Selection.GetEnd().GetOffset();
-
-		if (SelectionBeginningLineIndex == SelectionEndLineIndex)
-		{
-			TextLayout->RemoveAt(FTextLocation(SelectionBeginningLineIndex, SelectionBeginningLineOffset), SelectionEndLineOffset - SelectionBeginningLineOffset);
-			//Do nothing to the cursor as it is already in the correct location
-		}
-		else
-		{
-			const TArray< FTextLayout::FLineModel >& Lines = TextLayout->GetLineModels();
-			const FTextLayout::FLineModel& EndLine = Lines[SelectionEndLineIndex];
-
-			if (EndLine.Text->Len() == SelectionEndLineOffset)
-			{ 
-				TextLayout->RemoveLine(SelectionEndLineIndex);
-			}
-			else
-			{
-				TextLayout->RemoveAt(FTextLocation(SelectionEndLineIndex, 0), SelectionEndLineOffset);
-			}
-
-			for (int32 LineIndex = SelectionEndLineIndex - 1; LineIndex > SelectionBeginningLineIndex; LineIndex--)
-			{
-				TextLayout->RemoveLine(LineIndex);
-			}
-
-			const FTextLayout::FLineModel& BeginLine = Lines[SelectionBeginningLineIndex];
-			TextLayout->RemoveAt(FTextLocation(SelectionBeginningLineIndex, SelectionBeginningLineOffset), BeginLine.Text->Len() - SelectionBeginningLineOffset);
-
-			TextLayout->JoinLineWithNextLine(SelectionBeginningLineIndex);
-
-			if (Lines.Num() == 0)
-			{
-				const TSharedRef< FString > EmptyText = MakeShareable(new FString());
-				TArray< TSharedRef< IRun > > Runs;
-				Runs.Add(FSlateTextRun::Create(FRunInfo(), EmptyText, TextStyle));
-
-				TextLayout->AddLine(FTextLayout::FNewLineData(MoveTemp(EmptyText), MoveTemp(Runs)));
-			}
-		}
-
-		// Clear selection
-		ClearSelection();
-		const FTextLocation FinalCursorLocation = FTextLocation(SelectionBeginningLineIndex, SelectionBeginningLineOffset);
-		CursorInfo.SetCursorLocationAndCalculateAlignment(*TextLayout, FinalCursorLocation);
-		UpdateCursorHighlight();
-	}
-}
-
-FReply SMultiLineEditableText::MoveCursor( FMoveCursor Args )
-{
-	// We can't use the keyboard to move the cursor while composing, as the IME has control over it
-	if(TextInputMethodContext->IsComposing && Args.GetMoveMethod() != ECursorMoveMethod::ScreenPosition)
-	{
-		return FReply::Handled();
-	}
-
-	bool bAllowMoveCursor = true;
-	FTextLocation NewCursorPosition;
-	FTextLocation CursorPosition = CursorInfo.GetCursorInteractionLocation();
-
-	// If we have selected text, the cursor needs to:
-	// a) Jump to the start of the selection if moving the cursor Left or Up
-	// b) Jump to the end of the selection if moving the cursor Right or Down
-	// This is done regardless of whether the selection was made left-to-right, or right-to-left
-	// This also needs to be done *before* moving to word boundaries, or moving vertically, as the 
-	// start point needs to be the start or end of the selection depending on the above rules
-	if ( Args.GetAction() == ECursorAction::MoveCursor && Args.GetMoveMethod() != ECursorMoveMethod::ScreenPosition && AnyTextSelected( ) )
-	{
-		if (Args.IsHorizontalMovement())
-		{
-			// If we're moving the cursor horizontally, we just snap to the start or end of the selection rather than 
-			// move the cursor by the normal movement rules
-			bAllowMoveCursor = false;
-		}
-
-		// Work out which edge of the selection we need to start at
-		bool bSnapToSelectionStart =
-			Args.GetMoveMethod() == ECursorMoveMethod::Cardinal &&
-			(Args.GetMoveDirection().X < 0.0f || Args.GetMoveDirection().Y < 0.0f);
-
-
-		// Adjust the current cursor position - also set the new cursor position so that the bAllowMoveCursor == false case is handled
-		const FTextSelection Selection(SelectionStart.GetValue(), CursorPosition);
-		CursorPosition = bSnapToSelectionStart ? Selection.GetBeginning() : Selection.GetEnd();
-		NewCursorPosition = CursorPosition;
-
-		// If we're snapping to a word boundary, but the selection was already at a word boundary, don't let the cursor move any more
-		if (Args.GetGranularity() == ECursorMoveGranularity::Word && IsAtWordStart(NewCursorPosition))
-		{
-			bAllowMoveCursor = false;
-		}
-	}
-
-	TOptional<ECursorAlignment> NewCursorAlignment;
-	bool bUpdatePreferredCursorScreenOffsetInLine = false;
-	if (bAllowMoveCursor)
-	{
-		if ( Args.GetMoveMethod() == ECursorMoveMethod::Cardinal )
-		{
-			if ( Args.GetGranularity() == ECursorMoveGranularity::Character )
-			{
-				if ( Args.IsHorizontalMovement() )
-				{
-					NewCursorPosition = TranslatedLocation( CursorPosition, Args.GetMoveDirection().X );
-					bUpdatePreferredCursorScreenOffsetInLine = true;
-				}
-				else
-				{
-					TranslateLocationVertical( CursorPosition, Args.GetMoveDirection().Y, Args.GetGeometryScale(), NewCursorPosition, NewCursorAlignment );
-				}
-			}
-			else
-			{
-				checkSlow( Args.IsHorizontalMovement() );
-				checkSlow( Args.GetGranularity() == ECursorMoveGranularity::Word );
-				checkSlow( Args.GetMoveDirection().X != 0 );
-				NewCursorPosition = ScanForWordBoundary( CursorPosition, Args.GetMoveDirection().X );
-				bUpdatePreferredCursorScreenOffsetInLine = true;
-			}
-		}
-		else if ( Args.GetMoveMethod() == ECursorMoveMethod::ScreenPosition )
-		{
-			ETextHitPoint HitPoint = ETextHitPoint::WithinText;
-			NewCursorPosition = TextLayout->GetTextLocationAt( Args.GetLocalPosition() * Args.GetGeometryScale(), &HitPoint );
-			bUpdatePreferredCursorScreenOffsetInLine = true;
-
-			// Moving with the mouse behaves a bit differently to moving with the keyboard, as clicking at the end of a wrapped line needs to place the cursor there
-			// rather than at the start of the next line (which is tricky since they have the same index according to GetTextLocationAt!).
-			// We use the HitPoint to work this out and then adjust the cursor position accordingly
-			if ( HitPoint == ETextHitPoint::RightGutter )
-			{
-				NewCursorPosition = FTextLocation( NewCursorPosition, -1 );
-				NewCursorAlignment = ECursorAlignment::Right;
-			}
-		}
-		else
-		{
-			checkfSlow(false, TEXT("Unknown ECursorMoveMethod value"));
-		}
-	}
-
-	if (Args.GetAction() == ECursorAction::SelectText)
-	{
-		// We are selecting text. Just remember where the selection started.
-		// The cursor is implicitly the other endpoint.
-		if (!SelectionStart.IsSet())
-		{
-			this->SelectionStart = CursorPosition;
-		}
-	}
-	else
-	{
-		// No longer selection text; clear the selection!
-		this->ClearSelection();
-	}
-
-	if (NewCursorAlignment.IsSet())
-	{
-		CursorInfo.SetCursorLocationAndAlignment(*TextLayout, NewCursorPosition, NewCursorAlignment.GetValue());
-	}
-	else
-	{
-		CursorInfo.SetCursorLocationAndCalculateAlignment(*TextLayout, NewCursorPosition);
-	}
-
-	OnCursorMoved.ExecuteIfBound(CursorInfo.GetCursorInteractionLocation());
-
-	if (bUpdatePreferredCursorScreenOffsetInLine)
-	{
-		UpdatePreferredCursorScreenOffsetInLine();
-	}
-
-	UpdateCursorHighlight();
-
-	// If we've moved the cursor while composing, we need to end the current composition session
-	// Note: You should only be able to do this via the mouse due to the check at the top of this function
-	if(TextInputMethodContext->IsComposing)
-	{
-		ITextInputMethodSystem* const TextInputMethodSystem = FSlateApplication::Get().GetTextInputMethodSystem();
-		if(TextInputMethodSystem)
-		{
-			TextInputMethodSystem->DeactivateContext(TextInputMethodContext.ToSharedRef());
-			TextInputMethodSystem->ActivateContext(TextInputMethodContext.ToSharedRef());
-		}
-	}
-
-	return FReply::Handled();
-}
-
-void SMultiLineEditableText::UpdateCursorHighlight()
-{
-	PositionToScrollIntoView = FScrollInfo(CursorInfo.GetCursorInteractionLocation(), CursorInfo.GetCursorAlignment());
-	EnsureActiveTick();
-
-	RemoveCursorHighlight();
-
-	static const int32 CompositionRangeZOrder = 1; // draw above the text
-	static const int32 CursorZOrder = 2; // draw above the text and the composition
-
-	const FTextLocation CursorInteractionPosition = CursorInfo.GetCursorInteractionLocation();
-	const FTextLocation SelectionLocation = SelectionStart.Get( CursorInteractionPosition );
-
-	const bool bHasKeyboardFocus = HasKeyboardFocus();
-	const bool bIsComposing = TextInputMethodContext->IsComposing;
-	const bool bHasSelection = SelectionLocation != CursorInteractionPosition;
-
-	if ( bIsComposing )
-	{
-		FTextLayout::FTextOffsetLocations OffsetLocations;
-		TextLayout->GetTextOffsetLocations(OffsetLocations);
-
-		const FTextLocation CompositionBeginLocation = OffsetLocations.OffsetToTextLocation(TextInputMethodContext->CompositionBeginIndex);
-		const FTextLocation CompositionEndLocation = OffsetLocations.OffsetToTextLocation(TextInputMethodContext->CompositionBeginIndex + TextInputMethodContext->CompositionLength);
-
-		// Composition should never span more than one (hard) line
-		if ( CompositionBeginLocation.GetLineIndex() == CompositionEndLocation.GetLineIndex() )
-		{
-			const FTextRange Range(CompositionBeginLocation.GetOffset(), CompositionEndLocation.GetOffset());
-
-			// We only draw the composition highlight if the cursor is within the composition range
-			const bool bCursorInRange = (CompositionBeginLocation.GetLineIndex() == CursorInteractionPosition.GetLineIndex() && Range.InclusiveContains(CursorInteractionPosition.GetOffset()));
-			if (!Range.IsEmpty() && bCursorInRange)
-			{
-				TextLayout->AddLineHighlight(FTextLineHighlight(CompositionBeginLocation.GetLineIndex(), Range, CompositionRangeZOrder, TextCompositionHighlighter.ToSharedRef()));
-			}
-		}
-	}
-	else if ( bHasSelection )
-	{
-		const FTextSelection Selection( SelectionLocation, CursorInteractionPosition );
-
-		const int32 SelectionBeginningLineIndex = Selection.GetBeginning().GetLineIndex();
-		const int32 SelectionBeginningLineOffset = Selection.GetBeginning().GetOffset();
-		
-		const int32 SelectionEndLineIndex = Selection.GetEnd().GetLineIndex();
-		const int32 SelectionEndLineOffset = Selection.GetEnd().GetOffset();
-
-		TextSelectionRunRenderer->SetHasKeyboardFocus(bHasKeyboardFocus);
-
-		if ( SelectionBeginningLineIndex == SelectionEndLineIndex )
-		{
-			const FTextRange Range(SelectionBeginningLineOffset, SelectionEndLineOffset);
-			TextLayout->AddRunRenderer(FTextRunRenderer(SelectionBeginningLineIndex, Range, TextSelectionRunRenderer.ToSharedRef()));
-		}
-		else
-		{
-			const TArray< FTextLayout::FLineModel >& Lines = TextLayout->GetLineModels();
-
-			for (int32 LineIndex = SelectionBeginningLineIndex; LineIndex <= SelectionEndLineIndex; LineIndex++)
-			{
-				if ( LineIndex == SelectionBeginningLineIndex )
-				{
-					const FTextRange Range(SelectionBeginningLineOffset, Lines[LineIndex].Text->Len());
-					TextLayout->AddRunRenderer(FTextRunRenderer(LineIndex, Range, TextSelectionRunRenderer.ToSharedRef()));
-				}
-				else if ( LineIndex == SelectionEndLineIndex )
-				{
-					const FTextRange Range(0, SelectionEndLineOffset);
-					TextLayout->AddRunRenderer(FTextRunRenderer(LineIndex, Range, TextSelectionRunRenderer.ToSharedRef()));
-				}
-				else
-				{
-					const FTextRange Range(0, Lines[LineIndex].Text->Len());
-					TextLayout->AddRunRenderer(FTextRunRenderer(LineIndex, Range, TextSelectionRunRenderer.ToSharedRef()));
-				}
-			}
-		}
-	}
-	
-	if ( bHasKeyboardFocus )
-	{
-		// The cursor mode uses the literal position rather than the interaction position
-		const FTextLocation CursorPosition = CursorInfo.GetCursorLocation();
-
-		const TArray< FTextLayout::FLineModel >& Lines = TextLayout->GetLineModels();
-		const int32 LineTextLength = Lines[CursorPosition.GetLineIndex()].Text->Len();
-
-		if (LineTextLength == 0)
-		{
-			TextLayout->AddLineHighlight(FTextLineHighlight(CursorPosition.GetLineIndex(), FTextRange(0, 0), CursorZOrder, CursorLineHighlighter.ToSharedRef()));
-		}
-		else if (CursorPosition.GetOffset() == LineTextLength)
-		{
-			TextLayout->AddLineHighlight(FTextLineHighlight(CursorPosition.GetLineIndex(), FTextRange(LineTextLength - 1, LineTextLength), CursorZOrder, CursorLineHighlighter.ToSharedRef()));
-		}
-		else
-		{
-			TextLayout->AddLineHighlight(FTextLineHighlight(CursorPosition.GetLineIndex(), FTextRange(CursorPosition.GetOffset(), CursorPosition.GetOffset() + 1), CursorZOrder, CursorLineHighlighter.ToSharedRef()));
-		}
-	}
-}
-
-void SMultiLineEditableText::RemoveCursorHighlight()
-{
-	TextLayout->ClearRunRenderers();
-	TextLayout->ClearLineHighlights();
-}
-
-void SMultiLineEditableText::UpdatePreferredCursorScreenOffsetInLine()
-{
-	PreferredCursorScreenOffsetInLine = TextLayout->GetLocationAt(CursorInfo.GetCursorInteractionLocation(), CursorInfo.GetCursorAlignment() == ECursorAlignment::Right).X;
-}
-
-void SMultiLineEditableText::JumpTo(ETextLocation JumpLocation, ECursorAction Action)
-{
-	// Utility function to count the number of fully visible lines (vertically)
-	// We consider this to be the number of lines on the current page
-	auto CountVisibleLines = [](const TArray<FTextLayout::FLineView>& LineViews, const float VisibleHeight) -> int32
-	{
-		int32 LinesInView = 0;
-		for (const auto& LineView : LineViews)
-		{
-			// The line view is scrolled such that lines above the top of the text area have negative offsets
-			if (LineView.Offset.Y >= 0.0f)
-			{
-				const float EndOffsetY = LineView.Offset.Y + LineView.Size.Y;
-				if (EndOffsetY <= VisibleHeight)
-				{
-					// Line is completely in view
-					++LinesInView;
-				}
-				else
-				{
-					// Line extends beyond the bottom of the text area - we've finished finding visible lines
-					break;
-				}
-			}
-		}
-		return LinesInView;
-	};
-
-	switch (JumpLocation)
-	{
-		case ETextLocation::BeginningOfLine:
-		{
-			const FTextLocation CursorInteractionPosition = CursorInfo.GetCursorInteractionLocation();
-			const TArray< FTextLayout::FLineView >& LineViews = TextLayout->GetLineViews();
-			const int32 CurrentLineViewIndex = TextLayout->GetLineViewIndexForTextLocation(LineViews, CursorInteractionPosition, CursorInfo.GetCursorAlignment() == ECursorAlignment::Right );
-
-			if (LineViews.IsValidIndex(CurrentLineViewIndex))
-			{
-				const FTextLayout::FLineView& CurrentLineView = LineViews[CurrentLineViewIndex];
-				
-				const FTextLocation OldCursorPosition = CursorInteractionPosition;
-				const FTextLocation NewCursorPosition = FTextLocation(OldCursorPosition.GetLineIndex(), CurrentLineView.Range.BeginIndex);
-
-				if (Action == ECursorAction::SelectText)
-				{
-					if (!SelectionStart.IsSet())
-					{
-						this->SelectionStart = OldCursorPosition;
-					}
-				}
-				else
-				{
-					ClearSelection();
-				}
-
-				CursorInfo.SetCursorLocationAndCalculateAlignment(*TextLayout, NewCursorPosition);
-				OnCursorMoved.ExecuteIfBound(CursorInfo.GetCursorInteractionLocation());
-				UpdatePreferredCursorScreenOffsetInLine();
-				UpdateCursorHighlight();
-			}
-		}
-		break;
-
-		case ETextLocation::BeginningOfDocument:
-		{
-			const FTextLocation OldCursorPosition = CursorInfo.GetCursorInteractionLocation();
-			const FTextLocation NewCursorPosition = FTextLocation(0,0);
-
-			if (Action == ECursorAction::SelectText)
-			{
-				if (!SelectionStart.IsSet())
-				{
-					this->SelectionStart = OldCursorPosition;
-				}
-			}
-			else
-			{
-				ClearSelection();
-			}
-
-			CursorInfo.SetCursorLocationAndCalculateAlignment(*TextLayout, NewCursorPosition);
-			OnCursorMoved.ExecuteIfBound(CursorInfo.GetCursorInteractionLocation());
-			UpdatePreferredCursorScreenOffsetInLine();
-			UpdateCursorHighlight();
-		}
-		break;
-
-		case ETextLocation::EndOfLine:
-		{
-			const FTextLocation CursorInteractionPosition = CursorInfo.GetCursorInteractionLocation();
-			const TArray< FTextLayout::FLineView >& LineViews = TextLayout->GetLineViews();
-			const int32 CurrentLineViewIndex = TextLayout->GetLineViewIndexForTextLocation(LineViews, CursorInteractionPosition, CursorInfo.GetCursorAlignment() == ECursorAlignment::Right);
-
-			if (LineViews.IsValidIndex(CurrentLineViewIndex))
-			{
-				const FTextLayout::FLineView& CurrentLineView = LineViews[CurrentLineViewIndex];
-
-				const FTextLocation OldCursorPosition = CursorInteractionPosition;
-				const FTextLocation NewCursorPosition = FTextLocation(OldCursorPosition.GetLineIndex(), FMath::Max(0, CurrentLineView.Range.EndIndex - 1));
-
-				if (Action == ECursorAction::SelectText)
-				{
-					if (!SelectionStart.IsSet())
-					{
-						this->SelectionStart = OldCursorPosition;
-					}
-				}
-				else
-				{
-					ClearSelection();
-				}
-
-				CursorInfo.SetCursorLocationAndAlignment(*TextLayout, NewCursorPosition, ECursorAlignment::Right);
-				OnCursorMoved.ExecuteIfBound(CursorInfo.GetCursorInteractionLocation());
-				UpdatePreferredCursorScreenOffsetInLine();
-				UpdateCursorHighlight();
-			}
-		}
-		break;
-
-		case ETextLocation::EndOfDocument:
-		{
-			if (!TextLayout->IsEmpty())
-			{
-				const FTextLocation OldCursorPosition = CursorInfo.GetCursorInteractionLocation();
-				const TArray< FTextLayout::FLineModel >& Lines = TextLayout->GetLineModels();
-
-				const int32 LastLineIndex = Lines.Num() - 1;
-				const FTextLocation NewCursorPosition = FTextLocation(LastLineIndex, Lines[LastLineIndex].Text->Len());
-
-				if (Action == ECursorAction::SelectText)
-				{
-					if (!SelectionStart.IsSet())
-					{
-						this->SelectionStart = OldCursorPosition;
-					}
-				}
-				else
-				{
-					ClearSelection();
-				}
-
-				CursorInfo.SetCursorLocationAndCalculateAlignment(*TextLayout, NewCursorPosition);
-				OnCursorMoved.ExecuteIfBound(CursorInfo.GetCursorInteractionLocation());
-				UpdatePreferredCursorScreenOffsetInLine();
-				UpdateCursorHighlight();
-			}
-		}
-		break;
-
-		case ETextLocation::PreviousPage:
-		{
-			const FTextLocation CursorInteractionPosition = CursorInfo.GetCursorInteractionLocation();
-			const TArray< FTextLayout::FLineView >& LineViews = TextLayout->GetLineViews();
-			const int32 CurrentLineViewIndex = TextLayout->GetLineViewIndexForTextLocation(LineViews, CursorInteractionPosition, CursorInfo.GetCursorAlignment() == ECursorAlignment::Right );
-
-			if (LineViews.IsValidIndex(CurrentLineViewIndex))
-			{
-				const FTextLayout::FLineView& CurrentLineView = LineViews[CurrentLineViewIndex];
-				
-				const FTextLocation OldCursorPosition = CursorInteractionPosition;
-
-				FTextLocation NewCursorPosition;
-				TOptional<ECursorAlignment> NewCursorAlignment;
-				const int32 NumLinesToMove = FMath::Max(1, CountVisibleLines(LineViews, CachedSize.Y));
-				TranslateLocationVertical(OldCursorPosition, -NumLinesToMove, TextLayout->GetScale(), NewCursorPosition, NewCursorAlignment);
-
-				if (Action == ECursorAction::SelectText)
-				{
-					if (!SelectionStart.IsSet())
-					{
-						this->SelectionStart = OldCursorPosition;
-					}
-				}
-				else
-				{
-					ClearSelection();
-				}
-
-				if (NewCursorAlignment.IsSet())
-				{
-					CursorInfo.SetCursorLocationAndAlignment(*TextLayout, NewCursorPosition, NewCursorAlignment.GetValue());
-				}
-				else
-				{
-					CursorInfo.SetCursorLocationAndCalculateAlignment(*TextLayout, NewCursorPosition);
-				}
-				OnCursorMoved.ExecuteIfBound(CursorInfo.GetCursorInteractionLocation());
-				UpdatePreferredCursorScreenOffsetInLine();
-				UpdateCursorHighlight();
-
-				// We need to scroll by the delta vertical offset value of the old line and the new line
-				// This will (try to) keep the cursor in the same relative location after the page jump
-				const int32 NewLineViewIndex = TextLayout->GetLineViewIndexForTextLocation(LineViews, CursorInfo.GetCursorInteractionLocation(), CursorInfo.GetCursorAlignment() == ECursorAlignment::Right );
-				if (LineViews.IsValidIndex(NewLineViewIndex))
-				{
-					const FTextLayout::FLineView& NewLineView = LineViews[NewLineViewIndex];
-					const float DeltaScrollY = (NewLineView.Offset.Y - CurrentLineView.Offset.Y) / TextLayout->GetScale();
-					ScrollOffset.Y = FMath::Max(0.0f, ScrollOffset.Y + DeltaScrollY);
-
-					// Disable the normal cursor scrolling that UpdateCursorHighlight triggers
-					PositionToScrollIntoView.Reset();
-				}
-			}
-		}
-		break;
-
-		case ETextLocation::NextPage:
-		{
-			const FTextLocation CursorInteractionPosition = CursorInfo.GetCursorInteractionLocation();
-			const TArray< FTextLayout::FLineView >& LineViews = TextLayout->GetLineViews();
-			const int32 CurrentLineViewIndex = TextLayout->GetLineViewIndexForTextLocation(LineViews, CursorInteractionPosition, CursorInfo.GetCursorAlignment() == ECursorAlignment::Right );
-
-			if (LineViews.IsValidIndex(CurrentLineViewIndex))
-			{
-				const FTextLayout::FLineView& CurrentLineView = LineViews[CurrentLineViewIndex];
-				
-				const FTextLocation OldCursorPosition = CursorInteractionPosition;
-
-				FTextLocation NewCursorPosition;
-				TOptional<ECursorAlignment> NewCursorAlignment;
-				const int32 NumLinesToMove = FMath::Max(1, CountVisibleLines(LineViews, CachedSize.Y));
-				TranslateLocationVertical(OldCursorPosition, NumLinesToMove, TextLayout->GetScale(), NewCursorPosition, NewCursorAlignment);
-
-				if (Action == ECursorAction::SelectText)
-				{
-					if (!SelectionStart.IsSet())
-					{
-						this->SelectionStart = OldCursorPosition;
-					}
-				}
-				else
-				{
-					ClearSelection();
-				}
-
-				if (NewCursorAlignment.IsSet())
-				{
-					CursorInfo.SetCursorLocationAndAlignment(*TextLayout, NewCursorPosition, NewCursorAlignment.GetValue());
-				}
-				else
-				{
-					CursorInfo.SetCursorLocationAndCalculateAlignment(*TextLayout, NewCursorPosition);
-				}
-				OnCursorMoved.ExecuteIfBound(CursorInfo.GetCursorInteractionLocation());
-				UpdatePreferredCursorScreenOffsetInLine();
-				UpdateCursorHighlight();
-
-				// We need to scroll by the delta vertical offset value of the old line and the new line
-				// This will (try to) keep the cursor in the same relative location after the page jump
-				const int32 NewLineViewIndex = TextLayout->GetLineViewIndexForTextLocation(LineViews, CursorInfo.GetCursorInteractionLocation(), CursorInfo.GetCursorAlignment() == ECursorAlignment::Right );
-				if (LineViews.IsValidIndex(NewLineViewIndex))
-				{
-					const FTextLayout::FLineView& NewLineView = LineViews[NewLineViewIndex];
-					const float DeltaScrollY = (NewLineView.Offset.Y - CurrentLineView.Offset.Y) / TextLayout->GetScale();
-					ScrollOffset.Y = FMath::Min(TextLayout->GetSize().Y - CachedSize.Y, ScrollOffset.Y + DeltaScrollY);
-
-					// Disable the normal cursor scrolling that UpdateCursorHighlight triggers
-					PositionToScrollIntoView.Reset();
-				}
-			}
-		}
-		break;
-	}
-}
-
-void SMultiLineEditableText::ClearSelection()
-{
-	SelectionStart = TOptional<FTextLocation>();
-}
-
-FText SMultiLineEditableText::GetSelectedText() const
-{
-	FText Text;
-
-	if (AnyTextSelected())
-	{
-		const FTextLocation CursorInteractionPosition = CursorInfo.GetCursorInteractionLocation();
-		FTextLocation SelectionLocation = SelectionStart.Get(CursorInteractionPosition);
-		FTextSelection Selection(SelectionLocation, CursorInteractionPosition);
-
-		FString SelectedText;
-		TextLayout->GetSelectionAsText(SelectedText, Selection);
-
-		Text = FText::FromString(SelectedText);
-	}
-
-	return Text;
-}
-
-void SMultiLineEditableText::InsertTextAtCursor(const FText& InText)
-{
-	InsertTextAtCursor(InText.ToString());
-}
-
-void SMultiLineEditableText::InsertTextAtCursor(const FString& InString)
-{
-	StartChangingText();
-
-	DeleteSelectedText();
-	InsertTextAtCursorImpl(InString);
-
-	FinishChangingText();
-}
-
-void SMultiLineEditableText::InsertRunAtCursor(TSharedRef<IRun> InRun)
-{
-	StartChangingText();
-
-	DeleteSelectedText();
-
-	const FTextLocation CursorInteractionPosition = CursorInfo.GetCursorInteractionLocation();
-	TextLayout->InsertAt(CursorInteractionPosition, InRun, true); // true to preserve the run after the insertion point, even if it's empty - this preserves the text formatting correctly
-
-	// Move the cursor along since we've inserted some new text
-	FString InRunText;
-	InRun->AppendTextTo(InRunText);
-
-	const TArray<FTextLayout::FLineModel>& Lines = TextLayout->GetLineModels();
-	const FTextLayout::FLineModel& Line = Lines[CursorInteractionPosition.GetLineIndex()];
-	const FTextLocation FinalCursorLocation = FTextLocation(CursorInteractionPosition.GetLineIndex(), FMath::Min(CursorInteractionPosition.GetOffset() + InRunText.Len(), Line.Text->Len()));
-
-	CursorInfo.SetCursorLocationAndCalculateAlignment(*TextLayout, FinalCursorLocation);
-	UpdateCursorHighlight();
-
-	FinishChangingText();
-}
-
-void SMultiLineEditableText::GoTo(const FTextLocation& NewLocation)
-{
-	const TArray< FTextLayout::FLineModel >& Lines = TextLayout->GetLineModels();
-	if (Lines.IsValidIndex(NewLocation.GetLineIndex()))
-	{
-		const FTextLayout::FLineModel& Line = Lines[NewLocation.GetLineIndex()];
-		if (NewLocation.GetOffset() <= Line.Text->Len())
-		{
-			ClearSelection();
-
-			CursorInfo.SetCursorLocationAndCalculateAlignment(*TextLayout, NewLocation);
-			OnCursorMoved.ExecuteIfBound(CursorInfo.GetCursorInteractionLocation());
-			UpdatePreferredCursorScreenOffsetInLine();
-			UpdateCursorHighlight();
-		}
-	}
-}
-
-void SMultiLineEditableText::GoTo(ETextLocation GoToLocation)
-{
-	JumpTo(GoToLocation, ECursorAction::MoveCursor);
-}
-
-void SMultiLineEditableText::ScrollTo(const FTextLocation& NewLocation)
-{
-	const TArray< FTextLayout::FLineModel >& Lines = TextLayout->GetLineModels();
-	if (Lines.IsValidIndex(NewLocation.GetLineIndex()))
-	{
-		const FTextLayout::FLineModel& Line = Lines[NewLocation.GetLineIndex()];
-		if (NewLocation.GetOffset() <= Line.Text->Len())
-		{
-			PositionToScrollIntoView = FScrollInfo(NewLocation, ECursorAlignment::Left);
-			EnsureActiveTick();
-		}
-	}
-}
-
-void SMultiLineEditableText::ApplyToSelection(const FRunInfo& InRunInfo, const FTextBlockStyle& InStyle)
-{
-	StartChangingText();
-
-	const FTextLocation CursorInteractionPosition = CursorInfo.GetCursorInteractionLocation();
-	const FTextLocation SelectionLocation = SelectionStart.Get(CursorInteractionPosition);
-	const FTextSelection Selection(SelectionLocation, CursorInteractionPosition);
-
-	const int32 SelectionBeginningLineIndex = Selection.GetBeginning().GetLineIndex();
-	const int32 SelectionBeginningLineOffset = Selection.GetBeginning().GetOffset();
-
-	const int32 SelectionEndLineIndex = Selection.GetEnd().GetLineIndex();
-	const int32 SelectionEndLineOffset = Selection.GetEnd().GetOffset();
-
-	if (SelectionBeginningLineIndex == SelectionEndLineIndex)
-	{
-		TSharedRef<FString> SelectedText = MakeShareable(new FString);
-		TextLayout->GetSelectionAsText(*SelectedText, Selection);
-		
-		TextLayout->RemoveAt(Selection.GetBeginning(), SelectionEndLineOffset - SelectionBeginningLineOffset);
-			
-		TSharedRef<IRun> StyledRun = FSlateTextRun::Create(InRunInfo, SelectedText, InStyle);
-		TextLayout->InsertAt(Selection.GetBeginning(), StyledRun);
-	}
-	else
-	{
-		const TArray< FTextLayout::FLineModel >& Lines = TextLayout->GetLineModels();
-
-		{
-			const FTextLayout::FLineModel& Line = Lines[SelectionBeginningLineIndex];
-
-			const FTextLocation LineStartLocation(SelectionBeginningLineIndex, SelectionBeginningLineOffset);
-			const FTextLocation LineEndLocation(SelectionBeginningLineIndex, Line.Text->Len());
-
-			TSharedRef<FString> SelectedText = MakeShareable(new FString);
-			TextLayout->GetSelectionAsText(*SelectedText, FTextSelection(LineStartLocation, LineEndLocation));
-
-			TextLayout->RemoveAt(LineStartLocation, LineEndLocation.GetOffset() - LineStartLocation.GetOffset());
-			
-			TSharedRef<IRun> StyledRun = FSlateTextRun::Create(InRunInfo, SelectedText, InStyle);
-			TextLayout->InsertAt(LineStartLocation, StyledRun);
-		}
-			
-		for (int32 LineIndex = SelectionBeginningLineIndex + 1; LineIndex < SelectionEndLineIndex; ++LineIndex)
-		{
-			const FTextLayout::FLineModel& Line = Lines[LineIndex];
-				
-			const FTextLocation LineStartLocation(LineIndex, 0);
-			const FTextLocation LineEndLocation(LineIndex, Line.Text->Len());
-
-			TSharedRef<FString> SelectedText = MakeShareable(new FString);
-			TextLayout->GetSelectionAsText(*SelectedText, FTextSelection(LineStartLocation, LineEndLocation));
-
-			TextLayout->RemoveAt(LineStartLocation, LineEndLocation.GetOffset() - LineStartLocation.GetOffset());
-			
-			TSharedRef<IRun> StyledRun = FSlateTextRun::Create(InRunInfo, SelectedText, InStyle);
-			TextLayout->InsertAt(LineStartLocation, StyledRun);
-		}
-
-		{
-			const FTextLayout::FLineModel& Line = Lines[SelectionEndLineIndex];
-
-			const FTextLocation LineStartLocation(SelectionEndLineIndex, 0);
-			const FTextLocation LineEndLocation(SelectionEndLineIndex, SelectionEndLineOffset);
-
-			TSharedRef<FString> SelectedText = MakeShareable(new FString);
-			TextLayout->GetSelectionAsText(*SelectedText, FTextSelection(LineStartLocation, LineEndLocation));
-
-			TextLayout->RemoveAt(LineStartLocation, LineEndLocation.GetOffset() - LineStartLocation.GetOffset());
-			
-			TSharedRef<IRun> StyledRun = FSlateTextRun::Create(InRunInfo, SelectedText, InStyle);
-			TextLayout->InsertAt(LineStartLocation, StyledRun);
-		}
-	}
-
-	SelectionStart = SelectionLocation;
-	CursorInfo.SetCursorLocationAndCalculateAlignment(*TextLayout, CursorInteractionPosition);
-
-	UpdatePreferredCursorScreenOffsetInLine();
-	UpdateCursorHighlight();
-
-	FinishChangingText();
-}
-
-TSharedPtr<const IRun> SMultiLineEditableText::GetRunUnderCursor() const
-{
-	const TArray<FTextLayout::FLineModel>& Lines = TextLayout->GetLineModels();
-	
-	const FTextLocation CursorInteractionLocation = CursorInfo.GetCursorInteractionLocation();
-	if (Lines.IsValidIndex(CursorInteractionLocation.GetLineIndex()))
-	{
-		const FTextLayout::FLineModel& LineModel = Lines[CursorInteractionLocation.GetLineIndex()];
-		for (int32 RunIndex = 0; RunIndex < LineModel.Runs.Num(); ++RunIndex)
-		{
-			const FTextLayout::FRunModel& RunModel = LineModel.Runs[RunIndex];
-			const FTextRange RunRange = RunModel.GetTextRange();
-
-			const bool bIsLastRun = RunIndex == LineModel.Runs.Num() - 1;
-			if (RunRange.Contains(CursorInteractionLocation.GetOffset()) || bIsLastRun)
-			{
-				return RunModel.GetRun();
-			}
-		}
-	}
-
-	return nullptr;
-}
-
-const TArray<TSharedRef<const IRun>> SMultiLineEditableText::GetSelectedRuns() const
-{
-	TArray<TSharedRef<const IRun>> Runs;
-
-	if(AnyTextSelected())
-	{
-		const TArray<FTextLayout::FLineModel>& Lines = TextLayout->GetLineModels();
-		const FTextLocation CursorInteractionLocation = CursorInfo.GetCursorInteractionLocation();
-		if (Lines.IsValidIndex(SelectionStart.GetValue().GetLineIndex()) && Lines.IsValidIndex(CursorInteractionLocation.GetLineIndex()))
-		{
-			const FTextSelection Selection(SelectionStart.GetValue(), CursorInteractionLocation);
-			const int32 StartLine = Selection.GetBeginning().GetLineIndex();
-			const int32 EndLine = Selection.GetEnd().GetLineIndex();
-
-			// iterate over lines
-			for(int32 LineIndex = StartLine; LineIndex <= EndLine; LineIndex++)
-			{
-				const bool bIsFirstLine = LineIndex == StartLine;
-				const bool bIsLastLine = LineIndex == EndLine;
-
-				const FTextLayout::FLineModel& LineModel = Lines[LineIndex];
-				for (int32 RunIndex = 0; RunIndex < LineModel.Runs.Num(); ++RunIndex)
-				{
-					const FTextLayout::FRunModel& RunModel = LineModel.Runs[RunIndex];
-
-					// check what we should be intersecting with
-					if(!bIsFirstLine && !bIsLastLine)
-					{
-						// whole line is inside the range, so just add the run
-						Runs.Add(RunModel.GetRun());
-					}
-					else 
-					{
-						const FTextRange RunRange = RunModel.GetTextRange();
-						if(bIsFirstLine && !bIsLastLine)
-						{
-							// on first line of multi-line selection
-							const FTextRange IntersectedRange = RunRange.Intersect(FTextRange(Selection.GetBeginning().GetOffset(), LineModel.Text->Len()));
-							if (!IntersectedRange.IsEmpty())
-							{
-								Runs.Add(RunModel.GetRun());
-							}
-						}
-						else if(!bIsFirstLine && bIsLastLine)
-						{
-							// on last line of multi-line selection
-							const FTextRange IntersectedRange = RunRange.Intersect(FTextRange(0, Selection.GetEnd().GetOffset()));
-							if (!IntersectedRange.IsEmpty())
-							{
-								Runs.Add(RunModel.GetRun());
-							}
-						}
-						else
-						{
-							// single line selection
-							const FTextRange IntersectedRange = RunRange.Intersect(FTextRange(Selection.GetBeginning().GetOffset(), Selection.GetEnd().GetOffset()));
-							if (!IntersectedRange.IsEmpty())
-							{
-								Runs.Add(RunModel.GetRun());
-							}
-						}
-					}
-				}
-			}
-		}
-	}
-
-	return Runs;
-}
-
-TSharedPtr<const SScrollBar> SMultiLineEditableText::GetHScrollBar() const
-{
-	return HScrollBar;
-}
-
-TSharedPtr<const SScrollBar> SMultiLineEditableText::GetVScrollBar() const
-{
-	return VScrollBar;
-}
-
-void SMultiLineEditableText::Refresh()
-{
-	bool bHasSetText = false;
-
-	const FText& TextToSet = BoundText.Get(FText::GetEmpty());
-	if (!BoundTextLastTick.IdenticalTo(TextToSet))
-	{
-		// The pointer used by the bound text has changed, however the text may still be the same - check that now
-		if (!BoundTextLastTick.IsDisplayStringEqualTo(TextToSet))
-		{
-			// The source text has changed, so update the internal editable text
-			bHasSetText = true;
-			SetEditableText(TextToSet);
-		}
-
-		// Update this even if the text is lexically identical, as it will update the pointer compared by IdenticalTo for the next Tick
-		BoundTextLastTick = FTextSnapshot(TextToSet);
-	}
-
-	if (!bHasSetText && Marshaller->IsDirty())
-	{
-		ForceRefreshTextLayout(TextToSet);
-	}
-}
-
-void SMultiLineEditableText::RestoreOriginalText()
-{
-	if(HasTextChangedFromOriginal())
-	{
-		SetText(OriginalText.Text);
-		TextLayout->UpdateIfNeeded();
-
-		// Let outsiders know that the text content has been changed
-		OnTextCommitted.ExecuteIfBound(OriginalText.Text, ETextCommit::OnCleared);
-	}
-}
-
-bool SMultiLineEditableText::HasTextChangedFromOriginal() const
-{
-	bool bHasChanged = false;
-	if(!IsReadOnly.Get())
-	{
-		const FText EditedText = GetEditableText();
-		bHasChanged = !EditedText.EqualTo(OriginalText.Text);
-	}
-	return bHasChanged;
-}
-
-bool SMultiLineEditableText::CanExecuteSelectAll() const
-{
-	bool bCanExecute = true;
-
-	// Can't select all if string is empty
-	if (TextLayout->IsEmpty())
-	{
-		bCanExecute = false;
-	}
-
-	const TArray< FTextLayout::FLineModel >& Lines = TextLayout->GetLineModels();
-	const int32 NumberOfLines = Lines.Num();
-
-	// Can't select all if entire string is already selected
-	const FTextLocation CursorInteractionPosition = CursorInfo.GetCursorInteractionLocation();
-	if (SelectionStart.IsSet() &&
-		SelectionStart.GetValue() == FTextLocation(0, 0) &&
-		CursorInteractionPosition == FTextLocation(NumberOfLines - 1, Lines[NumberOfLines - 1].Text->Len()))
-	{
-		bCanExecute = false;
-	}
-
-	return bCanExecute;
-}
-
-void SMultiLineEditableText::SelectAllText()
-{
-	if (TextLayout->IsEmpty())
-	{
-		return;
-	}
-
-	const TArray< FTextLayout::FLineModel >& Lines = TextLayout->GetLineModels();
-	const int32 NumberOfLines = Lines.Num();
-
-	SelectionStart = FTextLocation(0,0);
-	const FTextLocation NewCursorPosition = FTextLocation( NumberOfLines - 1, Lines[ NumberOfLines - 1 ].Text->Len() );
-	CursorInfo.SetCursorLocationAndCalculateAlignment(*TextLayout, NewCursorPosition);
-	UpdateCursorHighlight();
-}
-
-bool SMultiLineEditableText::SelectAllTextWhenFocused()
-{
-	return bSelectAllTextWhenFocused.Get();
-}
-
-void SMultiLineEditableText::SelectWordAt( const FGeometry& MyGeometry, const FVector2D& ScreenSpacePosition )
-{
-	const FVector2D LocalPosition = MyGeometry.AbsoluteToLocal( ScreenSpacePosition );
-
-	FTextLocation InitialLocation = TextLayout->GetTextLocationAt(LocalPosition * MyGeometry.Scale);
-	FTextSelection WordSelection = TextLayout->GetWordAt(InitialLocation);
-
-	FTextLocation WordStart = WordSelection.GetBeginning();
-	FTextLocation WordEnd = WordSelection.GetEnd();
-
-	if (WordStart.IsValid() && WordEnd.IsValid())
-	{
-		// Deselect any text that was selected
-		ClearSelection();
-
-		if (WordStart != WordEnd)
-		{
-			SelectionStart = WordStart;
-		}
-
-		const FTextLocation NewCursorPosition = WordEnd;
-		CursorInfo.SetCursorLocationAndCalculateAlignment(*TextLayout, NewCursorPosition);
-		UpdateCursorHighlight();
-	}
-}
-
-void SMultiLineEditableText::BeginDragSelection() 
-{
-	bIsDragSelecting = true;
-}
-
-bool SMultiLineEditableText::IsDragSelecting() const
-{
-	return bIsDragSelecting;
-}
-
-void SMultiLineEditableText::EndDragSelection() 
-{
-	bIsDragSelecting = false;
-}
-
-bool SMultiLineEditableText::AnyTextSelected() const
-{
-	const FTextLocation CursorInteractionPosition = CursorInfo.GetCursorInteractionLocation();
-	const FTextLocation SelectionPosition = SelectionStart.Get( CursorInteractionPosition );
-	return SelectionPosition != CursorInteractionPosition;
-}
-
-bool SMultiLineEditableText::IsTextSelectedAt( const FGeometry& MyGeometry, const FVector2D& ScreenSpacePosition ) const
-{
-	const FTextLocation CursorInteractionPosition = CursorInfo.GetCursorInteractionLocation();
-	const FTextLocation SelectionPosition = SelectionStart.Get(CursorInteractionPosition);
-
-	if (SelectionPosition == CursorInteractionPosition)
-	{
-		return false;
-	}
-
-	const FVector2D LocalPosition = MyGeometry.AbsoluteToLocal(ScreenSpacePosition);
-	const FTextLocation ClickedPosition = TextLayout->GetTextLocationAt(LocalPosition * MyGeometry.Scale);
-
-	FTextLocation SelectionLocation = SelectionStart.Get(CursorInteractionPosition);
-	FTextSelection Selection(SelectionLocation, CursorInteractionPosition);
-
-	int32 SelectionBeginningLineIndex = Selection.GetBeginning().GetLineIndex();
-	int32 SelectionBeginningLineOffset = Selection.GetBeginning().GetOffset();
-
-	int32 SelectionEndLineIndex = Selection.GetEnd().GetLineIndex();
-	int32 SelectionEndLineOffset = Selection.GetEnd().GetOffset();
-
-	if (SelectionBeginningLineIndex == SelectionEndLineIndex)
-	{
-		return ClickedPosition.GetLineIndex() == SelectionBeginningLineIndex
-			&& SelectionBeginningLineOffset <= ClickedPosition.GetOffset()
-			&& SelectionEndLineOffset >= ClickedPosition.GetOffset();
-	}
-
-	if (SelectionBeginningLineIndex == ClickedPosition.GetLineIndex())
-	{
-		return SelectionBeginningLineOffset <= ClickedPosition.GetOffset();
-	}
-
-	if (SelectionEndLineIndex == ClickedPosition.GetLineIndex())
-	{
-		return SelectionEndLineOffset >= ClickedPosition.GetOffset();
-	}
-	
-	return SelectionBeginningLineIndex < ClickedPosition.GetLineIndex()
-		&& SelectionEndLineIndex > ClickedPosition.GetLineIndex();
-}
-
-void SMultiLineEditableText::SetWasFocusedByLastMouseDown( bool Value )
-{
-	bWasFocusedByLastMouseDown = Value;
-}
-
-bool SMultiLineEditableText::WasFocusedByLastMouseDown() const
-{
-	return bWasFocusedByLastMouseDown;
-}
-
-bool SMultiLineEditableText::HasDragSelectedSinceFocused() const
-{
-	return bHasDragSelectedSinceFocused;
-}
-
-void SMultiLineEditableText::SetHasDragSelectedSinceFocused( bool Value )
-{
-	bHasDragSelectedSinceFocused = Value;
-}
-
-FReply SMultiLineEditableText::OnEscape()
-{
-	FReply MyReply = FReply::Unhandled();
-
-	if(AnyTextSelected())
-	{
-		// Clear selection
-		ClearSelection();
-		UpdateCursorHighlight();
-		MyReply = FReply::Handled();
-	}
-
-	if(!GetIsReadOnly())
-	{
-		// Restore the text if the revert flag is set
-		if(bRevertTextOnEscape.Get() && HasTextChangedFromOriginal())
-		{
-			RestoreOriginalText();
-			MyReply = FReply::Handled();
-		}
-	}
-	return MyReply;
-}
-
-void SMultiLineEditableText::OnEnter()
-{
-	if ( FSlateApplication::Get().GetModifierKeys().AreModifersDown(ModiferKeyForNewLine) )
-	{
-		InsertNewLineAtCursorImpl();
-	}
-	else
-	{
-		//Always clear the local undo chain on commit.
-		ClearUndoStates();
-
-		const FText EditedText = GetEditableText();
-
-		// When enter is pressed text is committed.  Let anyone interested know about it.
-		OnTextCommitted.ExecuteIfBound(EditedText, ETextCommit::OnEnter);
-
-		// Reload underlying value now it is committed  (commit may alter the value) 
-		// so it can be re-displayed in the edit box if it retains focus
-		LoadText();
-		// Release input focus
-		if(bClearKeyboardFocusOnCommit.Get())
-		{
-			FSlateApplication::Get().ClearKeyboardFocus(EFocusCause::Cleared);
-		}
-	}
-}
-
-bool SMultiLineEditableText::CanExecuteCut() const
-{
-	bool bCanExecute = true;
-
-	// Can't execute if this is a read-only control
-	if (IsReadOnly.Get())
-	{
-		bCanExecute = false;
-	}
-
-	// Can't execute if there is no text selected
-	if (!AnyTextSelected())
-	{
-		bCanExecute = false;
-	}
-
-	return bCanExecute;
-}
-
-void SMultiLineEditableText::CutSelectedTextToClipboard()
-{
-	if (AnyTextSelected())
-	{
-		StartChangingText();
-
-		const FTextLocation CursorInteractionPosition = CursorInfo.GetCursorInteractionLocation();
-		FTextLocation SelectionLocation = SelectionStart.Get(CursorInteractionPosition);
-		FTextSelection Selection(SelectionLocation, CursorInteractionPosition);
-
-		// Grab the selected substring
-		FString SelectedText;
-		TextLayout->GetSelectionAsText(SelectedText, Selection);
-
-		// Copy text to clipboard
-		FPlatformMisc::ClipboardCopy(*SelectedText);
-
-		DeleteSelectedText();
-
-		UpdateCursorHighlight();
-
-		FinishChangingText();
-	}
-}
-
-bool SMultiLineEditableText::CanExecuteCopy() const
-{
-	bool bCanExecute = true;
-
-	// Can't execute if there is no text selected
-	if (!AnyTextSelected())
-	{
-		bCanExecute = false;
-	}
-
-	return bCanExecute;
-}
-
-void SMultiLineEditableText::CopySelectedTextToClipboard()
-{
-	if (AnyTextSelected())
-	{
-		const FTextLocation CursorInteractionPosition = CursorInfo.GetCursorInteractionLocation();
-		FTextLocation SelectionLocation = SelectionStart.Get(CursorInteractionPosition);
-		FTextSelection Selection(SelectionLocation, CursorInteractionPosition);
-
-		// Grab the selected substring
-		FString SelectedText;
-		TextLayout->GetSelectionAsText(SelectedText, Selection);
-
-		// Copy text to clipboard
-		FPlatformMisc::ClipboardCopy(*SelectedText);
-	}
-}
-
-bool SMultiLineEditableText::CanExecutePaste() const
-{
-	bool bCanExecute = true;
-
-	// Can't execute if this is a read-only control
-	if (IsReadOnly.Get())
-	{
-		bCanExecute = false;
-	}
-
-	// Can't paste unless the clipboard has a string in it
-	if (!DoesClipboardHaveAnyText())
-	{
-		bCanExecute = false;
-	}
-
-	return bCanExecute;
-}
-
-bool SMultiLineEditableText::DoesClipboardHaveAnyText() const
-{
-	FString ClipboardContent;
-	FPlatformMisc::ClipboardPaste(ClipboardContent);
-	return ClipboardContent.Len() > 0;
-}
-
-void SMultiLineEditableText::PasteTextFromClipboard()
-{
-	if (!IsReadOnly.Get())
-	{
-		StartChangingText(); 
-
-		DeleteSelectedText();
-
-		// Paste from the clipboard
-		FString PastedText;
-		FPlatformMisc::ClipboardPaste(PastedText);
-
-		InsertTextAtCursorImpl(PastedText);
-
-		FinishChangingText();
-	}
-}
-
-void SMultiLineEditableText::InsertTextAtCursorImpl(const FString& InString)
-{
-	if (InString.Len())
-	{
-		for (const TCHAR* Character = InString.GetCharArray().GetData(); *Character; ++Character)
-		{
-			if (*Character == '\n')
-			{
-				InsertNewLineAtCursorImpl();
-			}
-			else
-			{
-				TypeChar(*Character);
-			}
-		}
-	}
-}
-
-void SMultiLineEditableText::InsertNewLineAtCursorImpl()
-{
-	if ( AnyTextSelected() )
-	{
-		// Delete selected text
-		DeleteSelectedText();
-	}
-
-	const FTextLocation CursorInteractionPosition = CursorInfo.GetCursorInteractionLocation();
-	if ( TextLayout->SplitLineAt(CursorInteractionPosition) )
-	{
-		// Adjust the cursor position to be at the beginning of the new line
-		const FTextLocation NewCursorPosition = FTextLocation(CursorInteractionPosition.GetLineIndex() + 1, 0);
-		CursorInfo.SetCursorLocationAndCalculateAlignment(*TextLayout, NewCursorPosition);
-	}
-
-	ClearSelection();
-	UpdateCursorHighlight();
-}
-
-bool SMultiLineEditableText::CanExecuteUndo() const
-{
-	if(IsReadOnly.Get() || TextInputMethodContext->IsComposing)
-	{
-		return false;
-	}
-
-	return true;
-}
-
-void SMultiLineEditableText::Undo()
-{
-	if (!IsReadOnly.Get() && UndoStates.Num() > 0 && !TextInputMethodContext->IsComposing)
-	{
-		// Restore from undo state
-		int32 UndoStateIndex;
-		if (CurrentUndoLevel == INDEX_NONE)
-		{
-			// We haven't undone anything since the last time a new undo state was added
-			UndoStateIndex = UndoStates.Num() - 1;
-
-			// Store an undo state for the current state (before undo was pressed)
-			FUndoState NewUndoState;
-			MakeUndoState(NewUndoState);
-			PushUndoState(NewUndoState);
-		}
-		else
-		{
-			// Move down to the next undo level
-			UndoStateIndex = CurrentUndoLevel - 1;
-		}
-
-		// Is there anything else to undo?
-		if (UndoStateIndex >= 0)
-		{
-			{
-				// NOTE: It's important the no code called here creates or destroys undo states!
-				const FUndoState& UndoState = UndoStates[UndoStateIndex];
-
-				SaveText(UndoState.Text);
-
-				if (SetEditableText(UndoState.Text))
-				{
-					// Let outsiders know that the text content has been changed
-					OnTextChanged.ExecuteIfBound(UndoState.Text);
-				}
-
-				CursorInfo = UndoState.CursorInfo.CreateUndo();
-				SelectionStart = UndoState.SelectionStart;
-
-				OnCursorMoved.ExecuteIfBound(CursorInfo.GetCursorInteractionLocation());
-
-				UpdateCursorHighlight();
-			}
-
-			CurrentUndoLevel = UndoStateIndex;
-		}
-	}
-}
-
-void SMultiLineEditableText::Redo()
-{
-	// Is there anything to redo?  If we've haven't tried to undo since the last time new
-	// undo state was added, then CurrentUndoLevel will be INDEX_NONE
-	if (!IsReadOnly.Get() && CurrentUndoLevel != INDEX_NONE && !(TextInputMethodContext->IsComposing))
-	{
-		const int32 NextUndoLevel = CurrentUndoLevel + 1;
-		if (UndoStates.Num() > NextUndoLevel)
-		{
-			// Restore from undo state
-			{
-				// NOTE: It's important the no code called here creates or destroys undo states!
-				const FUndoState& UndoState = UndoStates[NextUndoLevel];
-
-				SaveText(UndoState.Text);
-
-				if (SetEditableText(UndoState.Text))
-				{
-					// Let outsiders know that the text content has been changed
-					OnTextChanged.ExecuteIfBound(UndoState.Text);
-				}
-
-				CursorInfo.RestoreFromUndo(UndoState.CursorInfo);
-				SelectionStart = UndoState.SelectionStart;
-
-				OnCursorMoved.ExecuteIfBound(CursorInfo.GetCursorInteractionLocation());
-
-				UpdateCursorHighlight();
-			}
-
-			CurrentUndoLevel = NextUndoLevel;
-
-			if (UndoStates.Num() <= CurrentUndoLevel + 1)
-			{
-				// We've redone all available undo states
-				CurrentUndoLevel = INDEX_NONE;
-
-				// Pop last undo state that we created on initial undo
-				UndoStates.RemoveAt(UndoStates.Num() - 1);
-			}
-		}
-	}
-}
-
-void SMultiLineEditableText::PushUndoState(const SMultiLineEditableText::FUndoState& InUndoState)
+FReply SMultiLineEditableText::OnFocusReceived( const FGeometry& MyGeometry, const FFocusEvent& InFocusEvent )
 {
-	// If we've already undone some state, then we'll remove any undo state beyond the level that
-	// we've already undone up to.
-	if (CurrentUndoLevel != INDEX_NONE)
-	{
-		UndoStates.RemoveAt(CurrentUndoLevel, UndoStates.Num() - CurrentUndoLevel);
-
-		// Reset undo level as we haven't undone anything since this latest undo
-		CurrentUndoLevel = INDEX_NONE;
-	}
-
-	// Cache new undo state
-	UndoStates.Add(InUndoState);
-
-	// If we've reached the maximum number of undo levels, then trim our array
-	if (UndoStates.Num() > EditableTextDefs::MaxUndoLevels)
+	// Skip the focus received code if it's due to the context menu closing
+	if (EditableTextLayout->HandleFocusReceived(InFocusEvent))
 	{
-		UndoStates.RemoveAt(0);
+		return SWidget::OnFocusReceived(MyGeometry, InFocusEvent);
 	}
-}
-
-void SMultiLineEditableText::ClearUndoStates()
-{
-	CurrentUndoLevel = INDEX_NONE;
-	UndoStates.Empty();
-}
-
-void SMultiLineEditableText::MakeUndoState(SMultiLineEditableText::FUndoState& OutUndoState)
-{
-	//@todo save and restoring the whole document is not ideal [3/31/2014 justin.sargent]
-	const FText EditedText = GetEditableText();
 
-	OutUndoState.Text = EditedText;
-	OutUndoState.CursorInfo = CursorInfo;
-	OutUndoState.SelectionStart = SelectionStart;
+	return FReply::Handled();
 }
 
-TSharedRef< SWidget > SMultiLineEditableText::GetWidget()
+void SMultiLineEditableText::OnFocusLost( const FFocusEvent& InFocusEvent )
 {
-	return SharedThis( this );
+	bIsSoftwareCursor = false;
+	EditableTextLayout->HandleFocusLost(InFocusEvent);
 }
 
-TSharedRef< IVirtualKeyboardEntry > SMultiLineEditableText::GetVirtualKeyboardEntry()
+bool SMultiLineEditableText::AnyTextSelected() const
 {
-	return SharedThis(this);
+	return EditableTextLayout->AnyTextSelected();
 }
 
-void SMultiLineEditableText::SummonContextMenu(const FVector2D& InLocation, TSharedPtr<SWindow> ParentWindow, const FWidgetPath& EventPath)
+void SMultiLineEditableText::SelectAllText()
 {
-	if (!AllowContextMenu.Get())
-	{
-		return;
-	}
-
-	TSharedPtr<SWidget> MenuContentWidget;
-
-	if (OnContextMenuOpening.IsBound())
-	{
-		MenuContentWidget = OnContextMenuOpening.Execute();
-	}
-	else
-	{
-		// Set the menu to automatically close when the user commits to a choice
-		const bool bShouldCloseWindowAfterMenuSelection = true;
-
-#define LOCTEXT_NAMESPACE "EditableTextContextMenu"
-
-		// This is a context menu which could be summoned from within another menu if this text block is in a menu
-		// it should not close the menu it is inside
-		bool bCloseSelfOnly = true;
-		FMenuBuilder MenuBuilder(bShouldCloseWindowAfterMenuSelection, UICommandList, MenuExtender, bCloseSelfOnly, &FCoreStyle::Get());
-		{
-			MenuBuilder.BeginSection("EditText", LOCTEXT("Heading", "Modify Text"));
-			{
-				// Undo
-				MenuBuilder.AddMenuEntry(FGenericCommands::Get().Undo);
-			}
-			MenuBuilder.EndSection();
-
-			MenuBuilder.BeginSection("EditableTextModify2");
-			{
-				// Cut
-				MenuBuilder.AddMenuEntry(FGenericCommands::Get().Cut);
-
-				// Copy
-				MenuBuilder.AddMenuEntry(FGenericCommands::Get().Copy);
-
-				// Paste
-				MenuBuilder.AddMenuEntry(FGenericCommands::Get().Paste);
-
-				// Delete
-				MenuBuilder.AddMenuEntry(FGenericCommands::Get().Delete);
-			}
-			MenuBuilder.EndSection();
-
-			MenuBuilder.BeginSection("EditableTextModify3");
-			{
-				// Select All
-				MenuBuilder.AddMenuEntry(FGenericCommands::Get().SelectAll);
-			}
-			MenuBuilder.EndSection();
-		}
-
-		MenuContentWidget = MenuBuilder.MakeWidget();
-
-#undef LOCTEXT_NAMESPACE
-	}
-
-	if (MenuContentWidget.IsValid())
-	{
-		ActiveContextMenu.PrepareToSummon();
-
-		const bool bFocusImmediately = true;
-		TSharedRef<SWidget> MenuParent = SharedThis(this);
-		if (ParentWindow.IsValid())
-		{
-			MenuParent = StaticCastSharedRef<SWidget>(ParentWindow.ToSharedRef());
-		}
-		TSharedPtr< IMenu > ContextMenu = FSlateApplication::Get().PushMenu(MenuParent, EventPath, MenuContentWidget.ToSharedRef(), InLocation, FPopupTransitionEffect(FPopupTransitionEffect::ContextMenu), bFocusImmediately);
-
-		// Make sure the window is valid.  It's possible for the parent to already be in the destroy queue, for example if the editable text was configured to dismiss it's window during OnTextCommitted.
-		if (ContextMenu.IsValid())
-		{
-			ContextMenu->GetOnMenuDismissed().AddSP(this, &SMultiLineEditableText::OnContextMenuClosed);
-			ActiveContextMenu.SummonSucceeded(ContextMenu.ToSharedRef());
-		}
-		else
-		{
-			ActiveContextMenu.SummonFailed();
-		}
-	}
+	EditableTextLayout->SelectAllText();
 }
 
-void SMultiLineEditableText::OnContextMenuClosed(TSharedRef<IMenu> Menu)
+void SMultiLineEditableText::ClearSelection()
 {
-	// Note: We don't reset the ActiveContextMenu here, as Slate hasn't yet finished processing window focus events, and we need 
-	// to know that the window is still available for OnFocusReceived and OnFocusLost even though it's about to be destroyed
-
-	// Give this focus when the context menu has been dismissed
-	FSlateApplication::Get().SetKeyboardFocus(AsShared(), EFocusCause::OtherWidgetLostFocus);
+	EditableTextLayout->ClearSelection();
 }
 
-void SMultiLineEditableText::LoadText()
+FText SMultiLineEditableText::GetSelectedText() const
 {
-	// We only need to do this if we're bound to a delegate, otherwise the text layout will already be up-to-date
-	// either from Construct, or a call to SetText
-	if (BoundText.IsBound())
-	{
-		SetText(BoundText);
-		TextLayout->UpdateIfNeeded();
-	}
+	return EditableTextLayout->GetSelectedText();
 }
 
-void SMultiLineEditableText::Tick( const FGeometry& AllottedGeometry, const double InCurrentTime, const float InDeltaTime )
-{
-	if (bTextCommittedByVirtualKeyboard)
-	{
-		// Let outsiders know that the text content has been changed
-		OnTextCommitted.ExecuteIfBound(GetEditableText(), VirtualKeyboardTextCommitType);
-		bTextCommittedByVirtualKeyboard = false;
-	}
-	else if (bTextChangedByVirtualKeyboard)
-	{
-		// Let outsiders know that the text content has been changed
-		OnTextChanged.ExecuteIfBound(GetEditableText());
-		bTextChangedByVirtualKeyboard = false;
-	}
-
-	if(TextInputMethodChangeNotifier.IsValid() && TextInputMethodContext.IsValid() && TextInputMethodContext->CachedGeometry != AllottedGeometry)
-	{
-		TextInputMethodContext->CachedGeometry = AllottedGeometry;
-		TextInputMethodChangeNotifier->NotifyLayoutChanged(ITextInputMethodChangeNotifier::ELayoutChangeType::Changed);
-	}
-
-	const bool bShouldAppearFocused = HasKeyboardFocus() || ActiveContextMenu.IsValid();
-
-	if (!bShouldAppearFocused)
-	{
-		Refresh();
-	}
-
-	const float FontMaxCharHeight = FTextEditHelper::GetFontHeight(TextStyle.Font);
-	const float CaretWidth = FTextEditHelper::CalculateCaretWidth(FontMaxCharHeight);
-
-	// If we're auto-wrapping, we need to hide the scrollbars until the first valid auto-wrap has been performed
-	// If we don't do this, then we can get some nasty layout shuffling as the scrollbars appear for one frame and then vanish again
-	const EVisibility ScrollBarVisiblityOverride = (AutoWrapText.Get() && CachedSize.IsZero()) ? EVisibility::Collapsed : EVisibility::Visible;
-
-	// Try and make sure that the line containing the cursor is in view
-	if (PositionToScrollIntoView.IsSet())
-	{
-		const FScrollInfo& ScrollInfo = PositionToScrollIntoView.GetValue();
-
-		const TArray<FTextLayout::FLineView>& LineViews = TextLayout->GetLineViews();
-		const int32 LineViewIndex = TextLayout->GetLineViewIndexForTextLocation(LineViews, ScrollInfo.Position, ScrollInfo.Alignment == ECursorAlignment::Right);
-		if (LineViews.IsValidIndex(LineViewIndex))
-		{
-			const FTextLayout::FLineView& LineView = LineViews[LineViewIndex];
-			const FSlateRect LocalLineViewRect(LineView.Offset / TextLayout->GetScale(), (LineView.Offset + LineView.Size) / TextLayout->GetScale());
+void SMultiLineEditableText::InsertTextAtCursor(const FText& InText)
+{
+	EditableTextLayout->InsertTextAtCursor(InText.ToString());
+}
 
-			const FVector2D LocalCursorLocation = TextLayout->GetLocationAt(ScrollInfo.Position, ScrollInfo.Alignment == ECursorAlignment::Right) / TextLayout->GetScale();
-			const FSlateRect LocalCursorRect(LocalCursorLocation, FVector2D(LocalCursorLocation.X + CaretWidth, LocalCursorLocation.Y + FontMaxCharHeight));
+void SMultiLineEditableText::InsertTextAtCursor(const FString& InString)
+{
+	EditableTextLayout->InsertTextAtCursor(InString);
+}
 
-			if (LocalCursorRect.Left < 0.0f)
-			{
-				ScrollOffset.X += LocalCursorRect.Left;
-			}
-			else if (LocalCursorRect.Right > AllottedGeometry.Size.X)
-			{
-				ScrollOffset.X += (LocalCursorRect.Right - AllottedGeometry.Size.X);
-			}
+void SMultiLineEditableText::InsertRunAtCursor(TSharedRef<IRun> InRun)
+{
+	EditableTextLayout->InsertRunAtCursor(InRun);
+}
 
-			if (LocalLineViewRect.Top < 0.0f)
-			{
-				ScrollOffset.Y += LocalLineViewRect.Top;
-			}
-			else if (LocalLineViewRect.Bottom > AllottedGeometry.Size.Y)
-			{
-				ScrollOffset.Y += (LocalLineViewRect.Bottom - AllottedGeometry.Size.Y);
-			}
-		}
+void SMultiLineEditableText::GoTo(const FTextLocation& NewLocation)
+{
+	EditableTextLayout->GoTo(NewLocation);
+}
 
-		PositionToScrollIntoView.Reset();
-	}
+void SMultiLineEditableText::GoTo(ETextLocation GoToLocation)
+{
+	EditableTextLayout->GoTo(GoToLocation);
+}
 
-	{
-		// Need to account for the caret width too
-		const float ContentSize = TextLayout->GetSize().X + CaretWidth;
-		const float VisibleSize = AllottedGeometry.Size.X;
-
-		// If this text box has no size, do not compute a view fraction because it will be wrong and causes pop in when the size is available
-		const float ViewFraction = (VisibleSize > 0.0f && ContentSize > 0.0f) ? VisibleSize / ContentSize : 1;
-		const float ViewOffset = (ContentSize > 0.0f && ViewFraction < 1.0f) ? FMath::Clamp<float>(ScrollOffset.X / ContentSize, 0.0f, 1.0f - ViewFraction) : 0.0f;
-	
-		// Update the scrollbar with the clamped version of the offset
-		ScrollOffset.X = ViewOffset * ContentSize;
-	
-		if (HScrollBar.IsValid())
-		{
-			HScrollBar->SetState(ViewOffset, ViewFraction);
-			HScrollBar->SetUserVisibility(ScrollBarVisiblityOverride);
-			if (!HScrollBar->IsNeeded())
-			{
-				// We cannot scroll, so ensure that there is no offset
-				ScrollOffset.X = 0.0f;
-			}
-		}
-	}
+void SMultiLineEditableText::ScrollTo(const FTextLocation& NewLocation)
+{
+	EditableTextLayout->ScrollTo(NewLocation);
+}
 
-	{
-		const float ContentSize = TextLayout->GetSize().Y;
-		const float VisibleSize = AllottedGeometry.Size.Y;
-
-		// If this text box has no size, do not compute a view fraction because it will be wrong and causes pop in when the size is available
-		const float ViewFraction = (VisibleSize > 0.0f && ContentSize > 0.0f) ? VisibleSize / ContentSize : 1;
-		const float ViewOffset = (ContentSize > 0.0f && ViewFraction < 1.0f) ? FMath::Clamp<float>(ScrollOffset.Y / ContentSize, 0.0f, 1.0f - ViewFraction) : 0.0f;
-	
-		// Update the scrollbar with the clamped version of the offset
-		ScrollOffset.Y = ViewOffset * ContentSize;
-	
-		if (VScrollBar.IsValid())
-		{
-			VScrollBar->SetState(ViewOffset, ViewFraction);
-			VScrollBar->SetUserVisibility(ScrollBarVisiblityOverride);
-			if (!VScrollBar->IsNeeded())
-			{
-				// We cannot scroll, so ensure that there is no offset
-				ScrollOffset.Y = 0.0f;
-			}
-		}
-	}
+void SMultiLineEditableText::ApplyToSelection(const FRunInfo& InRunInfo, const FTextBlockStyle& InStyle)
+{
+	EditableTextLayout->ApplyToSelection(InRunInfo, InStyle);
+}
 
-	TextLayout->SetVisibleRegion(AllottedGeometry.Size, ScrollOffset * TextLayout->GetScale());
+TSharedPtr<const IRun> SMultiLineEditableText::GetRunUnderCursor() const
+{
+	return EditableTextLayout->GetRunUnderCursor();
 }
 
-int32 SMultiLineEditableText::OnPaint( const FPaintArgs& Args, const FGeometry& AllottedGeometry, const FSlateRect& MyClippingRect, FSlateWindowElementList& OutDrawElements, int32 LayerId, const FWidgetStyle& InWidgetStyle, bool bParentEnabled ) const
+TArray<TSharedRef<const IRun>> SMultiLineEditableText::GetSelectedRuns() const
 {
-	// Update the auto-wrap size now that we have computed paint geometry; won't take affect until text frame
-	// Note: This is done here rather than in Tick(), because Tick() doesn't get called while resizing windows, but OnPaint() does
-	CachedSize = AllottedGeometry.Size;
+	return EditableTextLayout->GetSelectedRuns();
+}
 
-	// Only paint the hint text layout if we don't have any text set
-	if(TextLayout->IsEmpty() && HintTextLayout.IsValid())
-	{
-		// This should always be set when HintTextLayout is
-		check(HintTextStyle.IsValid());
+TSharedPtr<const SScrollBar> SMultiLineEditableText::GetHScrollBar() const
+{
+	return HScrollBar;
+}
 
-		const FLinearColor ThisColorAndOpacity = TextStyle.ColorAndOpacity.GetColor(InWidgetStyle);
+TSharedPtr<const SScrollBar> SMultiLineEditableText::GetVScrollBar() const
+{
+	return VScrollBar;
+}
 
-		// Make sure the hint text is the correct color before we paint it
-		HintTextStyle->ColorAndOpacity = FLinearColor(ThisColorAndOpacity.R, ThisColorAndOpacity.G, ThisColorAndOpacity.B, 0.35f);
-		HintTextLayout->OverrideTextStyle(*HintTextStyle);
+void SMultiLineEditableText::Refresh()
+{
+	EditableTextLayout->Refresh();
+}
 
-		LayerId = HintTextLayout->OnPaint( Args, AllottedGeometry, MyClippingRect, OutDrawElements, LayerId, InWidgetStyle, ShouldBeEnabled(bParentEnabled) );
-	}
+void SMultiLineEditableText::Tick( const FGeometry& AllottedGeometry, const double InCurrentTime, const float InDeltaTime )
+{
+	EditableTextLayout->Tick(AllottedGeometry, InCurrentTime, InDeltaTime);
+}
 
-	LayerId = TextLayout->OnPaint( Args, AllottedGeometry, MyClippingRect, OutDrawElements, LayerId, InWidgetStyle, ShouldBeEnabled( bParentEnabled ) );
+int32 SMultiLineEditableText::OnPaint( const FPaintArgs& Args, const FGeometry& AllottedGeometry, const FSlateRect& MyClippingRect, FSlateWindowElementList& OutDrawElements, int32 LayerId, const FWidgetStyle& InWidgetStyle, bool bParentEnabled ) const
+{
+	LayerId = EditableTextLayout->OnPaint(Args, AllottedGeometry, MyClippingRect, OutDrawElements, LayerId, InWidgetStyle, ShouldBeEnabled(bParentEnabled));
 
 	if (bIsSoftwareCursor)
 	{
@@ -2707,88 +438,23 @@ int32 SMultiLineEditableText::OnPaint( const FPaintArgs& Args, const FGeometry&
 
 void SMultiLineEditableText::CacheDesiredSize(float LayoutScaleMultiplier)
 {
-	// Get the wrapping width and font to see if they have changed
-	float WrappingWidth = WrapTextAt.Get();
-
-	// Text wrapping can either be used defined (WrapTextAt), automatic (AutoWrapText), or a mixture of both
-	// Take whichever has the smallest value (>1)
-	if(AutoWrapText.Get() && CachedSize.X >= 1.0f)
-	{
-		WrappingWidth = (WrappingWidth >= 1.0f) ? FMath::Min(WrappingWidth, CachedSize.X) : CachedSize.X;
-	}
-
-	TextLayout->SetScale( LayoutScaleMultiplier );
-	TextLayout->SetWrappingWidth( WrappingWidth );
-	TextLayout->SetMargin( Margin.Get() );
-	TextLayout->SetLineHeightPercentage( LineHeightPercentage.Get() );
-	TextLayout->SetJustification( Justification.Get() );
-	TextLayout->SetVisibleRegion( CachedSize, ScrollOffset * TextLayout->GetScale() );
-	TextLayout->UpdateIfNeeded();
-
+	EditableTextLayout->CacheDesiredSize(LayoutScaleMultiplier);
 	SWidget::CacheDesiredSize(LayoutScaleMultiplier);
 }
 
 FVector2D SMultiLineEditableText::ComputeDesiredSize( float LayoutScaleMultiplier ) const
 {
-	const float FontMaxCharHeight = FTextEditHelper::GetFontHeight(TextStyle.Font);
-	const float CaretWidth = FTextEditHelper::CalculateCaretWidth(FontMaxCharHeight);
-
-	const float WrappingWidth = WrapTextAt.Get();
-	float DesiredWidth = 0.0f;
-	float DesiredHeight = 0.0f;
-
-	// If we have hint text, make sure we include that in any size calculations
-	if(TextLayout->IsEmpty() && HintTextLayout.IsValid())
-	{
-		// This should always be set when HintTextLayout is
-		check(HintTextStyle.IsValid());
-
-		const FVector2D HintTextSize = HintTextLayout->ComputeDesiredSize(
-			FTextBlockLayout::FWidgetArgs(HintText, FText::GetEmpty(), WrapTextAt, AutoWrapText, Margin, LineHeightPercentage, Justification), 
-			LayoutScaleMultiplier, *HintTextStyle
-			);
-
-		// If a wrapping width has been provided, then we need to report that as the desired width
-		DesiredWidth = WrappingWidth > 0 ? WrappingWidth : HintTextSize.X;
-		DesiredHeight = HintTextSize.Y;
-	}
-	else
-	{
-		// If a wrapping width has been provided, then we need to report the wrapped size as the desired width
-		// (rather than the actual text layout size as that can have non-breaking lines that extend beyond the wrap width)
-		// Note: We don't do this when auto-wrapping as it would cause a feedback loop in the Slate sizing logic
-		const FVector2D TextLayoutBaseSize = WrappingWidth > 0 ? TextLayout->GetWrappedSize() : TextLayout->GetSize();
-		const FVector2D TextLayoutSize = TextLayoutBaseSize + FVector2D(CaretWidth, 0);
-
-		// The layouts current margin size. We should not report a size smaller then the margins.
-		const FMargin TextLayoutMargin = TextLayout->GetMargin();
-		DesiredWidth = FMath::Max(TextLayoutMargin.GetTotalSpaceAlong<Orient_Horizontal>(), TextLayoutSize.X);
-		DesiredHeight = FMath::Max(TextLayoutMargin.GetTotalSpaceAlong<Orient_Vertical>(), TextLayoutSize.Y);
-		DesiredHeight = FMath::Max(DesiredHeight, FontMaxCharHeight);
-	}
-	
-	return FVector2D(DesiredWidth, DesiredHeight);
+	return EditableTextLayout->ComputeDesiredSize(LayoutScaleMultiplier);
 }
 
 FChildren* SMultiLineEditableText::GetChildren()
 {
-	// Only use the hint text layout if we don't have any text set
-	return (TextLayout->IsEmpty() && HintTextLayout.IsValid())
-		? HintTextLayout->GetChildren()
-		: TextLayout->GetChildren();
+	return EditableTextLayout->GetChildren();
 }
 
 void SMultiLineEditableText::OnArrangeChildren( const FGeometry& AllottedGeometry, FArrangedChildren& ArrangedChildren ) const
 {
-	// Only arrange the hint text layout if we don't have any text set
-	if(TextLayout->IsEmpty() && HintTextLayout.IsValid())
-	{
-		HintTextLayout->ArrangeChildren( AllottedGeometry, ArrangedChildren );
-	}
-	else
-	{
-		TextLayout->ArrangeChildren( AllottedGeometry, ArrangedChildren );
-	}
+	EditableTextLayout->OnArrangeChildren(AllottedGeometry, ArrangedChildren);
 }
 
 bool SMultiLineEditableText::SupportsKeyboardFocus() const
@@ -2798,7 +464,7 @@ bool SMultiLineEditableText::SupportsKeyboardFocus() const
 
 FReply SMultiLineEditableText::OnKeyChar( const FGeometry& MyGeometry,const FCharacterEvent& InCharacterEvent )
 {
-	return FTextEditHelper::OnKeyChar( InCharacterEvent, SharedThis( this ) );
+	return EditableTextLayout->HandleKeyChar(InCharacterEvent);
 }
 
 FReply SMultiLineEditableText::OnKeyDown( const FGeometry& MyGeometry, const FKeyEvent& InKeyEvent )
@@ -2811,15 +477,9 @@ FReply SMultiLineEditableText::OnKeyDown( const FGeometry& MyGeometry, const FKe
 		Reply = OnKeyDownHandler.Execute(MyGeometry, InKeyEvent);
 	}
 
-	if( !Reply.IsEventHandled() )
+	if (!Reply.IsEventHandled())
 	{
-		Reply = FTextEditHelper::OnKeyDown( InKeyEvent, SharedThis( this ) );
-
-		// Process keybindings if the event wasn't already handled
-		if (!Reply.IsEventHandled() && UICommandList->ProcessCommandBindings(InKeyEvent))
-		{
-			Reply = FReply::Handled();
-		}
+		Reply = EditableTextLayout->HandleKeyDown(InKeyEvent);
 	}
 
 	return Reply;
@@ -2837,9 +497,7 @@ FReply SMultiLineEditableText::OnMouseButtonDown( const FGeometry& MyGeometry, c
 		AmountScrolledWhileRightMouseDown = 0.0f;
 	}
 
-	FReply Reply = FTextEditHelper::OnMouseButtonDown(MyGeometry, MouseEvent, SharedThis(this));
-//	UpdateCursorHighlight();
-	return Reply;
+	return EditableTextLayout->HandleMouseButtonDown(MyGeometry, MouseEvent);
 }
 
 FReply SMultiLineEditableText::OnMouseButtonUp( const FGeometry& MyGeometry, const FPointerEvent& MouseEvent ) 
@@ -2858,8 +516,7 @@ FReply SMultiLineEditableText::OnMouseButtonUp( const FGeometry& MyGeometry, con
 		}
 	}
 
-	FReply Reply = FTextEditHelper::OnMouseButtonUp( MyGeometry, MouseEvent, SharedThis( this ) );
-	return Reply;
+	return EditableTextLayout->HandleMouseButtonUp(MyGeometry, MouseEvent);
 }
 
 FReply SMultiLineEditableText::OnMouseMove( const FGeometry& MyGeometry, const FPointerEvent& MouseEvent )
@@ -2874,13 +531,11 @@ FReply SMultiLineEditableText::OnMouseMove( const FGeometry& MyGeometry, const F
 
 		if (IsRightClickScrolling())
 		{
-			const float PreviousScrollOffset = ScrollOffset.Y;
-			ScrollOffset.Y -= ScrollByAmount;
+			const FVector2D PreviousScrollOffset = EditableTextLayout->GetScrollOffset();
 
-			const float ContentSize = TextLayout->GetSize().Y;
-			const float ScrollMin = 0.0f;
-			const float ScrollMax = ContentSize - MyGeometry.Size.Y;
-			ScrollOffset.Y = FMath::Clamp(ScrollOffset.Y, ScrollMin, ScrollMax);
+			FVector2D NewScrollOffset = PreviousScrollOffset;
+			NewScrollOffset.Y -= ScrollByAmount;
+			EditableTextLayout->SetScrollOffset(NewScrollOffset, MyGeometry);
 
 			if (!bIsSoftwareCursor)
 			{
@@ -2888,38 +543,37 @@ FReply SMultiLineEditableText::OnMouseMove( const FGeometry& MyGeometry, const F
 				bIsSoftwareCursor = true;
 			}
 
-			if (PreviousScrollOffset != ScrollOffset.Y)
+			if (PreviousScrollOffset.Y != NewScrollOffset.Y)
 			{
-				const float ScrollbarOffset = (ScrollMax != 0.0f) ? ScrollOffset.Y / ScrollMax : 0.0f;
+				const float ScrollMax = EditableTextLayout->GetSize().Y - MyGeometry.Size.Y;
+				const float ScrollbarOffset = (ScrollMax != 0.0f) ? NewScrollOffset.Y / ScrollMax : 0.0f;
 				OnVScrollBarUserScrolled.ExecuteIfBound(ScrollbarOffset);
-				SoftwareCursorPosition.Y += (PreviousScrollOffset - ScrollOffset.Y);
+				SoftwareCursorPosition.Y += (PreviousScrollOffset.Y - NewScrollOffset.Y);
 			}
 
 			return FReply::Handled().UseHighPrecisionMouseMovement(AsShared());
 		}
 	}
 
-	FReply Reply = FTextEditHelper::OnMouseMove(MyGeometry, MouseEvent, SharedThis(this));
-	return Reply;
+	return EditableTextLayout->HandleMouseMove(MyGeometry, MouseEvent);
 }
 
 FReply SMultiLineEditableText::OnMouseWheel( const FGeometry& MyGeometry, const FPointerEvent& MouseEvent )
 {
 	if (VScrollBar.IsValid() && VScrollBar->IsNeeded())
 	{
-		const float PreviousScrollOffset = ScrollOffset.Y;
-
 		const float ScrollAmount = -MouseEvent.GetWheelDelta() * WheelScrollAmount;
-		ScrollOffset.Y += ScrollAmount;
 
-		const float ContentSize = TextLayout->GetSize().Y;
-		const float ScrollMin = 0.0f;
-		const float ScrollMax = ContentSize - MyGeometry.Size.Y;
-		ScrollOffset.Y = FMath::Clamp(ScrollOffset.Y, ScrollMin, ScrollMax);
+		const FVector2D PreviousScrollOffset = EditableTextLayout->GetScrollOffset();
+		
+		FVector2D NewScrollOffset = PreviousScrollOffset;
+		NewScrollOffset.Y += ScrollAmount;
+		EditableTextLayout->SetScrollOffset(NewScrollOffset, MyGeometry);
 
-		if (PreviousScrollOffset != ScrollOffset.Y)
+		if (PreviousScrollOffset.Y != NewScrollOffset.Y)
 		{
-			const float ScrollbarOffset = (ScrollMax != 0.0f) ? ScrollOffset.Y / ScrollMax : 0.0f;
+			const float ScrollMax = EditableTextLayout->GetSize().Y - MyGeometry.Size.Y;
+			const float ScrollbarOffset = (ScrollMax != 0.0f) ? NewScrollOffset.Y / ScrollMax : 0.0f;
 			OnVScrollBarUserScrolled.ExecuteIfBound(ScrollbarOffset);
 			return FReply::Handled();
 		}
@@ -2930,8 +584,7 @@ FReply SMultiLineEditableText::OnMouseWheel( const FGeometry& MyGeometry, const
 
 FReply SMultiLineEditableText::OnMouseButtonDoubleClick(const FGeometry& MyGeometry, const FPointerEvent& MouseEvent)
 {
-	FReply Reply = FTextEditHelper::OnMouseButtonDoubleClick(MyGeometry, MouseEvent, SharedThis(this));
-	return Reply;
+	return EditableTextLayout->HandleMouseButtonDoubleClick(MyGeometry, MouseEvent);
 }
 
 FCursorReply SMultiLineEditableText::OnCursorQuery( const FGeometry& MyGeometry, const FPointerEvent& CursorEvent ) const
@@ -2956,378 +609,4 @@ bool SMultiLineEditableText::IsRightClickScrolling() const
 	return AmountScrolledWhileRightMouseDown >= FSlateApplication::Get().GetDragTriggerDistance() && VScrollBar.IsValid() && VScrollBar->IsNeeded();
 }
 
-/** Remember where the cursor was when we started selecting. */
-void SMultiLineEditableText::BeginSelecting( const FTextLocation& Endpoint )
-{
-	SelectionStart = Endpoint;
-}
-
-FTextLocation SMultiLineEditableText::TranslatedLocation( const FTextLocation& Location, int8 Direction ) const
-{
-	const int32 OffsetInLine = Location.GetOffset() + Direction;
-	const TArray< FTextLayout::FLineModel >& Lines = TextLayout->GetLineModels();
-	const int32 NumberOfLines = Lines.Num();
-	const int32 LineTextLength = Lines[ Location.GetLineIndex() ].Text->Len();
-		
-	if ( Location.GetLineIndex() < NumberOfLines - 1 && OffsetInLine > LineTextLength )
-	{
-		// We're going over the end of the line and we aren't on the last line
-		return FTextLocation( Location.GetLineIndex() + 1, 0 );
-	}
-	else if ( OffsetInLine < 0 && Location.GetLineIndex() > 0 )
-	{
-		// We're stepping before the beginning of the line, and we're not on the first line.
-		const int32 NewLineIndex = Location.GetLineIndex() - 1;
-		return FTextLocation( NewLineIndex, Lines[ NewLineIndex ].Text->Len() );	
-	}
-	else
-	{
-		if ( LineTextLength == 0 )
-		{
-			return FTextLocation( Location.GetLineIndex(), 0 );
-		}
-
-		return FTextLocation( Location.GetLineIndex(), FMath::Clamp( OffsetInLine, 0, LineTextLength ) );
-	}
-}
-
-void SMultiLineEditableText::TranslateLocationVertical( const FTextLocation& Location, int32 NumLinesToMove, float GeometryScale, FTextLocation& OutCursorPosition, TOptional<ECursorAlignment>& OutCursorAlignment ) const
-{
-	const TArray< FTextLayout::FLineView >& LineViews = TextLayout->GetLineViews();
-	const int32 NumberOfLineViews = LineViews.Num();
-
-	const int32 CurrentLineViewIndex = TextLayout->GetLineViewIndexForTextLocation(LineViews, Location, CursorInfo.GetCursorAlignment() == ECursorAlignment::Right);
-	ensure(CurrentLineViewIndex != INDEX_NONE);
-	const FTextLayout::FLineView& CurrentLineView = LineViews[ CurrentLineViewIndex ];
-
-	const int32 NewLineViewIndex = FMath::Clamp( CurrentLineViewIndex + NumLinesToMove, 0, NumberOfLineViews - 1 );
-	const FTextLayout::FLineView& NewLineView = LineViews[ NewLineViewIndex ];
-	
-	// Our horizontal position is the clamped version of whatever the user explicitly set with horizontal movement.
-	ETextHitPoint HitPoint = ETextHitPoint::WithinText;
-	OutCursorPosition = TextLayout->GetTextLocationAt( NewLineView, FVector2D( PreferredCursorScreenOffsetInLine, NewLineView.Offset.Y ) * GeometryScale, &HitPoint );
-
-	// PreferredCursorScreenOffsetInLine can cause the cursor to move to the right hand gutter, and it needs to be placed there
-	// rather than at the start of the next line (which is tricky since they have the same index according to GetTextLocationAt!).
-	// We use the HitPoint to work this out and then adjust the cursor position accordingly
-	if(HitPoint == ETextHitPoint::RightGutter)
-	{
-		OutCursorPosition = FTextLocation(OutCursorPosition, -1);
-		OutCursorAlignment = ECursorAlignment::Right;
-	}
-}
-
-FTextLocation SMultiLineEditableText::ScanForWordBoundary( const FTextLocation& CurrentLocation, int8 Direction ) const
-{
-	FTextLocation Location = TranslatedLocation(CurrentLocation, Direction);
-
-	while ( !IsAtBeginningOfDocument(Location) && !IsAtBeginningOfLine(Location) && !IsAtEndOfDocument(Location) && !IsAtEndOfLine(Location) && !IsAtWordStart(Location) )
-	{
-		Location = TranslatedLocation(Location, Direction);
-	}	
-
-	return Location;
-}
-
-TCHAR SMultiLineEditableText::GetCharacterAt( const FTextLocation& Location ) const
-{
-	const TArray< FTextLayout::FLineModel >& Lines = TextLayout->GetLineModels();
-
-	const bool bIsLineEmpty = Lines[ Location.GetLineIndex() ].Text->IsEmpty();
-	return (bIsLineEmpty)
-		? '\n'
-		: (*Lines[ Location.GetLineIndex() ].Text )[Location.GetOffset()];
-}
-
-bool SMultiLineEditableText::IsAtBeginningOfDocument( const FTextLocation& Location ) const
-{
-	return Location.GetLineIndex() == 0 && Location.GetOffset() == 0;
-}
-	
-bool SMultiLineEditableText::IsAtEndOfDocument( const FTextLocation& Location ) const
-{
-	const TArray< FTextLayout::FLineModel >& Lines = TextLayout->GetLineModels();
-	const int32 NumberOfLines = Lines.Num();
-
-	return NumberOfLines == 0 || ( NumberOfLines - 1 == Location.GetLineIndex() && Lines[ Location.GetLineIndex() ].Text->Len() == Location.GetOffset() );
-}
-
-bool SMultiLineEditableText::IsAtBeginningOfLine( const FTextLocation& Location ) const
-{
-	return Location.GetOffset() == 0;
-}
-
-bool SMultiLineEditableText::IsAtEndOfLine( const FTextLocation& Location ) const
-{
-	const TArray< FTextLayout::FLineModel >& Lines = TextLayout->GetLineModels();
-	return Lines[ Location.GetLineIndex() ].Text->Len() == Location.GetOffset();
-}
-
-bool SMultiLineEditableText::IsAtWordStart( const FTextLocation& Location ) const
-{
-	const FTextSelection WordUnderCursor = TextLayout->GetWordAt(Location);
-	const FTextLocation WordStart = WordUnderCursor.GetBeginning();
-
-	return WordStart.IsValid() && WordStart == Location;
-}
-
-bool SMultiLineEditableText::FTextInputMethodContext::IsReadOnly()
-{
-	bool Return = true;
-	const TSharedPtr<SMultiLineEditableText> OwningWidgetPtr = OwningWidget.Pin();
-	if(OwningWidgetPtr.IsValid())
-	{
-		Return = OwningWidgetPtr->GetIsReadOnly();
-	}
-	return Return;
-}
-
-uint32 SMultiLineEditableText::FTextInputMethodContext::GetTextLength()
-{
-	uint32 TextLength = 0;
-	const TSharedPtr<SMultiLineEditableText> OwningWidgetPtr = OwningWidget.Pin();
-	if(OwningWidgetPtr.IsValid())
-	{
-		FTextLayout::FTextOffsetLocations OffsetLocations;
-		OwningWidgetPtr->TextLayout->GetTextOffsetLocations(OffsetLocations);
-
-		TextLength = OffsetLocations.GetTextLength();
-	}
-	return TextLength;
-}
-
-void SMultiLineEditableText::FTextInputMethodContext::GetSelectionRange(uint32& BeginIndex, uint32& Length, ECaretPosition& OutCaretPosition)
-{
-	const TSharedPtr<SMultiLineEditableText> OwningWidgetPtr = OwningWidget.Pin();
-	if(OwningWidgetPtr.IsValid())
-	{
-		const FTextLocation CursorInteractionPosition = OwningWidgetPtr->CursorInfo.GetCursorInteractionLocation();
-		const FTextLocation SelectionLocation = OwningWidgetPtr->SelectionStart.Get(CursorInteractionPosition);
-
-		FTextLayout::FTextOffsetLocations OffsetLocations;
-		OwningWidgetPtr->TextLayout->GetTextOffsetLocations(OffsetLocations);
-
-		const bool bHasSelection = SelectionLocation != CursorInteractionPosition;
-		if(bHasSelection)
-		{
-			// We need to translate the selection into "editable text" space
-			const FTextSelection Selection(SelectionLocation, CursorInteractionPosition);
-
-			const FTextLocation& BeginningOfSelectionInDocumentSpace = Selection.GetBeginning();
-			const int32 BeginningOfSelectionInEditableTextSpace = OffsetLocations.TextLocationToOffset(BeginningOfSelectionInDocumentSpace);
-
-			const FTextLocation& EndOfSelectionInDocumentSpace = Selection.GetEnd();
-			const int32 EndOfSelectionInEditableTextSpace = OffsetLocations.TextLocationToOffset(EndOfSelectionInDocumentSpace);
-
-			BeginIndex = BeginningOfSelectionInEditableTextSpace;
-			Length = EndOfSelectionInEditableTextSpace - BeginningOfSelectionInEditableTextSpace;
-
-			const bool bCursorIsBeforeSelection = CursorInteractionPosition < SelectionLocation;
-			OutCaretPosition = (bCursorIsBeforeSelection) ? ITextInputMethodContext::ECaretPosition::Beginning : ITextInputMethodContext::ECaretPosition::Ending;
-		}
-		else
-		{
-			// We need to translate the cursor position into "editable text" space
-			const int32 CursorInteractionPositionInEditableTextSpace = OffsetLocations.TextLocationToOffset(CursorInteractionPosition);
-
-			BeginIndex = CursorInteractionPositionInEditableTextSpace;
-			Length = 0;
-
-			OutCaretPosition = ITextInputMethodContext::ECaretPosition::Beginning;
-		}
-	}
-}
-
-void SMultiLineEditableText::FTextInputMethodContext::SetSelectionRange(const uint32 BeginIndex, const uint32 Length, const ECaretPosition InCaretPosition)
-{
-	const TSharedPtr<SMultiLineEditableText> OwningWidgetPtr = OwningWidget.Pin();
-	if(OwningWidgetPtr.IsValid())
-	{
-		const uint32 MinIndex = BeginIndex;
-		const uint32 MaxIndex = MinIndex + Length;
-
-		FTextLayout::FTextOffsetLocations OffsetLocations;
-		OwningWidgetPtr->TextLayout->GetTextOffsetLocations(OffsetLocations);
-
-		// We need to translate the indices into document space
-		const FTextLocation MinTextLocation = OffsetLocations.OffsetToTextLocation(MinIndex);
-		const FTextLocation MaxTextLocation = OffsetLocations.OffsetToTextLocation(MaxIndex);
-
-		OwningWidgetPtr->ClearSelection();
-
-		switch(InCaretPosition)
-		{
-		case ITextInputMethodContext::ECaretPosition::Beginning:
-			{
-				OwningWidgetPtr->CursorInfo.SetCursorLocationAndCalculateAlignment(*OwningWidgetPtr->TextLayout, MinTextLocation);
-				OwningWidgetPtr->SelectionStart = MaxTextLocation;
-			}
-			break;
-
-		case ITextInputMethodContext::ECaretPosition::Ending:
-			{
-				OwningWidgetPtr->SelectionStart = MinTextLocation;
-				OwningWidgetPtr->CursorInfo.SetCursorLocationAndCalculateAlignment(*OwningWidgetPtr->TextLayout, MaxTextLocation);
-			}
-			break;
-		}
-
-		OwningWidgetPtr->OnCursorMoved.ExecuteIfBound(OwningWidgetPtr->CursorInfo.GetCursorInteractionLocation());
-		OwningWidgetPtr->UpdateCursorHighlight();
-	}
-}
-
-void SMultiLineEditableText::FTextInputMethodContext::GetTextInRange(const uint32 BeginIndex, const uint32 Length, FString& OutString)
-{
-	const TSharedPtr<SMultiLineEditableText> OwningWidgetPtr = OwningWidget.Pin();
-	if(OwningWidgetPtr.IsValid())
-	{
-		const FText EditedText = OwningWidgetPtr->GetEditableText();
-		OutString = EditedText.ToString().Mid(BeginIndex, Length);
-	}
-}
-
-void SMultiLineEditableText::FTextInputMethodContext::SetTextInRange(const uint32 BeginIndex, const uint32 Length, const FString& InString)
-{
-	const TSharedPtr<SMultiLineEditableText> OwningWidgetPtr = OwningWidget.Pin();
-	if(OwningWidgetPtr.IsValid())
-	{
-		// We don't use Start/FinishEditing text here because the whole IME operation handles that.
-		// Also, we don't want to support undo for individual characters added during an IME context
-		const FText OldEditedText = OwningWidgetPtr->GetEditableText();
-
-		// We do this as a select, delete, and insert as it's the simplest way to keep the text layout correct
-		SetSelectionRange(BeginIndex, Length, ITextInputMethodContext::ECaretPosition::Beginning);
-		OwningWidgetPtr->DeleteSelectedText();
-		OwningWidgetPtr->InsertTextAtCursorImpl(InString);
-
-		// Has the text changed?
-		const FText EditedText = OwningWidgetPtr->GetEditableText();
-		const bool HasTextChanged = !EditedText.ToString().Equals(OldEditedText.ToString(), ESearchCase::CaseSensitive);
-		if(HasTextChanged)
-		{
-			OwningWidgetPtr->SaveText(EditedText);
-			OwningWidgetPtr->TextLayout->UpdateIfNeeded();
-			OwningWidgetPtr->OnTextChanged.ExecuteIfBound(EditedText);
-		}
-	}
-}
-
-int32 SMultiLineEditableText::FTextInputMethodContext::GetCharacterIndexFromPoint(const FVector2D& Point)
-{
-	int32 CharacterIndex = INDEX_NONE;
-	const TSharedPtr<SMultiLineEditableText> OwningWidgetPtr = OwningWidget.Pin();
-	if(OwningWidgetPtr.IsValid())
-	{
-		const FTextLocation CharacterPosition = OwningWidgetPtr->TextLayout->GetTextLocationAt(Point * OwningWidgetPtr->TextLayout->GetScale());
-
-		FTextLayout::FTextOffsetLocations OffsetLocations;
-		OwningWidgetPtr->TextLayout->GetTextOffsetLocations(OffsetLocations);
-
-		CharacterIndex = OffsetLocations.TextLocationToOffset(CharacterPosition);
-	}
-	return CharacterIndex;
-}
-
-bool SMultiLineEditableText::FTextInputMethodContext::GetTextBounds(const uint32 BeginIndex, const uint32 Length, FVector2D& Position, FVector2D& Size)
-{
-	bool IsClipped = false;
-	const TSharedPtr<SMultiLineEditableText> OwningWidgetPtr = OwningWidget.Pin();
-	if(OwningWidgetPtr.IsValid())
-	{
-		FTextLayout::FTextOffsetLocations OffsetLocations;
-		OwningWidgetPtr->TextLayout->GetTextOffsetLocations(OffsetLocations);
-
-		const FTextLocation BeginLocation = OffsetLocations.OffsetToTextLocation(BeginIndex);
-		const FTextLocation EndLocation = OffsetLocations.OffsetToTextLocation(BeginIndex + Length);
-
-		const FVector2D BeginPosition = OwningWidgetPtr->TextLayout->GetLocationAt(BeginLocation, false);
-		const FVector2D EndPosition = OwningWidgetPtr->TextLayout->GetLocationAt(EndLocation, false);
-
-		if(BeginPosition.Y == EndPosition.Y)
-		{
-			// The text range is contained within a single line
-			Position = BeginPosition;
-			Size = EndPosition - BeginPosition;
-		}
-		else
-		{
-			// If the two positions aren't on the same line, then we assume the worst case scenario, and make the size as wide as the text area itself
-			Position = FVector2D(0.0f, BeginPosition.Y);
-			Size = FVector2D(OwningWidgetPtr->TextLayout->GetDrawSize().X, EndPosition.Y - BeginPosition.Y);
-		}
-
-		// Translate the position (which is in local space) into screen (absolute) space
-		// Note: The local positions are pre-scaled, so we don't scale them again here
-		Position += CachedGeometry.AbsolutePosition;
-	}
-	return IsClipped;
-}
-
-void SMultiLineEditableText::FTextInputMethodContext::GetScreenBounds(FVector2D& Position, FVector2D& Size)
-{
-	const TSharedPtr<SMultiLineEditableText> OwningWidgetPtr = OwningWidget.Pin();
-	if(OwningWidgetPtr.IsValid())
-	{
-		Position = CachedGeometry.AbsolutePosition;
-		Size = CachedGeometry.GetDrawSize();
-	}
-}
-
-TSharedPtr<FGenericWindow> SMultiLineEditableText::FTextInputMethodContext::GetWindow()
-{
-	TSharedPtr<FGenericWindow> Window;
-	const TSharedPtr<SMultiLineEditableText> OwningWidgetPtr = OwningWidget.Pin();
-	if(OwningWidgetPtr.IsValid())
-	{
-		const TSharedPtr<SWindow> SlateWindow = FSlateApplication::Get().FindWidgetWindow( OwningWidgetPtr.ToSharedRef() );
-		Window = SlateWindow.IsValid() ? SlateWindow->GetNativeWindow() : nullptr;
-	}
-	return Window;
-}
-
-void SMultiLineEditableText::FTextInputMethodContext::BeginComposition()
-{
-	if(!IsComposing)
-	{
-		IsComposing = true;
-
-		const TSharedPtr<SMultiLineEditableText> OwningWidgetPtr = OwningWidget.Pin();
-		if(OwningWidgetPtr.IsValid())
-		{
-			OwningWidgetPtr->StartChangingText();
-			OwningWidgetPtr->UpdateCursorHighlight();
-		}
-	}
-}
-
-void SMultiLineEditableText::FTextInputMethodContext::UpdateCompositionRange(const int32 InBeginIndex, const uint32 InLength)
-{
-	check(IsComposing);
-
-	CompositionBeginIndex = InBeginIndex;
-	CompositionLength = InLength;
-
-	const TSharedPtr<SMultiLineEditableText> OwningWidgetPtr = OwningWidget.Pin();
-	if(OwningWidgetPtr.IsValid())
-	{
-		OwningWidgetPtr->UpdateCursorHighlight();
-	}
-}
-
-void SMultiLineEditableText::FTextInputMethodContext::EndComposition()
-{
-	if(IsComposing)
-	{
-		const TSharedPtr<SMultiLineEditableText> OwningWidgetPtr = OwningWidget.Pin();
-		if(OwningWidgetPtr.IsValid())
-		{
-			OwningWidgetPtr->FinishChangingText();
-			OwningWidgetPtr->UpdateCursorHighlight();
-		}
-
-		IsComposing = false;
-	}
-}
-
 #endif //WITH_FANCY_TEXT
