Copying //UE4/Dev-Editor to Dev-Main (//UE4/Dev-Main)

diff --git a/Engine/Source/Runtime/Slate/Private/Widgets/Text/SRichTextBlock.cpp b/Engine/Source/Runtime/Slate/Private/Widgets/Text/SRichTextBlock.cpp
--- a/Engine/Source/Runtime/Slate/Private/Widgets/Text/SRichTextBlock.cpp
+++ b/Engine/Source/Runtime/Slate/Private/Widgets/Text/SRichTextBlock.cpp
@@ -11,6 +11,15 @@
 #include "RichTextLayoutMarshaller.h"
 #include "ReflectionMetadata.h"
 
+SRichTextBlock::SRichTextBlock()
+{
+}
+
+SRichTextBlock::~SRichTextBlock()
+{
+	// Needed to avoid "deletion of pointer to incomplete type 'FTextBlockLayout'; no destructor called" error when using TUniquePtr
+}
+
 void SRichTextBlock::Construct( const FArguments& InArgs )
 {
 	BoundText = InArgs._Text;
@@ -41,7 +50,7 @@ void SRichTextBlock::Construct( const FArguments& InArgs )
 			Marshaller->AppendInlineDecorator(Decorator);
 		}
 
-		TextLayoutCache = FTextBlockLayout::Create(TextStyle, InArgs._TextShapingMethod, InArgs._TextFlowDirection, Marshaller.ToSharedRef(), nullptr);
+		TextLayoutCache = MakeUnique<FTextBlockLayout>(TextStyle, InArgs._TextShapingMethod, InArgs._TextFlowDirection, Marshaller.ToSharedRef(), nullptr);
 		TextLayoutCache->SetDebugSourceInfo(TAttribute<FString>::Create(TAttribute<FString>::FGetter::CreateLambda([this]{ return FReflectionMetaData::GetWidgetDebugInfo(this); })));
 	}
 }
diff --git a/Engine/Source/Runtime/Slate/Private/Widgets/Text/STextBlock.cpp b/Engine/Source/Runtime/Slate/Private/Widgets/Text/STextBlock.cpp
--- a/Engine/Source/Runtime/Slate/Private/Widgets/Text/STextBlock.cpp
+++ b/Engine/Source/Runtime/Slate/Private/Widgets/Text/STextBlock.cpp
@@ -10,6 +10,17 @@ DECLARE_CYCLE_STAT(TEXT("STextBlock::OnPaint Time"), Stat_SlateTextBlockOnPaint,
 DECLARE_CYCLE_STAT(TEXT("STextBlock::ComputeDesiredSize"), Stat_SlateTextBlockCDS, STATGROUP_SlateVerbose)
 DECLARE_CYCLE_STAT(TEXT("STextBlock::ComputeVolitility"), Stat_SlateTextBlockCV, STATGROUP_SlateVerbose)
 
+STextBlock::STextBlock()
+{
+	bCanTick = false;
+	bCanSupportFocus = false;
+}
+
+STextBlock::~STextBlock()
+{
+	// Needed to avoid "deletion of pointer to incomplete type 'FTextBlockLayout'; no destructor called" error when using TUniquePtr
+}
+
 void STextBlock::Construct( const FArguments& InArgs )
 {
 	TextStyle = InArgs._TextStyle;
@@ -33,13 +44,9 @@ void STextBlock::Construct( const FArguments& InArgs )
 
 	BoundText = InArgs._Text;
 
-#if WITH_FANCY_TEXT
-
 	// We use a dummy style here (as it may not be safe to call the delegates used to compute the style), but the correct style is set by ComputeDesiredSize
-	TextLayoutCache = FTextBlockLayout::Create(FTextBlockStyle::GetDefault(), InArgs._TextShapingMethod, InArgs._TextFlowDirection, FPlainTextLayoutMarshaller::Create(), InArgs._LineBreakPolicy);
+	TextLayoutCache = MakeUnique<FTextBlockLayout>(FTextBlockStyle::GetDefault(), InArgs._TextShapingMethod, InArgs._TextFlowDirection, FPlainTextLayoutMarshaller::Create(), InArgs._LineBreakPolicy);
 	TextLayoutCache->SetDebugSourceInfo(TAttribute<FString>::Create(TAttribute<FString>::FGetter::CreateLambda([this]{ return FReflectionMetaData::GetWidgetDebugInfo(this); })));
-
-#endif//WITH_FANCY_TEXT
 }
 
 FSlateFontInfo STextBlock::GetFont() const
@@ -137,8 +144,6 @@ int32 STextBlock::OnPaint( const FPaintArgs& Args, const FGeometry& AllottedGeom
 {
 	SCOPE_CYCLE_COUNTER(Stat_SlateTextBlockOnPaint);
 
-#if WITH_FANCY_TEXT
-
 	//FPlatformMisc::BeginNamedEvent(FColor::Orange, "STextBlock");
 
 	// OnPaint will also update the text layout cache if required
@@ -146,49 +151,6 @@ int32 STextBlock::OnPaint( const FPaintArgs& Args, const FGeometry& AllottedGeom
 
 	//FPlatformMisc::EndNamedEvent();
 
-#else//WITH_FANCY_TEXT
-
-	// Simple stubbed version for when WITH_FANCY_TEXT is disabled (which should only ever be for testing purposes)
-	// Just draw the text as is, with no attempt at handling wrapping or highlighting
-	{
-		const FSlateRect ClippingRect = AllottedGeometry.GetClippingRect().IntersectionWith(MyClippingRect);
-		const ESlateDrawEffect::Type DrawEffects = ShouldBeEnabled(bParentEnabled) ? ESlateDrawEffect::None : ESlateDrawEffect::DisabledEffect;
-		const FLinearColor CurShadowColor = GetShadowColorAndOpacity();
-		const FVector2D CurShadowOffset = GetShadowOffset();
-		const bool ShouldDropShadow = CurShadowColor.A > 0.f && CurShadowOffset.SizeSquared() > 0.f;
-		const FSlateFontInfo FontInfo = GetFont();
-		const FText& TextToDraw = BoundText.Get(FText::GetEmpty());
-
-		// Draw the optional shadow
-		if (ShouldDropShadow)
-		{
-			FSlateDrawElement::MakeText(
-				OutDrawElements,
-				LayerId,
-				AllottedGeometry.ToPaintGeometry(CurShadowOffset, AllottedGeometry.Size),
-				TextToDraw,
-				FontInfo,
-				ClippingRect,
-				DrawEffects,
-				CurShadowColor * InWidgetStyle.GetColorAndOpacityTint()
-				);
-		}
-
-		// Draw the text itself
-		FSlateDrawElement::MakeText(
-			OutDrawElements,
-			++LayerId,
-			AllottedGeometry.ToPaintGeometry(),
-			TextToDraw,
-			FontInfo,
-			ClippingRect,
-			DrawEffects,
-			InWidgetStyle.GetColorAndOpacityTint() * GetColorAndOpacity().GetColor(InWidgetStyle)
-			);
-	}
-
-#endif//WITH_FANCY_TEXT
-
 	return LayerId;
 }
 
@@ -208,7 +170,6 @@ FReply STextBlock::OnMouseButtonDoubleClick( const FGeometry& InMyGeometry, cons
 FVector2D STextBlock::ComputeDesiredSize(float LayoutScaleMultiplier) const
 {
 	SCOPE_CYCLE_COUNTER(Stat_SlateTextBlockCDS);
-#if WITH_FANCY_TEXT
 
 	// ComputeDesiredSize will also update the text layout cache if required
 	const FVector2D TextSize = TextLayoutCache->ComputeDesiredSize(
@@ -216,15 +177,6 @@ FVector2D STextBlock::ComputeDesiredSize(float LayoutScaleMultiplier) const
 		LayoutScaleMultiplier, GetComputedTextStyle()
 		);
 
-#else//WITH_FANCY_TEXT
-
-	// Simple stubbed version for when WITH_FANCY_TEXT is disabled (which should only ever be for testing purposes)
-	// Just measure the text as is, with no attempt at handling wrapping
-	const TSharedRef<FSlateFontMeasure> FontMeasureService = FSlateApplication::Get().GetRenderer()->GetFontMeasureService();
-	const FVector2D TextSize = FontMeasureService->Measure(BoundText.Get(FText::GetEmpty()), GetFont());
-
-#endif//WITH_FANCY_TEXT
-
 	return FVector2D(FMath::Max(MinDesiredWidth.Get(0.0f), TextSize.X), TextSize.Y);
 }
 
@@ -252,6 +204,13 @@ void STextBlock::SetColorAndOpacity(const TAttribute<FSlateColor>& InColorAndOpa
 void STextBlock::SetTextStyle(const FTextBlockStyle* InTextStyle)
 {
 	TextStyle = InTextStyle;
+
+	if (TextStyle == nullptr)
+	{
+		FArguments Defaults;
+		TextStyle = Defaults._TextStyle;
+	}
+
 	Invalidate(EInvalidateWidget::Layout);
 }
 
diff --git a/Engine/Source/Runtime/Slate/Private/Widgets/Text/TextBlockLayout.cpp b/Engine/Source/Runtime/Slate/Private/Widgets/Text/TextBlockLayout.cpp
--- a/Engine/Source/Runtime/Slate/Private/Widgets/Text/TextBlockLayout.cpp
+++ b/Engine/Source/Runtime/Slate/Private/Widgets/Text/TextBlockLayout.cpp
@@ -7,15 +7,8 @@
 #include "SlateTextHighlightRunRenderer.h"
 #include "SlateStats.h"
 
-#if WITH_FANCY_TEXT
-
 SLATE_DECLARE_CYCLE_COUNTER(GSlateTextBlockLayoutComputeDesiredSize, "FTextBlockLayout ComputeDesiredSize");
 
-TSharedRef<FTextBlockLayout> FTextBlockLayout::Create(FTextBlockStyle InDefaultTextStyle, const TOptional<ETextShapingMethod> InTextShapingMethod, const TOptional<ETextFlowDirection> InTextFlowDirection, TSharedRef<ITextLayoutMarshaller> InMarshaller, TSharedPtr<IBreakIterator> InLineBreakPolicy)
-{
-	return MakeShareable(new FTextBlockLayout(MoveTemp(InDefaultTextStyle), InTextShapingMethod, InTextFlowDirection, MoveTemp(InMarshaller), MoveTemp(InLineBreakPolicy)));
-}
-
 FTextBlockLayout::FTextBlockLayout(FTextBlockStyle InDefaultTextStyle, const TOptional<ETextShapingMethod> InTextShapingMethod, const TOptional<ETextFlowDirection> InTextFlowDirection, TSharedRef<ITextLayoutMarshaller> InMarshaller, TSharedPtr<IBreakIterator> InLineBreakPolicy)
 	: TextLayout(FSlateTextLayout::Create(MoveTemp(InDefaultTextStyle)))
 	, Marshaller(MoveTemp(InMarshaller))
@@ -252,5 +245,3 @@ float FTextBlockLayout::CalculateWrappingWidth(const FWidgetArgs& InWidgetArgs)
 
 	return FMath::Max(0.0f, WrappingWidth);
 }
-
-#endif //WITH_FANCY_TEXT
diff --git a/Engine/Source/Runtime/Slate/Public/Framework/Application/SlateApplication.h b/Engine/Source/Runtime/Slate/Public/Framework/Application/SlateApplication.h
--- a/Engine/Source/Runtime/Slate/Public/Framework/Application/SlateApplication.h
+++ b/Engine/Source/Runtime/Slate/Public/Framework/Application/SlateApplication.h
@@ -1058,6 +1058,11 @@ public:
 
 	//~ Begin FSlateApplicationBase Interface
 
+	virtual bool IsActive() const override
+	{
+		return bAppIsActive;
+	}
+
 	virtual TSharedRef<SWindow> AddWindow( TSharedRef<SWindow> InSlateWindow, const bool bShowImmediately = true ) override;
 
 	virtual void ArrangeWindowToFrontVirtual( TArray<TSharedRef<SWindow>>& Windows, const TSharedRef<SWindow>& WindowToBringToFront ) override
@@ -1123,6 +1128,7 @@ public:
 	virtual bool SetKeyboardFocus( const FWidgetPath& InFocusPath, const EFocusCause InCause ) override;
 	virtual bool SetUserFocus(const uint32 InUserIndex, const FWidgetPath& InFocusPath, const EFocusCause InCause) override;
 	virtual void SetAllUserFocus(const FWidgetPath& InFocusPath, const EFocusCause InCause) override;
+	virtual void SetAllUserFocusAllowingDescendantFocus(const FWidgetPath& InFocusPath, const EFocusCause InCause) override;
 
 public:
 
@@ -1200,6 +1206,15 @@ public:
 	 */
 	bool RoutePointerMoveEvent( const FWidgetPath& WidgetsUnderPointer, FPointerEvent& PointerEvent, bool bIsSynthetic );
 
+	/**
+	 * Directly routes a pointer mouse wheel or gesture event to the widgets in the specified widget path.
+	 * 
+	 * @param WidgetsUnderPointer	The path of widgets the event is routed to.
+	 * @param InWheelEvent			The event data that is is routed to the widget path
+	 * @param InGestureEvent		The event data that is is routed to the widget path
+	 */
+	FReply RouteMouseWheelOrGestureEvent(const FWidgetPath& WidgetsUnderPointer, const FPointerEvent& InWheelEvent, const FPointerEvent* InGestureEvent = nullptr);
+
 private:
 
 	TSharedRef< FGenericWindow > MakeWindow( TSharedRef<SWindow> InSlateWindow, const bool bShowImmediately );
diff --git a/Engine/Source/Runtime/Slate/Public/Framework/Text/BaseTextLayoutMarshaller.h b/Engine/Source/Runtime/Slate/Public/Framework/Text/BaseTextLayoutMarshaller.h
--- a/Engine/Source/Runtime/Slate/Public/Framework/Text/BaseTextLayoutMarshaller.h
+++ b/Engine/Source/Runtime/Slate/Public/Framework/Text/BaseTextLayoutMarshaller.h
@@ -1,8 +1,6 @@
 // Copyright 1998-2016 Epic Games, Inc. All Rights Reserved.
 #pragma once
 
-#if WITH_FANCY_TEXT
-
 #include "ITextLayoutMarshaller.h"
 
 /**
@@ -50,5 +48,3 @@ private:
 	bool bIsDirty;
 
 };
-
-#endif //WITH_FANCY_TEXT
diff --git a/Engine/Source/Runtime/Slate/Public/Framework/Text/ILayoutBlock.h b/Engine/Source/Runtime/Slate/Public/Framework/Text/ILayoutBlock.h
--- a/Engine/Source/Runtime/Slate/Public/Framework/Text/ILayoutBlock.h
+++ b/Engine/Source/Runtime/Slate/Public/Framework/Text/ILayoutBlock.h
@@ -5,6 +5,8 @@ class SLATE_API ILayoutBlock
 {
 public:
 
+	virtual ~ILayoutBlock() {}
+
 	virtual TSharedRef< class IRun > GetRun() const = 0;
 	virtual FTextRange GetTextRange() const = 0;
 	virtual FVector2D GetSize() const = 0;
diff --git a/Engine/Source/Runtime/Slate/Public/Framework/Text/IRun.h b/Engine/Source/Runtime/Slate/Public/Framework/Text/IRun.h
--- a/Engine/Source/Runtime/Slate/Public/Framework/Text/IRun.h
+++ b/Engine/Source/Runtime/Slate/Public/Framework/Text/IRun.h
@@ -123,6 +123,8 @@ class SLATE_API IRun
 {
 public:
 
+	virtual ~IRun() {}
+
 	virtual FTextRange GetTextRange() const = 0;
 	virtual void SetTextRange( const FTextRange& Value ) = 0;
 
diff --git a/Engine/Source/Runtime/Slate/Public/Framework/Text/ISlateRun.h b/Engine/Source/Runtime/Slate/Public/Framework/Text/ISlateRun.h
--- a/Engine/Source/Runtime/Slate/Public/Framework/Text/ISlateRun.h
+++ b/Engine/Source/Runtime/Slate/Public/Framework/Text/ISlateRun.h
@@ -1,8 +1,6 @@
 // Copyright 1998-2016 Epic Games, Inc. All Rights Reserved.
 #pragma once
 
-#if WITH_FANCY_TEXT
-
 #include "TextLayout.h"
 
 class SLATE_API ISlateRun : public IRun
@@ -15,5 +13,3 @@ public:
 
 	virtual void ArrangeChildren( const TSharedRef< ILayoutBlock >& Block, const FGeometry& AllottedGeometry, FArrangedChildren& ArrangedChildren ) const = 0;
 };
-
-#endif //WITH_FANCY_TEXT
\ No newline at end of file
diff --git a/Engine/Source/Runtime/Slate/Public/Framework/Text/ITextDecorator.h b/Engine/Source/Runtime/Slate/Public/Framework/Text/ITextDecorator.h
--- a/Engine/Source/Runtime/Slate/Public/Framework/Text/ITextDecorator.h
+++ b/Engine/Source/Runtime/Slate/Public/Framework/Text/ITextDecorator.h
@@ -70,6 +70,8 @@ class SLATE_API ITextDecorator
 {
 public:
 
+	virtual ~ITextDecorator() {}
+
 	virtual bool Supports( const FTextRunParseResults& RunInfo, const FString& Text ) const = 0;
 
 	virtual TSharedRef< ISlateRun > Create(const TSharedRef<class FTextLayout>& TextLayout, const FTextRunParseResults& RunInfo, const FString& OriginalText, const TSharedRef< FString >& ModelText, const ISlateStyle* Style ) = 0;
diff --git a/Engine/Source/Runtime/Slate/Public/Framework/Text/PlainTextLayoutMarshaller.h b/Engine/Source/Runtime/Slate/Public/Framework/Text/PlainTextLayoutMarshaller.h
--- a/Engine/Source/Runtime/Slate/Public/Framework/Text/PlainTextLayoutMarshaller.h
+++ b/Engine/Source/Runtime/Slate/Public/Framework/Text/PlainTextLayoutMarshaller.h
@@ -1,8 +1,6 @@
 // Copyright 1998-2016 Epic Games, Inc. All Rights Reserved.
 #pragma once
 
-#if WITH_FANCY_TEXT
-
 #include "BaseTextLayoutMarshaller.h"
 
 /**
@@ -16,6 +14,8 @@ public:
 
 	virtual ~FPlainTextLayoutMarshaller();
 	
+	void SetIsPassword(const TAttribute<bool>& InIsPassword);
+
 	// ITextLayoutMarshaller
 	virtual void SetText(const FString& SourceString, FTextLayout& TargetTextLayout) override;
 	virtual void GetText(FString& TargetString, const FTextLayout& SourceTextLayout) override;
@@ -24,6 +24,7 @@ protected:
 
 	FPlainTextLayoutMarshaller();
 
-};
+	/** This this marshaller displaying a password? */
+	TAttribute<bool> bIsPassword;
 
-#endif //WITH_FANCY_TEXT
+};
diff --git a/Engine/Source/Runtime/Slate/Public/Framework/Text/SlatePasswordRun.h b/Engine/Source/Runtime/Slate/Public/Framework/Text/SlatePasswordRun.h
new file mode 100644
--- /dev/null
+++ b/Engine/Source/Runtime/Slate/Public/Framework/Text/SlatePasswordRun.h
@@ -0,0 +1,39 @@
+// Copyright 1998-2016 Epic Games, Inc. All Rights Reserved.
+#pragma once
+
+#include "SlateTextRun.h"
+
+class SLATE_API FSlatePasswordRun : public FSlateTextRun
+{
+public:
+
+	static TSharedRef< FSlatePasswordRun > Create( const FRunInfo& InRunInfo, const TSharedRef< const FString >& InText, const FTextBlockStyle& Style );
+
+	static TSharedRef< FSlatePasswordRun > Create( const FRunInfo& InRunInfo, const TSharedRef< const FString >& InText, const FTextBlockStyle& Style, const FTextRange& InRange );
+
+public:
+
+	virtual ~FSlatePasswordRun() {}
+
+	virtual FVector2D Measure( int32 StartIndex, int32 EndIndex, float Scale, const FRunTextContext& TextContext ) const override;
+	virtual int8 GetKerning(int32 CurrentIndex, float Scale, const FRunTextContext& TextContext) const override;
+
+	virtual int32 OnPaint( const FPaintArgs& Args, const FTextLayout::FLineView& Line, const TSharedRef< ILayoutBlock >& Block, const FTextBlockStyle& DefaultStyle, const FGeometry& AllottedGeometry, const FSlateRect& MyClippingRect, FSlateWindowElementList& OutDrawElements, int32 LayerId, const FWidgetStyle& InWidgetStyle, bool bParentEnabled ) const override;
+
+	virtual int32 GetTextIndexAt( const TSharedRef< ILayoutBlock >& Block, const FVector2D& Location, float Scale, ETextHitPoint* const OutHitPoint = nullptr ) const override;
+
+	virtual FVector2D GetLocationAt(const TSharedRef< ILayoutBlock >& Block, int32 Offset, float Scale) const override;
+
+	virtual TSharedRef<IRun> Clone() const override;
+
+protected:
+
+	FSlatePasswordRun( const FRunInfo& InRunInfo, const TSharedRef< const FString >& InText, const FTextBlockStyle& InStyle );
+
+	FSlatePasswordRun( const FRunInfo& InRunInfo, const TSharedRef< const FString >& InText, const FTextBlockStyle& InStyle, const FTextRange& InRange );
+
+	FSlatePasswordRun( const FSlatePasswordRun& Run );
+
+	static TCHAR GetPasswordChar();
+	static FString BuildPasswordString(const int32 InLength);
+};
diff --git a/Engine/Source/Runtime/Slate/Public/Framework/Text/SlateTextHighlightRunRenderer.h b/Engine/Source/Runtime/Slate/Public/Framework/Text/SlateTextHighlightRunRenderer.h
--- a/Engine/Source/Runtime/Slate/Public/Framework/Text/SlateTextHighlightRunRenderer.h
+++ b/Engine/Source/Runtime/Slate/Public/Framework/Text/SlateTextHighlightRunRenderer.h
@@ -1,8 +1,6 @@
 // Copyright 1998-2016 Epic Games, Inc. All Rights Reserved.
 #pragma once
 
-#if WITH_FANCY_TEXT
-
 class SLATE_API FSlateTextHighlightRunRenderer : public ISlateRunRenderer
 {
 public:
@@ -18,4 +16,3 @@ private:
 	FSlateTextHighlightRunRenderer();
 
 };
-#endif //WITH_FANCY_TEXT
\ No newline at end of file
diff --git a/Engine/Source/Runtime/Slate/Public/Framework/Text/SlateTextLayout.h b/Engine/Source/Runtime/Slate/Public/Framework/Text/SlateTextLayout.h
--- a/Engine/Source/Runtime/Slate/Public/Framework/Text/SlateTextLayout.h
+++ b/Engine/Source/Runtime/Slate/Public/Framework/Text/SlateTextLayout.h
@@ -1,8 +1,6 @@
 // Copyright 1998-2016 Epic Games, Inc. All Rights Reserved.
 #pragma once
 
-#if WITH_FANCY_TEXT
-
 class SLATE_API FSlateTextLayout : public FTextLayout
 {
 public:
@@ -22,6 +20,8 @@ public:
 	void SetDefaultTextStyle(FTextBlockStyle InDefaultTextStyle);
 	const FTextBlockStyle& GetDefaultTextStyle() const;
 
+	void SetIsPassword(const TAttribute<bool>& InIsPassword);
+
 protected:
 
 	FSlateTextLayout(FTextBlockStyle InDefaultTextStyle);
@@ -39,10 +39,11 @@ private:
 	/** Default style used by the TextLayout */
 	FTextBlockStyle DefaultTextStyle;
 
+	/** This this layout displaying a password? */
+	TAttribute<bool> bIsPassword;
+
 	/** The localized fallback font revision the last time the text layout was updated. Used to force a flush if the font changes. */
 	int32 LocalizedFallbackFontRevision;
 
 	friend class FSlateTextLayoutFactory;
 };
-
-#endif //WITH_FANCY_TEXT
\ No newline at end of file
diff --git a/Engine/Source/Runtime/Slate/Public/Widgets/Input/IVirtualKeyboardEntry.h b/Engine/Source/Runtime/Slate/Public/Widgets/Input/IVirtualKeyboardEntry.h
--- a/Engine/Source/Runtime/Slate/Public/Widgets/Input/IVirtualKeyboardEntry.h
+++ b/Engine/Source/Runtime/Slate/Public/Widgets/Input/IVirtualKeyboardEntry.h
@@ -23,6 +23,8 @@ class SLATE_API IVirtualKeyboardEntry
 {
 
 public:
+	virtual ~IVirtualKeyboardEntry() {}
+
 	/**
 	* Sets the text to that entered by the virtual keyboard
 	*
@@ -37,14 +39,14 @@ public:
 	*
 	* @return  Text
 	*/
-	virtual const FText& GetText() const = 0;
+	virtual FText GetText() const = 0;
 
 	/**
 	* Returns the hint text.
 	*
 	* @return  HintText
 	*/
-	virtual const FText GetHintText() const = 0;
+	virtual FText GetHintText() const = 0;
 
 	/**
 	* Returns the virtual keyboard type.
diff --git a/Engine/Source/Runtime/Slate/Public/Widgets/Input/SEditableText.h b/Engine/Source/Runtime/Slate/Public/Widgets/Input/SEditableText.h
--- a/Engine/Source/Runtime/Slate/Public/Widgets/Input/SEditableText.h
+++ b/Engine/Source/Runtime/Slate/Public/Widgets/Input/SEditableText.h
@@ -2,20 +2,20 @@
 
 #pragma once
 
-#include "IBreakIterator.h"
-#include "ITextEditorWidget.h"
-#include "ITextInputMethodSystem.h"
-#include "IVirtualKeyboardEntry.h"
-#include "SlateScrollHelper.h"
+#include "ISlateEditableTextWidget.h"
+#include "UniquePtr.h"
+
+class IBreakIterator;
+class ITextLayoutMarshaller;
+class FSlateEditableTextLayout;
+class FPlainTextLayoutMarshaller;
 
 /**
  * Editable text widget
  */
-class SLATE_API SEditableText : public SLeafWidget, public ITextEditorWidget, public IVirtualKeyboardEntry
+class SLATE_API SEditableText : public SWidget, public ISlateEditableTextWidget
 {
-
 public:
-
 	SLATE_BEGIN_ARGS( SEditableText )
 		: _Text()
 		, _HintText()
@@ -35,6 +35,8 @@ public:
 		, _MinDesiredWidth(0.0f)
 		, _SelectAllTextOnCommit( false )
 		, _VirtualKeyboardType(EKeyboardType::Keyboard_Default)
+		, _TextShapingMethod()
+		, _TextFlowDirection()
 		{}
 
 		/** Sets the text content for this editable text widget */
@@ -114,11 +116,19 @@ public:
 		/** The type of virtual keyboard to use on mobile devices */
 		SLATE_ATTRIBUTE( EKeyboardType, VirtualKeyboardType)
 
+		/** Which text shaping method should we use? (unset to use the default returned by GetDefaultTextShapingMethod) */
+		SLATE_ARGUMENT(TOptional<ETextShapingMethod>, TextShapingMethod)
+
+		/** Which text flow direction should we use? (unset to use the default returned by GetDefaultTextFlowDirection) */
+		SLATE_ARGUMENT(TOptional<ETextFlowDirection>, TextFlowDirection)
+
 	SLATE_END_ARGS()
 
 	/** Constructor */
 	SEditableText();
-	virtual ~SEditableText();
+
+	/** Destructor */
+	~SEditableText();
 
 	/**
 	 * Construct this widget
@@ -134,9 +144,19 @@ public:
 	 */
 	void SetText( const TAttribute< FText >& InNewText );
 	
+	/**
+	 * Returns the text string
+	 *
+	 * @return  Text string
+	 */
+	FText GetText() const;
+
 	/** See the HintText attribute */
 	void SetHintText( const TAttribute< FText >& InHintText );
 	
+	/** Get the text that appears when there is no text in the text box */
+	FText GetHintText() const;
+
 	/** See the IsReadOnly attribute */
 	void SetIsReadOnly( TAttribute< bool > InIsReadOnly );
 	
@@ -149,16 +169,6 @@ public:
 	/** See the AllowContextMenu attribute */
 	void SetAllowContextMenu(const TAttribute< bool >& InAllowContextMenu);
 
-	/**
-	 * Restores the text to the original state
-	 */
-	void RestoreOriginalText();
-
-	/**
-	 *	Returns whether the current text varies from the original
-	 */
-	bool HasTextChangedFromOriginal() const;
-
 	/**
 	 * Sets the font used to draw the text
 	 *
@@ -218,197 +228,32 @@ public:
 		OnKeyDownHandler = InOnKeyDownHandler;
 	}
 
-protected:
-
-	friend class FTextInputMethodContext;
-	class FTextInputMethodContext : public ITextInputMethodContext
-	{
-	public:
-		FTextInputMethodContext(const TWeakPtr<SEditableText>& InOwningWidget)
-			:	OwningWidget(InOwningWidget)
-			,	IsComposing(false)
-			,	CompositionBeginIndex(INDEX_NONE)
-			,	CompositionLength(0)
-		{}
-
-		virtual ~FTextInputMethodContext() {}
-
-	private:
-		virtual bool IsReadOnly() override;
-		virtual uint32 GetTextLength() override;
-		virtual void GetSelectionRange(uint32& BeginIndex, uint32& Length, ECaretPosition& CaretPosition) override;
-		virtual void SetSelectionRange(const uint32 BeginIndex, const uint32 Length, const ECaretPosition CaretPosition) override;
-		virtual void GetTextInRange(const uint32 BeginIndex, const uint32 Length, FString& OutString) override;
-		virtual void SetTextInRange(const uint32 BeginIndex, const uint32 Length, const FString& InString) override;
-		virtual int32 GetCharacterIndexFromPoint(const FVector2D& Point) override;
-		virtual bool GetTextBounds(const uint32 BeginIndex, const uint32 Length, FVector2D& Position, FVector2D& Size) override;
-		virtual void GetScreenBounds(FVector2D& Position, FVector2D& Size) override;
-		virtual TSharedPtr<FGenericWindow> GetWindow() override;
-		virtual void BeginComposition() override;
-		virtual void UpdateCompositionRange(const int32 InBeginIndex, const uint32 InLength) override;
-		virtual void EndComposition() override;
-
-	private:
-		TWeakPtr<SEditableText> OwningWidget;
-
-	public:
-		FGeometry CachedGeometry;
-		bool IsComposing;
-		int32 CompositionBeginIndex;
-		uint32 CompositionLength;
-	};
-
-	/**
-	 * Range of selected characters
-	 */
-	class FTextSelection
-	{
-
-	public:
-
-		/** Index of character we started selection on.  Note that this can be larger than LastIndex! */
-		int32 StartIndex;
-
-		/** Index of last character included in the selection range.  Note that this can be smaller than StartIndex! */
-		int32 FinishIndex;
-
-
-	public:
-
-		/** Constructor */
-		FTextSelection()
-			: StartIndex( INDEX_NONE ),
-			  FinishIndex( INDEX_NONE )
-		{
-		}
-
-		/** Clears selection */
-		void Clear()
-		{
-			StartIndex = FinishIndex = INDEX_NONE;
-		}
-
-		/**
-		 * Returns the first character index in the selection range
-		 *
-		 * @return  First character index
-		 */
-		int32 GetMinIndex() const
-		{
-			return FMath::Min( StartIndex, FinishIndex );
-		}
-
-		/**
-		 * Returns the last character index in the selection range
-		 *
-		 * @return  Last character index
-		 */
-		int32 GetMaxIndex() const
-		{
-			return FMath::Max( StartIndex, FinishIndex );
-		}
-	};
-
-
-
-	/**
-	 * Stores a single undo level for editable text
-	 */
-	class FUndoState
-	{
+	/** See TextShapingMethod attribute */
+	void SetTextShapingMethod(const TOptional<ETextShapingMethod>& InTextShapingMethod);
 
-	public:
+	/** See TextFlowDirection attribute */
+	void SetTextFlowDirection(const TOptional<ETextFlowDirection>& InTextFlowDirection);
 
-		/** Text string */
-		FText Text;
+	/** Query to see if any text is selected within the document */
+	bool AnyTextSelected() const;
 
-		/** Selection state */
-		FTextSelection Selection;
+	/** Select all the text in the document */
+	void SelectAllText();
 
-		/** Caret position */
-		int32 CaretPosition;
+	/** Clear the active text selection */
+	void ClearSelection();
 
-	public:
-
-		FUndoState()
-			: Text(),
-			  Selection(),
-			  CaretPosition( INDEX_NONE )
-		{
-		}
-	};
-	
-public:
-	//~ Begin ITextEditorWidget Interface
-	virtual void StartChangingText() override;
-	virtual void FinishChangingText() override;
-	virtual bool GetIsReadOnly() const override;
-	virtual void BackspaceChar() override;
-	virtual void DeleteChar() override;
-	virtual bool CanTypeCharacter(const TCHAR CharInQuestion) const override;
-	virtual void TypeChar( const int32 Character ) override;
-	virtual FReply MoveCursor( FMoveCursor Args ) override;
-	virtual void JumpTo(ETextLocation JumpLocation, ECursorAction Action) override;
-	virtual void ClearSelection() override;
-	virtual void SelectAllText() override;
-	virtual bool SelectAllTextWhenFocused() override;
-	virtual void SelectWordAt( const FGeometry& MyGeometry, const FVector2D& ScreenSpacePosition ) override;
-	virtual void BeginDragSelection() override;
-	virtual bool IsDragSelecting() const override;
-	virtual void EndDragSelection() override;
-	virtual bool AnyTextSelected() const override;
-	virtual bool IsTextSelectedAt( const FGeometry& MyGeometry, const FVector2D& ScreenSpacePosition ) const override;
-	virtual void SetWasFocusedByLastMouseDown( bool Value ) override;
-	virtual bool WasFocusedByLastMouseDown() const override;
-	virtual void SetHasDragSelectedSinceFocused( bool Value ) override;
-	virtual bool HasDragSelectedSinceFocused() const override;
-	virtual FReply OnEscape() override;
-	virtual void OnEnter() override;
-	virtual bool CanExecuteCut() const override;
-	virtual void CutSelectedTextToClipboard() override;
-	virtual bool CanExecuteCopy() const override;
-	virtual void CopySelectedTextToClipboard() override;
-	virtual bool CanExecutePaste() const override;
-	virtual void PasteTextFromClipboard() override;
-	virtual bool CanExecuteUndo() const override;
-	virtual void Undo() override;
-	virtual void Redo() override;
-	virtual TSharedRef< SWidget > GetWidget() override;
-	virtual TSharedRef<IVirtualKeyboardEntry> GetVirtualKeyboardEntry() override;
-	virtual void SummonContextMenu(const FVector2D& InLocation, TSharedPtr<SWindow> ParentWindow, const FWidgetPath& EventPath) override;
-	virtual void LoadText() override;
-	//~ End ITextEditorWidget Interface
-
-public:
-	//~ Begin IVirtualKeyboardEntry Interface
-	virtual void SetTextFromVirtualKeyboard(const FText& InNewText, ESetTextType SetTextType, ETextCommit::Type CommitType) override;
-
-	virtual const FText& GetText() const override
-	{
-		return Text.Get();
-	}
-
-	virtual const FText GetHintText() const override
-	{
-		return HintText.Get();
-	}
-
-	virtual EKeyboardType GetVirtualKeyboardType() const override
-	{
-		return VirtualKeyboardType.Get();
-	}
-
-	virtual bool IsMultilineEntry() const override
-	{
-		return false;
-	}
-	//~ End IVirtualKeyboardEntry Interface
+	/** Get the currently selected text */
+	FText GetSelectedText() const;
 
 protected:
 	//~ Begin SWidget Interface
 	virtual void Tick( const FGeometry& AllottedGeometry, const double InCurrentTime, const float InDeltaTime ) override;
-	virtual FVector2D ComputeDesiredSize(float) const override;
-	virtual int32 OnPaint( const FPaintArgs& Args, const FGeometry& AllottedGeometry, const FSlateRect& MyClippingRect, FSlateWindowElementList& OutDrawElements, int32 LayerId, const FWidgetStyle& InWidgetStyle, bool bParentEnabled ) const override;
+	virtual int32 OnPaint(const FPaintArgs& Args, const FGeometry& AllottedGeometry, const FSlateRect& MyClippingRect, FSlateWindowElementList& OutDrawElements, int32 LayerId, const FWidgetStyle& InWidgetStyle, bool bParentEnabled) const override;
+	virtual void CacheDesiredSize(float LayoutScaleMultiplier) override;
+	virtual FVector2D ComputeDesiredSize(float LayoutScaleMultiplier) const override;
+	virtual FChildren* GetChildren() override;
+	virtual void OnArrangeChildren(const FGeometry& AllottedGeometry, FArrangedChildren& ArrangedChildren) const override;
 	virtual FReply OnDragOver( const FGeometry& MyGeometry, const FDragDropEvent& DragDropEvent ) override;
 	virtual FReply OnDrop( const FGeometry& MyGeometry, const FDragDropEvent& DragDropEvent ) override;
 	virtual bool SupportsKeyboardFocus() const override;
@@ -427,170 +272,74 @@ protected:
 	//~ End SWidget Interface
 
 protected:
-	/**
-	 * Updates text selection after the caret moved
-	 *
-	 * @param  InOldCaretPosition  Previous location of the caret
-	 */
-	void SelectText( const int32 InOldCaretPosition );
+	/** Synchronize the text style currently set (including from overrides) and update the text layout if required */
+	void SynchronizeTextStyle();
 
-	/**
-	 * Deletes selected text
-	 */
-	void DeleteSelectedText();
-	
-	/**
-	 * Determines which character was clicked on
-	 *
-	 * @param  InMyGeometry				Widget geometry
-	 * @param  InScreenCursorPosition	Click position (in screen space)
-	 *
-	 * @return  The clicked character index
-	 */
-	int32 FindClickedCharacterIndex( const FGeometry& InMyGeometry, const FVector2D& InScreenCursorPosition ) const;
-
-	/**
-	 * Determines which character was clicked on
-	 *
-	 * @param  InLocalCursorPosition       Click position (in local space)
-	 * @param  GeometryScale               The DPI scale with which this widget is being drawn.
-	 *
-	 * @return  The clicked character index
-	 */
-	int32 FindClickedCharacterIndex( const FVector2D& InLocalCursorPosition, float GeometryScale ) const;
-
-	/** Find the closest word boundary */
-	int32 ScanForWordBoundary( const int32 Location, int8 Direction ) const; 
-
-	/** Are we currently at the beginning of a word */
-	bool IsAtWordStart( const int32 Location ) const;
-
-	/**
-	 * Adds the specified undo state to the undo stack
-	 *
-	 * @param  InUndoState  The undo state to push.  Usually, this is created using MakeUndoState()
-	 */
-	void PushUndoState( const SEditableText::FUndoState& InUndoState );
-
-
-	/**
-	 * Clears the local undo buffer
-	 */
-	void ClearUndoStates();
-
-
-	/**
-	 * Creates an undo state structure for the control's current state and returns it
-	 *
-	 * @param  OutUndoState  Undo state for current state of control
-	 */
-	void MakeUndoState( SEditableText::FUndoState& OutUndoState );
-
-	/**
-	 * Saves edited text into our text attribute, if appropriate
-	 */
-	void SaveText();
-
-	/**
-	 * Executes action to delete selected text
-	 */
-	void ExecuteDeleteAction();
-
-	/**
-	 * Returns true if 'Select All' is currently a valid action
-	 *
-	 * @return	True to allow this action, otherwise false
-	 */
-	bool CanExecuteSelectAll() const;
-
-	/**
-	 * Returns true if 'Delete' is currently a valid action
-	 *
-	 * @return	True to allow this action, otherwise false
-	 */
-	bool CanExecuteDelete() const;
-
-	/**
-	 * Sets the caret position
-	 * 
-	 * @param Position The character index for the new position
-	 */
-	void SetCaretPosition( int32 Position );
-
-	/**
-	 * Called when the context menu is closed
-	 */
-	void OnContextMenuClosed(TSharedRef<IMenu>);
-
-private:
-	/** @return Whether the editable text should appear focused */
-	bool ShouldAppearFocused() const;
-
-	/** 
-	 * @return whether there is anything in the clipboard
-	 */
-	bool DoesClipboardHaveAnyText() const;
-
-	/**
-	 * @return  the string that needs to be rendered
-	 */
-	FString GetStringToRender() const;
-
-	/**
-	 * Ensure that we will get a Tick() soon (either due to having active focus, or something having changed progmatically and requiring an update)
-	 * Does nothing if the active tick timer is already enabled
-	 */
-	void EnsureActiveTick();
-
-private:
+public:
+	//~ Begin ISlateEditableTextWidget Interface
+	virtual bool IsTextReadOnly() const override;
+	virtual bool IsTextPassword() const override;
+	virtual bool IsMultiLineTextEdit() const override;
+	//~ End ISlateEditableTextWidget Interface
 
-	/** A pointer to the editable text currently under the mouse cursor. nullptr when there isn't one. */
-	static SEditableText* EditableTextUnderCursor;
+protected:
+	//~ Begin ISlateEditableTextWidget Interface
+	virtual bool ShouldSelectAllTextWhenFocused() const override;
+	virtual bool ShouldClearTextSelectionOnFocusLoss() const override;
+	virtual bool ShouldRevertTextOnEscape() const override;
+	virtual bool ShouldClearKeyboardFocusOnCommit() const override;
+	virtual bool CanInsertCarriageReturn() const override;
+	virtual bool CanTypeCharacter(const TCHAR InChar) const override;
+	virtual void EnsureActiveTick() override;
+	virtual EKeyboardType GetVirtualKeyboardType() const override;
+	virtual TSharedRef<SWidget> GetSlateWidget() override;
+	virtual TSharedPtr<SWidget> BuildContextMenuContent() const override;
+	virtual void OnTextChanged(const FText& InText) override;
+	virtual void OnTextCommitted(const FText& InText, const ETextCommit::Type InTextAction) override;
+	virtual void OnCursorMoved(const FTextLocation& InLocation) override;
+	virtual float UpdateAndClampHorizontalScrollBar(const float InViewOffset, const float InViewFraction, const EVisibility InVisiblityOverride) override;
+	virtual float UpdateAndClampVerticalScrollBar(const float InViewOffset, const float InViewFraction, const EVisibility InVisiblityOverride) override;
+	//~ End ISlateEditableTextWidget Interface
 
-	/** The text content for this editable text widget */
-	TAttribute< FText > Text;
+protected:
+	/** Text marshaller used by the editable text layout */
+	TSharedPtr<FPlainTextLayoutMarshaller> PlainTextMarshaller;
 
-	/** The text content for watermark/hinting what text to type here */
-	TAttribute< FText > HintText;
+	/** The text layout that deals with the editable text */
+	TUniquePtr<FSlateEditableTextLayout> EditableTextLayout;
 
 	/** The font used to draw the text */
-	TAttribute< FSlateFontInfo > Font;
+	TAttribute<FSlateFontInfo> Font;
 
 	/** Text color and opacity */
 	TAttribute<FSlateColor> ColorAndOpacity;
 
 	/** Background image for the selected text */
-	TAttribute< const FSlateBrush* > BackgroundImageSelected;
-
-	/** Background image for the composing text */
-	TAttribute< const FSlateBrush* > BackgroundImageComposing;
-
-	/** Image brush used for the caret */
-	TAttribute< const FSlateBrush* > CaretImage;
+	TAttribute<const FSlateBrush*> BackgroundImageSelected;
 
 	/** Sets whether this text box can actually be modified interactively by the user */
-	TAttribute< bool > IsReadOnly;
+	TAttribute<bool> bIsReadOnly;
 
 	/** Sets whether this text box is for storing a password */
-	TAttribute< bool > IsPassword;
+	TAttribute<bool> bIsPassword;
 
 	/** Workaround as we loose focus when the auto completion closes. */
-	TAttribute< bool > IsCaretMovedWhenGainFocus;
+	TAttribute<bool> bIsCaretMovedWhenGainFocus;
 	
 	/** Whether to select all text when the user clicks to give focus on the widget */
-	TAttribute< bool > bSelectAllTextWhenFocused;
+	TAttribute<bool> bSelectAllTextWhenFocused;
 
 	/** Whether to allow the user to back out of changes when they press the escape key */
-	TAttribute< bool > RevertTextOnEscape;
+	TAttribute<bool> bRevertTextOnEscape;
 
 	/** Whether to clear keyboard focus when pressing enter to commit changes */
-	TAttribute< bool > ClearKeyboardFocusOnCommit;
+	TAttribute<bool> bClearKeyboardFocusOnCommit;
 
 	/** Whether to select all text when pressing enter to commit changes */
-	TAttribute< bool > SelectAllTextOnCommit;
+	TAttribute<bool> bSelectAllTextOnCommit;
 
 	/** Whether to disable the context menu */
-	TAttribute< bool > AllowContextMenu;
+	TAttribute<bool> bAllowContextMenu;
 
 	/** Delegate to call before a context menu is opened */
 	FOnContextMenuOpening OnContextMenuOpening;
@@ -599,74 +348,17 @@ private:
 	FOnIsTypedCharValid OnIsTypedCharValid;
 
 	/** Called whenever the text is changed interactively by the user */
-	FOnTextChanged OnTextChanged;
+	FOnTextChanged OnTextChangedCallback;
 
 	/** Called whenever the text is committed.  This happens when the user presses enter or the text box loses focus. */
-	FOnTextCommitted OnTextCommitted;
-
-	/** Text string currently being edited */
-	FText EditedText;
-
-	/** Original text prior to user edits.  This is used to restore the text if the user uses the escape key. */
-	FText OriginalText;
+	FOnTextCommitted OnTextCommittedCallback;
 
-	/** Current scrolled position, in text-area space; use ConvertGeometryFromRenderSpaceToTextSpace to convert before using values with ScrollHelper */
-	FScrollHelper ScrollHelper;
-
-	/** Caret position as an index into the editable text string's character array */
-	int32 CaretPosition;
-
-	/** Stores the last time that the caret was moved by the user.  Used to prevent the caret from blinking. */
-	double LastCaretInteractionTime;
-
-	/** Current selection state */
-	FTextSelection Selection;
-
-	/** True if we're currently selecting text by dragging the mouse cursor with the left button held down */
-	bool bIsDragSelecting;
-
-	/** True if the last mouse down caused us to receive keyboard focus */
-	bool bWasFocusedByLastMouseDown;
-
-	/** True if characters were selected by dragging since the last keyboard focus.  Used for text selection. */
-	bool bHasDragSelectedSinceFocused;
-
-	/** True if characters were selected by dragging since the last mouse down. */
-	bool bHasDragSelectedSinceMouseDown;
-
-	/** Undo states */
-	TArray< FUndoState > UndoStates;
-
-	/** Current undo state level that we've rolled back to, or INDEX_NONE if we haven't undone.  Used for 'Redo'. */
-	int32 CurrentUndoLevel;
-
-	/** True if we're currently (potentially) changing the text string */
-	bool bIsChangingText;
-
-	/** Undo state that will be pushed if text is actually changed between calls to StartChangingText() and FinishChangingText() */
-	FUndoState StateBeforeChangingText;
-
-	/** Information about any active context menu widgets */
-	FActiveTextEditContextMenu ActiveContextMenu;
-
-	/** Prevents the editabletext from being smaller than desired in certain cases (e.g. when it is empty) */
+	/** Prevents the editable text from being smaller than desired in certain cases (e.g. when it is empty) */
 	TAttribute<float> MinDesiredWidth;
 
-	/** A list commands to execute if a user presses the corresponding keybinding in the text box */
-	TSharedRef< FUICommandList > UICommandList;
-
 	/** Menu extender for right-click context menu */
 	TSharedPtr<FExtender> MenuExtender;
 
-	/** Implemented context object for text input method systems. */
-	TSharedPtr<FTextInputMethodContext> TextInputMethodContext;
-
-	/** Notification interface object for text input method systems. */
-	TSharedPtr<ITextInputMethodChangeNotifier> TextInputMethodChangeNotifier;
-
-	/** True if the text has been changed by a virtual keyboard */
-	bool bTextChangedByVirtualKeyboard;
-
 	/** The timer that is actively driving this widget to Tick() even when Slate is idle */
 	TWeakPtr<FActiveTimerHandle> ActiveTickTimer;
 
@@ -678,7 +370,4 @@ private:
 
 	/** The type of virtual keyboard to use for editing this text on mobile */
 	TAttribute<EKeyboardType> VirtualKeyboardType;
-
-	/** Keep track of whether we showed the virtual keyboard when focus was received */
-	bool bShowingVirtualKeyboard;
 };
diff --git a/Engine/Source/Runtime/Slate/Public/Widgets/Input/SEditableTextBox.h b/Engine/Source/Runtime/Slate/Public/Widgets/Input/SEditableTextBox.h
--- a/Engine/Source/Runtime/Slate/Public/Widgets/Input/SEditableTextBox.h
+++ b/Engine/Source/Runtime/Slate/Public/Widgets/Input/SEditableTextBox.h
@@ -100,6 +100,12 @@ public:
 		/** The type of virtual keyboard to use on mobile devices */
 		SLATE_ATTRIBUTE(EKeyboardType, VirtualKeyboardType)
 
+		/** Which text shaping method should we use? (unset to use the default returned by GetDefaultTextShapingMethod) */
+		SLATE_ARGUMENT(TOptional<ETextShapingMethod>, TextShapingMethod)
+
+		/** Which text flow direction should we use? (unset to use the default returned by GetDefaultTextFlowDirection) */
+		SLATE_ARGUMENT(TOptional<ETextFlowDirection>, TextFlowDirection)
+
 	SLATE_END_ARGS()
 	
 	/**
@@ -114,7 +120,7 @@ public:
 	 *
 	 * @return  Text string
 	 */
-	const FText& GetText() const
+	FText GetText() const
 	{
 		return EditableText->GetText();
 	}
@@ -225,11 +231,29 @@ public:
 	 */
 	void SetOnKeyDownHandler(FOnKeyDown InOnKeyDownHandler);
 
+	/** See TextShapingMethod attribute */
+	void SetTextShapingMethod(const TOptional<ETextShapingMethod>& InTextShapingMethod);
+
+	/** See TextFlowDirection attribute */
+	void SetTextFlowDirection(const TOptional<ETextFlowDirection>& InTextFlowDirection);
+
+	/** Query to see if any text is selected within the document */
+	bool AnyTextSelected() const;
+
+	/** Select all the text in the document */
+	void SelectAllText();
+
+	/** Clear the active text selection */
+	void ClearSelection();
+
+	/** Get the currently selected text */
+	FText GetSelectedText() const;
+
 	// SWidget overrides
 	virtual bool SupportsKeyboardFocus() const override;
 	virtual bool HasKeyboardFocus() const override;
 	virtual FReply OnFocusReceived( const FGeometry& MyGeometry, const FFocusEvent& InFocusEvent ) override;
-	virtual FReply OnKeyDown( const FGeometry& MyGeometry, const FKeyEvent& InKeyEvent ) override;
+	virtual FReply OnKeyDown(const FGeometry& MyGeometry, const FKeyEvent& InKeyEvent) override;
 
 protected:
 	const FEditableTextBoxStyle* Style;
diff --git a/Engine/Source/Runtime/Slate/Public/Widgets/Input/SMultiLineEditableTextBox.h b/Engine/Source/Runtime/Slate/Public/Widgets/Input/SMultiLineEditableTextBox.h
--- a/Engine/Source/Runtime/Slate/Public/Widgets/Input/SMultiLineEditableTextBox.h
+++ b/Engine/Source/Runtime/Slate/Public/Widgets/Input/SMultiLineEditableTextBox.h
@@ -186,7 +186,7 @@ public:
 	 *
 	 * @return  Text string
 	 */
-	const FText& GetText() const
+	FText GetText() const
 	{
 		return EditableText->GetText();
 	}
@@ -196,7 +196,7 @@ public:
 	 *
 	 * @return  Text string
 	 */
-	const FText GetPlainText() const
+	FText GetPlainText() const
 	{
 		return EditableText->GetPlainText();
 	}
@@ -275,6 +275,15 @@ public:
 	virtual bool HasKeyboardFocus() const override;
 	virtual FReply OnFocusReceived( const FGeometry& MyGeometry, const FFocusEvent& InFocusEvent ) override;
 
+	/** Query to see if any text is selected within the document */
+	bool AnyTextSelected() const;
+
+	/** Select all the text in the document */
+	void SelectAllText();
+
+	/** Clear the active text selection */
+	void ClearSelection();
+
 	/** Get the currently selected text */
 	FText GetSelectedText() const;
 
@@ -304,7 +313,7 @@ public:
 	TSharedPtr<const IRun> GetRunUnderCursor() const;
 
 	/** Get the runs currently that are current selected, some of which may be only partially selected */
-	const TArray<TSharedRef<const IRun>> GetSelectedRuns() const;
+	TArray<TSharedRef<const IRun>> GetSelectedRuns() const;
 
 	/** Get the horizontal scroll bar widget */
 	TSharedPtr<const SScrollBar> GetHScrollBar() const;
diff --git a/Engine/Source/Runtime/Slate/Public/Widgets/Input/SNumericEntryBox.h b/Engine/Source/Runtime/Slate/Public/Widgets/Input/SNumericEntryBox.h
--- a/Engine/Source/Runtime/Slate/Public/Widgets/Input/SNumericEntryBox.h
+++ b/Engine/Source/Runtime/Slate/Public/Widgets/Input/SNumericEntryBox.h
@@ -72,7 +72,7 @@ public:
 		/** Whether or not the user should be able to change the value by dragging with the mouse cursor */
 		SLATE_ARGUMENT( bool, AllowSpin )
 		/** Delta to increment the value as the slider moves.  If not specified will determine automatically */
-		SLATE_ARGUMENT( NumericType, Delta )
+		SLATE_ATTRIBUTE( NumericType, Delta )
 		/** The minimum value that can be entered into the text edit box */
 		SLATE_ATTRIBUTE( TOptional< NumericType >, MinValue )
 		/** The maximum value that can be entered into the text edit box */
diff --git a/Engine/Source/Runtime/Slate/Public/Widgets/Input/SSpinBox.h b/Engine/Source/Runtime/Slate/Public/Widgets/Input/SSpinBox.h
--- a/Engine/Source/Runtime/Slate/Public/Widgets/Input/SSpinBox.h
+++ b/Engine/Source/Runtime/Slate/Public/Widgets/Input/SSpinBox.h
@@ -51,7 +51,7 @@ public:
 		/** The maximum value that can be specified by using the slider, defaults to MaxValue */
 		SLATE_ATTRIBUTE( TOptional< NumericType >, MaxSliderValue )
 		/** Delta to increment the value as the slider moves.  If not specified will determine automatically */
-		SLATE_ARGUMENT( NumericType, Delta )
+		SLATE_ATTRIBUTE( NumericType, Delta )
 		/** Use exponential scale for the slider */
 		SLATE_ATTRIBUTE( float, SliderExponent )
 		/** Font used to display text in the slider */
@@ -207,9 +207,10 @@ public:
 		if (!bUnlimitedSpinRange)
 		{
 			double Value = ValueAttribute.Get();
-			if( Delta != 0.0f )
+			NumericType CurrentDelta = Delta.Get();
+			if( CurrentDelta != 0.0f )
 			{
-				Value= Snap(Value, Delta); // snap floating point value to nearest Delta
+				Value= Snap(Value, CurrentDelta); // snap floating point value to nearest Delta
 			}
 
 			float FractionFilled = Fraction(Value, GetMinSliderValue(), GetMaxSliderValue());
@@ -417,13 +418,13 @@ public:
 		}
 		else if ( Key == EKeys::Up || Key == EKeys::Right )
 		{
-			CommitValue( InternalValue + Delta, CommittedViaSpin, ETextCommit::OnEnter );
+			CommitValue( InternalValue + Delta.Get(), CommittedViaSpin, ETextCommit::OnEnter );
 			ExitTextMode();
 			return FReply::Handled();	
 		}
 		else if ( Key == EKeys::Down || Key == EKeys::Left )
 		{
-			CommitValue( InternalValue - Delta, CommittedViaSpin, ETextCommit::OnEnter );
+			CommitValue( InternalValue - Delta.Get(), CommittedViaSpin, ETextCommit::OnEnter );
 			ExitTextMode();
 			return FReply::Handled();
 		}
@@ -485,8 +486,8 @@ public:
 	}
 
 	/** See the Delta attribute */
-	NumericType GetDelta() const { return Delta; }
-	void SetDelta(NumericType InDelta) { Delta = InDelta; }
+	NumericType GetDelta() const { return Delta.Get(); }
+	void SetDelta(NumericType InDelta) { Delta.Set( InDelta ); }
 
 	/** See the SliderExponent attribute */
 	float GetSliderExponent() const { return SliderExponent.Get(); }
@@ -515,9 +516,10 @@ protected:
 	FString GetValueAsString() const
 	{
 		auto CurrentValue = ValueAttribute.Get();
-		if( Delta != 0 )
+		NumericType CurrentDelta = Delta.Get();
+		if( CurrentDelta != 0 )
 		{
-			CurrentValue = (NumericType)Snap(CurrentValue, Delta);
+			CurrentValue = (NumericType)Snap(CurrentValue, CurrentDelta);
 		}
 		
 		return Interface->ToString(CurrentValue);
@@ -604,9 +606,10 @@ protected:
 		// If needed, round this value to the delta. Internally the value is not held to the Delta but externally it appears to be.
 		if ( CommitMethod == CommittedViaSpin )
 		{
-			if( Delta != 0 )
+			NumericType CurrentDelta = Delta.Get();
+			if( CurrentDelta != 0 )
 			{
-				NewValue = Snap(NewValue, Delta); // snap numeric point value to nearest Delta
+				NewValue = Snap(NewValue, CurrentDelta); // snap numeric point value to nearest Delta
 			}
 		}
 
@@ -684,7 +687,7 @@ private:
 	const FSlateBrush* InactiveFillBrush;
 
 	float DistanceDragged;
-	NumericType Delta;
+	TAttribute<NumericType> Delta;
 	TAttribute<float> SliderExponent;
 	TAttribute< TOptional<NumericType> > MinValue;
 	TAttribute< TOptional<NumericType> > MaxValue;
diff --git a/Engine/Source/Runtime/Slate/Public/Widgets/Input/SSuggestionTextBox.h b/Engine/Source/Runtime/Slate/Public/Widgets/Input/SSuggestionTextBox.h
--- a/Engine/Source/Runtime/Slate/Public/Widgets/Input/SSuggestionTextBox.h
+++ b/Engine/Source/Runtime/Slate/Public/Widgets/Input/SSuggestionTextBox.h
@@ -128,7 +128,7 @@ public:
 	 *
 	 * @return The text string.
 	 */
-	const FText& GetText( ) const
+	FText GetText( ) const
 	{
 		return TextBox->GetText();
 	}
diff --git a/Engine/Source/Runtime/Slate/Public/Widgets/Input/SVirtualKeyboardEntry.h b/Engine/Source/Runtime/Slate/Public/Widgets/Input/SVirtualKeyboardEntry.h
--- a/Engine/Source/Runtime/Slate/Public/Widgets/Input/SVirtualKeyboardEntry.h
+++ b/Engine/Source/Runtime/Slate/Public/Widgets/Input/SVirtualKeyboardEntry.h
@@ -3,6 +3,7 @@
 #pragma once
 
 #include "IVirtualKeyboardEntry.h"
+#include "SlateScrollHelper.h"
 
 class SLATE_API SVirtualKeyboardEntry : public SLeafWidget, public IVirtualKeyboardEntry
 {
@@ -90,12 +91,12 @@ public:
 	//~ Begin IVirtualKeyboardEntry Interface
 	virtual void SetTextFromVirtualKeyboard(const FText& InNewText, ESetTextType SetTextType, ETextCommit::Type CommitType) override;
 
-	virtual const FText& GetText() const override
+	virtual FText GetText() const override
 	{
 		return Text.Get();
 	}
 
-	virtual const FText GetHintText() const override
+	virtual FText GetHintText() const override
 	{
 		return HintText.Get();
 	}
diff --git a/Engine/Source/Runtime/Slate/Public/Widgets/SViewport.h b/Engine/Source/Runtime/Slate/Public/Widgets/SViewport.h
--- a/Engine/Source/Runtime/Slate/Public/Widgets/SViewport.h
+++ b/Engine/Source/Runtime/Slate/Public/Widgets/SViewport.h
@@ -109,6 +109,15 @@ public:
 	 */
 	void OnWindowClosed( const TSharedRef<SWindow>& InWindowBeingClosed );
 
+	/**
+	 * A delegate called when the viewports top level window is deactivated
+	 */
+	FReply OnViewportActivated(const FWindowActivateEvent& InActivateEvent);
+
+	/**
+	 * A delegate called when the viewports top level window is deactivated
+	 */
+	void OnViewportDeactivated(const FWindowActivateEvent& InActivateEvent);
 
 	/** @return Whether or not this viewport renders directly to the backbuffer */
 	bool ShouldRenderDirectly() const { return bRenderDirectlyToWindow; }
@@ -116,23 +125,6 @@ public:
 	/** @return Whether or not this viewport supports stereo rendering */
 	bool IsStereoRenderingAllowed() const { return bEnableStereoRendering; }
 
-	/**
-	 * Sets a widget that should become focused when this window is next activated
-	 *
-	 * @param	InWidget	The widget to set focus to when this window is activated
-	 */
-	void SetWidgetToFocusOnActivate( const TSharedPtr<SWidget>& InWidget )
-	{
-		
-		WidgetToFocusOnActivate = InWidget;
-	}
-
-	/** Removes the widget to focus on activate so the viewport will be focused */
-	void ClearWidgetToFocusOnActivate()
-	{
-		WidgetToFocusOnActivate.Reset();
-	}
-
 	/** 
 	 * Sets whether this viewport is active. 
 	 * While active, a persistent Active Timer is registered and a Slate tick/paint pass is guaranteed every frame.
@@ -140,6 +132,9 @@ public:
 	 */
 	void SetActive(bool bActive);
 
+	DEPRECATED(4.11, "SetWidgetToFocusOnActivate is no longer needed, remove this call.")
+	void SetWidgetToFocusOnActivate(TSharedPtr<SWidget> WidgetToFocus) { }
+
 public:
 
 	// SWidget interface
@@ -204,12 +199,6 @@ private:
 	/** Size of the viewport. */
 	TAttribute<FVector2D> ViewportSize;
 
-	/** 
-	 * Widget to transfer keyboard focus to when this window becomes active, if any.
-	 * This is used to restore focus to a widget after a popup has been dismissed.
-	 */
-	TWeakPtr< SWidget > WidgetToFocusOnActivate;
-
 	TSharedPtr<ICustomHitTestPath> CustomHitTestPath;
 
 	/** Whether or not this viewport renders directly to the window back-buffer. */
diff --git a/Engine/Source/Runtime/Slate/Public/Widgets/Text/SInlineEditableTextBlock.h b/Engine/Source/Runtime/Slate/Public/Widgets/Text/SInlineEditableTextBlock.h
--- a/Engine/Source/Runtime/Slate/Public/Widgets/Text/SInlineEditableTextBlock.h
+++ b/Engine/Source/Runtime/Slate/Public/Widgets/Text/SInlineEditableTextBlock.h
@@ -138,8 +138,10 @@ protected:
 	/** The widget used when in editing mode */ 
 	TSharedPtr< SEditableTextBox > TextBox;
 
+#if WITH_FANCY_TEXT
 	/** The widget used when in editing mode */ 
 	TSharedPtr< SMultiLineEditableTextBox > MultiLineTextBox;
+#endif //WITH_FANCY_TEXT
 
 	/** Delegate to execute when the text starts to be edited */
 	FOnBeginTextEdit OnBeginTextEditDelegate;
diff --git a/Engine/Source/Runtime/Slate/Public/Widgets/Text/SMultiLineEditableText.h b/Engine/Source/Runtime/Slate/Public/Widgets/Text/SMultiLineEditableText.h
--- a/Engine/Source/Runtime/Slate/Public/Widgets/Text/SMultiLineEditableText.h
+++ b/Engine/Source/Runtime/Slate/Public/Widgets/Text/SMultiLineEditableText.h
@@ -4,16 +4,15 @@
 
 #if WITH_FANCY_TEXT
 
-#include "ITextEditorWidget.h"
-#include "IVirtualKeyboardEntry.h"
-#include "ITextInputMethodSystem.h"
-#include "ITextLayoutMarshaller.h"
+#include "ISlateEditableTextWidget.h"
 #include "SScrollBar.h"
+#include "UniquePtr.h"
 
-class FTextBlockLayout;
+class ITextLayoutMarshaller;
+class FSlateEditableTextLayout;
 
 /** An editable text widget that supports multiple lines and soft word-wrapping. */
-class SLATE_API SMultiLineEditableText : public SWidget, public ITextEditorWidget, public IVirtualKeyboardEntry
+class SLATE_API SMultiLineEditableText : public SWidget, public ISlateEditableTextWidget
 {
 public:
 	/** Called when the cursor is moved within the text area */
@@ -136,7 +135,7 @@ public:
 	SLATE_END_ARGS()
 
 	SMultiLineEditableText();
-	virtual ~SMultiLineEditableText();
+	~SMultiLineEditableText();
 
 	void Construct( const FArguments& InArgs );
 
@@ -145,22 +144,31 @@ public:
 	 */
 	void SetText(const TAttribute< FText >& InText);
 
+	/**
+	 * Returns the text string
+	 *
+	 * @return  Text string
+	 */
+	FText GetText() const;
+
 	/**
 	 * Returns the plain text string without richtext formatting
 	 * @return  Text string
 	 */
-	const FText GetPlainText() const;
+	FText GetPlainText() const;
 
 	/**
 	 * Sets the text that appears when there is no text in the text box
 	 */
 	void SetHintText(const TAttribute< FText >& InHintText);
 
-	/**
-	* Sets the font used to draw the text
-	*
-	* @param  InNewFont	The new font to use
-	*/
+	/** Get the text that appears when there is no text in the text box */
+	FText GetHintText() const;
+
+	/** See attribute TextStyle */
+	void SetTextStyle(const FTextBlockStyle* InTextStyle);
+
+	/** See attribute Font */
 	void SetFont(const TAttribute< FSlateFontInfo >& InNewFont);
 
 	/** See TextShapingMethod attribute */
@@ -187,8 +195,14 @@ public:
 	/** See the AllowContextMenu attribute */
 	void SetAllowContextMenu(const TAttribute< bool >& InAllowContextMenu);
 
-	virtual bool GetIsReadOnly() const override;
-	virtual void ClearSelection() override;
+	/** Query to see if any text is selected within the document */
+	bool AnyTextSelected() const;
+
+	/** Select all the text in the document */
+	void SelectAllText();
+
+	/** Clear the active text selection */
+	void ClearSelection();
 
 	/** Get the currently selected text */
 	FText GetSelectedText() const;
@@ -216,7 +230,7 @@ public:
 	TSharedPtr<const IRun> GetRunUnderCursor() const;
 
 	/** Get the runs currently that are current selected, some of which may be only partially selected */
-	const TArray<TSharedRef<const IRun>> GetSelectedRuns() const;
+	TArray<TSharedRef<const IRun>> GetSelectedRuns() const;
 
 	/** Get the horizontal scroll bar widget */
 	TSharedPtr<const SScrollBar> GetHScrollBar() const;
@@ -224,15 +238,9 @@ public:
 	/** Get the vertical scroll bar widget */
 	TSharedPtr<const SScrollBar> GetVScrollBar() const;
 
-	/** Refresh this text box immediately, rather than wait for the usual caching mechanisms to take affect on the text Tick */
+	/** Refresh this editable text immediately, rather than wait for the usual caching mechanisms to take affect on the text Tick */
 	void Refresh();
 
-	/** Restores the text to the original state */
-	void RestoreOriginalText();
-
-	/** Returns whether the current text varies from the original */
-	bool HasTextChangedFromOriginal() const;
-
 	/**
 	 * Sets the OnKeyDownHandler to provide first chance handling of the OnKeyDown event
 	 *
@@ -243,305 +251,11 @@ public:
 		OnKeyDownHandler = InOnKeyDownHandler;
 	}
 
-private:
-	
-	/** 
-	 * Note: The IME interface for the multiline editable text uses the pre-flowed version of the string since the IME APIs are designed to work with flat strings
-	 *		 This means we have to do a bit of juggling to convert between the two
-	 */
-	friend class FTextInputMethodContext;
-	class FTextInputMethodContext : public ITextInputMethodContext
-	{
-	public:
-		FTextInputMethodContext(const TWeakPtr<SMultiLineEditableText>& InOwningWidget)
-			:	OwningWidget(InOwningWidget)
-			,	IsComposing(false)
-			,	CompositionBeginIndex(INDEX_NONE)
-			,	CompositionLength(0)
-		{}
-
-		virtual ~FTextInputMethodContext() {}
-
-	private:
-		virtual bool IsReadOnly() override;
-		virtual uint32 GetTextLength() override;
-		virtual void GetSelectionRange(uint32& BeginIndex, uint32& Length, ECaretPosition& CaretPosition) override;
-		virtual void SetSelectionRange(const uint32 BeginIndex, const uint32 Length, const ECaretPosition CaretPosition) override;
-		virtual void GetTextInRange(const uint32 BeginIndex, const uint32 Length, FString& OutString) override;
-		virtual void SetTextInRange(const uint32 BeginIndex, const uint32 Length, const FString& InString) override;
-		virtual int32 GetCharacterIndexFromPoint(const FVector2D& Point) override;
-		virtual bool GetTextBounds(const uint32 BeginIndex, const uint32 Length, FVector2D& Position, FVector2D& Size) override;
-		virtual void GetScreenBounds(FVector2D& Position, FVector2D& Size) override;
-		virtual TSharedPtr<FGenericWindow> GetWindow() override;
-		virtual void BeginComposition() override;
-		virtual void UpdateCompositionRange(const int32 InBeginIndex, const uint32 InLength) override;
-		virtual void EndComposition() override;
-
-	private:
-		TWeakPtr<SMultiLineEditableText> OwningWidget;
-
-	public:
-		FGeometry CachedGeometry;
-		bool IsComposing;
-		int32 CompositionBeginIndex;
-		uint32 CompositionLength;
-	};
-
-	enum class ECursorAlignment : uint8
-	{
-		/** Visually align the cursor to the left of the character its placed at, and insert text before the character */
-		Left,
-
-		/** Visually align the cursor to the right of the character its placed at, and insert text after the character */
-		Right,
-	};
-
-	/** Store the information about the current cursor position */
-	class FCursorInfo
-	{
-	public:
-		FCursorInfo()
-			: CursorPosition()
-			, CursorAlignment(ECursorAlignment::Left)
-			, CursorTextDirection(TextBiDi::ETextDirection::LeftToRight)
-			, LastCursorInteractionTime(0)
-		{
-		}
-
-		/** Get the literal position of the cursor (note: this may not be the correct place to insert text to, use GetCursorInteractionLocation for that) */
-		FORCEINLINE FTextLocation GetCursorLocation() const
-		{
-			return CursorPosition;
-		}
-
-		/** Get the alignment of the cursor */
-		FORCEINLINE ECursorAlignment GetCursorAlignment() const
-		{
-			return CursorAlignment;
-		}
-
-		/** Get the direction of the text under the cursor */
-		FORCEINLINE TextBiDi::ETextDirection GetCursorTextDirection() const
-		{
-			return CursorTextDirection;
-		}
-
-		/** Get the interaction position of the cursor (where to insert, delete, etc, text from/to) */
-		FORCEINLINE FTextLocation GetCursorInteractionLocation() const
-		{
-			// If the cursor is right aligned, we treat it as if it were one character further along for insert/delete purposes
-			return FTextLocation(CursorPosition, (CursorAlignment == ECursorAlignment::Right) ? 1 : 0);
-		}
-
-		FORCEINLINE double GetLastCursorInteractionTime() const
-		{
-			return LastCursorInteractionTime;
-		}
-
-		/** Set the position of the cursor, and then work out the correct alignment based on the current text layout */
-		void SetCursorLocationAndCalculateAlignment(const FTextLayout& InTextLayout, const FTextLocation& InCursorPosition);
-
-		/** Set the literal position and alignment of the cursor */
-		void SetCursorLocationAndAlignment(const FTextLayout& InTextLayout, const FTextLocation& InCursorPosition, const ECursorAlignment InCursorAlignment);
-
-		/** Create an undo for this cursor data */
-		FCursorInfo CreateUndo() const;
-
-		/** Restore this cursor data from an undo */
-		void RestoreFromUndo(const FCursorInfo& UndoData);
-
-	private:
-		/** Current cursor position; there is always one. */
-		FTextLocation CursorPosition;
-
-		/** Cursor alignment (horizontal) within its position. This affects the rendering and behavior of the cursor when adding new characters */
-		ECursorAlignment CursorAlignment;
-
-		/** The direction of the text under the cursor */
-		TextBiDi::ETextDirection CursorTextDirection;
-
-		/** Last time the user did anything with the cursor.*/
-		double LastCursorInteractionTime;
-	};
-
-	/**
-	* Stores a single undo level for editable text
-	*/
-	class FUndoState
-	{
-
-	public:
-
-		/** Text */
-		FText Text;
-
-		/** Selection state */
-		TOptional<FTextLocation> SelectionStart;
-
-		/** Cursor data */
-		FCursorInfo CursorInfo;
-
-	public:
-
-		FUndoState()
-			: Text()
-			, SelectionStart()
-			, CursorInfo()
-		{
-		}
-	};
-
-	/** Information needed to be able to scroll to a given point */
-	struct FScrollInfo
-	{
-		FScrollInfo()
-			: Position()
-			, Alignment(ECursorAlignment::Left)
-		{
-		}
-
-		FScrollInfo(const FTextLocation InPosition, const ECursorAlignment InAlignment)
-			: Position(InPosition)
-			, Alignment(InAlignment)
-		{
-		}
-
-		/** The location in the document to scroll to (in line model space) */
-		FTextLocation Position;
-
-		/** The alignment at the given position. This may affect which line view the character maps to when converted from line model space */
-		ECursorAlignment Alignment;
-	};
-
-	/** Run highlighter used to draw the cursor */
-	class FCursorLineHighlighter : public ISlateLineHighlighter
-	{
-	public:
-
-		static TSharedRef< FCursorLineHighlighter > Create(const FCursorInfo* InCursorInfo);
-
-		virtual ~FCursorLineHighlighter() {}
-
-		virtual int32 OnPaint( const FPaintArgs& Args, const FTextLayout::FLineView& Line, const float OffsetX, const float Width, const FTextBlockStyle& DefaultStyle, const FGeometry& AllottedGeometry, const FSlateRect& MyClippingRect, FSlateWindowElementList& OutDrawElements, int32 LayerId, const FWidgetStyle& InWidgetStyle, bool bParentEnabled ) const override;
-
-	protected:
-
-		FCursorLineHighlighter(const FCursorInfo* InCursorInfo);
-
-		/** Cursor data that this highlighter is tracking */
-		const FCursorInfo* CursorInfo;
-	};
-
-	/** Run highlighter used to draw the composition range */
-	class FTextCompositionHighlighter : public ISlateLineHighlighter
-	{
-	public:
-
-		static TSharedRef< FTextCompositionHighlighter > Create();
-
-		virtual ~FTextCompositionHighlighter() {}
-
-		virtual int32 OnPaint( const FPaintArgs& Args, const FTextLayout::FLineView& Line, const float OffsetX, const float Width, const FTextBlockStyle& DefaultStyle, const FGeometry& AllottedGeometry, const FSlateRect& MyClippingRect, FSlateWindowElementList& OutDrawElements, int32 LayerId, const FWidgetStyle& InWidgetStyle, bool bParentEnabled ) const override;
-
-	protected:
-
-		FTextCompositionHighlighter();
-	};
-
-	/** Run renderer used to draw selection ranges */
-	class FTextSelectionRunRenderer : public ISlateRunRenderer
-	{
-	public:
-
-		static TSharedRef< FTextSelectionRunRenderer > Create();
-
-		virtual ~FTextSelectionRunRenderer() {}
-
-		virtual int32 OnPaint( const FPaintArgs& Args, const FTextLayout::FLineView& Line, const TSharedRef< ISlateRun >& Run, const TSharedRef< ILayoutBlock >& Block, const FTextBlockStyle& DefaultStyle, const FGeometry& AllottedGeometry, const FSlateRect& MyClippingRect, FSlateWindowElementList& OutDrawElements, int32 LayerId, const FWidgetStyle& InWidgetStyle, bool bParentEnabled ) const override;
-
-		void SetHasKeyboardFocus(const bool bInHasKeyboardFocus)
-		{
-			bHasKeyboardFocus = bInHasKeyboardFocus;
-		}
-
-	protected:
-
-		FTextSelectionRunRenderer();
-
-		/** true if the parent widget has keyboard focus, false otherwise */
-		bool bHasKeyboardFocus;
-	};
-
-private:
-
-	//~ Begin ITextEditorWidget Interface
-	virtual void StartChangingText() override;
-	virtual void FinishChangingText() override;
-	virtual void BackspaceChar() override;
-	virtual void DeleteChar() override;
-	virtual bool CanTypeCharacter(const TCHAR CharInQuestion) const override;
-	virtual void TypeChar( const int32 Character ) override;
-	virtual FReply MoveCursor( FMoveCursor Args ) override;
-	virtual void JumpTo(ETextLocation JumpLocation, ECursorAction Action) override;
-	virtual void SelectAllText() override;
-	virtual bool SelectAllTextWhenFocused() override;
-	virtual void SelectWordAt( const FGeometry& MyGeometry, const FVector2D& ScreenSpacePosition ) override;
-	virtual void BeginDragSelection() override;
-	virtual bool IsDragSelecting() const override;
-	virtual void EndDragSelection() override;
-	virtual bool AnyTextSelected() const override;
-	virtual bool IsTextSelectedAt( const FGeometry& MyGeometry, const FVector2D& ScreenSpacePosition ) const override;
-	virtual void SetWasFocusedByLastMouseDown( bool Value ) override;
-	virtual bool WasFocusedByLastMouseDown() const override;
-	virtual void SetHasDragSelectedSinceFocused( bool Value ) override;
-	virtual bool HasDragSelectedSinceFocused() const override;
-	virtual FReply OnEscape() override;
-	virtual void OnEnter() override;
-	virtual bool CanExecuteCut() const override;
-	virtual void CutSelectedTextToClipboard() override;
-	virtual bool CanExecuteCopy() const override;
-	virtual void CopySelectedTextToClipboard() override;
-	virtual bool CanExecutePaste() const override;
-	virtual void PasteTextFromClipboard() override;
-	virtual bool CanExecuteUndo() const override;
-	virtual void Undo() override;
-	virtual void Redo() override;
-	virtual TSharedRef< SWidget > GetWidget() override;
-	virtual TSharedRef<IVirtualKeyboardEntry> GetVirtualKeyboardEntry() override;
-	virtual void SummonContextMenu(const FVector2D& InLocation, TSharedPtr<SWindow> ParentWindow, const FWidgetPath& EventPath) override;
-	virtual void LoadText() override;
-	//~ End ITextEditorWidget Interface
-
-public:
-	//~ Begin IVirtualKeyboardEntry Interface
-	virtual void SetTextFromVirtualKeyboard(const FText& InNewText, ESetTextType SetTextType, ETextCommit::Type CommitType) override;
-
-	virtual const FText& GetText() const override
-	{
-		return BoundText.Get(FText::GetEmpty());
-	}
-
-	virtual const FText GetHintText() const override
-	{
-		return HintText.Get(FText::GetEmpty());
-	}
-
-	virtual EKeyboardType GetVirtualKeyboardType() const override
-	{
-		return EKeyboardType::Keyboard_Default;
-	}
-
-	virtual bool IsMultilineEntry() const override
-	{
-		return true;
-	}
-	//~ End IVirtualKeyboardEntry Interface
-
 protected:
 	//~ Begin SWidget Interface
 	virtual void Tick( const FGeometry& AllottedGeometry, const double InCurrentTime, const float InDeltaTime ) override;
 	virtual int32 OnPaint( const FPaintArgs& Args, const FGeometry& AllottedGeometry, const FSlateRect& MyClippingRect, FSlateWindowElementList& OutDrawElements, int32 LayerId, const FWidgetStyle& InWidgetStyle, bool bParentEnabled ) const override;
-	virtual void CacheDesiredSize(float) override;
+	virtual void CacheDesiredSize(float LayoutScaleMultiplier) override;
 	virtual FVector2D ComputeDesiredSize(float LayoutScaleMultiplier) const override;
 	virtual FChildren* GetChildren() override;
 	virtual void OnArrangeChildren( const FGeometry& AllottedGeometry, FArrangedChildren& ArrangedChildren ) const override;
@@ -560,236 +274,78 @@ protected:
 	virtual bool IsInteractable() const override;
 	//~ End SWidget Interface
 
-private:
-	
-	void OnContextMenuClosed(TSharedRef<IMenu>);
-
-	/** Remember where the cursor was when we started selecting. */
-	void BeginSelecting( const FTextLocation& Endpoint );
-
-	void DeleteSelectedText();
-
-	void ExecuteDeleteAction();
-	bool CanExecuteDelete() const;
-
-	bool CanExecuteSelectAll() const;
-
-	bool DoesClipboardHaveAnyText() const;
-
-	/** Insert the given text at the current cursor position, correctly taking into account new line characters */
-	void InsertTextAtCursorImpl(const FString& InString);
-
-	/** Insert a new-line at the current cursor position */
-	void InsertNewLineAtCursorImpl();
-
-	/**
-	 * Given a location and a Direction to offset, return a new location.
-	 *
-	 * @param Location    Cursor location from which to offset
-	 * @param Direction   Positive means right, negative means left.
-	 */
-	FTextLocation TranslatedLocation( const FTextLocation& CurrentLocation, int8 Direction ) const;
-
-	/**
-	 * Given a location and a Direction to offset, return a new location.
-	 *
-	 * @param Location              Cursor location from which to offset
-	 * @param NumLinesToMove        Number of lines to move in a given direction. Positive means down, negative means up.
-	 * @param GeometryScale         Geometry DPI scale at which the widget is being rendered
-	 * @param OutCursorPosition     Fill with the updated cursor position.
-	 * @param OutCursorAlignment    Optionally fill with a new cursor alignment (will be auto-calculated if not set).
-	 */
-	void TranslateLocationVertical( const FTextLocation& Location, int32 NumLinesToMove, float GeometryScale, FTextLocation& OutCursorPosition, TOptional<ECursorAlignment>& OutCursorAlignment ) const;
-
-	/** Find the closest word boundary */
-	FTextLocation ScanForWordBoundary( const FTextLocation& Location, int8 Direction ) const; 
-
-	/** Get the character at Location */
-	TCHAR GetCharacterAt( const FTextLocation& Location ) const;
-
-	/** Are we at the beginning of all the text. */
-	bool IsAtBeginningOfDocument( const FTextLocation& Location ) const;
-	
-	/** Are we at the end of all the text. */
-	bool IsAtEndOfDocument( const FTextLocation& Location ) const;
-
-	/** Is this location the beginning of a line */
-	bool IsAtBeginningOfLine( const FTextLocation& Location ) const;
-
-	/** Is this location the end of a line. */
-	bool IsAtEndOfLine( const FTextLocation& Location ) const;
-
-	/** Are we currently at the beginning of a word */
-	bool IsAtWordStart( const FTextLocation& Location ) const;
-
-	void UpdateCursorHighlight();
-
-	void RemoveCursorHighlight();
-
-	void UpdatePreferredCursorScreenOffsetInLine();
-
-	void PushUndoState(const SMultiLineEditableText::FUndoState& InUndoState);
-
-	void ClearUndoStates();
-
-	void MakeUndoState(SMultiLineEditableText::FUndoState& OutUndoState);
-
-	void SaveText(const FText& TextToSave);
-
-	/**
-	 * Sets the current editable text for this text block
-	 * Note: Doesn't update the value of BoundText, nor does it call OnTextChanged
-	 *
-	 * @param TextToSet		The new text to set in the internal TextLayout
-	 * @param bForce		True to force the update, even if the text currently matches what's in the TextLayout
-	 *
-	 * @return true if the text was updated, false if the text wasn't update (because it was already up-to-date)
-	 */
-	bool SetEditableText(const FText& TextToSet, const bool bForce = false);
-
-	/**
-	 * Gets the current editable text for this text block
-	 * Note: We don't store text in this form (it's stored as lines in the text layout) so every call to this function has to reconstruct it
-	 */
-	FText GetEditableText() const;
-
-	/**
-	 * Force the text layout to be updated from the marshaller
-	 */
-	void ForceRefreshTextLayout(const FText& CurrentText);
-
+protected:
 	void OnHScrollBarMoved(const float InScrollOffsetFraction);
 	void OnVScrollBarMoved(const float InScrollOffsetFraction);
 
 	/** Return whether a RMB+Drag scroll operation is taking place */
 	bool IsRightClickScrolling() const;
 
-	/**
-	 * Ensure that we will get a Tick() soon (either due to having active focus, or something having changed progmatically and requiring an update)
-	 * Does nothing if the active tick timer is already enabled
-	 */
-	void EnsureActiveTick();
-
-private:
-
-	/** The text displayed in this text block */
-	TAttribute<FText> BoundText;
-
-	/** The state of BoundText last Tick() (used to allow updates when the text is changed) */
-	FTextSnapshot BoundTextLastTick;
-
-	/** The text that appears when there is no text in the text box */
-	TAttribute<FText> HintText;
-
-	/** The marshaller used to get/set the BoundText text to/from the text layout. */
-	TSharedPtr< ITextLayoutMarshaller > Marshaller;
-
-	/** In control of the layout and wrapping of the BoundText */
-	TSharedPtr< FSlateTextLayout > TextLayout;
-
-	/** In control of the layout and wrapping of the HintText */
-	TSharedPtr< FTextBlockLayout > HintTextLayout;
-
-	/** Default style used by the TextLayout */
-	FTextBlockStyle TextStyle;
-	
-	/** Style used to draw the hint text (only valid when HintTextLayout is set) */
-	TSharedPtr< FTextBlockStyle > HintTextStyle;
-
-	/** Whether text wraps onto a new line when it's length exceeds this width; if this value is zero or negative, no wrapping occurs. */
-	TAttribute< float > WrapTextAt;
-
-	/** True if we're wrapping text automatically based on the computed horizontal space for this widget */
-	TAttribute< bool > AutoWrapText;
-
-	/** The last known size of the control from the previous OnPaint, used to recalculate wrapping. */
-	mutable FVector2D CachedSize;
-
-	/** The scroll offset (in unscaled Slate units) for this text */
-	FVector2D ScrollOffset;
-
-	TAttribute< FMargin > Margin;
-	TAttribute< ETextJustify::Type > Justification; 
-	TAttribute< float > LineHeightPercentage;
-
-	TSharedPtr< FCursorLineHighlighter > CursorLineHighlighter;
-	TSharedPtr< FTextCompositionHighlighter > TextCompositionHighlighter;
-	TSharedPtr< FTextSelectionRunRenderer > TextSelectionRunRenderer;
-
-	/** That start of the selection when there is a selection. The end is implicitly wherever the cursor happens to be. */
-	TOptional<FTextLocation> SelectionStart;
+public:
+	//~ Begin ISlateEditableTextWidget Interface
+	virtual bool IsTextReadOnly() const override;
+	virtual bool IsTextPassword() const override;
+	virtual bool IsMultiLineTextEdit() const override;
+	//~ End ISlateEditableTextWidget Interface
 
-	/** Current cursor data */
-	FCursorInfo CursorInfo;
+protected:
+	//~ Begin ISlateEditableTextWidget Interface
+	virtual bool ShouldSelectAllTextWhenFocused() const override;
+	virtual bool ShouldClearTextSelectionOnFocusLoss() const override;
+	virtual bool ShouldRevertTextOnEscape() const override;
+	virtual bool ShouldClearKeyboardFocusOnCommit() const override;
+	virtual bool CanInsertCarriageReturn() const override;
+	virtual bool CanTypeCharacter(const TCHAR InChar) const override;
+	virtual void EnsureActiveTick() override;
+	virtual EKeyboardType GetVirtualKeyboardType() const override;
+	virtual TSharedRef<SWidget> GetSlateWidget() override;
+	virtual TSharedPtr<SWidget> BuildContextMenuContent() const override;
+	virtual void OnTextChanged(const FText& InText) override;
+	virtual void OnTextCommitted(const FText& InText, const ETextCommit::Type InTextAction) override;
+	virtual void OnCursorMoved(const FTextLocation& InLocation) override;
+	virtual float UpdateAndClampHorizontalScrollBar(const float InViewOffset, const float InViewFraction, const EVisibility InVisiblityOverride) override;
+	virtual float UpdateAndClampVerticalScrollBar(const float InViewOffset, const float InViewFraction, const EVisibility InVisiblityOverride) override;
+	//~ End ISlateEditableTextWidget Interface
 
-	/** The user probably wants the cursor where they last explicitly positioned it horizontally. */
-	float PreferredCursorScreenOffsetInLine;
+protected:
+	/** The text layout that deals with the editable text */
+	TUniquePtr<FSlateEditableTextLayout> EditableTextLayout;
 
 	/** Whether to select all text when the user clicks to give focus on the widget */
-	TAttribute< bool > bSelectAllTextWhenFocused;
+	TAttribute<bool> bSelectAllTextWhenFocused;
 
 	/** Whether to clear text selection when focus is lost */
-	TAttribute< bool > bClearTextSelectionOnFocusLoss;
+	TAttribute<bool> bClearTextSelectionOnFocusLoss;
 
-	/** True if any changes should be reverted if we recieve an escape key */
-	TAttribute< bool > bRevertTextOnEscape;
+	/** True if any changes should be reverted if we receive an escape key */
+	TAttribute<bool> bRevertTextOnEscape;
 
 	/** True if we want the text control to lose focus on an text commit/revert events */
-	TAttribute< bool > bClearKeyboardFocusOnCommit;
+	TAttribute<bool> bClearKeyboardFocusOnCommit;
 
 	/** Whether the context menu can be opened */
-	TAttribute< bool > AllowContextMenu;
-
-	/** True if we're currently selecting text by dragging the mouse cursor with the left button held down */
-	bool bIsDragSelecting;
-
-	/** True if the last mouse down caused us to receive keyboard focus */
-	bool bWasFocusedByLastMouseDown;
-
-	/** True if characters were selected by dragging since the last keyboard focus.  Used for text selection. */
-	bool bHasDragSelectedSinceFocused;
-
-	/** If set, the pending data containing a position that should be scrolled into view */
-	TOptional< FScrollInfo > PositionToScrollIntoView;
-
-	/** Undo states */
-	TArray< FUndoState > UndoStates;
-
-	/** Current undo state level that we've rolled back to, or INDEX_NONE if we haven't undone.  Used for 'Redo'. */
-	int32 CurrentUndoLevel;
-
-	/** Undo state that will be pushed if text is actually changed between calls to StartChangingText() and FinishChangingText() */
-	FUndoState StateBeforeChangingText;
-
-	/** Original text undo state */
-	FUndoState OriginalText;
-
-	/** True if we're currently (potentially) changing the text string */
-	bool bIsChangingText;
+	TAttribute<bool> bAllowContextMenu;
 
 	/** Sets whether this text box can actually be modified interactively by the user */
-	TAttribute< bool > IsReadOnly;
-
-	/** A list commands to execute if a user presses the corresponding keybinding in the text box */
-	TSharedRef< FUICommandList > UICommandList;
+	TAttribute<bool> bIsReadOnly;
 
 	/** Delegate to call before a context menu is opened */
 	FOnContextMenuOpening OnContextMenuOpening;
 
 	/** Called whenever the text is changed interactively by the user */
-	FOnTextChanged OnTextChanged;
+	FOnTextChanged OnTextChangedCallback;
 
 	/** Called whenever the text is committed.  This happens when the user presses enter or the text box loses focus. */
-	FOnTextCommitted OnTextCommitted;
+	FOnTextCommitted OnTextCommittedCallback;
 
 	/** Called when the cursor is moved within the text area */
-	FOnCursorMoved OnCursorMoved;
+	FOnCursorMoved OnCursorMovedCallback;
 
 	/** The horizontal scroll bar widget */
-	TSharedPtr< SScrollBar > HScrollBar;
+	TSharedPtr<SScrollBar> HScrollBar;
 
 	/** The vertical scroll bar widget */
-	TSharedPtr< SScrollBar > VScrollBar;
+	TSharedPtr<SScrollBar> VScrollBar;
 
 	/** Called whenever the horizontal scrollbar is moved by the user */
 	FOnUserScrolled OnHScrollBarUserScrolled;
@@ -800,24 +356,6 @@ private:
 	/** Menu extender for right-click context menu */
 	TSharedPtr<FExtender> MenuExtender;
 
-	/** Information about any active context menu widgets */
-	FActiveTextEditContextMenu ActiveContextMenu;
-
-	/** Implemented context object for text input method systems. */
-	TSharedPtr<FTextInputMethodContext> TextInputMethodContext;
-
-	/** Notification interface object for text input method systems. */
-	TSharedPtr<ITextInputMethodChangeNotifier> TextInputMethodChangeNotifier;
-
-	/** Whether the text has been changed by a virtual keyboard */
-	bool bTextChangedByVirtualKeyboard;
-
-	/** Whether the text has been committed by a virtual keyboard */
-	bool bTextCommittedByVirtualKeyboard;
-
-	/** How the text was committed by the virtual keyboard */
-	ETextCommit::Type VirtualKeyboardTextCommitType;
-
 	/** The optional modifier key necessary to create a newline when typing into the editor. */
 	EModifierKey::Type ModiferKeyForNewLine;
 
@@ -837,5 +375,4 @@ private:
 	FOnKeyDown OnKeyDownHandler;
 };
 
-
-#endif //WITH_FANCY_TEXT
\ No newline at end of file
+#endif //WITH_FANCY_TEXT
diff --git a/Engine/Source/Runtime/Slate/Public/Widgets/Text/SRichTextBlock.h b/Engine/Source/Runtime/Slate/Public/Widgets/Text/SRichTextBlock.h
--- a/Engine/Source/Runtime/Slate/Public/Widgets/Text/SRichTextBlock.h
+++ b/Engine/Source/Runtime/Slate/Public/Widgets/Text/SRichTextBlock.h
@@ -1,6 +1,8 @@
 // Copyright 1998-2016 Epic Games, Inc. All Rights Reserved.
 #pragma once
 
+#include "UniquePtr.h"
+
 #if WITH_FANCY_TEXT
 
 /**
@@ -113,6 +115,12 @@ public:
 		return FHyperlinkDecorator::Create( Id, FSlateHyperlinkRun::FOnClick::CreateSP( InUserObjectPtr, NavigateFunc ) );
 	}
 
+	/** Constructor */
+	SRichTextBlock();
+
+	/** Destructor */
+	~SRichTextBlock();
+
 	//~ Begin SWidget Interface
 	void Construct( const FArguments& InArgs );
 	virtual int32 OnPaint( const FPaintArgs& Args, const FGeometry& AllottedGeometry, const FSlateRect& MyClippingRect, FSlateWindowElementList& OutDrawElements, int32 LayerId, const FWidgetStyle& InWidgetStyle, bool bParentEnabled ) const override;
@@ -174,7 +182,7 @@ private:
 	TAttribute< FText > BoundText;
 
 	/** The wrapped layout for this text block */
-	TSharedPtr< FTextBlockLayout > TextLayoutCache;
+	TUniquePtr< FTextBlockLayout > TextLayoutCache;
 
 	/** Default style used by the TextLayout */
 	FTextBlockStyle TextStyle;
diff --git a/Engine/Source/Runtime/Slate/Public/Widgets/Text/STextBlock.h b/Engine/Source/Runtime/Slate/Public/Widgets/Text/STextBlock.h
--- a/Engine/Source/Runtime/Slate/Public/Widgets/Text/STextBlock.h
+++ b/Engine/Source/Runtime/Slate/Public/Widgets/Text/STextBlock.h
@@ -3,6 +3,7 @@
 #pragma once
 
 #include "IBreakIterator.h"
+#include "UniquePtr.h"
 
 class FTextBlockLayout;
 
@@ -108,11 +109,10 @@ public:
 	SLATE_END_ARGS()
 
 	/** Constructor */
-	STextBlock()
-	{
-		bCanTick = false;
-		bCanSupportFocus = false;
-	}
+	STextBlock();
+
+	/** Destructor */
+	~STextBlock();
 
 	/**
 	 * Construct this widget
@@ -231,12 +231,8 @@ private:
 	/** The text displayed in this text block */
 	TAttribute< FText > BoundText;
 
-#if WITH_FANCY_TEXT
-
 	/** The wrapped layout for this text block */
-	TSharedPtr< FTextBlockLayout > TextLayoutCache;
-
-#endif//WITH_FANCY_TEXT
+	TUniquePtr< FTextBlockLayout > TextLayoutCache;
 
 	/** Default style used by the TextLayout */
 	const FTextBlockStyle* TextStyle;
