Copying //UE4/Dev-Editor to Dev-Main (//UE4/Dev-Main)

diff --git a/Engine/Source/Runtime/Slate/Public/Widgets/Text/SlateEditableTextLayout.h b/Engine/Source/Runtime/Slate/Public/Widgets/Text/SlateEditableTextLayout.h
new file mode 100644
--- /dev/null
+++ b/Engine/Source/Runtime/Slate/Public/Widgets/Text/SlateEditableTextLayout.h
@@ -0,0 +1,547 @@
+// Copyright 1998-2016 Epic Games, Inc. All Rights Reserved.
+
+#pragma once
+
+#include "SlateEditableTextTypes.h"
+#include "IVirtualKeyboardEntry.h"
+#include "ITextInputMethodSystem.h"
+#include "ISlateEditableTextWidget.h"
+#include "UniquePtr.h"
+
+class FSlateTextLayout;
+class FTextBlockLayout;
+class ITextLayoutMarshaller;
+class ISlateRunRenderer;
+
+/** Class to handle the cached layout of SEditableText/SMultiLineEditableText by proxying around a FTextLayout */
+class SLATE_API FSlateEditableTextLayout
+{
+public:
+	FSlateEditableTextLayout(ISlateEditableTextWidget& InOwnerWidget, const TAttribute<FText>& InInitialText, FTextBlockStyle InTextStyle, const TOptional<ETextShapingMethod> InTextShapingMethod, const TOptional<ETextFlowDirection> InTextFlowDirection, TSharedRef<ITextLayoutMarshaller> InTextMarshaller, TSharedRef<ITextLayoutMarshaller> InHintTextMarshaller);
+	~FSlateEditableTextLayout();
+
+	void SetText(const TAttribute<FText>& InText);
+	FText GetText() const;
+
+	void SetHintText(const TAttribute<FText>& InHintText);
+	FText GetHintText() const;
+
+	void SetTextStyle(const FTextBlockStyle& InTextStyle);
+	const FTextBlockStyle& GetTextStyle() const;
+
+	/** Set the brush to use when drawing the cursor */
+	void SetCursorBrush(const TAttribute<const FSlateBrush*>& InCursorBrush);
+
+	/** Set the brush to use when drawing the composition highlight */
+	void SetCompositionBrush(const TAttribute<const FSlateBrush*>& InCompositionBrush);
+
+	/** Get the plain text string without rich-text formatting */
+	FText GetPlainText() const;
+
+	/**
+	 * Sets the current editable text for this text block
+	 * Note: Doesn't update the value of BoundText, nor does it call OnTextChanged
+	 *
+	 * @param TextToSet	The new text to set in the internal TextLayout
+	 * @param bForce	True to force the update, even if the text currently matches what's in the TextLayout
+	 *
+	 * @return true if the text was updated, false if the text wasn't update (because it was already up-to-date)
+	 */
+	bool SetEditableText(const FText& TextToSet, const bool bForce = false);
+
+	/**
+	 * Gets the current editable text for this text block
+	 * Note: We don't store text in this form (it's stored as lines in the text layout) so every call to this function has to reconstruct it
+	 */
+	FText GetEditableText() const;
+
+	/** Get the currently selected text */
+	FText GetSelectedText() const;
+
+	/** Set the text shaping method that the internal text layout should use */
+	void SetTextShapingMethod(const TOptional<ETextShapingMethod>& InTextShapingMethod);
+
+	/** Set the text flow direction that the internal text layout should use */
+	void SetTextFlowDirection(const TOptional<ETextFlowDirection>& InTextFlowDirection);
+
+	/** Set the wrapping to use for this document */
+	void SetTextWrapping(const TAttribute<float>& InWrapTextAt, const TAttribute<bool>& InAutoWrapText);
+
+	/** Set whether text wraps onto a new line when it's length exceeds this width; if this value is zero or negative, no wrapping occurs */
+	void SetWrapTextAt(const TAttribute<float>& InWrapTextAt);
+
+	/** Set whether to wrap text automatically based on the widget's computed horizontal space */
+	void SetAutoWrapText(const TAttribute<bool>& InAutoWrapText);
+
+	/** Set the amount of blank space left around the edges of text area */
+	void SetMargin(const TAttribute<FMargin>& InMargin);
+
+	/** Set how the text should be aligned with the margin */
+	void SetJustification(const TAttribute<ETextJustify::Type>& InJustification);
+
+	/** Set the amount to scale each lines height by */
+	void SetLineHeightPercentage(const TAttribute<float>& InLineHeightPercentage);
+
+	/** Set the information used to help identify who owns this text layout in the case of an error */
+	void SetDebugSourceInfo(const TAttribute<FString>& InDebugSourceInfo);
+
+	/** Get the virtual keyboard handler for this text layout */
+	TSharedRef<IVirtualKeyboardEntry> GetVirtualKeyboardEntry() const;
+
+	/** Get the IME context for this text layout */
+	TSharedRef<ITextInputMethodContext> GetTextInputMethodContext() const;
+
+	/** Refresh this editable text immediately, rather than wait for the usual caching mechanisms to take affect on the text Tick */
+	bool Refresh();
+
+	/** Force the text layout to be updated from the marshaller */
+	void ForceRefreshTextLayout(const FText& CurrentText);
+
+	/** Update the horizontal scroll amount from the given fraction */
+	FVector2D SetHorizontalScrollFraction(const float InScrollOffsetFraction);
+
+	/** Update the vertical scroll amount from the given fraction */
+	FVector2D SetVerticalScrollFraction(const float InScrollOffsetFraction);
+
+	/** Set the absolute scroll offset value */
+	FVector2D SetScrollOffset(const FVector2D& InScrollOffset, const FGeometry& InGeometry);
+
+	/** Get the absolute scroll offset value */
+	FVector2D GetScrollOffset() const;
+
+	/** Called when our parent widget receives focus */
+	bool HandleFocusReceived(const FFocusEvent& InFocusEvent);
+
+	/** Called when our parent widget loses focus */
+	bool HandleFocusLost(const FFocusEvent& InFocusEvent);
+
+	/** Called to handle an OnKeyChar event from our parent widget */
+	FReply HandleKeyChar(const FCharacterEvent& InCharacterEvent);
+
+	/** Called to handle an OnKeyDown event from our parent widget */
+	FReply HandleKeyDown(const FKeyEvent& InKeyEvent);
+
+	/** Called to handle an OnMouseButtonDown event from our parent widget */
+	FReply HandleMouseButtonDown(const FGeometry& MyGeometry, const FPointerEvent& InMouseEvent);
+
+	/** Called to handle an OnMouseButtonUp event from our parent widget */
+	FReply HandleMouseButtonUp(const FGeometry& MyGeometry, const FPointerEvent& InMouseEvent);
+
+	/** Called to handle an OnMouseMove event from our parent widget */
+	FReply HandleMouseMove(const FGeometry& InMyGeometry, const FPointerEvent& InMouseEvent);
+
+	/** Called to handle an OnMouseButtonDoubleClick event from our parent widget */
+	FReply HandleMouseButtonDoubleClick(const FGeometry& InMyGeometry, const FPointerEvent& InMouseEvent);
+
+	/** Called to handle an escape action acting on the current selection */
+	bool HandleEscape();
+
+	/** Called to handle a backspace action acting on the current selection or at the cursor position */
+	bool HandleBackspace();
+
+	/** Called to handle a delete action acting on the current selection or at the cursor position */
+	bool HandleDelete();
+
+	/** Called to handle a typing a character on the current selection or at the cursor position */
+	bool HandleTypeChar(const TCHAR InChar);
+
+	/** Called to handle a carriage return action acting on the current selection or at the cursor position */
+	bool HandleCarriageReturn();
+
+	/** Are we able to delete the currently selected text? */
+	bool CanExecuteDelete() const;
+
+	/** Delete any currently selected text */
+	void DeleteSelectedText();
+
+	/** Query to see if any text is selected within the document */
+	bool AnyTextSelected() const;
+
+	/** Query to see if the text under the given position is currently selected */
+	bool IsTextSelectedAt(const FGeometry& MyGeometry, const FVector2D& ScreenSpacePosition) const;
+
+	/** Query to see if the text under the given position is currently selected (the position is local to the text layout space) */
+	bool IsTextSelectedAt(const FVector2D& InLocalPosition) const;
+
+	/** Are we able to execute the "Select All" command? */
+	bool CanExecuteSelectAll() const;
+
+	/** Select all the text in the document */
+	void SelectAllText();
+
+	/** Select the word under the given position */
+	void SelectWordAt(const FGeometry& MyGeometry, const FVector2D& ScreenSpacePosition);
+
+	/** Select the word under the given position (the position is local to the text layout space) */
+	void SelectWordAt(const FVector2D& InLocalPosition);
+
+	/** Clear the active text selection */
+	void ClearSelection();
+
+	/** Are we able to cut the currently selected text? */
+	bool CanExecuteCut() const;
+
+	/** Cut the currently selected text and place it on the clipboard */
+	void CutSelectedTextToClipboard();
+
+	/** Are we able to copy the currently selected text? */
+	bool CanExecuteCopy() const;
+
+	/** Copy the currently selected text and place it on the clipboard */
+	void CopySelectedTextToClipboard();
+
+	/** Are we able to paste the text from the clipboard into this document? */
+	bool CanExecutePaste() const;
+
+	/** Paste the text from the clipboard into this document */
+	void PasteTextFromClipboard();
+
+	/** Insert the given text at the current cursor position, correctly taking into account new line characters */
+	void InsertTextAtCursor(const FString& InString);
+
+	/** Insert the given run at the current cursor position */
+	void InsertRunAtCursor(TSharedRef<IRun> InRun);
+
+	/** Move the cursor in the document using the specified move method */
+	void MoveCursor(const FMoveCursor& InArgs);
+
+	/** Move the cursor to the given location in the document (will also scroll to this point) */
+	void GoTo(const FTextLocation& NewLocation);
+
+	/** Move the cursor specified location */
+	void GoTo(ETextLocation NewLocation);
+
+	/** Jump the cursor to the given location in the document */
+	void JumpTo(ETextLocation JumpLocation, ECursorAction Action);
+
+	/** Scroll to the given location in the document (without moving the cursor) */
+	void ScrollTo(const FTextLocation& NewLocation);
+
+	/** Update the active cursor highlight based on the state of the text layout */
+	void UpdateCursorHighlight();
+
+	/** Remove any active cursor highlights */
+	void RemoveCursorHighlight();
+
+	/** Update the preferred offset of the cursor, based on the current state of the text layout */
+	void UpdatePreferredCursorScreenOffsetInLine();
+
+	/** Apply the given style to the currently selected text (or insert a new run at the current cursor position if no text is selected) */
+	void ApplyToSelection(const FRunInfo& InRunInfo, const FTextBlockStyle& InStyle);
+
+	/** Get the run currently under the cursor, or null if there is no run currently under the cursor */
+	TSharedPtr<const IRun> GetRunUnderCursor() const;
+
+	/** Get the runs currently that are current selected, some of which may be only partially selected */
+	TArray<TSharedRef<const IRun>> GetSelectedRuns() const;
+
+	/**
+	 * Given a location and a Direction to offset, return a new location.
+	 *
+	 * @param Location    Cursor location from which to offset
+	 * @param Direction   Positive means right, negative means left.
+	 */
+	FTextLocation TranslatedLocation(const FTextLocation& CurrentLocation, int8 Direction) const;
+
+	/**
+	 * Given a location and a Direction to offset, return a new location.
+	 *
+	 * @param Location              Cursor location from which to offset
+	 * @param NumLinesToMove        Number of lines to move in a given direction. Positive means down, negative means up.
+	 * @param GeometryScale         Geometry DPI scale at which the widget is being rendered
+	 * @param OutCursorPosition     Fill with the updated cursor position.
+	 * @param OutCursorAlignment    Optionally fill with a new cursor alignment (will be auto-calculated if not set).
+	 */
+	void TranslateLocationVertical(const FTextLocation& Location, int32 NumLinesToMove, float GeometryScale, FTextLocation& OutCursorPosition, TOptional<SlateEditableTextTypes::ECursorAlignment>& OutCursorAlignment) const;
+
+	/** Find the closest word boundary */
+	FTextLocation ScanForWordBoundary(const FTextLocation& Location, int8 Direction) const;
+
+	/** Get the character at Location */
+	TCHAR GetCharacterAt(const FTextLocation& Location) const;
+
+	/** Are we at the beginning of all the text. */
+	bool IsAtBeginningOfDocument(const FTextLocation& Location) const;
+
+	/** Are we at the end of all the text. */
+	bool IsAtEndOfDocument(const FTextLocation& Location) const;
+
+	/** Is this location the beginning of a line */
+	bool IsAtBeginningOfLine(const FTextLocation& Location) const;
+
+	/** Is this location the end of a line. */
+	bool IsAtEndOfLine(const FTextLocation& Location) const;
+
+	/** Are we currently at the beginning of a word */
+	bool IsAtWordStart(const FTextLocation& Location) const;
+
+	/** Restores the text to the original state */
+	void RestoreOriginalText();
+
+	/** Returns whether the current text varies from the original */
+	bool HasTextChangedFromOriginal() const;
+
+	/** Called to begin an undoable editable text transaction */
+	void BeginEditTransation();
+
+	/** Called to end an undoable editable text transaction */
+	void EndEditTransaction();
+
+	/** Push the given undo state onto the undo stack */
+	void PushUndoState(const SlateEditableTextTypes::FUndoState& InUndoState);
+
+	/** Clear the current undo stack */
+	void ClearUndoStates();
+
+	/** Create an undo state that reflects the current state of the document */
+	void MakeUndoState(SlateEditableTextTypes::FUndoState& OutUndoState);
+
+	/** Are we currently able to execute an undo action? */
+	bool CanExecuteUndo() const;
+	
+	/** Execute an undo action */
+	void Undo();
+
+	/** Are we currently able to execute a redo action? */
+	bool CanExecuteRedo() const;
+
+	/** Execute a redo action */
+	void Redo();
+
+	void SaveText(const FText& TextToSave);
+
+	void LoadText();
+
+	void Tick(const FGeometry& AllottedGeometry, const double InCurrentTime, const float InDeltaTime);
+
+	int32 OnPaint(const FPaintArgs& Args, const FGeometry& AllottedGeometry, const FSlateRect& MyClippingRect, FSlateWindowElementList& OutDrawElements, int32 LayerId, const FWidgetStyle& InWidgetStyle, bool bParentEnabled);
+
+	void CacheDesiredSize(float LayoutScaleMultiplier);
+
+	FVector2D ComputeDesiredSize(float LayoutScaleMultiplier) const;
+
+	FChildren* GetChildren();
+
+	void OnArrangeChildren(const FGeometry& AllottedGeometry, FArrangedChildren& ArrangedChildren) const;
+
+	FVector2D GetSize() const;
+
+	TSharedRef<SWidget> BuildDefaultContextMenu(const TSharedPtr<FExtender>& InMenuExtender) const;
+
+	bool HasActiveContextMenu() const;
+
+private:
+	/** Insert the given text at the current cursor position, correctly taking into account new line characters */
+	void InsertTextAtCursorImpl(const FString& InString);
+
+	/** Insert a new-line at the current cursor position */
+	void InsertNewLineAtCursorImpl();
+
+	/** Implementation of Refresh that actually updates the layout. Optionally takes text to set, or will use the current editable text if none if provided */
+	bool RefreshImpl(const FText* InTextToSet);
+
+	/** Create a text or password run using the given text and style */
+	TSharedRef<IRun> CreateTextOrPasswordRun(const FRunInfo& InRunInfo, const TSharedRef<const FString>& InText, const FTextBlockStyle& InStyle);
+
+	/** Called when the active context menu is closed */
+	void OnContextMenuClosed(TSharedRef<IMenu> Menu);
+
+private:
+	/** Virtual keyboard handler for an editable text layout */
+	friend class FVirtualKeyboardEntry;
+	class FVirtualKeyboardEntry : public IVirtualKeyboardEntry
+	{
+	public:
+		static TSharedRef<FVirtualKeyboardEntry> Create(FSlateEditableTextLayout& InOwnerLayout);
+
+		virtual void SetTextFromVirtualKeyboard(const FText& InNewText, ESetTextType SetTextType, ETextCommit::Type CommitType) override;
+		virtual FText GetText() const override;
+		virtual FText GetHintText() const override;
+		virtual EKeyboardType GetVirtualKeyboardType() const override;
+		virtual bool IsMultilineEntry() const override;
+
+	private:
+		FVirtualKeyboardEntry(FSlateEditableTextLayout& InOwnerLayout);
+		FSlateEditableTextLayout* OwnerLayout;
+	};
+
+private:
+	/**
+	 * Note: The IME interface for the multiline editable text uses the pre-flowed version of the string since the IME APIs are designed to work with flat strings
+	 *		 This means we have to do a bit of juggling to convert between the two
+	 */
+	friend class FTextInputMethodContext;
+	class FTextInputMethodContext : public ITextInputMethodContext
+	{
+	public:
+		static TSharedRef<FTextInputMethodContext> Create(FSlateEditableTextLayout& InOwnerLayout);
+
+		FORCEINLINE bool IsComposing() const
+		{
+			return bIsComposing;
+		}
+
+		FORCEINLINE FTextRange GetCompositionRange() const
+		{
+			return FTextRange(CompositionBeginIndex, CompositionBeginIndex + CompositionLength);
+		}
+
+		bool UpdateCachedGeometry(const FGeometry& InAllottedGeometry)
+		{
+			if (CachedGeometry != InAllottedGeometry)
+			{
+				CachedGeometry = InAllottedGeometry;
+				return true;
+			}
+			return false;
+		}
+
+	private:
+		virtual bool IsReadOnly() override;
+		virtual uint32 GetTextLength() override;
+		virtual void GetSelectionRange(uint32& BeginIndex, uint32& Length, ECaretPosition& CaretPosition) override;
+		virtual void SetSelectionRange(const uint32 BeginIndex, const uint32 Length, const ECaretPosition CaretPosition) override;
+		virtual void GetTextInRange(const uint32 BeginIndex, const uint32 Length, FString& OutString) override;
+		virtual void SetTextInRange(const uint32 BeginIndex, const uint32 Length, const FString& InString) override;
+		virtual int32 GetCharacterIndexFromPoint(const FVector2D& Point) override;
+		virtual bool GetTextBounds(const uint32 BeginIndex, const uint32 Length, FVector2D& Position, FVector2D& Size) override;
+		virtual void GetScreenBounds(FVector2D& Position, FVector2D& Size) override;
+		virtual TSharedPtr<FGenericWindow> GetWindow() override;
+		virtual void BeginComposition() override;
+		virtual void UpdateCompositionRange(const int32 InBeginIndex, const uint32 InLength) override;
+		virtual void EndComposition() override;
+
+	private:
+		FTextInputMethodContext(FSlateEditableTextLayout& InOwnerLayout);
+		FSlateEditableTextLayout* OwnerLayout;
+
+		FGeometry CachedGeometry;
+		bool bIsComposing;
+		int32 CompositionBeginIndex;
+		uint32 CompositionLength;
+	};
+
+private:
+	/** Pointer to the interface for our owner widget */
+	ISlateEditableTextWidget* OwnerWidget;
+
+	/** The marshaller used to get/set the BoundText text to/from the text layout. */
+	TSharedPtr<ITextLayoutMarshaller> Marshaller;
+
+	/** The marshaller used to get/set the HintText text to/from the text layout. */
+	TSharedPtr<ITextLayoutMarshaller> HintMarshaller;
+
+	/** In control of the layout and wrapping of the BoundText */
+	TSharedPtr<FSlateTextLayout> TextLayout;
+
+	/** In control of the layout and wrapping of the HintText */
+	TUniquePtr<FTextBlockLayout> HintTextLayout;
+
+	/** Default style used by the TextLayout */
+	FTextBlockStyle TextStyle;
+
+	/** Style used to draw the hint text (only valid when HintTextLayout is set) */
+	FTextBlockStyle HintTextStyle;
+
+	/** The text displayed in this text block */
+	TAttribute<FText> BoundText;
+
+	/** The state of BoundText last Tick() (used to allow updates when the text is changed) */
+	FTextSnapshot BoundTextLastTick;
+
+	/** Was the editable text showing a password last Tick() (allows a forcible text layout update when changing state) */
+	bool bWasPasswordLastTick;
+
+	/** The text that appears when there is no text in the text box */
+	TAttribute<FText> HintText;
+
+	/** Whether text wraps onto a new line when it's length exceeds this width; if this value is zero or negative, no wrapping occurs. */
+	TAttribute<float> WrapTextAt;
+
+	/** True if we're wrapping text automatically based on the computed horizontal space for this widget */
+	TAttribute<bool> AutoWrapText;
+
+	/** The amount of blank space left around the edges of text area */
+	TAttribute<FMargin> Margin;
+
+	/** How the text should be aligned with the margin */
+	TAttribute<ETextJustify::Type> Justification;
+
+	/** The amount to scale each lines height by */
+	TAttribute<float> LineHeightPercentage;
+
+	/** The information used to help identify who owns this text layout in the case of an error */
+	TAttribute<FString> DebugSourceInfo;
+
+	/** Virtual keyboard handler for this text layout */
+	TSharedPtr<FVirtualKeyboardEntry> VirtualKeyboardEntry;
+
+	/** IME context for this text layout */
+	TSharedPtr<FTextInputMethodContext> TextInputMethodContext;
+
+	/** Notification interface object for IMEs */
+	TSharedPtr<ITextInputMethodChangeNotifier> TextInputMethodChangeNotifier;
+
+	/** Layout highlighter used to draw the cursor */
+	TSharedPtr<SlateEditableTextTypes::FCursorLineHighlighter> CursorLineHighlighter;
+
+	/** Layout highlighter used to draw an active text composition */
+	TSharedPtr<SlateEditableTextTypes::FTextCompositionHighlighter> TextCompositionHighlighter;
+
+	/** Run renderer used to draw the active text selection */
+	TSharedPtr<SlateEditableTextTypes::FTextSelectionRunRenderer> TextSelectionRunRenderer;
+
+	/** The scroll offset (in unscaled Slate units) for this text */
+	FVector2D ScrollOffset;
+
+	/** If set, the pending data containing a position that should be scrolled into view */
+	TOptional<SlateEditableTextTypes::FScrollInfo> PositionToScrollIntoView;
+
+	/** That start of the selection when there is a selection. The end is implicitly wherever the cursor happens to be. */
+	TOptional<FTextLocation> SelectionStart;
+
+	/** The user probably wants the cursor where they last explicitly positioned it horizontally. */
+	float PreferredCursorScreenOffsetInLine;
+
+	/** Current cursor data */
+	SlateEditableTextTypes::FCursorInfo CursorInfo;
+
+	/** Undo states */
+	TArray<SlateEditableTextTypes::FUndoState> UndoStates;
+
+	/** Current undo state level that we've rolled back to, or INDEX_NONE if we haven't undone. Used for 'Redo'. */
+	int32 CurrentUndoLevel;
+
+	/** Undo state that will be pushed if text is actually changed between calls to BeginEditTransation() and EndEditTransaction() */
+	TOptional<SlateEditableTextTypes::FUndoState> StateBeforeChangingText;
+
+	/** Original text undo state */
+	SlateEditableTextTypes::FUndoState OriginalText;
+
+	/** True if we're currently selecting text by dragging the mouse cursor with the left button held down */
+	bool bIsDragSelecting;
+
+	/** True if the last mouse down caused us to receive keyboard focus */
+	bool bWasFocusedByLastMouseDown;
+
+	/** True if characters were selected by dragging since the last keyboard focus.  Used for text selection. */
+	bool bHasDragSelectedSinceFocused;
+
+	/** Whether the text has been changed by a virtual keyboard */
+	bool bTextChangedByVirtualKeyboard;
+
+	/** Whether the text has been committed by a virtual keyboard */
+	bool bTextCommittedByVirtualKeyboard;
+
+	/** How the text was committed by the virtual keyboard */
+	ETextCommit::Type VirtualKeyboardTextCommitType;
+
+	/** The last known size of the widget from the previous OnPaint, used to recalculate wrapping */
+	FVector2D CachedSize;
+
+	/** A list commands to execute if a user presses the corresponding key-binding in the text box */
+	TSharedPtr<FUICommandList> UICommandList;
+
+	/** Information about any active context menu widgets */
+	FActiveTextEditContextMenu ActiveContextMenu;
+};
diff --git a/Engine/Source/Runtime/Slate/Public/Widgets/Text/SlateEditableTextTypes.h b/Engine/Source/Runtime/Slate/Public/Widgets/Text/SlateEditableTextTypes.h
new file mode 100644
--- /dev/null
+++ b/Engine/Source/Runtime/Slate/Public/Widgets/Text/SlateEditableTextTypes.h
@@ -0,0 +1,177 @@
+// Copyright 1998-2016 Epic Games, Inc. All Rights Reserved.
+
+#pragma once
+
+namespace SlateEditableTextTypes
+{
+
+enum class ECursorAlignment : uint8
+{
+	/** Visually align the cursor to the left of the character its placed at, and insert text before the character */
+	Left,
+
+	/** Visually align the cursor to the right of the character its placed at, and insert text after the character */
+	Right,
+};
+
+/** Store the information about the current cursor position */
+class SLATE_API FCursorInfo
+{
+public:
+	FCursorInfo()
+		: CursorPosition()
+		, CursorAlignment(ECursorAlignment::Left)
+		, CursorTextDirection(TextBiDi::ETextDirection::LeftToRight)
+		, LastCursorInteractionTime(0)
+	{
+	}
+
+	/** Get the literal position of the cursor (note: this may not be the correct place to insert text to, use GetCursorInteractionLocation for that) */
+	FORCEINLINE FTextLocation GetCursorLocation() const
+	{
+		return CursorPosition;
+	}
+
+	/** Get the alignment of the cursor */
+	FORCEINLINE ECursorAlignment GetCursorAlignment() const
+	{
+		return CursorAlignment;
+	}
+
+	/** Get the direction of the text under the cursor */
+	FORCEINLINE TextBiDi::ETextDirection GetCursorTextDirection() const
+	{
+		return CursorTextDirection;
+	}
+
+	/** Get the interaction position of the cursor (where to insert, delete, etc, text from/to) */
+	FORCEINLINE FTextLocation GetCursorInteractionLocation() const
+	{
+		// If the cursor is right aligned, we treat it as if it were one character further along for insert/delete purposes
+		return FTextLocation(CursorPosition, (CursorAlignment == ECursorAlignment::Right) ? 1 : 0);
+	}
+
+	FORCEINLINE double GetLastCursorInteractionTime() const
+	{
+		return LastCursorInteractionTime;
+	}
+
+	/** Set the position of the cursor, and then work out the correct alignment based on the current text layout */
+	void SetCursorLocationAndCalculateAlignment(const FTextLayout& InTextLayout, const FTextLocation& InCursorPosition);
+
+	/** Set the literal position and alignment of the cursor */
+	void SetCursorLocationAndAlignment(const FTextLayout& InTextLayout, const FTextLocation& InCursorPosition, const ECursorAlignment InCursorAlignment);
+
+	/** Create an undo for this cursor data */
+	FCursorInfo CreateUndo() const;
+
+	/** Restore this cursor data from an undo */
+	void RestoreFromUndo(const FCursorInfo& UndoData);
+
+private:
+	/** Current cursor position; there is always one. */
+	FTextLocation CursorPosition;
+
+	/** Cursor alignment (horizontal) within its position. This affects the rendering and behavior of the cursor when adding new characters */
+	ECursorAlignment CursorAlignment;
+
+	/** The direction of the text under the cursor */
+	TextBiDi::ETextDirection CursorTextDirection;
+
+	/** Last time the user did anything with the cursor.*/
+	double LastCursorInteractionTime;
+};
+
+/** Stores a single undo level for editable text */
+struct FUndoState
+{
+	/** Text */
+	FText Text;
+
+	/** Selection state */
+	TOptional<FTextLocation> SelectionStart;
+
+	/** Cursor data */
+	FCursorInfo CursorInfo;
+};
+
+/** Information needed to be able to scroll to a given point */
+struct FScrollInfo
+{
+	FScrollInfo()
+		: Position()
+		, Alignment(ECursorAlignment::Left)
+	{
+	}
+
+	FScrollInfo(const FTextLocation InPosition, const ECursorAlignment InAlignment)
+		: Position(InPosition)
+		, Alignment(InAlignment)
+	{
+	}
+
+	/** The location in the document to scroll to (in line model space) */
+	FTextLocation Position;
+
+	/** The alignment at the given position. This may affect which line view the character maps to when converted from line model space */
+	ECursorAlignment Alignment;
+};
+
+/** Run highlighter used to draw the cursor */
+class SLATE_API FCursorLineHighlighter : public ISlateLineHighlighter
+{
+public:
+	static TSharedRef< FCursorLineHighlighter > Create(const FCursorInfo* InCursorInfo);
+
+	void SetCursorBrush(const TAttribute<const FSlateBrush*>& InCursorBrush);
+
+	virtual int32 OnPaint(const FPaintArgs& Args, const FTextLayout::FLineView& Line, const float OffsetX, const float Width, const FTextBlockStyle& DefaultStyle, const FGeometry& AllottedGeometry, const FSlateRect& MyClippingRect, FSlateWindowElementList& OutDrawElements, int32 LayerId, const FWidgetStyle& InWidgetStyle, bool bParentEnabled) const override;
+
+protected:
+	FCursorLineHighlighter(const FCursorInfo* InCursorInfo);
+
+	/** Cursor data that this highlighter is tracking */
+	const FCursorInfo* CursorInfo;
+
+	/** Brush used to draw the cursor */
+	TAttribute<const FSlateBrush*> CursorBrush;
+};
+
+/** Run highlighter used to draw the composition range */
+class SLATE_API FTextCompositionHighlighter : public ISlateLineHighlighter
+{
+public:
+	static TSharedRef< FTextCompositionHighlighter > Create();
+
+	void SetCompositionBrush(const TAttribute<const FSlateBrush*>& InCompositionBrush);
+
+	virtual int32 OnPaint(const FPaintArgs& Args, const FTextLayout::FLineView& Line, const float OffsetX, const float Width, const FTextBlockStyle& DefaultStyle, const FGeometry& AllottedGeometry, const FSlateRect& MyClippingRect, FSlateWindowElementList& OutDrawElements, int32 LayerId, const FWidgetStyle& InWidgetStyle, bool bParentEnabled) const override;
+
+protected:
+	FTextCompositionHighlighter();
+
+	/** Brush used to draw the composition highlight */
+	TAttribute<const FSlateBrush*> CompositionBrush;
+};
+
+/** Run renderer used to draw selection ranges */
+class SLATE_API FTextSelectionRunRenderer : public ISlateRunRenderer
+{
+public:
+	static TSharedRef< FTextSelectionRunRenderer > Create();
+
+	virtual int32 OnPaint(const FPaintArgs& Args, const FTextLayout::FLineView& Line, const TSharedRef< ISlateRun >& Run, const TSharedRef< ILayoutBlock >& Block, const FTextBlockStyle& DefaultStyle, const FGeometry& AllottedGeometry, const FSlateRect& MyClippingRect, FSlateWindowElementList& OutDrawElements, int32 LayerId, const FWidgetStyle& InWidgetStyle, bool bParentEnabled) const override;
+
+	void SetHasKeyboardFocus(const bool bInHasKeyboardFocus)
+	{
+		bHasKeyboardFocus = bInHasKeyboardFocus;
+	}
+
+protected:
+	FTextSelectionRunRenderer();
+
+	/** true if the parent widget has keyboard focus, false otherwise */
+	bool bHasKeyboardFocus;
+};
+
+} // namespace SlateEditableTextTypes
diff --git a/Engine/Source/Runtime/Slate/Public/Widgets/Text/TextBlockLayout.h b/Engine/Source/Runtime/Slate/Public/Widgets/Text/TextBlockLayout.h
--- a/Engine/Source/Runtime/Slate/Public/Widgets/Text/TextBlockLayout.h
+++ b/Engine/Source/Runtime/Slate/Public/Widgets/Text/TextBlockLayout.h
@@ -6,10 +6,8 @@ class FSlateTextLayout;
 class ITextLayoutMarshaller;
 class ISlateRunRenderer;
 
-#if WITH_FANCY_TEXT
-
 /** Class to handle the cached layout of STextBlock/SRichTextBlock by proxying around a FTextLayout */
-class FTextBlockLayout
+class SLATE_API FTextBlockLayout
 {
 public:
 	struct FWidgetArgs
@@ -42,7 +40,10 @@ public:
 		const TAttribute<ETextJustify::Type>& Justification;
 	};
 
-	static TSharedRef<FTextBlockLayout> Create(FTextBlockStyle InDefaultTextStyle, const TOptional<ETextShapingMethod> InTextShapingMethod, const TOptional<ETextFlowDirection> InTextFlowDirection, TSharedRef<ITextLayoutMarshaller> InMarshaller, TSharedPtr<IBreakIterator> InLineBreakPolicy);
+	/**
+	 * Constructor
+	 */
+	FTextBlockLayout(FTextBlockStyle InDefaultTextStyle, const TOptional<ETextShapingMethod> InTextShapingMethod, const TOptional<ETextFlowDirection> InTextFlowDirection, TSharedRef<ITextLayoutMarshaller> InMarshaller, TSharedPtr<IBreakIterator> InLineBreakPolicy);
 
 	/**
 	 * Get the computed desired size for this layout, updating the internal cache as required
@@ -72,7 +73,7 @@ public:
 	void SetTextShapingMethod(const TOptional<ETextShapingMethod>& InTextShapingMethod);
 
 	/**
-	 * See the text flow direction that the internal text layout should use
+	 * Set the text flow direction that the internal text layout should use
 	 */
 	void SetTextFlowDirection(const TOptional<ETextFlowDirection>& InTextFlowDirection);
 
@@ -107,8 +108,6 @@ private:
 	/** Calculate the wrapping width based on the given fixed wrap width, and whether we're auto-wrapping */
 	float CalculateWrappingWidth(const FWidgetArgs& InWidgetArgs) const;
 
-	FTextBlockLayout(FTextBlockStyle InDefaultTextStyle, const TOptional<ETextShapingMethod> InTextShapingMethod, const TOptional<ETextFlowDirection> InTextFlowDirection, TSharedRef<ITextLayoutMarshaller> InMarshaller, TSharedPtr<IBreakIterator> InLineBreakPolicy);
-
 	/** In control of the layout and wrapping of the text */
 	TSharedPtr<FSlateTextLayout> TextLayout;
 
@@ -127,5 +126,3 @@ private:
 	/** The state of the highlight text the last time it was updated (used to allow updates when the text is changed) */
 	FTextSnapshot HighlightTextLastUpdate;
 };
-
-#endif //WITH_FANCY_TEXT
diff --git a/Engine/Source/Runtime/SlateCore/Private/Fonts/FontCache.cpp b/Engine/Source/Runtime/SlateCore/Private/Fonts/FontCache.cpp
--- a/Engine/Source/Runtime/SlateCore/Private/Fonts/FontCache.cpp
+++ b/Engine/Source/Runtime/SlateCore/Private/Fonts/FontCache.cpp
@@ -431,6 +431,12 @@ FShapedGlyphSequencePtr FShapedGlyphSequence::GetSubSequence(const int32 InStart
 
 bool FShapedGlyphSequence::EnumerateGlyphsInClusterRange(const int32 InStartIndex, const int32 InEndIndex, const FForEachShapedGlyphEntryCallback& InGlyphCallback) const
 {
+	if (InStartIndex == InEndIndex)
+	{
+		// Nothing to enumerate, but don't say we failed
+		return true;
+	}
+
 	// The given range is exclusive, but we use an inclusive range when performing all the bounds testing below (as it makes things simpler)
 	const FClusterIndexToGlyphData* StartClusterIndexToGlyphData = ClusterIndicesToGlyphData.GetGlyphData(InStartIndex);
 	const FClusterIndexToGlyphData* EndClusterIndexToGlyphData = ClusterIndicesToGlyphData.GetGlyphData(InEndIndex - 1);
@@ -540,6 +546,12 @@ bool FShapedGlyphSequence::EnumerateGlyphsInClusterRange(const int32 InStartInde
 
 bool FShapedGlyphSequence::EnumerateGlyphsInClusterRange(const int32 InStartIndex, const int32 InEndIndex, const FForEachShapedGlyphClusterBlockCallback& InBeginClusterBlockCallback, const FForEachShapedGlyphClusterBlockCallback& InEndClusterBlockCallback, const FForEachShapedGlyphEntryCallback& InGlyphCallback) const
 {
+	if (InStartIndex == InEndIndex)
+	{
+		// Nothing to enumerate, but don't say we failed
+		return true;
+	}
+
 	// The given range is exclusive, but we use an inclusive range when performing all the bounds testing below (as it makes things simpler)
 	const FClusterIndexToGlyphData* StartClusterIndexToGlyphData = ClusterIndicesToGlyphData.GetGlyphData(InStartIndex);
 	const FClusterIndexToGlyphData* EndClusterIndexToGlyphData = ClusterIndicesToGlyphData.GetGlyphData(InEndIndex - 1);
diff --git a/Engine/Source/Runtime/SlateCore/Private/Widgets/SWidget.cpp b/Engine/Source/Runtime/SlateCore/Private/Widgets/SWidget.cpp
--- a/Engine/Source/Runtime/SlateCore/Private/Widgets/SWidget.cpp
+++ b/Engine/Source/Runtime/SlateCore/Private/Widgets/SWidget.cpp
@@ -448,6 +448,11 @@ bool SWidget::HasFocusedDescendants() const
 	return FSlateApplicationBase::Get().HasFocusedDescendants(SharedThis(this));
 }
 
+bool SWidget::HasAnyUserFocusOrFocusedDescendants() const
+{
+	return HasAnyUserFocus().IsSet() || HasFocusedDescendants();
+}
+
 const FSlateBrush* SWidget::GetFocusBrush() const
 {
 	return FCoreStyle::Get().GetBrush("FocusRectangle");
diff --git a/Engine/Source/Runtime/SlateCore/Private/Widgets/SWindow.cpp b/Engine/Source/Runtime/SlateCore/Private/Widgets/SWindow.cpp
--- a/Engine/Source/Runtime/SlateCore/Private/Widgets/SWindow.cpp
+++ b/Engine/Source/Runtime/SlateCore/Private/Widgets/SWindow.cpp
@@ -1348,50 +1348,46 @@ FCursorReply SWindow::OnCursorQuery( const FGeometry& MyGeometry, const FPointer
 
 bool SWindow::OnIsActiveChanged( const FWindowActivateEvent& ActivateEvent )
 {
-	const bool bWasDeactivated = ( ActivateEvent.GetActivationType() == FWindowActivateEvent::EA_Deactivate );
-	if ( bWasDeactivated )
+	const bool bWasDeactivated = ActivateEvent.GetActivationType() == FWindowActivateEvent::EA_Deactivate;
+	if (bWasDeactivated)
 	{
 		OnWindowDeactivated.ExecuteIfBound();	// deprecated
 		WindowDeactivatedEvent.Broadcast();
 
 		const EWindowMode::Type WindowMode = GetWindowMode();
 		// If the window is not fullscreen, we do not want to automatically recapture the mouse unless an external UI such as Steam is open. Fullscreen windows we do.
-		if( WindowMode != EWindowMode::Fullscreen && WidgetToFocusOnActivate.IsValid() && WidgetToFocusOnActivate.Pin()->HasMouseCapture() && !FSlateApplicationBase::Get().IsExternalUIOpened())
+		if (WindowMode != EWindowMode::Fullscreen && WidgetToFocusOnActivate.IsValid() && WidgetToFocusOnActivate.Pin()->HasMouseCapture() && !FSlateApplicationBase::Get().IsExternalUIOpened())
 		{
-			//For a windowed application with an OS border, if the user is giving focus back to the application by clicking on the close/(X) button, then we must clear 
-			//the weak pointer to WidgetToFocus--so that the application's main viewport does not steal focus immediately (thus canceling the close attempt).
-			
-			//This change introduces a different bug where slate context is lost when closing popup menus.  However, this issue is negated by a 
-			//change to FMenuStack::PushMenu, where we ReleaseMouseCapture when immediately shifting focus.
 			WidgetToFocusOnActivate.Reset();
 		}
 	}
 	else
 	{
-		if (SupportsKeyboardFocus() && ActivateEvent.GetActivationType() == FWindowActivateEvent::EA_Activate)
+		if (ActivateEvent.GetActivationType() == FWindowActivateEvent::EA_Activate)
 		{
 			TArray< TSharedRef<SWindow> > JustThisWindow;
-			JustThisWindow.Add( SharedThis(this) );
-			
+			JustThisWindow.Add(SharedThis(this));
+
 			// If we're becoming active and we were set to restore keyboard focus to a specific widget
 			// after reactivating, then do so now
-			TSharedPtr< SWidget > PinnedWidgetToFocus( WidgetToFocusOnActivate.Pin() );
-			
+			TSharedPtr< SWidget > PinnedWidgetToFocus(WidgetToFocusOnActivate.Pin());
 			if (PinnedWidgetToFocus.IsValid())
 			{
 				FWidgetPath WidgetToFocusPath;
-				if( FSlateWindowHelper::FindPathToWidget( JustThisWindow, PinnedWidgetToFocus.ToSharedRef(), WidgetToFocusPath ) )
+				if (FSlateWindowHelper::FindPathToWidget(JustThisWindow, PinnedWidgetToFocus.ToSharedRef(), WidgetToFocusPath))
 				{
-					FSlateApplicationBase::Get().SetAllUserFocus( WidgetToFocusPath, EFocusCause::SetDirectly );
+					FSlateApplicationBase::Get().SetAllUserFocus(WidgetToFocusPath, EFocusCause::SetDirectly);
 				}
 			}
-			else
+
+			// If we didn't have a specified widget to focus (above)
+			// We'll make sure all the users focus this window, however if they are already focusing something in the window we leave them be.
+			else if (SupportsKeyboardFocus())
 			{
 				FWidgetPath WindowWidgetPath;
-				if( FSlateWindowHelper::FindPathToWidget( JustThisWindow, AsShared(), WindowWidgetPath ) )
+				if (FSlateWindowHelper::FindPathToWidget(JustThisWindow, AsShared(), WindowWidgetPath))
 				{
-					FWeakWidgetPath WeakWindowPath(WindowWidgetPath);
-					FSlateApplicationBase::Get().SetAllUserFocus( WeakWindowPath.ToNextFocusedPath(EUINavigation::Next), EFocusCause::SetDirectly );
+					FSlateApplicationBase::Get().SetAllUserFocusAllowingDescendantFocus(WindowWidgetPath, EFocusCause::SetDirectly);
 				}
 			}
 		}
diff --git a/Engine/Source/Runtime/SlateCore/Public/Application/SlateApplicationBase.h b/Engine/Source/Runtime/SlateCore/Public/Application/SlateApplicationBase.h
--- a/Engine/Source/Runtime/SlateCore/Public/Application/SlateApplicationBase.h
+++ b/Engine/Source/Runtime/SlateCore/Public/Application/SlateApplicationBase.h
@@ -61,6 +61,13 @@ public:
 
 	FSlateApplicationBase();
 
+	/**
+	 * Whether the application is active.
+	 *
+	 * @return application active or not
+	 */
+	virtual bool IsActive() const = 0;
+
 	/**
 	 * Gets the renderer being used to draw this application.
 	 *
@@ -347,6 +354,14 @@ public:
 	 */
 	virtual void SetAllUserFocus(const FWidgetPath& InFocusPath, const EFocusCause InCause) = 0;
 
+	/**
+	 * Sets the focus for all users to the specified widget unless that user is focused on a descendant.  The widget must be allowed to receive focus.
+	 *
+	 * @param InWidget WidgetPath to the Widget to being focused.
+	 * @param InCause The reason that focus is changing.
+	 */
+	virtual void SetAllUserFocusAllowingDescendantFocus(const FWidgetPath& InFocusPath, const EFocusCause InCause) = 0;
+
 	/**
 	 * Gets a delegate that is invoked when a global invalidate of all widgets should occur
 	 */
diff --git a/Engine/Source/Runtime/SlateCore/Public/Fonts/FontCache.h b/Engine/Source/Runtime/SlateCore/Public/Fonts/FontCache.h
--- a/Engine/Source/Runtime/SlateCore/Public/Fonts/FontCache.h
+++ b/Engine/Source/Runtime/SlateCore/Public/Fonts/FontCache.h
@@ -4,6 +4,7 @@
 
 #include "TextureAtlas.h"
 #include "UniquePtr.h"
+#include "ShapedTextFwd.h"
 #include "FontCache.generated.h"
 
 class FFreeTypeLibrary;
@@ -211,11 +212,6 @@ struct FShapedGlyphClusterBlock
 	}
 };
 
-/** Immutable pointer/reference to a shaped text sequence. This is what gets returned by the font cache, and is used throughout the rest of the rendering pipeline to avoid copying data */
-class FShapedGlyphSequence;
-typedef TSharedPtr<const FShapedGlyphSequence> FShapedGlyphSequencePtr;
-typedef TSharedRef<const FShapedGlyphSequence> FShapedGlyphSequenceRef;
-
 /** Information for rendering a shaped text sequence */
 class SLATECORE_API FShapedGlyphSequence
 {
diff --git a/Engine/Source/Runtime/SlateCore/Public/Rendering/RenderingCommon.h b/Engine/Source/Runtime/SlateCore/Public/Rendering/RenderingCommon.h
--- a/Engine/Source/Runtime/SlateCore/Public/Rendering/RenderingCommon.h
+++ b/Engine/Source/Runtime/SlateCore/Public/Rendering/RenderingCommon.h
@@ -524,6 +524,21 @@ public:
 	virtual void OnViewportClosed()
 	{
 	}
+
+	/**
+	 * Called when the viewports top level window is being Activated
+	 */
+	virtual FReply OnViewportActivated(const FWindowActivateEvent& InActivateEvent)
+	{
+		return FReply::Unhandled();
+	}
+
+	/**
+	 * Called when the viewports top level window is being Deactivated
+	 */
+	virtual void OnViewportDeactivated(const FWindowActivateEvent& InActivateEvent)
+	{
+	}
 };
 
 /**
diff --git a/Engine/Source/Runtime/SlateCore/Public/Widgets/SWidget.h b/Engine/Source/Runtime/SlateCore/Public/Widgets/SWidget.h
--- a/Engine/Source/Runtime/SlateCore/Public/Widgets/SWidget.h
+++ b/Engine/Source/Runtime/SlateCore/Public/Widgets/SWidget.h
@@ -644,6 +644,11 @@ public:
 	 */
 	bool HasFocusedDescendants() const;
 
+	/**
+	 * @return whether or not any users have this widget focused, or any descendant focused.
+	 */
+	bool HasAnyUserFocusOrFocusedDescendants() const;
+
 	/**
 	 * Checks to see if this widget is the current mouse captor
 	 *
diff --git a/Engine/Source/Runtime/UMG/Private/Components/EditableText.cpp b/Engine/Source/Runtime/UMG/Private/Components/EditableText.cpp
--- a/Engine/Source/Runtime/UMG/Private/Components/EditableText.cpp
+++ b/Engine/Source/Runtime/UMG/Private/Components/EditableText.cpp
@@ -71,6 +71,8 @@ void UEditableText::SynchronizeProperties()
 	MyEditableText->SetIsPassword(IsPassword);
 	MyEditableText->SetAllowContextMenu(AllowContextMenu);
 	// TODO UMG Complete making all properties settable on SEditableText
+
+	ShapedTextOptions.SynchronizeShapedTextProperties(*MyEditableText);
 }
 
 FText UEditableText::GetText() const
diff --git a/Engine/Source/Runtime/UMG/Private/Components/EditableTextBox.cpp b/Engine/Source/Runtime/UMG/Private/Components/EditableTextBox.cpp
--- a/Engine/Source/Runtime/UMG/Private/Components/EditableTextBox.cpp
+++ b/Engine/Source/Runtime/UMG/Private/Components/EditableTextBox.cpp
@@ -81,6 +81,8 @@ void UEditableTextBox::SynchronizeProperties()
 	MyEditableTextBlock->SetClearKeyboardFocusOnCommit(ClearKeyboardFocusOnCommit);
 	MyEditableTextBlock->SetSelectAllTextOnCommit(SelectAllTextOnCommit);
 	MyEditableTextBlock->SetAllowContextMenu(AllowContextMenu);
+
+	ShapedTextOptions.SynchronizeShapedTextProperties(*MyEditableTextBlock);
 }
 
 FText UEditableTextBox::GetText() const
diff --git a/Engine/Source/Runtime/UMG/Private/Components/MultiLineEditableText.cpp b/Engine/Source/Runtime/UMG/Private/Components/MultiLineEditableText.cpp
--- a/Engine/Source/Runtime/UMG/Private/Components/MultiLineEditableText.cpp
+++ b/Engine/Source/Runtime/UMG/Private/Components/MultiLineEditableText.cpp
@@ -19,6 +19,8 @@ UMultiLineEditableText::UMultiLineEditableText(const FObjectInitializer& ObjectI
 	{
 		static ConstructorHelpers::FObjectFinder<UFont> RobotoFontObj(TEXT("/Engine/EngineFonts/Roboto"));
 		Font_DEPRECATED = FSlateFontInfo(RobotoFontObj.Object, 12, FName("Bold"));
+
+		WidgetStyle.SetFont(Font_DEPRECATED);
 	}
 }
 
@@ -57,6 +59,7 @@ void UMultiLineEditableText::SynchronizeProperties()
 
 	TAttribute<FText> HintTextBinding = OPTIONAL_BINDING(FText, HintText);
 
+	MyMultiLineEditableText->SetTextStyle(&WidgetStyle);
 	MyMultiLineEditableText->SetText(Text);
 	MyMultiLineEditableText->SetHintText(HintTextBinding);
 	MyMultiLineEditableText->SetAllowContextMenu(AllowContextMenu);
diff --git a/Engine/Source/Runtime/UMG/Private/Components/MultiLineEditableTextBox.cpp b/Engine/Source/Runtime/UMG/Private/Components/MultiLineEditableTextBox.cpp
--- a/Engine/Source/Runtime/UMG/Private/Components/MultiLineEditableTextBox.cpp
+++ b/Engine/Source/Runtime/UMG/Private/Components/MultiLineEditableTextBox.cpp
@@ -15,18 +15,22 @@ UMultiLineEditableTextBox::UMultiLineEditableTextBox(const FObjectInitializer& O
 	BackgroundColor_DEPRECATED = FLinearColor::White;
 	ReadOnlyForegroundColor_DEPRECATED = FLinearColor::Black;
 
+	SMultiLineEditableTextBox::FArguments Defaults;
+	WidgetStyle = *Defaults._Style;
+	TextStyle = *Defaults._TextStyle;
+	AllowContextMenu = Defaults._AllowContextMenu.Get();
+	AutoWrapText = true;
+
 	if (!IsRunningDedicatedServer())
 	{
 		static ConstructorHelpers::FObjectFinder<UFont> RobotoFontObj(TEXT("/Engine/EngineFonts/Roboto"));
 		Font_DEPRECATED = FSlateFontInfo(RobotoFontObj.Object, 12, FName("Bold"));
-	}
 
-	AutoWrapText = true;
-
-	SMultiLineEditableTextBox::FArguments Defaults;
-	WidgetStyle = *Defaults._Style;
-	TextStyle = *Defaults._TextStyle;
-	AllowContextMenu = Defaults._AllowContextMenu.Get();
+		WidgetStyle.SetFont(Font_DEPRECATED);
+		WidgetStyle.SetForegroundColor(ForegroundColor_DEPRECATED);
+		WidgetStyle.SetBackgroundColor(BackgroundColor_DEPRECATED);
+		WidgetStyle.SetReadOnlyForegroundColor(ReadOnlyForegroundColor_DEPRECATED);
+	}
 }
 
 void UMultiLineEditableTextBox::ReleaseSlateResources(bool bReleaseChildren)
diff --git a/Engine/Source/Runtime/UMG/Private/Components/SpinBox.cpp b/Engine/Source/Runtime/UMG/Private/Components/SpinBox.cpp
--- a/Engine/Source/Runtime/UMG/Private/Components/SpinBox.cpp
+++ b/Engine/Source/Runtime/UMG/Private/Components/SpinBox.cpp
@@ -24,7 +24,7 @@ USpinBox::USpinBox(const FObjectInitializer& ObjectInitializer)
 	MaxValue = Defaults._MaxValue.Get().Get(0.0f);
 	MinSliderValue = Defaults._MinSliderValue.Get().Get(0.0f);
 	MaxSliderValue = Defaults._MaxSliderValue.Get().Get(0.0f);
-	Delta = Defaults._Delta;
+	Delta = Defaults._Delta.Get();
 	SliderExponent = Defaults._SliderExponent.Get();
 	MinDesiredWidth = Defaults._MinDesiredWidth.Get();
 	ClearKeyboardFocusOnCommit = Defaults._ClearKeyboardFocusOnCommit.Get();
diff --git a/Engine/Source/Runtime/UMG/Public/Components/CircularThrobber.h b/Engine/Source/Runtime/UMG/Public/Components/CircularThrobber.h
--- a/Engine/Source/Runtime/UMG/Public/Components/CircularThrobber.h
+++ b/Engine/Source/Runtime/UMG/Public/Components/CircularThrobber.h
@@ -75,7 +75,7 @@ protected:
 
 private:
 
-	UPROPERTY(Transient)
+	UPROPERTY(Transient, EditAnywhere, Category="Appearance", meta=(InlineEditConditionToggle))
 	bool bEnableRadius;
 
 	/** The CircularThrobber widget managed by this object. */
diff --git a/Engine/Source/Runtime/UMG/Public/Components/EditableText.h b/Engine/Source/Runtime/UMG/Public/Components/EditableText.h
--- a/Engine/Source/Runtime/UMG/Public/Components/EditableText.h
+++ b/Engine/Source/Runtime/UMG/Public/Components/EditableText.h
@@ -3,7 +3,7 @@
 #pragma once
 
 #include "EditableTextWidgetStyle.h"
-
+#include "TextWidgetTypes.h"
 #include "EditableText.generated.h"
 
 /**
@@ -99,12 +99,16 @@ public:
 	bool SelectAllTextOnCommit;
 
 	/** Whether the context menu can be opened */
-	UPROPERTY(EditAnywhere, Category = Behavior, AdvancedDisplay)
+	UPROPERTY(EditAnywhere, Category=Behavior, AdvancedDisplay)
 	bool AllowContextMenu;
 
 	/** If we're on a platform that requires a virtual keyboard, what kind of keyboard should this widget use? */
-	UPROPERTY( EditAnywhere, Category = Behavior, AdvancedDisplay )
-	TEnumAsByte< EVirtualKeyboardType::Type > KeyboardType;
+	UPROPERTY(EditAnywhere, Category=Behavior, AdvancedDisplay)
+	TEnumAsByte<EVirtualKeyboardType::Type> KeyboardType;
+
+	/** Controls how the text within this widget should be shaped. */
+	UPROPERTY(EditAnywhere, BlueprintReadOnly, Category=Localization, AdvancedDisplay, meta=(ShowOnlyInnerProperties))
+	FShapedTextOptions ShapedTextOptions;
 
 public:
 
diff --git a/Engine/Source/Runtime/UMG/Public/Components/EditableTextBox.h b/Engine/Source/Runtime/UMG/Public/Components/EditableTextBox.h
--- a/Engine/Source/Runtime/UMG/Public/Components/EditableTextBox.h
+++ b/Engine/Source/Runtime/UMG/Public/Components/EditableTextBox.h
@@ -2,6 +2,7 @@
 
 #pragma once
 
+#include "TextWidgetTypes.h"
 #include "EditableTextBox.generated.h"
 
 /**
@@ -100,12 +101,16 @@ public:
 	bool SelectAllTextOnCommit;
 
 	/** Whether the context menu can be opened */
-	UPROPERTY(EditAnywhere, Category = Behavior, AdvancedDisplay)
+	UPROPERTY(EditAnywhere, Category=Behavior, AdvancedDisplay)
 	bool AllowContextMenu;
 
 	/** If we're on a platform that requires a virtual keyboard, what kind of keyboard should this widget use? */
-	UPROPERTY(EditAnywhere, Category = Behavior, AdvancedDisplay)
-	TEnumAsByte< EVirtualKeyboardType::Type > KeyboardType;
+	UPROPERTY(EditAnywhere, Category=Behavior, AdvancedDisplay)
+	TEnumAsByte<EVirtualKeyboardType::Type> KeyboardType;
+
+	/** Controls how the text within this widget should be shaped. */
+	UPROPERTY(EditAnywhere, BlueprintReadOnly, Category=Localization, AdvancedDisplay, meta=(ShowOnlyInnerProperties))
+	FShapedTextOptions ShapedTextOptions;
 
 public:
 
diff --git a/Engine/Source/Runtime/UMG/Public/Components/MultiLineEditableText.h b/Engine/Source/Runtime/UMG/Public/Components/MultiLineEditableText.h
--- a/Engine/Source/Runtime/UMG/Public/Components/MultiLineEditableText.h
+++ b/Engine/Source/Runtime/UMG/Public/Components/MultiLineEditableText.h
@@ -2,6 +2,7 @@
 
 #pragma once
 
+#include "TextWidgetTypes.h"
 #include "MultiLineEditableText.generated.h"
 
 /**
@@ -32,7 +33,7 @@ public:
 public:
 
 	/** The style */
-	UPROPERTY(EditAnywhere, BlueprintReadWrite, Category="Style", meta=( DisplayName="Style" ))
+	UPROPERTY(EditAnywhere, BlueprintReadWrite, Category="Appearance", meta=(ShowOnlyInnerProperties))
 	FTextBlockStyle WidgetStyle;
 
 	/** Font color and opacity (overrides Style) */
diff --git a/Engine/Source/Runtime/UMG/Public/Components/MultiLineEditableTextBox.h b/Engine/Source/Runtime/UMG/Public/Components/MultiLineEditableTextBox.h
--- a/Engine/Source/Runtime/UMG/Public/Components/MultiLineEditableTextBox.h
+++ b/Engine/Source/Runtime/UMG/Public/Components/MultiLineEditableTextBox.h
@@ -2,6 +2,7 @@
 
 #pragma once
 
+#include "TextWidgetTypes.h"
 #include "MultiLineEditableTextBox.generated.h"
 
 /**
@@ -33,11 +34,11 @@ public:
 public:
 
 	/** The style */
-	UPROPERTY(EditAnywhere, BlueprintReadWrite, Category="Style", meta=( DisplayName="Style" ))
+	UPROPERTY(EditAnywhere, BlueprintReadWrite, Category="Appearance", meta=(DisplayName="Style", ShowOnlyInnerProperties))
 	FEditableTextBoxStyle WidgetStyle;
 
 	/** The text style */
-	UPROPERTY(EditAnywhere, BlueprintReadWrite, Category="Style", meta=( DisplayName="Text Style" ))
+	UPROPERTY(EditAnywhere, BlueprintReadWrite, Category="Appearance", meta=(DisplayName="Text Style"))
 	FTextBlockStyle TextStyle;
 
 	/** Whether the context menu can be opened */
diff --git a/Engine/Source/Runtime/UMG/Public/Components/SpinBox.h b/Engine/Source/Runtime/UMG/Public/Components/SpinBox.h
--- a/Engine/Source/Runtime/UMG/Public/Components/SpinBox.h
+++ b/Engine/Source/Runtime/UMG/Public/Components/SpinBox.h
@@ -174,19 +174,19 @@ protected:
 
 protected:
 	/** Whether the optional MinValue attribute of the widget is set */
-	UPROPERTY()
+	UPROPERTY(EditAnywhere, Category = Content, meta=(InlineEditConditionToggle))
 	uint32 bOverride_MinValue : 1;
 
 	/** Whether the optional MaxValue attribute of the widget is set */
-	UPROPERTY()
+	UPROPERTY(EditAnywhere, Category = Content, meta=(InlineEditConditionToggle))
 	uint32 bOverride_MaxValue : 1;
 
 	/** Whether the optional MinSliderValue attribute of the widget is set */
-	UPROPERTY()
+	UPROPERTY(EditAnywhere, Category = Content, meta=(InlineEditConditionToggle))
 	uint32 bOverride_MinSliderValue : 1;
 
 	/** Whether the optional MaxSliderValue attribute of the widget is set */
-	UPROPERTY()
+	UPROPERTY(EditAnywhere, Category = Content, meta=(InlineEditConditionToggle))
 	uint32 bOverride_MaxSliderValue : 1;
 
 	/** The minimum allowable value that can be manually entered into the spin box */
diff --git a/Engine/Source/Runtime/UMG/Public/Components/TextWidgetTypes.h b/Engine/Source/Runtime/UMG/Public/Components/TextWidgetTypes.h
--- a/Engine/Source/Runtime/UMG/Public/Components/TextWidgetTypes.h
+++ b/Engine/Source/Runtime/UMG/Public/Components/TextWidgetTypes.h
@@ -24,11 +24,11 @@ struct UMG_API FShapedTextOptions
 	}
 
 	/**  */
-	UPROPERTY()
+	UPROPERTY(EditAnywhere, Category=Localization, meta=(InlineEditConditionToggle))
 	uint32 bOverride_TextShapingMethod : 1;
 
 	/**  */
-	UPROPERTY()
+	UPROPERTY(EditAnywhere, Category=Localization, meta=(InlineEditConditionToggle))
 	uint32 bOverride_TextFlowDirection : 1;
 
 	/** Which text shaping method should the text within this widget use? (unset to use the default returned by GetDefaultTextShapingMethod) */
diff --git a/Engine/Source/Runtime/UMG/Public/Components/Widget.h b/Engine/Source/Runtime/UMG/Public/Components/Widget.h
--- a/Engine/Source/Runtime/UMG/Public/Components/Widget.h
+++ b/Engine/Source/Runtime/UMG/Public/Components/Widget.h
@@ -170,7 +170,7 @@ public:
 	FGetSlateVisibility VisibilityDelegate;
 
 	/**  */
-	UPROPERTY()
+	UPROPERTY(EditAnywhere, Category="Behavior", meta=(InlineEditConditionToggle))
 	uint32 bOverride_Cursor : 1;
 
 	/** The cursor to show when the mouse is over the widget */
