Slate: Allow manual routing of pointer double click events

diff --git a/Engine/Source/Runtime/Slate/Private/Framework/Application/SlateApplication.cpp b/Engine/Source/Runtime/Slate/Private/Framework/Application/SlateApplication.cpp
--- a/Engine/Source/Runtime/Slate/Private/Framework/Application/SlateApplication.cpp
+++ b/Engine/Source/Runtime/Slate/Private/Framework/Application/SlateApplication.cpp
@@ -4898,22 +4898,32 @@ bool FSlateApplication::ProcessMouseButtonDoubleClickEvent( const TSharedPtr< FG
 	
 	FWidgetPath WidgetsUnderCursor = LocateWindowUnderMouse(InMouseEvent.GetScreenSpacePosition(), GetInteractiveTopLevelWindows());
 
+	FReply Reply = RoutePointerDoubleClickEvent( WidgetsUnderCursor, InMouseEvent );
+
+	PointerIndexLastPositionMap.Add(InMouseEvent.GetPointerIndex(), InMouseEvent.GetScreenSpacePosition());
+	
+	return Reply.IsEventHandled();
+}
+
+
+FReply FSlateApplication::RoutePointerDoubleClickEvent(FWidgetPath& WidgetsUnderPointer, FPointerEvent& PointerEvent)
+{
+	FReply Reply = FReply::Unhandled();
+
 	// Switch worlds widgets in the current path
-	FScopedSwitchWorldHack SwitchWorld(WidgetsUnderCursor);
+	FScopedSwitchWorldHack SwitchWorld( WidgetsUnderPointer );
 
-	FReply Reply = FEventRouter::Route<FReply>(this, FEventRouter::FBubblePolicy(WidgetsUnderCursor), InMouseEvent, [](const FArrangedWidget& TargetWidget, const FPointerEvent& Event)
+	Reply = FEventRouter::Route<FReply>( this, FEventRouter::FBubblePolicy( WidgetsUnderPointer ), PointerEvent, []( const FArrangedWidget& TargetWidget, const FPointerEvent& Event )
 	{
-		return TargetWidget.Widget->OnMouseButtonDoubleClick(TargetWidget.Geometry, Event);
-	});
+		return TargetWidget.Widget->OnMouseButtonDoubleClick( TargetWidget.Geometry, Event );
+	} );
 
+	LOG_EVENT( EEventLog::MouseButtonDoubleClick, Reply );
 
-	LOG_EVENT(EEventLog::MouseButtonDoubleClick, Reply);
-
-	PointerIndexLastPositionMap.Add(InMouseEvent.GetPointerIndex(), InMouseEvent.GetScreenSpacePosition());
-	
-	return Reply.IsEventHandled();
+	return Reply;
 }
 
+
 bool FSlateApplication::OnMouseUp( const EMouseButtons::Type Button )
 {
 	// convert to touch event if we are faking it	
diff --git a/Engine/Source/Runtime/Slate/Public/Framework/Application/SlateApplication.h b/Engine/Source/Runtime/Slate/Public/Framework/Application/SlateApplication.h
--- a/Engine/Source/Runtime/Slate/Public/Framework/Application/SlateApplication.h
+++ b/Engine/Source/Runtime/Slate/Public/Framework/Application/SlateApplication.h
@@ -1206,6 +1206,14 @@ public:
 	 */
 	bool RoutePointerMoveEvent( const FWidgetPath& WidgetsUnderPointer, FPointerEvent& PointerEvent, bool bIsSynthetic );
 
+	/**
+	 * Directly routes a pointer double click event to the widgets in the specified widget path
+	 *
+	 * @param WidgetsUnderPointer	The path of widgets the event is routed to.
+	 * @param PointerEvent		The event data that is is routed to the widget path
+	 */
+	FReply RoutePointerDoubleClickEvent( FWidgetPath& WidgetsUnderPointer, FPointerEvent& PointerEvent );
+
 	/**
 	 * Directly routes a pointer mouse wheel or gesture event to the widgets in the specified widget path.
 	 * 
