Slate - Fixing a crash in the routing of mouse downs where the widget index is 0, and there are no widgets under the cursor.

diff --git a/Engine/Source/Runtime/Slate/Private/Framework/Application/SlateApplication.cpp b/Engine/Source/Runtime/Slate/Private/Framework/Application/SlateApplication.cpp
--- a/Engine/Source/Runtime/Slate/Private/Framework/Application/SlateApplication.cpp
+++ b/Engine/Source/Runtime/Slate/Private/Framework/Application/SlateApplication.cpp
@@ -4488,15 +4488,18 @@ FReply FSlateApplication::RoutePointerDownEvent(FWidgetPath& WidgetsUnderPointer
 		//                work well with 3D widgets, so we're just going to use the existing widgets under pointer
 		//                that we already passed in.
 
-		bool bFocusCandidateFound = false;
-		for ( int32 WidgetIndex = WidgetHandledIndex; !bFocusCandidateFound && WidgetIndex >= 0; --WidgetIndex )
+		if ( WidgetIndex < WidgetsUnderPointer.Widgets.Num() )
 		{
-			FArrangedWidget& CurWidget = WidgetsUnderPointer.Widgets[WidgetIndex];
-			if ( CurWidget.Widget->SupportsKeyboardFocus() )
+			bool bFocusCandidateFound = false;
+			for ( int32 WidgetIndex = WidgetHandledIndex; !bFocusCandidateFound && WidgetIndex >= 0; --WidgetIndex )
 			{
-				bFocusCandidateFound = true;
-				FWidgetPath NewFocusedWidgetPath = WidgetsUnderPointer.GetPathDownTo(CurWidget.Widget);
-				SetKeyboardFocus(NewFocusedWidgetPath, EFocusCause::Mouse);
+				FArrangedWidget& CurWidget = WidgetsUnderPointer.Widgets[WidgetIndex];
+				if ( CurWidget.Widget->SupportsKeyboardFocus() )
+				{
+					bFocusCandidateFound = true;
+					FWidgetPath NewFocusedWidgetPath = WidgetsUnderPointer.GetPathDownTo(CurWidget.Widget);
+					SetKeyboardFocus(NewFocusedWidgetPath, EFocusCause::Mouse);
+				}
 			}
 		}
 
