Slate - Small fix, wrong variable being checked to avoid crash.

diff --git a/Engine/Source/Runtime/Slate/Private/Framework/Application/SlateApplication.cpp b/Engine/Source/Runtime/Slate/Private/Framework/Application/SlateApplication.cpp
--- a/Engine/Source/Runtime/Slate/Private/Framework/Application/SlateApplication.cpp
+++ b/Engine/Source/Runtime/Slate/Private/Framework/Application/SlateApplication.cpp
@@ -4488,7 +4488,7 @@ FReply FSlateApplication::RoutePointerDownEvent(FWidgetPath& WidgetsUnderPointer
 		//                work well with 3D widgets, so we're just going to use the existing widgets under pointer
 		//                that we already passed in.
 
-		if ( WidgetIndex < WidgetsUnderPointer.Widgets.Num() )
+		if ( WidgetHandledIndex < WidgetsUnderPointer.Widgets.Num() )
 		{
 			bool bFocusCandidateFound = false;
 			for ( int32 WidgetIndex = WidgetHandledIndex; !bFocusCandidateFound && WidgetIndex >= 0; --WidgetIndex )
