Slate: Rolled back some changes from CL 2893139 that prevented VR input from working

diff --git a/Engine/Source/Runtime/Slate/Private/Framework/Application/SlateApplication.cpp b/Engine/Source/Runtime/Slate/Private/Framework/Application/SlateApplication.cpp
--- a/Engine/Source/Runtime/Slate/Private/Framework/Application/SlateApplication.cpp
+++ b/Engine/Source/Runtime/Slate/Private/Framework/Application/SlateApplication.cpp
@@ -4439,21 +4439,15 @@ FReply FSlateApplication::RoutePointerDownEvent(FWidgetPath& WidgetsUnderPointer
 
 	const TSharedPtr<SWidget> PreviouslyFocusedWidget = GetKeyboardFocusedWidget();
 
-	int32 WidgetHandledIndex = 0;
-
-	FReply Reply = FEventRouter::Route<FReply>(this, FEventRouter::FTunnelPolicy(WidgetsUnderPointer), PointerEvent, [&WidgetHandledIndex] (const FArrangedWidget TargetWidget, const FPointerEvent& Event)
+	FReply Reply = FEventRouter::Route<FReply>(this, FEventRouter::FTunnelPolicy(WidgetsUnderPointer), PointerEvent, [] (const FArrangedWidget TargetWidget, const FPointerEvent& Event)
 	{
-		WidgetHandledIndex++;
 		return TargetWidget.Widget->OnPreviewMouseButtonDown(TargetWidget.Geometry, Event);
 	});
 
 	if ( !Reply.IsEventHandled() )
 	{
-		WidgetHandledIndex = WidgetsUnderPointer.Widgets.Num();
-
-		Reply = FEventRouter::Route<FReply>(this, FEventRouter::FBubblePolicy(WidgetsUnderPointer), PointerEvent, [this, &WidgetHandledIndex] (const FArrangedWidget TargetWidget, const FPointerEvent& Event)
+		Reply = FEventRouter::Route<FReply>(this, FEventRouter::FBubblePolicy(WidgetsUnderPointer), PointerEvent, [this] (const FArrangedWidget TargetWidget, const FPointerEvent& Event)
 		{
-			WidgetHandledIndex--;
 			FReply ThisReply = FReply::Unhandled();
 			if ( !ThisReply.IsEventHandled() )
 			{
@@ -4469,12 +4463,6 @@ FReply FSlateApplication::RoutePointerDownEvent(FWidgetPath& WidgetsUnderPointer
 			return ThisReply;
 		});
 	}
-	else
-	{
-		// Since we always advance 1, need to back that out if it was handled by the preview mouse down code.
-		WidgetHandledIndex--;
-	}
-
 	LOG_EVENT(EEventLog::MouseButtonDown, Reply);
 
 	// If none of the widgets requested keyboard focus to be set (or set the keyboard focus explicitly), set it to the leaf-most widget under the mouse.
@@ -4483,23 +4471,19 @@ FReply FSlateApplication::RoutePointerDownEvent(FWidgetPath& WidgetsUnderPointer
 	if ( ( !Reply.GetUserFocusRecepient().IsValid() || ( PLATFORM_MAC && PointerEvent.GetEffectingButton() == EKeys::LeftMouseButton && !DragDetector.DetectDragForWidget.IsValid() ) )
 		&& ( !bFocusChangedByEventHandler || bNeedToActivateWindow ) )
 	{
-		// If the event wasn't handled, we need to resample the screen to find
-		// @HACK VREDITOR No longer resampling from the screen down, that won't
-		//                work well with 3D widgets, so we're just going to use the existing widgets under pointer
-		//                that we already passed in.
+		// The event handler for OnMouseButton down may have altered the widget hierarchy.
+		// Refresh the previously-cached widget path.
+		WidgetsUnderPointer = LocateWindowUnderMouse(PointerEvent.GetScreenSpacePosition(), GetInteractiveTopLevelWindows());
 
-		if ( WidgetHandledIndex < WidgetsUnderPointer.Widgets.Num() )
+		bool bFocusCandidateFound = false;
+		for ( int32 WidgetIndex = WidgetsUnderPointer.Widgets.Num() - 1; !bFocusCandidateFound && WidgetIndex >= 0; --WidgetIndex )
 		{
-			bool bFocusCandidateFound = false;
-			for ( int32 WidgetIndex = WidgetHandledIndex; !bFocusCandidateFound && WidgetIndex >= 0; --WidgetIndex )
+			FArrangedWidget& CurWidget = WidgetsUnderPointer.Widgets[WidgetIndex];
+			if ( CurWidget.Widget->SupportsKeyboardFocus() )
 			{
-				FArrangedWidget& CurWidget = WidgetsUnderPointer.Widgets[WidgetIndex];
-				if ( CurWidget.Widget->SupportsKeyboardFocus() )
-				{
-					bFocusCandidateFound = true;
-					FWidgetPath NewFocusedWidgetPath = WidgetsUnderPointer.GetPathDownTo(CurWidget.Widget);
-					SetKeyboardFocus(NewFocusedWidgetPath, EFocusCause::Mouse);
-				}
+				bFocusCandidateFound = true;
+				FWidgetPath NewFocusedWidgetPath = WidgetsUnderPointer.GetPathDownTo(CurWidget.Widget);
+				SetKeyboardFocus(NewFocusedWidgetPath, EFocusCause::Mouse);
 			}
 		}
 
