VREditor - Disabling tab dragging for asset editors opened while you're in vr editor mode.  Adding back in the changes to the way we send down mouse down messages so that popups work again and we don't refocus from the cursors standpoint every time, now that the VR editor's input is no longer based on the viewport having focus.

diff --git a/Engine/Source/Editor/VREditor/VREditorUISystem.cpp b/Engine/Source/Editor/VREditor/VREditorUISystem.cpp
--- a/Engine/Source/Editor/VREditor/VREditorUISystem.cpp
+++ b/Engine/Source/Editor/VREditor/VREditorUISystem.cpp
@@ -839,6 +839,8 @@ void FVREditorUISystem::CreateUIs()
 
 			// We're going to start stealing tabs from the global tab manager inserting them into the world instead.
 			FGlobalTabmanager::Get()->SetProxyTabManager(ProxyTabManager);
+
+			FAssetEditorManager::Get().OnAssetEditorOpened().AddRaw(this, &FVREditorUISystem::OnAssetEditorOpened);
 		}
 	}
 }
@@ -863,8 +865,18 @@ void FVREditorUISystem::CleanUpActorsBeforeMapChangeOrSimulate()
 
 	// Remove the proxy tab manager, we don't want to steal tabs any more.
 	FGlobalTabmanager::Get()->SetProxyTabManager(TSharedPtr<FProxyTabmanager>());
+	FAssetEditorManager::Get().OnAssetEditorOpened().RemoveAll(this);
 }
 
+void FVREditorUISystem::OnAssetEditorOpened(UObject* Asset)
+{
+	// We need to disable drag drop on the tabs spawned in VR mode.
+	TArray<IAssetEditorInstance*> Editors = FAssetEditorManager::Get().FindEditorsForAsset(Asset);
+	for ( IAssetEditorInstance* Editor : Editors )
+	{
+		Editor->GetAssociatedTabManager()->SetCanDoDragOperation(false);
+	}
+}
 
 bool FVREditorUISystem::IsWidgetAnEditorUIWidget( const UActorComponent* WidgetComponent ) const
 {
diff --git a/Engine/Source/Editor/VREditor/VREditorUISystem.h b/Engine/Source/Editor/VREditor/VREditorUISystem.h
--- a/Engine/Source/Editor/VREditor/VREditorUISystem.h
+++ b/Engine/Source/Editor/VREditor/VREditorUISystem.h
@@ -133,6 +133,9 @@ protected:
 	/** Show the asset editor if it's not currently visible. */
 	void ShowAssetEditor();
 
+	/** Called when an asset editor opens an asset while in VR Editor Mode. */
+	void OnAssetEditorOpened(UObject* Asset);
+
 protected:
 
 	/** Owning object */
diff --git a/Engine/Source/Runtime/Slate/Private/Framework/Application/SlateApplication.cpp b/Engine/Source/Runtime/Slate/Private/Framework/Application/SlateApplication.cpp
--- a/Engine/Source/Runtime/Slate/Private/Framework/Application/SlateApplication.cpp
+++ b/Engine/Source/Runtime/Slate/Private/Framework/Application/SlateApplication.cpp
@@ -4431,23 +4431,29 @@ bool FSlateApplication::ProcessMouseButtonDownEvent( const TSharedPtr< FGenericW
 FReply FSlateApplication::RoutePointerDownEvent(FWidgetPath& WidgetsUnderPointer, FPointerEvent& PointerEvent)
 {
 #if PLATFORM_MAC
-	NSWindow* ActiveWindow = [NSApp keyWindow];
-	const bool bNeedToActivateWindow = (ActiveWindow == nullptr);
+	NSWindow* ActiveWindow =[NSApp keyWindow];
+	const bool bNeedToActivateWindow = ( ActiveWindow == nullptr );
 #else
 	const bool bNeedToActivateWindow = false;
 #endif
 
 	const TSharedPtr<SWidget> PreviouslyFocusedWidget = GetKeyboardFocusedWidget();
 
-	FReply Reply = FEventRouter::Route<FReply>(this, FEventRouter::FTunnelPolicy(WidgetsUnderPointer), PointerEvent, [] (const FArrangedWidget TargetWidget, const FPointerEvent& Event)
+	int32 WidgetHandledIndex = 0;
+
+	FReply Reply = FEventRouter::Route<FReply>(this, FEventRouter::FTunnelPolicy(WidgetsUnderPointer), PointerEvent, [&WidgetHandledIndex] (const FArrangedWidget TargetWidget, const FPointerEvent& Event)
 	{
+		WidgetHandledIndex++;
 		return TargetWidget.Widget->OnPreviewMouseButtonDown(TargetWidget.Geometry, Event);
 	});
 
 	if ( !Reply.IsEventHandled() )
 	{
-		Reply = FEventRouter::Route<FReply>(this, FEventRouter::FBubblePolicy(WidgetsUnderPointer), PointerEvent, [this] (const FArrangedWidget TargetWidget, const FPointerEvent& Event)
+		WidgetHandledIndex = WidgetsUnderPointer.Widgets.Num();
+
+		Reply = FEventRouter::Route<FReply>(this, FEventRouter::FBubblePolicy(WidgetsUnderPointer), PointerEvent, [this, &WidgetHandledIndex] (const FArrangedWidget TargetWidget, const FPointerEvent& Event)
 		{
+			WidgetHandledIndex--;
 			FReply ThisReply = FReply::Unhandled();
 			if ( !ThisReply.IsEventHandled() )
 			{
@@ -4463,6 +4469,12 @@ FReply FSlateApplication::RoutePointerDownEvent(FWidgetPath& WidgetsUnderPointer
 			return ThisReply;
 		});
 	}
+	else
+	{
+		// Since we always advance 1, need to back that out if it was handled by the preview mouse down code.
+		WidgetHandledIndex--;
+	}
+
 	LOG_EVENT(EEventLog::MouseButtonDown, Reply);
 
 	// If none of the widgets requested keyboard focus to be set (or set the keyboard focus explicitly), set it to the leaf-most widget under the mouse.
@@ -4471,19 +4483,23 @@ FReply FSlateApplication::RoutePointerDownEvent(FWidgetPath& WidgetsUnderPointer
 	if ( ( !Reply.GetUserFocusRecepient().IsValid() || ( PLATFORM_MAC && PointerEvent.GetEffectingButton() == EKeys::LeftMouseButton && !DragDetector.DetectDragForWidget.IsValid() ) )
 		&& ( !bFocusChangedByEventHandler || bNeedToActivateWindow ) )
 	{
-		// The event handler for OnMouseButton down may have altered the widget hierarchy.
-		// Refresh the previously-cached widget path.
-		WidgetsUnderPointer = LocateWindowUnderMouse(PointerEvent.GetScreenSpacePosition(), GetInteractiveTopLevelWindows());
+		// If the event wasn't handled, we need to resample the screen to find
+		// @HACK VREDITOR No longer resampling from the screen down, that won't
+		//                work well with 3D widgets, so we're just going to use the existing widgets under pointer
+		//                that we already passed in.
 
-		bool bFocusCandidateFound = false;
-		for ( int32 WidgetIndex = WidgetsUnderPointer.Widgets.Num() - 1; !bFocusCandidateFound && WidgetIndex >= 0; --WidgetIndex )
+		if ( WidgetHandledIndex < WidgetsUnderPointer.Widgets.Num() )
 		{
-			FArrangedWidget& CurWidget = WidgetsUnderPointer.Widgets[WidgetIndex];
-			if ( CurWidget.Widget->SupportsKeyboardFocus() )
+			bool bFocusCandidateFound = false;
+			for ( int32 WidgetIndex = WidgetHandledIndex; !bFocusCandidateFound && WidgetIndex >= 0; --WidgetIndex )
 			{
-				bFocusCandidateFound = true;
-				FWidgetPath NewFocusedWidgetPath = WidgetsUnderPointer.GetPathDownTo(CurWidget.Widget);
-				SetKeyboardFocus(NewFocusedWidgetPath, EFocusCause::Mouse);
+				FArrangedWidget& CurWidget = WidgetsUnderPointer.Widgets[WidgetIndex];
+				if ( CurWidget.Widget->SupportsKeyboardFocus() )
+				{
+					bFocusCandidateFound = true;
+					FWidgetPath NewFocusedWidgetPath = WidgetsUnderPointer.GetPathDownTo(CurWidget.Widget);
+					SetKeyboardFocus(NewFocusedWidgetPath, EFocusCause::Mouse);
+				}
 			}
 		}
 
