[INTEGRATE] Fixed debug console history not working (merge from CL 2892825)

diff --git a/Engine/Source/Runtime/Slate/Private/Widgets/Text/SlateEditableTextLayout.cpp b/Engine/Source/Runtime/Slate/Private/Widgets/Text/SlateEditableTextLayout.cpp
--- a/Engine/Source/Runtime/Slate/Private/Widgets/Text/SlateEditableTextLayout.cpp
+++ b/Engine/Source/Runtime/Slate/Private/Widgets/Text/SlateEditableTextLayout.cpp
@@ -185,16 +185,12 @@ void FSlateEditableTextLayout::SetText(const TAttribute<FText>& InText)
 {
 	BoundText = InText;
 
-	const bool bShouldAppearFocused = OwnerWidget->GetSlateWidget()->HasKeyboardFocus() || HasActiveContextMenu();
-	if (!bShouldAppearFocused)
+	// We don't have focus, so we can perform an immediate update
+	const FText& TextToSet = BoundText.Get(FText::GetEmpty());
+	if (RefreshImpl(&TextToSet))
 	{
-		// We don't have focus, so we can perform an immediate update
-		const FText& TextToSet = BoundText.Get(FText::GetEmpty());
-		if (RefreshImpl(&TextToSet))
-		{
-			// Let outsiders know that the text content has been changed
-			OwnerWidget->OnTextChanged(TextToSet);
-		}
+		// Let outsiders know that the text content has been changed
+		OwnerWidget->OnTextChanged(TextToSet);
 	}
 }
 
