Slate: Fixed focus being set to incorrect widget/viewport

diff --git a/Engine/Source/Runtime/Slate/Private/Framework/Application/SlateApplication.cpp b/Engine/Source/Runtime/Slate/Private/Framework/Application/SlateApplication.cpp
--- a/Engine/Source/Runtime/Slate/Private/Framework/Application/SlateApplication.cpp
+++ b/Engine/Source/Runtime/Slate/Private/Framework/Application/SlateApplication.cpp
@@ -4470,7 +4470,7 @@ bool FSlateApplication::ProcessMouseButtonDownEvent( const TSharedPtr< FGenericW
 FReply FSlateApplication::RoutePointerDownEvent(FWidgetPath& WidgetsUnderPointer, FPointerEvent& PointerEvent)
 {
 #if PLATFORM_MAC
-	NSWindow* ActiveWindow =[NSApp keyWindow];
+	NSWindow* ActiveWindow = [ NSApp keyWindow ];
 	const bool bNeedToActivateWindow = ( ActiveWindow == nullptr );
 #else
 	const bool bNeedToActivateWindow = false;
@@ -4478,84 +4478,64 @@ FReply FSlateApplication::RoutePointerDownEvent(FWidgetPath& WidgetsUnderPointer
 
 	const TSharedPtr<SWidget> PreviouslyFocusedWidget = GetKeyboardFocusedWidget();
 
-	int32 WidgetHandledIndex = 0;
-
-	FReply Reply = FEventRouter::Route<FReply>(this, FEventRouter::FTunnelPolicy(WidgetsUnderPointer), PointerEvent, [&WidgetHandledIndex] (const FArrangedWidget TargetWidget, const FPointerEvent& Event)
+	FReply Reply = FEventRouter::Route<FReply>( this, FEventRouter::FTunnelPolicy( WidgetsUnderPointer ), PointerEvent, []( const FArrangedWidget TargetWidget, const FPointerEvent& Event )
 	{
-		WidgetHandledIndex++;
-		return TargetWidget.Widget->OnPreviewMouseButtonDown(TargetWidget.Geometry, Event);
-	});
+		return TargetWidget.Widget->OnPreviewMouseButtonDown( TargetWidget.Geometry, Event );
+	} );
 
-	if ( !Reply.IsEventHandled() )
+	if( !Reply.IsEventHandled() )
 	{
-		WidgetHandledIndex = WidgetsUnderPointer.Widgets.Num();
-
-		Reply = FEventRouter::Route<FReply>(this, FEventRouter::FBubblePolicy(WidgetsUnderPointer), PointerEvent, [this, &WidgetHandledIndex] (const FArrangedWidget TargetWidget, const FPointerEvent& Event)
+		Reply = FEventRouter::Route<FReply>( this, FEventRouter::FBubblePolicy( WidgetsUnderPointer ), PointerEvent, [this]( const FArrangedWidget TargetWidget, const FPointerEvent& Event )
 		{
-			WidgetHandledIndex--;
 			FReply ThisReply = FReply::Unhandled();
-			if ( !ThisReply.IsEventHandled() )
+			if( !ThisReply.IsEventHandled() )
 			{
-				if ( Event.IsTouchEvent() )
+				if( Event.IsTouchEvent() )
 				{
-					ThisReply = TargetWidget.Widget->OnTouchStarted(TargetWidget.Geometry, Event);
+					ThisReply = TargetWidget.Widget->OnTouchStarted( TargetWidget.Geometry, Event );
 				}
-				if ( !Event.IsTouchEvent() || ( !ThisReply.IsEventHandled() && this->bTouchFallbackToMouse ) )
+				if( !Event.IsTouchEvent() || ( !ThisReply.IsEventHandled() && this->bTouchFallbackToMouse ) )
 				{
-					ThisReply = TargetWidget.Widget->OnMouseButtonDown(TargetWidget.Geometry, Event);
+					ThisReply = TargetWidget.Widget->OnMouseButtonDown( TargetWidget.Geometry, Event );
 				}
 			}
 			return ThisReply;
-		});
-	}
-	else
-	{
-		// Since we always advance 1, need to back that out if it was handled by the preview mouse down code.
-		WidgetHandledIndex--;
+		} );
 	}
-
-	LOG_EVENT(EEventLog::MouseButtonDown, Reply);
+	LOG_EVENT( EEventLog::MouseButtonDown, Reply );
 
 	// If none of the widgets requested keyboard focus to be set (or set the keyboard focus explicitly), set it to the leaf-most widget under the mouse.
 	// On Mac we prevent the OS from activating the window on mouse down, so we have full control and can activate only if there's nothing draggable under the mouse cursor.
 	const bool bFocusChangedByEventHandler = PreviouslyFocusedWidget != GetKeyboardFocusedWidget();
-	if ( ( !Reply.GetUserFocusRecepient().IsValid() || ( PLATFORM_MAC && PointerEvent.GetEffectingButton() == EKeys::LeftMouseButton && !DragDetector.DetectDragForWidget.IsValid() ) )
+	if( ( !Reply.GetUserFocusRecepient().IsValid() || ( PLATFORM_MAC && PointerEvent.GetEffectingButton() == EKeys::LeftMouseButton && !DragDetector.DetectDragForWidget.IsValid() ) )
 		&& ( !bFocusChangedByEventHandler || bNeedToActivateWindow ) )
 	{
-		// If the event wasn't handled, we need to resample the screen to find
-		// @HACK VREDITOR No longer resampling from the screen down, that won't
-		//                work well with 3D widgets, so we're just going to use the existing widgets under pointer
-		//                that we already passed in.
-
-		if ( WidgetHandledIndex < WidgetsUnderPointer.Widgets.Num() )
+		bool bFocusCandidateFound = false;
+		for( int32 WidgetIndex = WidgetsUnderPointer.Widgets.Num() - 1; !bFocusCandidateFound && WidgetIndex >= 0; --WidgetIndex )
 		{
-			bool bFocusCandidateFound = false;
-			for ( int32 WidgetIndex = WidgetHandledIndex; !bFocusCandidateFound && WidgetIndex >= 0; --WidgetIndex )
+			FArrangedWidget& CurWidget = WidgetsUnderPointer.Widgets[ WidgetIndex ];
+			if( CurWidget.Widget->SupportsKeyboardFocus() )
 			{
-				FArrangedWidget& CurWidget = WidgetsUnderPointer.Widgets[WidgetIndex];
-				if ( CurWidget.Widget->SupportsKeyboardFocus() )
-				{
-					bFocusCandidateFound = true;
-					FWidgetPath NewFocusedWidgetPath = WidgetsUnderPointer.GetPathDownTo(CurWidget.Widget);
-					SetKeyboardFocus(NewFocusedWidgetPath, EFocusCause::Mouse);
-				}
+				bFocusCandidateFound = true;
+				FWidgetPath NewFocusedWidgetPath = WidgetsUnderPointer.GetPathDownTo( CurWidget.Widget );
+				SetKeyboardFocus( NewFocusedWidgetPath, EFocusCause::Mouse );
 			}
 		}
 
 #if PLATFORM_MAC
 		TSharedPtr<SWindow> TopLevelWindow = WidgetsUnderPointer.TopLevelWindow;
-		if ( bNeedToActivateWindow || ( TopLevelWindow.IsValid() && TopLevelWindow->GetNativeWindow()->GetOSWindowHandle() != ActiveWindow ) )
+		if( bNeedToActivateWindow || ( TopLevelWindow.IsValid() && TopLevelWindow->GetNativeWindow()->GetOSWindowHandle() != ActiveWindow ) )
 		{
 			// Clicking on a context menu should not activate anything
 			// @todo: This needs to be updated when we have window type in SWindow and we no longer have to guess if WidgetsUnderCursor.TopLevelWindow is a menu
 			const bool bIsContextMenu = TopLevelWindow.IsValid() && !TopLevelWindow->IsRegularWindow() && TopLevelWindow->HasMinimizeBox() && TopLevelWindow->HasMaximizeBox();
-			if ( !bIsContextMenu && PointerEvent.GetEffectingButton() == EKeys::LeftMouseButton && !DragDetector.DetectDragForWidget.IsValid() && ActiveWindow == [NSApp keyWindow] )
+			if( !bIsContextMenu && PointerEvent.GetEffectingButton() == EKeys::LeftMouseButton && !DragDetector.DetectDragForWidget.IsValid() && ActiveWindow == [ NSApp keyWindow ] )
 			{
 				MouseCaptorHelper Captor = MouseCaptor;
 				FPlatformMisc::ActivateApplication();
-				if ( TopLevelWindow.IsValid() )
+				if( TopLevelWindow.IsValid() )
 				{
-					TopLevelWindow->BringToFront(true);
+					TopLevelWindow->BringToFront( true );
 				}
 				MouseCaptor = Captor;
 			}
