Copying //UE4/Dev-Editor to //UE4/Dev-Main (Source: //UE4/Dev-Editor @ 2973866)

diff --git a/Engine/Source/Runtime/Engine/Private/Model.cpp b/Engine/Source/Runtime/Engine/Private/Model.cpp
--- a/Engine/Source/Runtime/Engine/Private/Model.cpp
+++ b/Engine/Source/Runtime/Engine/Private/Model.cpp
@@ -382,26 +382,63 @@ void UModel::PostEditUndo()
 
 void UModel::ModifySurf( int32 InIndex, bool UpdateMaster )
 {
-	// TODO: Remove this method and replace calls to it with Modify()
-	//Surfs.ModifyItem( InIndex );
+	Modify();
+
 	FBspSurf& Surf = Surfs[InIndex];
 	if( UpdateMaster && Surf.Actor )
 	{
-		//Surf.Actor->Brush->Polys->Element.ModifyItem( Surf.iBrushPoly );
+		Surf.Actor->Brush->Modify();
 	}
 }
+
 void UModel::ModifyAllSurfs( bool UpdateMaster )
 {
-	for( int32 i=0; i<Surfs.Num(); i++ )
-		ModifySurf( i, UpdateMaster );
+	Modify();
+
+	if (UpdateMaster)
+	{
+		TArray<UModel*> MasterModels;
+		MasterModels.Reset(Surfs.Num());
+
+		for (const FBspSurf& Surf : Surfs)
+		{
+			if (Surf.Actor)
+			{
+				check(Surf.Actor->Brush);
+				MasterModels.AddUnique(Surf.Actor->Brush);
+			}
+		}
 
+		for (UModel* MasterModel : MasterModels)
+		{
+			MasterModel->Modify();
+		}
+	}
 }
+
 void UModel::ModifySelectedSurfs( bool UpdateMaster )
 {
-	for( int32 i=0; i<Surfs.Num(); i++ )
-		if( Surfs[i].PolyFlags & PF_Selected )
-			ModifySurf( i, UpdateMaster );
+	Modify();
+
+	if (UpdateMaster)
+	{
+		TArray<UModel*> MasterModels;
+		MasterModels.Reset(Surfs.Num());
 
+		for (const FBspSurf& Surf : Surfs)
+		{
+			if (Surf.Actor && (Surf.PolyFlags & PF_Selected))
+			{
+				check(Surf.Actor->Brush);
+				MasterModels.AddUnique(Surf.Actor->Brush);
+			}
+		}
+
+		for (UModel* MasterModel : MasterModels)
+		{
+			MasterModel->Modify();
+		}
+	}
 }
 
 bool UModel::HasSelectedSurfaces() const
diff --git a/Engine/Source/Runtime/Engine/Private/ObjectLibrary.cpp b/Engine/Source/Runtime/Engine/Private/ObjectLibrary.cpp
--- a/Engine/Source/Runtime/Engine/Private/ObjectLibrary.cpp
+++ b/Engine/Source/Runtime/Engine/Private/ObjectLibrary.cpp
@@ -405,10 +405,10 @@ int32 UObjectLibrary::LoadBlueprintAssetDataFromPaths(const TArray<FString>& Pat
 			FAssetData& Data = AssetDataList[AssetIdx];
 
 			bool bShouldRemove = true;
-			const FString* ParentClassFromData = Data.TagsAndValues.Find("ParentClass");
-			if (ParentClassFromData && !ParentClassFromData->IsEmpty())
+			const FString ParentClassFromData = Data.GetTagValueRef<FString>("ParentClass");
+			if (!ParentClassFromData.IsEmpty())
 			{
-				const FString ClassObjectPath = FPackageName::ExportTextPathToObjectPath(*ParentClassFromData);
+				const FString ClassObjectPath = FPackageName::ExportTextPathToObjectPath(ParentClassFromData);
 				const FString ClassName = FPackageName::ObjectPathToObjectName(ClassObjectPath);
 				if (DerivedClassNames.Contains(FName(*ClassName)))
 				{
diff --git a/Engine/Source/Runtime/Engine/Private/StaticMeshBuild.cpp b/Engine/Source/Runtime/Engine/Private/StaticMeshBuild.cpp
--- a/Engine/Source/Runtime/Engine/Private/StaticMeshBuild.cpp
+++ b/Engine/Source/Runtime/Engine/Private/StaticMeshBuild.cpp
@@ -136,56 +136,6 @@ void UStaticMesh::Build(bool bSilent, TArray<FText>* OutErrors)
 		{
 			if ( It->StaticMesh == this )
 			{
-				// Initialize override vertex colors on any new LODs which have just been created
-				It->SetLODDataCount(NumLODs, It->LODData.Num());
-
-				FStaticMeshComponentLODInfo& LOD0Info = It->LODData[0];
-				if (LOD0Info.OverrideVertexColors)
-				{
-					for (uint32 LODIndex = 1; LODIndex < NumLODs; ++LODIndex)
-					{
-						FStaticMeshComponentLODInfo& LODInfo = It->LODData[LODIndex];
-
-						if (LODInfo.OverrideVertexColors == nullptr && LODInfo.PaintedVertices.Num() == 0)
-						{
-							LODInfo.OverrideVertexColors = new FColorVertexBuffer;
-
-							FStaticMeshLODResources& CurRenderData = RenderData->LODResources[LODIndex];
-
-							TArray<FColor> NewOverrideColors;
-
-							if (LOD0Info.PaintedVertices.Num() > 0)
-							{
-								// Build override colors for LOD, based on LOD0
-								RemapPaintedVertexColors(
-									LOD0Info.PaintedVertices,
-									*LOD0Info.OverrideVertexColors,
-									CurRenderData.PositionVertexBuffer,
-									&CurRenderData.VertexBuffer,
-									NewOverrideColors
-									);
-							}
-							if (NewOverrideColors.Num())
-							{
-								LODInfo.OverrideVertexColors->InitFromColorArray(NewOverrideColors);
-
-								// Update the PaintedVertices array
-								const int32 NumVerts = CurRenderData.GetNumVertices();
-								check(NumVerts == NewOverrideColors.Num());
-
-								LODInfo.PaintedVertices.Reserve(NumVerts);
-								for (int32 VertIndex = 0; VertIndex < NumVerts; ++VertIndex)
-								{
-									FPaintedVertex* Vertex = new(LODInfo.PaintedVertices) FPaintedVertex;
-									Vertex->Position = CurRenderData.PositionVertexBuffer.VertexPosition(VertIndex);
-									Vertex->Normal = CurRenderData.VertexBuffer.VertexTangentZ(VertIndex);
-									Vertex->Color = LODInfo.OverrideVertexColors->VertexColor(VertIndex);
-								}
-							}
-						}
-					}
-				}
-
 				It->FixupOverrideColorsIfNecessary(true);
 				It->InvalidateLightingCache();
 			}
@@ -353,7 +303,7 @@ void RemapPaintedVertexColors(
 				{
 					FOREACH_OCTREE_CHILD_NODE( OctreeChildRef )
 					{
-						if( CurNode.HasChild( OctreeChildRef ) && CurNode.GetChild( OctreeChildRef )->GetInclusiveElementCount() > 0 )
+						if( CurNode.HasChild( OctreeChildRef ) )
 						{
 							OctreeIter.PushChild( OctreeChildRef );
 						}
diff --git a/Engine/Source/Runtime/Engine/Private/SubtitleManager.cpp b/Engine/Source/Runtime/Engine/Private/SubtitleManager.cpp
--- a/Engine/Source/Runtime/Engine/Private/SubtitleManager.cpp
+++ b/Engine/Source/Runtime/Engine/Private/SubtitleManager.cpp
@@ -66,7 +66,7 @@ void FSubtitleManager::QueueSubtitles( PTRINT SubtitleID, float Priority, bool b
 	FActiveSubtitle& NewSubtitle = ActiveSubtitles.Add( SubtitleID, FActiveSubtitle( 0, Priority, bManualWordWrap, bSingleLine, Subtitles ) );
 
 	// Resolve time offsets to absolute time
-	for( FSubtitleCue& Subtitle : Subtitles )
+	for( FSubtitleCue& Subtitle : NewSubtitle.Subtitles )
 	{
 		if (Subtitle.Time > SoundDuration)
 		{
@@ -78,9 +78,9 @@ void FSubtitleManager::QueueSubtitles( PTRINT SubtitleID, float Priority, bool b
 	}
 
 	// Add on a blank at the end to clear
-	FSubtitleCue* Temp = new( NewSubtitle.Subtitles ) FSubtitleCue();
-	Temp->Text = FText::GetEmpty();
-	Temp->Time = InStartTime + SoundDuration;
+	FSubtitleCue& Temp = NewSubtitle.Subtitles[NewSubtitle.Subtitles.AddDefaulted()];
+	Temp.Text = FText::GetEmpty();
+	Temp.Time = InStartTime + SoundDuration;
 }
 
 /**
@@ -217,18 +217,18 @@ void FSubtitleManager::SplitLinesToSafeZone( FCanvas* Canvas, FIntRect& Subtitle
 		Cumulative = 0.0f;
 		for( i = 0; i < Lines.Num(); i++ )
 		{
-			FSubtitleCue* Line = new( Subtitle.Subtitles ) FSubtitleCue();
+			FSubtitleCue& Line = Subtitle.Subtitles[Subtitle.Subtitles.AddDefaulted()];
 			
-			Line->Text = FText::FromString(Lines[ i ].Value);
-			Line->Time = StartTime + Cumulative;
+			Line.Text = FText::FromString(Lines[ i ].Value);
+			Line.Time = StartTime + Cumulative;
 
-			Cumulative += SecondsPerChar * Line->Text.ToString().Len();
+			Cumulative += SecondsPerChar * Line.Text.ToString().Len();
 		}
 
 		// Add in the blank terminating line
-		FSubtitleCue* Temp = new( Subtitle.Subtitles ) FSubtitleCue();
-		Temp->Text = FText::GetEmpty();
-		Temp->Time = StartTime + SoundDuration;
+		FSubtitleCue& Temp = Subtitle.Subtitles[Subtitle.Subtitles.AddDefaulted()];
+		Temp.Text = FText::GetEmpty();
+		Temp.Time = StartTime + SoundDuration;
 
 		UE_LOG(LogAudio, Log, TEXT( "Splitting subtitle:" ) );
 		for( i = 0; i < Subtitle.Subtitles.Num() - 1; i++ )
diff --git a/Engine/Source/Runtime/Engine/Private/UnrealEngine.cpp b/Engine/Source/Runtime/Engine/Private/UnrealEngine.cpp
--- a/Engine/Source/Runtime/Engine/Private/UnrealEngine.cpp
+++ b/Engine/Source/Runtime/Engine/Private/UnrealEngine.cpp
@@ -1048,6 +1048,11 @@ void UEngine::Init(IEngineLoop* InEngineLoop)
 	RecordHMDAnalytics();
 }
 
+void UEngine::Start()
+{
+	// Start the game!
+}
+
 void UEngine::RegisterBeginStreamingPauseRenderingDelegate( FBeginStreamingPauseDelegate* InDelegate )
 {
 	BeginStreamingPauseDelegate = InDelegate;
@@ -6349,22 +6354,15 @@ void UEngine::OnLostFocusPause(bool EnablePause)
 
 void UEngine::StartHardwareSurvey()
 {
-	bool bEnabled = true;
-#if !WITH_EDITORONLY_DATA
-	bEnabled = AreGameAnalyticsEnabled();
-#endif
 
 	// The hardware survey costs time and we don't want to slow down debug builds.
 	// This is mostly because of the CPU benchmark running in the survey and the results in debug are not being valid.
-#if UE_BUILD_DEBUG
-	bEnabled = false;
-#endif
-
-	if (bEnabled && FEngineAnalytics::IsAvailable())
+	// Never run the survey in games, only in the editor.
+	if (WITH_EDITOR && GIsEditor && !IsRunningCommandlet() && FEngineAnalytics::IsAvailable())
 	{
 		IHardwareSurveyModule::Get().StartHardwareSurvey(FEngineAnalytics::GetProvider());
-		}
-	}	
+	}
+}
 
 void UEngine::InitHardwareSurvey()
 {
diff --git a/Engine/Source/Runtime/GameplayAbilities/Private/GameplayCueManager.cpp b/Engine/Source/Runtime/GameplayAbilities/Private/GameplayCueManager.cpp
--- a/Engine/Source/Runtime/GameplayAbilities/Private/GameplayCueManager.cpp
+++ b/Engine/Source/Runtime/GameplayAbilities/Private/GameplayCueManager.cpp
@@ -460,19 +460,19 @@ void UGameplayCueManager::BuildCuesToAddToGlobalSet(const TArray<FAssetData>& As
 			continue;
 		}
 
-		const FString* FoundGameplayTag = Data.TagsAndValues.Find(TagPropertyName);
-		if (FoundGameplayTag && FoundGameplayTag->Equals(TEXT("None")) == false)
+		const FName FoundGameplayTag = Data.GetTagValueRef<FName>(TagPropertyName);
+		if (!FoundGameplayTag.IsNone())
 		{
-			const FString* GeneratedClassTag = Data.TagsAndValues.Find(TEXT("GeneratedClass"));
-			if (GeneratedClassTag == nullptr)
+			const FString GeneratedClassTag = Data.GetTagValueRef<FString>("GeneratedClass");
+			if (GeneratedClassTag.IsEmpty())
 			{
 				ABILITY_LOG(Warning, TEXT("Unable to find GeneratedClass value for AssetData %s"), *Data.ObjectPath.ToString());
 				continue;
 			}
 
-			ABILITY_LOG(Log, TEXT("GameplayCueManager Found: %s / %s"), **FoundGameplayTag, **GeneratedClassTag);
+			ABILITY_LOG(Log, TEXT("GameplayCueManager Found: %s / %s"), *FoundGameplayTag.ToString(), **GeneratedClassTag);
 
-			FGameplayTag  GameplayCueTag = GameplayTagsModule.GetGameplayTagsManager().RequestGameplayTag(FName(**FoundGameplayTag), false);
+			FGameplayTag  GameplayCueTag = GameplayTagsModule.GetGameplayTagsManager().RequestGameplayTag(FoundGameplayTag, false);
 			if (GameplayCueTag.IsValid())
 			{
 				// Add a new NotifyData entry to our flat list for this one
@@ -485,7 +485,7 @@ void UGameplayCueManager::BuildCuesToAddToGlobalSet(const TArray<FAssetData>& As
 			}
 			else
 			{
-				ABILITY_LOG(Warning, TEXT("Found GameplayCue tag %s in asset %s but there is no corresponding tag in the GameplayTagManager."), **FoundGameplayTag, *Data.PackageName.ToString());
+				ABILITY_LOG(Warning, TEXT("Found GameplayCue tag %s in asset %s but there is no corresponding tag in the GameplayTagManager."), *FoundGameplayTag.ToString(), *Data.PackageName.ToString());
 			}
 		}
 	}
@@ -667,11 +667,9 @@ void UGameplayCueManager::HandleAssetDeleted(UObject *Object)
 /** Handles cleaning up an object library if it matches the passed in object */
 void UGameplayCueManager::HandleAssetRenamed(const FAssetData& Data, const FString& String)
 {
-	const FString* ParentClassNamePtr = Data.TagsAndValues.Find(TEXT("ParentClass"));
-	if (ParentClassNamePtr)
+	const FString ParentClassName = Data.GetTagValueRef<FString>("ParentClass");
+	if (!ParentClassName.IsEmpty())
 	{
-		FString ParentClassName = *ParentClassNamePtr;
-		
 		UClass* DataClass = FindObject<UClass>(nullptr, *ParentClassName);
 		if (DataClass)
 		{
diff --git a/Engine/Source/Runtime/Launch/Private/LaunchEngineLoop.cpp b/Engine/Source/Runtime/Launch/Private/LaunchEngineLoop.cpp
--- a/Engine/Source/Runtime/Launch/Private/LaunchEngineLoop.cpp
+++ b/Engine/Source/Runtime/Launch/Private/LaunchEngineLoop.cpp
@@ -2313,6 +2313,8 @@ int32 FEngineLoop::Init()
 		GIsRequestingExit = true;
 		return 1;
 	}
+
+	GEngine->Start();
 	
 	GetMoviePlayer()->WaitForMovieToFinish();
 
@@ -2412,6 +2414,9 @@ void FEngineLoop::Exit()
 		GEngine->PreExit();
 	}
 
+	// close all windows
+	FSlateApplication::Shutdown();
+
 #if !UE_SERVER
 	if ( FEngineFontServices::IsInitialized() )
 	{
@@ -2419,9 +2424,6 @@ void FEngineLoop::Exit()
 	}
 #endif
 
-	// close all windows
-	FSlateApplication::Shutdown();
-
 #if !PLATFORM_ANDROID 	// AppPreExit doesn't work on Android
 	AppPreExit();
 
diff --git a/Engine/Source/Runtime/Projects/Private/PluginManager.cpp b/Engine/Source/Runtime/Projects/Private/PluginManager.cpp
--- a/Engine/Source/Runtime/Projects/Private/PluginManager.cpp
+++ b/Engine/Source/Runtime/Projects/Private/PluginManager.cpp
@@ -215,16 +215,17 @@ void FPluginManager::ReadPluginsInDirectory(const FString& PluginsDirectory, con
 
 		for(const FString& FileName: FileNames)
 		{
+			FString FullPath = FPaths::ConvertRelativePathToFull(FileName);
+
 			FPluginDescriptor Descriptor;
 			FText FailureReason;
-			if(Descriptor.Load(FileName, FailureReason))
+			if ( Descriptor.Load(FullPath, FailureReason) )
 			{
-				Plugins.Add(MakeShareable(new FPlugin(FileName, Descriptor, LoadedFrom)));
+				Plugins.Add(MakeShareable(new FPlugin(FullPath, Descriptor, LoadedFrom)));
 			}
 			else
 			{
 				// NOTE: Even though loading of this plugin failed, we'll keep processing other plugins
-				FString FullPath = FPaths::ConvertRelativePathToFull(FileName);
 				FText FailureMessage = FText::Format(LOCTEXT("FailureFormat", "{0} ({1})"), FailureReason, FText::FromString(FullPath));
 				FText DialogTitle = LOCTEXT("PluginFailureTitle", "Failed to load Plugin");
 				UE_LOG(LogPluginManager, Error, TEXT("%s"), *FailureMessage.ToString());
diff --git a/Engine/Source/Runtime/Projects/Public/ModuleDescriptor.h b/Engine/Source/Runtime/Projects/Public/ModuleDescriptor.h
--- a/Engine/Source/Runtime/Projects/Public/ModuleDescriptor.h
+++ b/Engine/Source/Runtime/Projects/Public/ModuleDescriptor.h
@@ -13,20 +13,20 @@ namespace ELoadingPhase
 {
 	enum Type
 	{
-		/** Loaded at the default loading point during startup (during engine init, after game modules are loaded.) */
-		Default,
+		/** Loaded before the engine is fully initialized, immediately after the config system has been initialized.  Necessary only for very low-level hooks */
+		PostConfigInit,
 
-		/** Right after the default phase */
-		PostDefault,
+		/** Loaded before the engine is fully initialized for modules that need to hook into the loading screen before it triggers */
+		PreLoadingScreen,
 
 		/** Right before the default phase */
 		PreDefault,
 
-		/** Loaded before the engine is fully initialized, immediately after the config system has been initialized.  Necessary only for very low-level hooks */
-		PostConfigInit,
+		/** Loaded at the default loading point during startup (during engine init, after game modules are loaded.) */
+		Default,
 
-		/** Loaded before the engine is fully initialized for modules that need to hook into the loading screen before it triggers */
-		PreLoadingScreen,
+		/** Right after the default phase */
+		PostDefault,
 
 		/** After the engine has been initialized */
 		PostEngineInit,
diff --git a/Engine/Source/Runtime/Slate/Private/Framework/Application/SlateApplication.cpp b/Engine/Source/Runtime/Slate/Private/Framework/Application/SlateApplication.cpp
--- a/Engine/Source/Runtime/Slate/Private/Framework/Application/SlateApplication.cpp
+++ b/Engine/Source/Runtime/Slate/Private/Framework/Application/SlateApplication.cpp
@@ -5860,6 +5860,20 @@ EDropEffect::Type FSlateApplication::OnDragEnterFiles( const TSharedRef< FGeneri
 	return Result;
 }
 
+EDropEffect::Type FSlateApplication::OnDragEnterExternal( const TSharedRef< FGenericWindow >& Window, const FString& Text, const TArray< FString >& Files )
+{
+	const TSharedPtr< FExternalDragOperation > DragDropOperation = FExternalDragOperation::NewOperation( Text, Files );
+	const TSharedPtr< SWindow > EffectingWindow = FSlateWindowHelper::FindWindowByPlatformWindow( SlateWindows, Window );
+
+	EDropEffect::Type Result = EDropEffect::None;
+	if ( DragDropOperation.IsValid() && EffectingWindow.IsValid() )
+	{
+		Result = OnDragEnter( EffectingWindow.ToSharedRef(), DragDropOperation.ToSharedRef() );
+	}
+
+	return Result;
+}
+
 EDropEffect::Type FSlateApplication::OnDragEnter( const TSharedRef< SWindow >& Window, const TSharedRef<FExternalDragOperation>& DragDropOperation )
 {
 	// We are encountering a new drag and drop operation.
diff --git a/Engine/Source/Runtime/Slate/Private/Framework/Commands/GenericCommands.cpp b/Engine/Source/Runtime/Slate/Private/Framework/Commands/GenericCommands.cpp
--- a/Engine/Source/Runtime/Slate/Private/Framework/Commands/GenericCommands.cpp
+++ b/Engine/Source/Runtime/Slate/Private/Framework/Commands/GenericCommands.cpp
@@ -3,7 +3,7 @@
 #include "SlatePrivatePCH.h"
 #include "GenericCommands.h"
 
-#define LOCTEXT_NAMESPACE ""
+#define LOCTEXT_NAMESPACE "GenericCommands"
 
 void FGenericCommands::RegisterCommands()
 {
diff --git a/Engine/Source/Runtime/Slate/Private/Framework/Docking/TabCommands.cpp b/Engine/Source/Runtime/Slate/Private/Framework/Docking/TabCommands.cpp
--- a/Engine/Source/Runtime/Slate/Private/Framework/Docking/TabCommands.cpp
+++ b/Engine/Source/Runtime/Slate/Private/Framework/Docking/TabCommands.cpp
@@ -3,7 +3,7 @@
 #include "SlatePrivatePCH.h"
 #include "TabCommands.h"
 
-#define LOCTEXT_NAMESPACE ""
+#define LOCTEXT_NAMESPACE "TabCommands"
 
 void FTabCommands::RegisterCommands()
 {
diff --git a/Engine/Source/Runtime/Slate/Private/Framework/Text/TextLayout.cpp b/Engine/Source/Runtime/Slate/Private/Framework/Text/TextLayout.cpp
--- a/Engine/Source/Runtime/Slate/Private/Framework/Text/TextLayout.cpp
+++ b/Engine/Source/Runtime/Slate/Private/Framework/Text/TextLayout.cpp
@@ -1334,6 +1334,11 @@ int32 FTextLayout::GetLineViewIndexForTextLocation(const TArray< FTextLayout::FL
 	const int32 LineModelIndex = Location.GetLineIndex();
 	const int32 Offset = Location.GetOffset();
 
+	if (!LineModels.IsValidIndex(LineModelIndex))
+	{
+		return INDEX_NONE;
+	}
+
 	const FLineModel& LineModel = LineModels[LineModelIndex];
 	for(int32 Index = 0; Index < InLineViews.Num(); Index++)
 	{
@@ -1646,7 +1651,7 @@ bool FTextLayout::InsertAt(const FTextLocation& Location, TSharedRef<IRun> InRun
 			InRun->Move(LineModel.Text, FTextRange(InsertLocation, InsertLocationEnd));
 
 			// Remove the old run (it may get re-added again as the right hand run)
-			LineModel.Runs.RemoveAt(RunIndex--);
+			LineModel.Runs.RemoveAt(RunIndex--, /*bAllowShrinking*/false);
 
 			// Insert the new runs at the correct place, and then skip over these new array entries
 			const bool LeftRunHasText = !LeftRun->GetTextRange().IsEmpty();
@@ -1714,7 +1719,7 @@ bool FTextLayout::JoinLineWithNextLine(int32 LineIndex)
 	}
 
 	//Remove the next line from the list of line models
-	LineModels.RemoveAt(LineIndex + 1);
+	LineModels.RemoveAt(LineIndex + 1, /*bAllowShrinking*/false);
 
 	DirtyFlags |= ETextLayoutDirtyState::Layout;
 	return true;
@@ -1804,8 +1809,7 @@ bool FTextLayout::SplitLineAt(const FTextLocation& Location)
 		}
 	}
 
-	LineModels.RemoveAt(LineIndex);
-	LineModels.Insert(LeftLineModel, LineIndex);
+	LineModels[LineIndex] = LeftLineModel;
 	LineModels.Insert(RightLineModel, LineIndex + 1);
 
 	DirtyFlags |= ETextLayoutDirtyState::Layout;
diff --git a/Engine/Source/Runtime/Slate/Private/Widgets/Layout/SScrollBox.cpp b/Engine/Source/Runtime/Slate/Private/Widgets/Layout/SScrollBox.cpp
--- a/Engine/Source/Runtime/Slate/Private/Widgets/Layout/SScrollBox.cpp
+++ b/Engine/Source/Runtime/Slate/Private/Widgets/Layout/SScrollBox.cpp
@@ -400,9 +400,6 @@ bool SScrollBox::ScrollDescendantIntoView(const FGeometry& MyGeometry, const TSh
 
 	if ( WidgetGeometry )
 	{
-		// @todo: This is a workaround because DesiredScrollOffset can exceed the ScrollMax when mouse dragging on the scroll bar and we need it clamped here or the offset is wrong
-		ScrollBy(MyGeometry, 0, EAllowOverscroll::No, false);
-
 		float ScrollOffset = 0.0f;
 		if ( InDestination == EDescendantScrollDestination::TopOrLeft )
 		{
@@ -568,7 +565,7 @@ void SScrollBox::Tick( const FGeometry& AllottedGeometry, const double InCurrent
 
 	if ( bScrollToEnd )
 	{
-		DesiredScrollOffset = ContentSize;
+		DesiredScrollOffset = FMath::Max(ContentSize - GetScrollComponentFromVector(ScrollPanelGeometry.Size), 0.0f);
 		bScrollToEnd = false;
 	}
 
@@ -939,7 +936,11 @@ int32 SScrollBox::OnPaint( const FPaintArgs& Args, const FGeometry& AllottedGeom
 
 void SScrollBox::ScrollBar_OnUserScrolled( float InScrollOffsetFraction )
 {
-	DesiredScrollOffset = InScrollOffsetFraction * GetScrollComponentFromVector(ScrollPanel->GetDesiredSize());
+	const float ContentSize = GetScrollComponentFromVector(ScrollPanel->GetDesiredSize());
+	const FGeometry ScrollPanelGeometry = FindChildGeometry(CachedGeometry, ScrollPanel.ToSharedRef());
+
+	// Clamp to max scroll offset
+	DesiredScrollOffset = FMath::Min(InScrollOffsetFraction * ContentSize, ContentSize - GetScrollComponentFromVector(ScrollPanelGeometry.Size));
 	OnUserScrolled.ExecuteIfBound(DesiredScrollOffset);
 }
 
diff --git a/Engine/Source/Runtime/Slate/Private/Widgets/Text/SlateEditableTextLayout.cpp b/Engine/Source/Runtime/Slate/Private/Widgets/Text/SlateEditableTextLayout.cpp
--- a/Engine/Source/Runtime/Slate/Private/Widgets/Text/SlateEditableTextLayout.cpp
+++ b/Engine/Source/Runtime/Slate/Private/Widgets/Text/SlateEditableTextLayout.cpp
@@ -42,6 +42,20 @@ FORCEINLINE FReply BoolToReply(const bool bHandled)
 	return (bHandled) ? FReply::Handled() : FReply::Unhandled();
 }
 
+bool IsCharAllowed(const TCHAR InChar)
+{
+	// Certain characters are not allowed
+	if (InChar == TEXT('\t'))
+	{
+		return true;
+	}
+	else if (InChar <= 0x1F)
+	{
+		return false;
+	}
+	return true;
+}
+
 }
 
 FSlateEditableTextLayout::FSlateEditableTextLayout(ISlateEditableTextWidget& InOwnerWidget, const TAttribute<FText>& InInitialText, FTextBlockStyle InTextStyle, const TOptional<ETextShapingMethod> InTextShapingMethod, const TOptional<ETextFlowDirection> InTextFlowDirection, TSharedRef<ITextLayoutMarshaller> InTextMarshaller, TSharedRef<ITextLayoutMarshaller> InHintTextMarshaller)
@@ -1205,18 +1219,7 @@ bool FSlateEditableTextLayout::HandleTypeChar(const TCHAR InChar)
 	}
 
 	// Certain characters are not allowed
-	bool bIsCharAllowed = true;
-	{
-		if (InChar == TEXT('\t'))
-		{
-			bIsCharAllowed = true;
-		}
-		else if (InChar <= 0x1F)
-		{
-			bIsCharAllowed = false;
-		}
-	}
-
+	const bool bIsCharAllowed = IsCharAllowed(InChar);
 	if (bIsCharAllowed)
 	{
 		const FTextLocation CursorInteractionPosition = CursorInfo.GetCursorInteractionLocation();
@@ -1626,7 +1629,11 @@ void FSlateEditableTextLayout::PasteTextFromClipboard()
 	FString PastedText;
 	FPlatformMisc::ClipboardPaste(PastedText);
 
-	InsertTextAtCursorImpl(PastedText);
+	if (PastedText.Len() > 0)
+	{
+		InsertTextAtCursorImpl(PastedText);
+		TextLayout->UpdateIfNeeded();
+	}
 }
 
 void FSlateEditableTextLayout::InsertTextAtCursor(const FString& InString)
@@ -1639,27 +1646,74 @@ void FSlateEditableTextLayout::InsertTextAtCursor(const FString& InString)
 	FScopedEditableTextTransaction TextTransaction(*this);
 
 	DeleteSelectedText();
-	InsertTextAtCursorImpl(InString);
+
+	if (InString.Len() > 0)
+	{
+		InsertTextAtCursorImpl(InString);
+		TextLayout->UpdateIfNeeded();
+	}
 }
 
 void FSlateEditableTextLayout::InsertTextAtCursorImpl(const FString& InString)
 {
-	if (InString.Len())
+	if (OwnerWidget->IsTextReadOnly() || InString.Len() == 0)
 	{
-		for (const TCHAR* Character = InString.GetCharArray().GetData(); *Character; ++Character)
+		return;
+	}
+
+	// Sanitize out any invalid characters
+	FString SanitizedString = InString;
+	{
+		const bool bIsMultiLine = OwnerWidget->IsMultiLineTextEdit();
+		SanitizedString.GetCharArray().RemoveAll([&](const TCHAR InChar) -> bool
+		{
+			const bool bIsCharAllowed = IsCharAllowed(InChar) || (bIsMultiLine || !FChar::IsLinebreak(InChar));
+			return !bIsCharAllowed;
+		});
+	}
+
+	// Split into lines
+	TArray<FTextRange> LineRanges;
+	FTextRange::CalculateLineRangesFromString(SanitizedString, LineRanges);
+
+	if (AnyTextSelected())
+	{
+		// Delete selected text
+		DeleteSelectedText();
+	}
+
+	// Insert each line
+	{
+		bool bIsFirstLine = true;
+		for (const FTextRange& LineRange : LineRanges)
 		{
-			if (*Character == '\n')
+			if (!bIsFirstLine)
 			{
-				if (OwnerWidget->IsMultiLineTextEdit())
+				const FTextLocation CursorInteractionPosition = CursorInfo.GetCursorInteractionLocation();
+				if (TextLayout->SplitLineAt(CursorInteractionPosition))
 				{
-					InsertNewLineAtCursorImpl();
+					// Adjust the cursor position to be at the beginning of the new line
+					const FTextLocation NewCursorPosition = FTextLocation(CursorInteractionPosition.GetLineIndex() + 1, 0);
+					CursorInfo.SetCursorLocationAndCalculateAlignment(*TextLayout, NewCursorPosition);
 				}
 			}
-			else
-			{
-				HandleTypeChar(*Character);
-			}
+			bIsFirstLine = false;
+
+			const FString NewLineText = FString(SanitizedString.Mid(LineRange.BeginIndex, LineRange.Len()));
+
+			const FTextLocation CursorInteractionPosition = CursorInfo.GetCursorInteractionLocation();
+			const TArray< FTextLayout::FLineModel >& Lines = TextLayout->GetLineModels();
+			const FTextLayout::FLineModel& Line = Lines[CursorInteractionPosition.GetLineIndex()];
+
+			// Insert character at caret position
+			TextLayout->InsertAt(CursorInteractionPosition, NewLineText);
+
+			// Advance caret position
+			const FTextLocation NewCursorPosition = FTextLocation(CursorInteractionPosition.GetLineIndex(), FMath::Min(CursorInteractionPosition.GetOffset() + NewLineText.Len(), Line.Text->Len()));
+			CursorInfo.SetCursorLocationAndCalculateAlignment(*TextLayout, NewCursorPosition);
 		}
+
+		UpdateCursorHighlight();
 	}
 }
 
diff --git a/Engine/Source/Runtime/Slate/Public/Framework/Application/SlateApplication.h b/Engine/Source/Runtime/Slate/Public/Framework/Application/SlateApplication.h
--- a/Engine/Source/Runtime/Slate/Public/Framework/Application/SlateApplication.h
+++ b/Engine/Source/Runtime/Slate/Public/Framework/Application/SlateApplication.h
@@ -1201,6 +1201,7 @@ public:
 	virtual void OnWindowClose( const TSharedRef< FGenericWindow >& PlatformWindow ) override;
 	virtual EDropEffect::Type OnDragEnterText( const TSharedRef< FGenericWindow >& Window, const FString& Text ) override;
 	virtual EDropEffect::Type OnDragEnterFiles( const TSharedRef< FGenericWindow >& Window, const TArray< FString >& Files ) override;
+	virtual EDropEffect::Type OnDragEnterExternal( const TSharedRef< FGenericWindow >& Window, const FString& Text, const TArray< FString >& Files ) override;
 
 	EDropEffect::Type OnDragEnter( const TSharedRef< SWindow >& Window, const TSharedRef<FExternalDragOperation>& DragDropOperation );
 
diff --git a/Engine/Source/Runtime/SlateCore/Private/Input/DragAndDrop.cpp b/Engine/Source/Runtime/SlateCore/Private/Input/DragAndDrop.cpp
--- a/Engine/Source/Runtime/SlateCore/Private/Input/DragAndDrop.cpp
+++ b/Engine/Source/Runtime/SlateCore/Private/Input/DragAndDrop.cpp
@@ -122,20 +122,30 @@ void FDragDropOperation::DestroyCursorDecoratorWindow()
 /* FExternalDragOperation implementation
  *****************************************************************************/
 
-TSharedRef<FExternalDragOperation> FExternalDragOperation::NewText( const FString& InText )
+TSharedRef<FExternalDragOperation> FExternalDragOperation::NewText( FString InText )
 {
 	TSharedRef<FExternalDragOperation> Operation = MakeShareable(new FExternalDragOperation);
 	Operation->DragType = DragText;
-	Operation->DraggedText = InText;
+	Operation->DraggedText = MoveTemp(InText);
 	Operation->Construct();
 	return Operation;
 }
 
-TSharedRef<FExternalDragOperation> FExternalDragOperation::NewFiles( const TArray<FString>& InFileNames )
+TSharedRef<FExternalDragOperation> FExternalDragOperation::NewFiles( TArray<FString> InFileNames )
 {
 	TSharedRef<FExternalDragOperation> Operation = MakeShareable(new FExternalDragOperation);
 	Operation->DragType = DragFiles;
-	Operation->DraggedFileNames = InFileNames;
+	Operation->DraggedFileNames = MoveTemp(InFileNames);
+	Operation->Construct();
+	return Operation;
+}
+
+TSharedRef<FExternalDragOperation> FExternalDragOperation::NewOperation( FString InText, TArray<FString> InFileNames )
+{
+	TSharedRef<FExternalDragOperation> Operation = MakeShareable(new FExternalDragOperation);
+	Operation->DragType = 0 | (InText.IsEmpty() ? 0 : DragText) | (InFileNames.Num() == 0 ? 0 : DragFiles);
+	Operation->DraggedText = MoveTemp(InText);
+	Operation->DraggedFileNames = MoveTemp(InFileNames);
 	Operation->Construct();
 	return Operation;
 }
diff --git a/Engine/Source/Runtime/SlateCore/Private/Styling/SlateIconFinder.cpp b/Engine/Source/Runtime/SlateCore/Private/Styling/SlateIconFinder.cpp
new file mode 100644
--- /dev/null
+++ b/Engine/Source/Runtime/SlateCore/Private/Styling/SlateIconFinder.cpp
@@ -0,0 +1,68 @@
+// Copyright 1998-2016 Epic Games, Inc. All Rights Reserved.
+
+#include "SlateCorePrivatePCH.h"
+#include "SlateIconFinder.h"
+#include "SlateIcon.h"
+
+FSlateIcon FSlateIconFinder::FindIconForClass(const UClass* InClass, const FName& InDefaultName)
+{
+	return FindCustomIconForClass(InClass, TEXT("ClassIcon"), InDefaultName);
+}
+
+FSlateIcon FSlateIconFinder::FindCustomIconForClass(const UClass* InClass, const TCHAR* StyleBasePath, const FName& InDefaultName)
+{
+	FSlateIcon Icon;
+
+	FString IconPath(StyleBasePath);
+	IconPath += TEXT(".");
+	const int32 BastPathLength = IconPath.Len();
+
+	if (InClass)
+	{
+		// walk up class hierarchy until we find an icon
+		const UClass* CurrentClass = InClass;
+		while ( CurrentClass)
+		{
+			CurrentClass->AppendName(IconPath);
+			Icon = FindIcon(*IconPath);
+
+			if (Icon.IsSet())
+			{
+				return Icon;
+			}
+
+			CurrentClass = CurrentClass->GetSuperClass();
+			IconPath.RemoveAt(BastPathLength, IconPath.Len() - BastPathLength, false);
+		}
+	}
+
+	// If we didn't supply an override name for the default icon use default class icon.
+	if (InDefaultName.IsNone())
+	{
+		IconPath += TEXT("Default");
+		return FindIcon(*IconPath);
+	}
+
+	return FindIcon(InDefaultName);
+}
+
+FSlateIcon FSlateIconFinder::FindIcon(const FName& IconName)
+{
+	FSlateIcon Icon;
+
+	FSlateStyleRegistry::IterateAllStyles(
+		[&](const ISlateStyle& Style)
+		{
+			if (Style.GetOptionalBrush(IconName, nullptr, nullptr))
+			{
+				Icon = FSlateIcon(Style.GetStyleSetName(), IconName);
+				// terminate iteration
+				return false;
+			}
+			// continue iteration
+			return true;
+		}
+	);
+
+	return Icon;
+}
\ No newline at end of file
diff --git a/Engine/Source/Runtime/SlateCore/Private/Styling/SlateStyleRegistry.cpp b/Engine/Source/Runtime/SlateCore/Private/Styling/SlateStyleRegistry.cpp
--- a/Engine/Source/Runtime/SlateCore/Private/Styling/SlateStyleRegistry.cpp
+++ b/Engine/Source/Runtime/SlateCore/Private/Styling/SlateStyleRegistry.cpp
@@ -49,6 +49,17 @@ const ISlateStyle* FSlateStyleRegistry::FindSlateStyle( const FName& InSlateStyl
 	return (SlateStylePtr) ? *SlateStylePtr : nullptr;
 }
 
+bool FSlateStyleRegistry::IterateAllStyles(const TFunctionRef<bool(const ISlateStyle&)>& Iter)
+{
+	for (auto& Pair : SlateStyleRepository)
+	{
+		if (!Iter(*Pair.Value))
+		{
+			return false;
+		}
+	}
+	return true;
+}
 
 void FSlateStyleRegistry::GetAllResources( TArray< const FSlateBrush* >& OutResources )
 {
diff --git a/Engine/Source/Runtime/SlateCore/Private/Textures/SlateIcon.cpp b/Engine/Source/Runtime/SlateCore/Private/Textures/SlateIcon.cpp
new file mode 100644
--- /dev/null
+++ b/Engine/Source/Runtime/SlateCore/Private/Textures/SlateIcon.cpp
@@ -0,0 +1,79 @@
+// Copyright 1998-2016 Epic Games, Inc. All Rights Reserved.
+
+#include "SlateCorePrivatePCH.h"
+#include "SlateIcon.h"
+
+/* FSlateIcon structors
+ *****************************************************************************/
+
+FSlateIcon::FSlateIcon( )
+	: StyleSetName(NAME_None)
+	, StyleName(NAME_None)
+	, SmallStyleName(NAME_None)
+	, bIsSet(false)
+{ }
+
+
+FSlateIcon::FSlateIcon( const FName& InStyleSetName, const FName& InStyleName )
+	: StyleSetName(InStyleSetName)
+	, StyleName(InStyleName)
+	, SmallStyleName(ISlateStyle::Join(InStyleName, ".Small"))
+	, bIsSet(true)
+{ }
+
+
+FSlateIcon::FSlateIcon( const FName& InStyleSetName, const FName& InStyleName, const FName& InSmallStyleName )
+	: StyleSetName(InStyleSetName)
+	, StyleName(InStyleName)
+	, SmallStyleName(InSmallStyleName)
+	, bIsSet(true)
+{ }
+
+
+/* FSlateIcon interface
+ *****************************************************************************/
+
+const FSlateBrush* FSlateIcon::GetSmallIcon( ) const
+{
+	const ISlateStyle* StyleSet = GetStyleSet();
+	
+	if (StyleSet)
+	{
+		return StyleSet->GetOptionalBrush(SmallStyleName);
+	}
+
+	return FStyleDefaults::GetNoBrush();
+}
+
+
+const ISlateStyle* FSlateIcon::GetStyleSet( ) const
+{
+	return StyleSetName.IsNone() ? nullptr : FSlateStyleRegistry::FindSlateStyle(StyleSetName);
+}
+
+
+const FSlateBrush* FSlateIcon::GetIcon( ) const
+{
+	const ISlateStyle* StyleSet = GetStyleSet();
+	
+	if(StyleSet)
+	{
+		return StyleSet->GetOptionalBrush(StyleName);
+	}
+
+	return FStyleDefaults::GetNoBrush();
+}
+
+
+const FSlateBrush* FSlateIcon::GetOptionalIcon( ) const
+{
+	const ISlateStyle* StyleSet = GetStyleSet();
+	return StyleSet ? StyleSet->GetOptionalBrush(StyleName) : nullptr;
+}
+
+
+const FSlateBrush* FSlateIcon::GetOptionalSmallIcon( ) const
+{
+	const ISlateStyle* StyleSet = GetStyleSet();
+	return StyleSet ? StyleSet->GetOptionalBrush(SmallStyleName) : nullptr;
+}
\ No newline at end of file
diff --git a/Engine/Source/Runtime/SlateCore/Public/Input/DragAndDrop.h b/Engine/Source/Runtime/SlateCore/Public/Input/DragAndDrop.h
--- a/Engine/Source/Runtime/SlateCore/Public/Input/DragAndDrop.h
+++ b/Engine/Source/Runtime/SlateCore/Public/Input/DragAndDrop.h
@@ -210,14 +210,16 @@ public:
 	DRAG_DROP_OPERATOR_TYPE(FExternalDragOperation, FDragDropOperation)
 
 	/** Creates a new external text drag operation */
-	static TSharedRef<FExternalDragOperation> NewText( const FString& InText );
+	static TSharedRef<FExternalDragOperation> NewText( FString InText );
 	/** Creates a new external file drag operation */
-	static TSharedRef<FExternalDragOperation> NewFiles( const TArray<FString>& InFileNames );
+	static TSharedRef<FExternalDragOperation> NewFiles( TArray<FString> InFileNames );
+	/** Creates a new external combined drag operation */
+	static TSharedRef<FExternalDragOperation> NewOperation( FString InText, TArray<FString> InFileNames );
 
 	/** @return true if this is a text drag operation */
-	bool HasText() const {return DragType == DragText;}
+	bool HasText() const {return !!(DragType & DragText);}
 	/** @return true if this is a file drag operation */
-	bool HasFiles() const {return DragType == DragFiles;}
+	bool HasFiles() const {return !!(DragType & DragFiles);}
 	
 	/** @return the text from this drag operation */
 	const FString& GetText() const { ensure(HasText()); return DraggedText;}
@@ -228,11 +230,12 @@ private:
 	FString DraggedText;
 	TArray<FString> DraggedFileNames;
 
-	enum ExternalDragType
+	enum EExternalDragType
 	{
-		DragText,
-		DragFiles
-	} DragType;
+		DragText = 1<<0,
+		DragFiles = 1<<1,
+	};
+	uint8 DragType;
 };
 
 
diff --git a/Engine/Source/Runtime/SlateCore/Public/Styling/SlateIconFinder.h b/Engine/Source/Runtime/SlateCore/Public/Styling/SlateIconFinder.h
new file mode 100644
--- /dev/null
+++ b/Engine/Source/Runtime/SlateCore/Public/Styling/SlateIconFinder.h
@@ -0,0 +1,68 @@
+// Copyright 1998-2016 Epic Games, Inc. All Rights Reserved.
+
+#pragma once
+
+#include "NameTypes.h"
+#include "SlateIcon.h"
+
+struct FSlateIcon;
+
+class UClass;
+class ISlateStyle;
+
+/** Class used for finding icons within a registered set of styles */
+class SLATECORE_API FSlateIconFinder
+{
+public:
+
+	/**
+	 * Find the icon to use for the supplied class
+	 *
+	 * @param InClass			The class to locate an icon for
+	 * @param InDefaultName		The default icon name to use, if no specialized icon could be found
+	 * @return An FSlateIcon structure defining the class's icon
+	 */
+	static FSlateIcon FindIconForClass(const UClass* InClass, const FName& InDefaultName = FName());
+
+	/**
+	 * Find a custom icon to use for the supplied class, according to the specified base style
+	 *
+	 * @param InClass			The class to locate an icon for
+	 * @param InStyleBasePath	Style base path to use for the search (e.g. ClassIcon, or ClassThumbnail)
+	 * @param InDefaultName		The default icon name to use, if no specialized icon could be found
+	 * @return An FSlateIcon structure defining the class's icon
+	 */
+	static FSlateIcon FindCustomIconForClass(const UClass* InClass, const TCHAR* StyleBasePath, const FName& InDefaultName = FName());
+
+	/**
+	 * Find a slate brush to use for the supplied class's icon
+	 *
+	 * @param InClass			The class to locate an icon for
+	 * @param InDefaultName		The default icon name to use, if no specialized icon could be found
+	 * @return A slate brush, or nullptr if no icon was found
+	 */
+	static const FSlateBrush* FindIconBrushForClass(const UClass* InClass, const FName& InDefaultName = FName())
+	{
+		return FindIconForClass(InClass, InDefaultName).GetOptionalIcon();
+	}
+	
+	/**
+	 * Find a custom icon to use for the supplied class, according to the specified base style
+	 *
+	 * @param InClass			The class to locate an icon for
+	 * @param InStyleBasePath	Style base path to use for the search (e.g. ClassIcon, or ClassThumbnail)
+	 * @param InDefaultName		The default icon name to use, if no specialized icon could be found
+	 * @return A slate brush, or nullptr if no icon was found
+	 */
+	static const FSlateBrush* FindCustomIconBrushForClass(const UClass* InClass, const TCHAR* StyleBasePath, const FName& InDefaultName = FName())
+	{
+		return FindCustomIconForClass(InClass, StyleBasePath, InDefaultName).GetOptionalIcon();
+	}
+
+	/**
+	 * Find the first occurance of a brush represented by the specified IconName in any of the registered style sets
+	 * @param InIconName 		The fully qualified style name of the icon to find
+	 * @return An FSlateIcon structure
+	 */
+	static FSlateIcon FindIcon(const FName& InIconName);
+};
\ No newline at end of file
diff --git a/Engine/Source/Runtime/SlateCore/Public/Styling/SlateStyleRegistry.h b/Engine/Source/Runtime/SlateCore/Public/Styling/SlateStyleRegistry.h
--- a/Engine/Source/Runtime/SlateCore/Public/Styling/SlateStyleRegistry.h
+++ b/Engine/Source/Runtime/SlateCore/Public/Styling/SlateStyleRegistry.h
@@ -39,6 +39,14 @@ public:
 	 */
 	static const ISlateStyle* FindSlateStyle( const FName& InSlateStyleName );
 
+	/**
+	 * Iterate all the slate style known to this registry
+	 *
+	 * @param Iter A predicate to call for each known style that returns true to continue iteration, and false to terminate iteration
+	 * @return true where to loop completed, of false if it was terminated by the predicate
+	 */
+	static bool IterateAllStyles(const TFunctionRef<bool(const ISlateStyle&)>& Iter);
+
 	/**
 	 * Populate an array of slate brushes with all of the resources used by the registered styles.
 	 *
diff --git a/Engine/Source/Runtime/SlateCore/Public/Textures/SlateIcon.h b/Engine/Source/Runtime/SlateCore/Public/Textures/SlateIcon.h
--- a/Engine/Source/Runtime/SlateCore/Public/Textures/SlateIcon.h
+++ b/Engine/Source/Runtime/SlateCore/Public/Textures/SlateIcon.h
@@ -9,7 +9,7 @@ struct FSlateBrush;
 /**
  * Struct used to represent an icon in Slate
  */
-struct SLATE_API FSlateIcon
+struct SLATECORE_API FSlateIcon
 {
 	/**
 	 * Default constructor (empty icon).
@@ -33,6 +33,24 @@ struct SLATE_API FSlateIcon
 	 */
 	FSlateIcon( const FName& InStyleSetName, const FName& InStyleName, const FName& InSmallStyleName );
 
+public:
+	
+	/**
+	 * Compare 2 slate icons for equality
+	 */
+	friend bool operator==(const FSlateIcon& A, const FSlateIcon& B)
+	{
+		return A.bIsSet == B.bIsSet && A.StyleSetName == B.StyleSetName && A.StyleName == B.StyleName && A.SmallStyleName == B.SmallStyleName;
+	}
+
+	/**
+	 * Compare 2 slate icons for inequality
+	 */
+	friend bool operator!=(const FSlateIcon& A, const FSlateIcon& B)
+	{
+		return !(A == B);
+	}
+
 public:
 
 	/**
@@ -43,6 +61,13 @@ public:
 	 */
 	const FSlateBrush* GetIcon( ) const;
 
+	/**
+	 * Optionally gets the resolved icon, returning nullptr if it's not defined
+	 *
+	 * @return Icon brush, or nullptr if the icon wasn't found.
+	 */
+	const FSlateBrush* GetOptionalIcon( ) const;
+
 	/**
 	 * Gets the resolved small icon.
 	 *
@@ -51,6 +76,14 @@ public:
 	 */
 	const FSlateBrush* GetSmallIcon( ) const;
 
+
+	/**
+	 * Optionally gets the resolved small icon, returning nullptr if it's not defined
+	 *
+	 * @return Icon brush, or nullptr if the icon wasn't found.
+	 */
+	const FSlateBrush* GetOptionalSmallIcon( ) const;
+
 	/**
 	 * Gets the name of the style for the icon.
 	 *
