Fixed major stability regression with editable text controls

diff --git a/Engine/Source/Runtime/Slate/Private/Framework/Text/TextLayout.cpp b/Engine/Source/Runtime/Slate/Private/Framework/Text/TextLayout.cpp
--- a/Engine/Source/Runtime/Slate/Private/Framework/Text/TextLayout.cpp
+++ b/Engine/Source/Runtime/Slate/Private/Framework/Text/TextLayout.cpp
@@ -1651,7 +1651,7 @@ bool FTextLayout::InsertAt(const FTextLocation& Location, TSharedRef<IRun> InRun
 			InRun->Move(LineModel.Text, FTextRange(InsertLocation, InsertLocationEnd));
 
 			// Remove the old run (it may get re-added again as the right hand run)
-			LineModel.Runs.RemoveAt(RunIndex--, /*bAllowShrinking*/false);
+			LineModel.Runs.RemoveAt(RunIndex--, 1, /*bAllowShrinking*/false);
 
 			// Insert the new runs at the correct place, and then skip over these new array entries
 			const bool LeftRunHasText = !LeftRun->GetTextRange().IsEmpty();
@@ -1719,7 +1719,7 @@ bool FTextLayout::JoinLineWithNextLine(int32 LineIndex)
 	}
 
 	//Remove the next line from the list of line models
-	LineModels.RemoveAt(LineIndex + 1, /*bAllowShrinking*/false);
+	LineModels.RemoveAt(LineIndex + 1, 1, /*bAllowShrinking*/false);
 
 	DirtyFlags |= ETextLayoutDirtyState::Layout;
 	return true;
