Fixed crash that could happen when deleting a line from a ML editable text

diff --git a/Engine/Source/Runtime/Slate/Private/Widgets/Text/SlateEditableTextLayout.cpp b/Engine/Source/Runtime/Slate/Private/Widgets/Text/SlateEditableTextLayout.cpp
--- a/Engine/Source/Runtime/Slate/Private/Widgets/Text/SlateEditableTextLayout.cpp
+++ b/Engine/Source/Runtime/Slate/Private/Widgets/Text/SlateEditableTextLayout.cpp
@@ -2357,10 +2357,16 @@ void FSlateEditableTextLayout::UpdateCursorHighlight()
 
 void FSlateEditableTextLayout::RemoveCursorHighlight()
 {
+	const TArray<FTextLayout::FLineModel>& Lines = TextLayout->GetLineModels();
+
 	for (const FTextLineHighlight& LineHighlight : ActiveLineHighlights)
 	{
-		TextLayout->RemoveLineHighlight(LineHighlight);
+		if (Lines.IsValidIndex(LineHighlight.LineIndex))
+		{
+			TextLayout->RemoveLineHighlight(LineHighlight);
+		}
 	}
+
 	ActiveLineHighlights.Empty();
 }
 
