Fix for OnMouseButtonUp not resetting SButton Brush state from Pressed to Hovered.

diff --git a/Engine/Source/Runtime/Slate/Private/Framework/MultiBox/SMenuEntryBlock.cpp b/Engine/Source/Runtime/Slate/Private/Framework/MultiBox/SMenuEntryBlock.cpp
--- a/Engine/Source/Runtime/Slate/Private/Framework/MultiBox/SMenuEntryBlock.cpp
+++ b/Engine/Source/Runtime/Slate/Private/Framework/MultiBox/SMenuEntryBlock.cpp
@@ -372,8 +372,6 @@ public:
 
 	virtual FReply OnMouseButtonUp(const FGeometry& MyGeometry, const FPointerEvent& MouseEvent) override
 	{
-		bool bWasPressed = bIsPressed;
-
 		if (ResponseToMouseUp == EResponseToMouseUp::Handle)
 		{
 			bIsPressed = true;
@@ -381,7 +379,6 @@ public:
 
 		FReply Reply = SButton::OnMouseButtonUp(MyGeometry, MouseEvent);
 
-		bIsPressed = bWasPressed;
 		bIsHovered = true;
 
 		return Reply;
