Fixing non-unity compilation error with DriverMetaData.h

diff --git a/Engine/Source/Runtime/Slate/Public/Framework/MetaData/DriverMetaData.h b/Engine/Source/Runtime/Slate/Public/Framework/MetaData/DriverMetaData.h
--- a/Engine/Source/Runtime/Slate/Public/Framework/MetaData/DriverMetaData.h
+++ b/Engine/Source/Runtime/Slate/Public/Framework/MetaData/DriverMetaData.h
@@ -1,6 +1,8 @@
 // Copyright 1998-2016 Epic Games, Inc. All Rights Reserved.
 #pragma once
 
+#include "CoreMinimal.h"
+
 class ISlateMetaData;
 
 /**
