Copying //UE4/Dev-Editor to //UE4/Dev-Main (Source: //UE4/Dev-Editor @ 3228984)

diff --git a/Engine/Source/Runtime/Launch/Private/LaunchEngineLoop.cpp b/Engine/Source/Runtime/Launch/Private/LaunchEngineLoop.cpp
--- a/Engine/Source/Runtime/Launch/Private/LaunchEngineLoop.cpp
+++ b/Engine/Source/Runtime/Launch/Private/LaunchEngineLoop.cpp
@@ -1662,6 +1662,64 @@ int32 FEngineLoop::PreInit( const TCHAR* CmdLine )
 			GetDerivedDataCacheRef();
 		}
 
+
+		// If platforms support early movie playback we have to start the rendering thread much earlier
+#if PLATFORM_SUPPORTS_EARLY_MOVIE_PLAYBACK
+		RHIPostInit();
+
+		if(GUseThreadedRendering)
+		{
+			if(GRHISupportsRHIThread)
+			{
+				const bool DefaultUseRHIThread = true;
+				GUseRHIThread = DefaultUseRHIThread;
+				if(FParse::Param(FCommandLine::Get(), TEXT("rhithread")))
+				{
+					GUseRHIThread = true;
+				}
+				else if(FParse::Param(FCommandLine::Get(), TEXT("norhithread")))
+				{
+					GUseRHIThread = false;
+				}
+			}
+			StartRenderingThread();
+		}
+#endif
+
+#if !UE_SERVER// && !UE_EDITOR
+		if(!IsRunningDedicatedServer() && !IsRunningCommandlet())
+		{
+			TSharedRef<FSlateRenderer> SlateRenderer = GUsingNullRHI ?
+				FModuleManager::Get().LoadModuleChecked<ISlateNullRendererModule>("SlateNullRenderer").CreateSlateNullRenderer() :
+				FModuleManager::Get().GetModuleChecked<ISlateRHIRendererModule>("SlateRHIRenderer").CreateSlateRHIRenderer();
+
+			// If Slate is being used, initialize the renderer after RHIInit 
+			FSlateApplication& CurrentSlateApp = FSlateApplication::Get();
+			CurrentSlateApp.InitializeRenderer(SlateRenderer);
+
+			// Create the engine font services now that the Slate renderer is ready
+			FEngineFontServices::Create();
+
+			GetMoviePlayer()->SetSlateRenderer(SlateRenderer);
+
+			// allow the movie player to load a sequence from the .inis (a PreLoadingScreen module could have already initialized a sequence, in which case
+			// it wouldn't have anything in it's .ini file)
+			GetMoviePlayer()->SetupLoadingScreenFromIni();
+
+			if(GetMoviePlayer()->HasEarlyStartupMovie())
+			{
+				GetMoviePlayer()->Initialize();
+
+				// hide splash screen now
+				FPlatformMisc::PlatformPostInit(false);
+
+				// only allowed to play any movies marked as early startup.  These movies or widgets can have no interaction whatsoever with uobjects or engine features
+				GetMoviePlayer()->PlayEarlyStartupMovies();
+			}
+
+		}
+#endif
+
 		// In order to be able to use short script package names get all script
 		// package names from ini files and register them with FPackageName system.
 		FPackageName::RegisterShortPackageNamesForUObjectModules();
@@ -1706,44 +1764,30 @@ int32 FEngineLoop::PreInit( const TCHAR* CmdLine )
 		return 1;
 	}
 
-#if !UE_SERVER// && !UE_EDITOR
-	if (!IsRunningDedicatedServer() && !IsRunningCommandlet())
-	{
-		TSharedRef<FSlateRenderer> SlateRenderer = GUsingNullRHI ?
-			FModuleManager::Get().LoadModuleChecked<ISlateNullRendererModule>("SlateNullRenderer").CreateSlateNullRenderer() :
-			FModuleManager::Get().GetModuleChecked<ISlateRHIRendererModule>("SlateRHIRenderer").CreateSlateRHIRenderer();
-
-		// If Slate is being used, initialize the renderer after RHIInit 
-		FSlateApplication& CurrentSlateApp = FSlateApplication::Get();
-		CurrentSlateApp.InitializeRenderer( SlateRenderer );
-
-		GetMoviePlayer()->SetSlateRenderer(SlateRenderer);
-	}
-
-	// Create the engine font services now that the Slate renderer is ready
-	FEngineFontServices::Create();
-#endif
 
 	SlowTask.EnterProgressFrame(10);
-	
+
 	// Load up all modules that need to hook into the loading screen
 	if (!IProjectManager::Get().LoadModulesForProject(ELoadingPhase::PreLoadingScreen) || !IPluginManager::Get().LoadModulesForEnabledPlugins(ELoadingPhase::PreLoadingScreen))
 	{
 		return 1;
 	}
 
-#if !UE_SERVER
-	if ( !IsRunningDedicatedServer() )
-	{
-		// @todo ps4: If a loading movie starts earlier, which it probably should, then please see PS4's PlatformPostInit() implementation!
-
-		// allow the movie player to load a sequence from the .inis (a PreLoadingScreen module could have already initialized a sequence, in which case
-		// it wouldn't have anything in it's .ini file)
-		GetMoviePlayer()->SetupLoadingScreenFromIni();
 
+#if !UE_SERVER// && !UE_EDITOR
+	if(!IsRunningDedicatedServer() && !IsRunningCommandlet() && !GetMoviePlayer()->IsMovieCurrentlyPlaying())
+	{
 		GetMoviePlayer()->Initialize();
+
+		// Play any non-early startup loading movies.
 		GetMoviePlayer()->PlayMovie();
+	
+	}
+#endif
 
+#if !UE_SERVER
+	if ( !IsRunningDedicatedServer() )
+	{
 		// do any post appInit processing, before the render thread is started.
 		FPlatformMisc::PlatformPostInit(!GetMoviePlayer()->IsMovieCurrentlyPlaying());
 	}
@@ -1755,6 +1799,7 @@ int32 FEngineLoop::PreInit( const TCHAR* CmdLine )
 	}
 	SlowTask.EnterProgressFrame(5);
 
+#if !PLATFORM_SUPPORTS_EARLY_MOVIE_PLAYBACK
 	RHIPostInit();
 
 	if (GUseThreadedRendering)
@@ -1774,6 +1819,7 @@ int32 FEngineLoop::PreInit( const TCHAR* CmdLine )
 		}
 		StartRenderingThread();
 	}
+#endif // !PLATFORM_SUPPORTS_EARLY_MOVIE_PLAYBACK
 	
 #if WITH_EDITOR
 	// We need to mount the shared resources for templates (if there are any) before we try and load and game classes
diff --git a/Engine/Source/Runtime/MoviePlayer/Private/DefaultGameMoviePlayer.cpp b/Engine/Source/Runtime/MoviePlayer/Private/DefaultGameMoviePlayer.cpp
--- a/Engine/Source/Runtime/MoviePlayer/Private/DefaultGameMoviePlayer.cpp
+++ b/Engine/Source/Runtime/MoviePlayer/Private/DefaultGameMoviePlayer.cpp
@@ -17,6 +17,8 @@
 #include "ShaderCompiler.h"
 #include "IHeadMountedDisplay.h"
 #include "IStereoLayers.h"
+#include "ConfigCacheIni.h"
+#include "FileManager.h"
 
 DEFINE_LOG_CATEGORY_STATIC(LogMoviePlayer, Log, All);
 
@@ -143,6 +145,11 @@ void FDefaultGameMoviePlayer::SetSlateRenderer(TSharedPtr<FSlateRenderer> InSlat
 
 void FDefaultGameMoviePlayer::Initialize()
 {
+	if(bInitialized)
+	{
+		return;
+	}
+
 	UE_LOG(LogMoviePlayer, Log, TEXT("Initializing movie player"));
 
 	ENQUEUE_UNIQUE_RENDER_COMMAND_ONEPARAMETER(RegisterMoviePlayerTickable, FDefaultGameMoviePlayer*, MoviePlayer, this,
@@ -265,6 +272,25 @@ void FDefaultGameMoviePlayer::SetupLoadingScreen(const FLoadingScreenAttributes&
 	LoadingScreenAttributes = InLoadingScreenAttributes;
 }
 
+bool FDefaultGameMoviePlayer::HasEarlyStartupMovie() const
+{
+#if PLATFORM_SUPPORTS_EARLY_MOVIE_PLAYBACK
+	return LoadingScreenAttributes.bAllowInEarlyStartup == true;
+#else
+	return false;
+#endif
+}
+
+bool FDefaultGameMoviePlayer::PlayEarlyStartupMovies()
+{
+	if(HasEarlyStartupMovie())
+	{
+		return PlayMovie();
+	}
+
+	return false;
+}
+
 bool FDefaultGameMoviePlayer::PlayMovie()
 {
 	bool bBeganPlaying = false;
@@ -524,30 +550,49 @@ void FDefaultGameMoviePlayer::SetupLoadingScreenFromIni()
 	{
 		// fill out the attributes
 		FLoadingScreenAttributes LoadingScreen;
-		LoadingScreen.bAutoCompleteWhenLoadingCompletes = !GetDefault<UMoviePlayerSettings>()->bWaitForMoviesToComplete;
-		LoadingScreen.bMoviesAreSkippable = GetDefault<UMoviePlayerSettings>()->bMoviesAreSkippable;
 
-		// look in the settings for a list of loading movies
-		GetMutableDefault<UMoviePlayerSettings>()->LoadConfig();
-		const TArray<FString>& StartupMovies = GetDefault<UMoviePlayerSettings>()->StartupMovies;
+		// Note: this code is executed too early so we cannot access UMoviePlayerSettings because the configs for that object have not been loaded and coalesced .  Have to read directly from the configs instead
+		GConfig->GetBool(TEXT("/Script/MoviePlayer.MoviePlayerSettings"), TEXT("bWaitForMoviesToComplete"), LoadingScreen.bAutoCompleteWhenLoadingCompletes, GGameIni);
+		GConfig->GetBool(TEXT("/Script/MoviePlayer.MoviePlayerSettings"), TEXT("bMoviesAreSkippable"), LoadingScreen.bMoviesAreSkippable, GGameIni);
+
+		TArray<FString> StartupMovies;
+		GConfig->GetArray(TEXT("/Script/MoviePlayer.MoviePlayerSettings"), TEXT("StartupMovies"), StartupMovies, GGameIni);
 
 		if (StartupMovies.Num() == 0)
 		{
-			LoadingScreen.MoviePaths.Add(TEXT("Default_Startup"));
+			StartupMovies.Add(TEXT("Default_Startup"));
 		}
-		else
+
+		// double check that the movies exist
+		// We dont know the extension so compare against any file in the directory with the same name for now
+		// @todo New Movie Player: movies should have the extension on them when set via the project settings
+		TArray<FString> ExistingMovieFiles;
+		IFileManager::Get().FindFiles(ExistingMovieFiles, *(FPaths::GameContentDir() + TEXT("Movies")));
+
+		bool bHasValidMovie = false;
+		for(const FString& Movie : StartupMovies)
 		{
-			for (const FString& Movie : StartupMovies)
-			{
-				if (!Movie.IsEmpty())
+			bool bFound = ExistingMovieFiles.ContainsByPredicate(
+				[&Movie](const FString& ExistingMovie)
 				{
-					LoadingScreen.MoviePaths.Add(Movie);
-				}
+					return ExistingMovie.Contains(Movie);
+				});
+
+			if(bFound)
+			{
+				bHasValidMovie = true;
+				LoadingScreen.MoviePaths.Add(Movie);
 			}
 		}
 
-		// now setup the actual loading screen
-		SetupLoadingScreen(LoadingScreen);
+		if(bHasValidMovie)
+		{
+			// These movies are all considered safe to play in very early startup sequences
+			LoadingScreen.bAllowInEarlyStartup = true;
+
+			// now setup the actual loading screen
+			SetupLoadingScreen(LoadingScreen);
+		}
 	}
 }
 
diff --git a/Engine/Source/Runtime/MoviePlayer/Private/DefaultGameMoviePlayer.h b/Engine/Source/Runtime/MoviePlayer/Private/DefaultGameMoviePlayer.h
--- a/Engine/Source/Runtime/MoviePlayer/Private/DefaultGameMoviePlayer.h
+++ b/Engine/Source/Runtime/MoviePlayer/Private/DefaultGameMoviePlayer.h
@@ -30,6 +30,8 @@ public:
 	virtual void Shutdown() override;
 	virtual void PassLoadingScreenWindowBackToGame() const override;
 	virtual void SetupLoadingScreen(const FLoadingScreenAttributes& LoadingScreenAttributes) override;
+	virtual bool HasEarlyStartupMovie() const override;
+	virtual bool PlayEarlyStartupMovies() override;
 	virtual bool PlayMovie() override;
 	virtual void StopMovie() override;
 	virtual void WaitForMovieToFinish() override;
diff --git a/Engine/Source/Runtime/MoviePlayer/Private/NullMoviePlayer.h b/Engine/Source/Runtime/MoviePlayer/Private/NullMoviePlayer.h
--- a/Engine/Source/Runtime/MoviePlayer/Private/NullMoviePlayer.h
+++ b/Engine/Source/Runtime/MoviePlayer/Private/NullMoviePlayer.h
@@ -26,6 +26,8 @@ public:
 	virtual void Shutdown() override {}
 	virtual void PassLoadingScreenWindowBackToGame() const override {}
 	virtual void SetupLoadingScreen(const FLoadingScreenAttributes& InLoadingScreenAttributes) override {}
+	virtual bool HasEarlyStartupMovie() const override { return false; }
+	virtual bool PlayEarlyStartupMovies() override { return false; }
 	virtual bool PlayMovie() override { return false; }
 	virtual void StopMovie() override {}
 	virtual void WaitForMovieToFinish() override {}
diff --git a/Engine/Source/Runtime/MoviePlayer/Public/MoviePlayer.h b/Engine/Source/Runtime/MoviePlayer/Public/MoviePlayer.h
--- a/Engine/Source/Runtime/MoviePlayer/Public/MoviePlayer.h
+++ b/Engine/Source/Runtime/MoviePlayer/Public/MoviePlayer.h
@@ -9,6 +9,10 @@
 #include "Widgets/SWidget.h"
 #include "Slate/SlateTextures.h"
 
+// In order for a platform to support early movie playback, the platform must support the rendering thread 
+// starting very early and support rendering as soon as it is started and the module containing the movie streamer for the platform must already be loaded
+#define PLATFORM_SUPPORTS_EARLY_MOVIE_PLAYBACK 0
+
 UENUM()
 enum EMoviePlaybackType
 {
@@ -105,6 +109,7 @@ struct MOVIEPLAYER_API FLoadingScreenAttributes
 		, bAutoCompleteWhenLoadingCompletes(true)
 		, bMoviesAreSkippable(true)
 		, bWaitForManualStop(false)
+		, bAllowInEarlyStartup(false)
 		, PlaybackType(EMoviePlaybackType::MT_Normal) {}
 
 	/** The widget to be displayed on top of the movie or simply standalone if there is no movie. */
@@ -125,6 +130,9 @@ struct MOVIEPLAYER_API FLoadingScreenAttributes
 	/** If true, movie playback continues until Stop is called. */
 	bool bWaitForManualStop;
 
+	/** If true loading screens here cannot have any uobjects of any kind or use any engine features at all. This will start the movies very early as a result on platforms that support it */
+	bool bAllowInEarlyStartup;
+
 	/** Should we just play back, loop, etc.  NOTE: if the playback type is MT_LoopLast, then bAutoCompleteWhenLoadingCompletes will be togged on when the last movie is hit*/
 	TEnumAsByte<EMoviePlaybackType> PlaybackType;
 
@@ -162,6 +170,18 @@ public:
 	DECLARE_EVENT(IGameMoviePlayer, FOnPrepareLoadingScreen)
 	virtual FOnPrepareLoadingScreen& OnPrepareLoadingScreen() = 0;
 
+	/**
+	 * @return true if the movie player has been set up with an early startup movie
+	 */
+	virtual bool HasEarlyStartupMovie() const = 0;
+
+	/**
+	 * Play any early start up movies that have been set up
+	 *
+	 * @return true if a movie started playing
+	 */
+	virtual bool PlayEarlyStartupMovies() = 0;
+
 	/** 
 	 * Starts playing the movie given the last FLoadingScreenAttributes passed in
 	 * @return true of a movie started playing.
diff --git a/Engine/Source/Runtime/RenderCore/Private/RenderingThread.cpp b/Engine/Source/Runtime/RenderCore/Private/RenderingThread.cpp
--- a/Engine/Source/Runtime/RenderCore/Private/RenderingThread.cpp
+++ b/Engine/Source/Runtime/RenderCore/Private/RenderingThread.cpp
@@ -922,7 +922,9 @@ static void GameThreadWaitForTask(const FGraphEventRef& Task, bool bEmptyGameThr
 				}
 				bDone = Event->Wait(WaitTime);
 
-				if (!bDone && !GIsEditor) // editor threads can block for quite a while...
+				// editor threads can block for quite a while... 
+				// Also dont do this when attached to the debugger as long pauses on breakpoints will cause this to trigger
+				if (!bDone && !GIsEditor && !FPlatformMisc::IsDebuggerPresent())
 				{
 					// Fatal timeout if we run out of time and this thread is being monitor for heartbeats
 					// (We could just let the heartbeat monitor error for us, but this leads to better diagnostics).
diff --git a/Engine/Source/Runtime/Slate/Private/Framework/Application/SlateApplication.cpp b/Engine/Source/Runtime/Slate/Private/Framework/Application/SlateApplication.cpp
--- a/Engine/Source/Runtime/Slate/Private/Framework/Application/SlateApplication.cpp
+++ b/Engine/Source/Runtime/Slate/Private/Framework/Application/SlateApplication.cpp
@@ -2149,6 +2149,7 @@ TSharedPtr<IMenu> FSlateApplication::PushMenu(const TSharedRef<SWidget>& InParen
 		return MenuStack.Push(WidgetPath, InContent, SummonLocation, TransitionEffect, bFocusImmediately, SummonLocationSize, Method, bIsCollapsedByParent);
 	}
 
+	UE_LOG(LogSlate, Warning, TEXT("Menu could not be pushed.  A path to the parent widget(%s) could not be found"), *InParentWidget->ToString());
 	return TSharedPtr<IMenu>();
 }
 
@@ -3771,7 +3772,7 @@ void FSlateApplication::EnterDebuggingMode()
 	// The scissor rect stack must be reset when re-entering the tick loop to avoid graphical artifacts with existing clip rects applied new widgets
 	TOptional<FShortRect> PreviousScissorRect = GSlateScissorRect;
 	GSlateScissorRect.Reset();
-
+	
 	// Tick slate from here in the event that we should not return until the modal window is closed.
 	while (!bRequestLeaveDebugMode)
 	{
@@ -5264,7 +5265,7 @@ bool FSlateApplication::RoutePointerMoveEvent(const FWidgetPath& WidgetsUnderPoi
 				}
 				return FNoReply();
 			});
-			
+
 			FReply Reply = FEventRouter::Route<FReply>( this, FEventRouter::FToLeafmostPolicy(MouseCaptorPath), PointerEvent, [this]( const FArrangedWidget& MouseCaptorWidget, const FPointerEvent& Event )
 			{
 				FReply TempReply = FReply::Unhandled();
@@ -5638,6 +5639,11 @@ bool FSlateApplication::OnMouseMove()
 
 bool FSlateApplication::OnRawMouseMove( const int32 X, const int32 Y )
 {
+	if (bIsFakingTouched)
+	{
+		return OnTouchMoved(GetCursorPos(), 0, 0);
+	}
+	
 	if ( X != 0 || Y != 0 )
 	{
 		FPointerEvent MouseEvent(
@@ -5906,7 +5912,6 @@ void FSlateApplication::ProcessMotionDetectedEvent( FMotionEvent& MotionEvent )
 	
 	if ( FSlateUser* User = GetOrCreateUser(MotionEvent.GetUserIndex()) )
 	{
-
 		if (InputPreProcessor.IsValid() && InputPreProcessor->HandleMotionDetectedEvent(*this, MotionEvent))
 		{
 			return;
diff --git a/Engine/Source/Runtime/Slate/Private/Framework/Application/SlateSettings.cpp b/Engine/Source/Runtime/Slate/Private/Framework/Application/SlateSettings.cpp
--- a/Engine/Source/Runtime/Slate/Private/Framework/Application/SlateSettings.cpp
+++ b/Engine/Source/Runtime/Slate/Private/Framework/Application/SlateSettings.cpp
@@ -10,5 +10,6 @@ USlateSettings::USlateSettings(const FObjectInitializer& ObjectInitializer)
 	: Super(ObjectInitializer)
 	, bVirtualKeyboardSendsTextChanged(false)
 	, bVirtualKeyboardDisplayOnFocus(true)
+	, bExplicitCanvasChildZOrder(false)
 {
 }
diff --git a/Engine/Source/Runtime/Slate/Private/Framework/Commands/UICommandInfo.cpp b/Engine/Source/Runtime/Slate/Private/Framework/Commands/UICommandInfo.cpp
--- a/Engine/Source/Runtime/Slate/Private/Framework/Commands/UICommandInfo.cpp
+++ b/Engine/Source/Runtime/Slate/Private/Framework/Commands/UICommandInfo.cpp
@@ -10,7 +10,7 @@
 #include "Widgets/SToolTip.h"
 
 
-FSimpleMulticastDelegate FBindingContext::CommandsChanged;
+FOnBindingContextChanged FBindingContext::CommandsChanged;
 
 
 FUICommandInfoDecl FBindingContext::NewCommand( const FName InCommandName, const FText& InCommandLabel, const FText& InCommandDesc )
diff --git a/Engine/Source/Runtime/Slate/Private/Framework/Commands/UICommandList.cpp b/Engine/Source/Runtime/Slate/Private/Framework/Commands/UICommandList.cpp
--- a/Engine/Source/Runtime/Slate/Private/Framework/Commands/UICommandList.cpp
+++ b/Engine/Source/Runtime/Slate/Private/Framework/Commands/UICommandList.cpp
@@ -78,13 +78,16 @@ void FUICommandList::Append( const TSharedRef<FUICommandList>& InCommandsToAppen
 
 void FUICommandList::UnmapAction( const TSharedPtr< const FUICommandInfo > InUICommandInfo )
 {
-	// Check against already mapped actions
-	checkfSlow(UICommandBindingMap.Contains(InUICommandInfo), TEXT("Command list does not already contain a command named '%s'"), *InUICommandInfo->GetCommandName().ToString());
-
 	UICommandBindingMap.Remove( InUICommandInfo );
 }
 
-bool FUICommandList::ExecuteAction( const TSharedRef< const FUICommandInfo > InUICommandInfo ) const
+bool FUICommandList::IsActionMapped(const TSharedPtr< const FUICommandInfo > InUICommandInfo) const
+{
+	// Check against already mapped actions
+	return UICommandBindingMap.Contains(InUICommandInfo);
+}
+
+bool FUICommandList::ExecuteAction(const TSharedRef< const FUICommandInfo > InUICommandInfo) const
 {
 	const FUIAction* Action = GetActionForCommand(InUICommandInfo);
 
diff --git a/Engine/Source/Runtime/Slate/Private/Framework/Text/SlateTextRun.cpp b/Engine/Source/Runtime/Slate/Private/Framework/Text/SlateTextRun.cpp
--- a/Engine/Source/Runtime/Slate/Private/Framework/Text/SlateTextRun.cpp
+++ b/Engine/Source/Runtime/Slate/Private/Framework/Text/SlateTextRun.cpp
@@ -112,11 +112,29 @@ int32 FSlateTextRun::OnPaint( const FPaintArgs& Args, const FTextLayout::FLineVi
 	// Draw the optional shadow
 	if (ShouldDropShadow)
 	{
+		FShapedGlyphSequenceRef ShadowShapedText = ShapedText;
+		if(Style.ShadowColorAndOpacity != Style.Font.OutlineSettings.OutlineColor)
+		{
+			// Copy font info for shadow to replace the outline color
+			FSlateFontInfo ShadowFontInfo = Style.Font;
+			ShadowFontInfo.OutlineSettings.OutlineColor = Style.ShadowColorAndOpacity;
+			ShadowFontInfo.OutlineSettings.OutlineMaterial = nullptr;
+			
+			// Create new shaped text for drop shadow
+			ShadowShapedText = ShapedTextCacheUtil::GetShapedTextSubSequence(
+				BlockTextContext.ShapedTextCache,
+				FCachedShapedTextKey(Line.Range, AllottedGeometry.GetAccumulatedLayoutTransform().GetScale(), BlockTextContext, ShadowFontInfo),
+				BlockRange,
+				**Text,
+				BlockTextContext.TextDirection
+				);
+		}
+
 		FSlateDrawElement::MakeShapedText(
 			OutDrawElements,
 			++LayerId,
 			AllottedGeometry.ToPaintGeometry(TransformVector(InverseScale, Block->GetSize()), FSlateLayoutTransform(TransformPoint(InverseScale, Block->GetLocationOffset() + DrawShadowOffset))),
-			ShapedText,
+			ShadowShapedText,
 			ClippingRect,
 			DrawEffects,
 			InWidgetStyle.GetColorAndOpacityTint() * Style.ShadowColorAndOpacity
diff --git a/Engine/Source/Runtime/Slate/Private/Framework/Text/TextLayout.cpp b/Engine/Source/Runtime/Slate/Private/Framework/Text/TextLayout.cpp
--- a/Engine/Source/Runtime/Slate/Private/Framework/Text/TextLayout.cpp
+++ b/Engine/Source/Runtime/Slate/Private/Framework/Text/TextLayout.cpp
@@ -1977,7 +1977,8 @@ bool FTextLayout::RemoveLine(int32 LineIndex)
 	if (!(DirtyFlags & ETextLayoutDirtyState::Layout))
 	{
 		//Lots of room for additional optimization
-		float OffsetAdjustment = 0;
+		float OffsetAdjustment = 0.0f;
+		float HeightAdjustment = 0.0f;
 
 		for (int32 ViewIndex = 0; ViewIndex < LineViews.Num(); ViewIndex++)
 		{
@@ -1985,15 +1986,11 @@ bool FTextLayout::RemoveLine(int32 LineIndex)
 
 			if (LineView.ModelIndex == LineIndex)
 			{
-				if (ViewIndex - 1 <= 0)
-				{
-					OffsetAdjustment += LineView.Offset.Y;
-				}
-				else
+				HeightAdjustment += LineView.Size.Y;
+
+				if (LineViews.IsValidIndex(ViewIndex + 1))
 				{
-					//Since the offsets are not relative to other lines, if we aren't removing the top line then
-					//we don't aggregate the offset from any previous removals as we'd be double counting.
-					OffsetAdjustment = (LineView.Offset.Y - LineViews[ViewIndex - 1].Offset.Y);
+					OffsetAdjustment += (LineViews[ViewIndex + 1].Offset.Y - LineView.Offset.Y);
 				}
 
 				LineViews.RemoveAt(ViewIndex);
@@ -2014,6 +2011,8 @@ bool FTextLayout::RemoveLine(int32 LineIndex)
 				}
 			}
 		}
+
+		TextLayoutSize.Height -= HeightAdjustment;
 	}
 
 	return true;
diff --git a/Engine/Source/Runtime/Slate/Private/Widgets/Colors/SColorGradingWheel.cpp b/Engine/Source/Runtime/Slate/Private/Widgets/Colors/SColorGradingWheel.cpp
new file mode 100644
--- /dev/null
+++ b/Engine/Source/Runtime/Slate/Private/Widgets/Colors/SColorGradingWheel.cpp
@@ -0,0 +1,187 @@
+// Copyright 1998-2016 Epic Games, Inc. All Rights Reserved.
+#include "SColorGradingWheel.h"
+#include "CoreStyle.h"
+#include "RenderingCommon.h"
+#include "DrawElements.h"
+
+
+
+/* SColorGradingWheel methods
+ *****************************************************************************/
+
+void SColorGradingWheel::Construct(const FArguments& InArgs)
+{
+	Image = FCoreStyle::Get().GetBrush("ColorGradingWheel.HueValueCircle");
+	SelectorImage = FCoreStyle::Get().GetBrush("ColorGradingWheel.Selector");
+	SelectedColor = InArgs._SelectedColor;
+	DesiredWheelSize = InArgs._DesiredWheelSize;
+	ExponentDisplacement = InArgs._ExponentDisplacement;
+	OnMouseCaptureBegin = InArgs._OnMouseCaptureBegin;
+	OnMouseCaptureEnd = InArgs._OnMouseCaptureEnd;
+	OnValueChanged = InArgs._OnValueChanged;
+}
+
+
+/* SWidget overrides
+ *****************************************************************************/
+
+FVector2D SColorGradingWheel::ComputeDesiredSize(float) const
+{
+	if (DesiredWheelSize.IsSet())
+	{
+		int32 CachedDesiredWheelSize = DesiredWheelSize.Get();
+		return FVector2D(CachedDesiredWheelSize, CachedDesiredWheelSize);
+	}
+	return Image->ImageSize;
+}
+
+
+FReply SColorGradingWheel::OnMouseButtonDoubleClick(const FGeometry& InMyGeometry, const FPointerEvent& InMouseEvent)
+{
+	return FReply::Handled();
+}
+
+
+FReply SColorGradingWheel::OnMouseButtonDown(const FGeometry& MyGeometry, const FPointerEvent& MouseEvent)
+{
+	if (MouseEvent.GetEffectingButton() == EKeys::LeftMouseButton)
+	{
+		OnMouseCaptureBegin.ExecuteIfBound();
+
+		if (!ProcessMouseAction(MyGeometry, MouseEvent, false))
+		{
+			OnMouseCaptureEnd.ExecuteIfBound();
+			return FReply::Unhandled();
+		}
+
+		return FReply::Handled().CaptureMouse(SharedThis(this));
+	}
+
+	return FReply::Unhandled();
+}
+
+
+FReply SColorGradingWheel::OnMouseButtonUp(const FGeometry& MyGeometry, const FPointerEvent& MouseEvent)
+{
+	if (MouseEvent.GetEffectingButton() == EKeys::LeftMouseButton && HasMouseCapture())
+	{
+		OnMouseCaptureEnd.ExecuteIfBound();
+
+		return FReply::Handled().ReleaseMouseCapture();
+	}
+
+	return FReply::Unhandled();
+}
+
+
+FReply SColorGradingWheel::OnMouseMove(const FGeometry& MyGeometry, const FPointerEvent& MouseEvent)
+{
+	if (!HasMouseCapture())
+	{
+		return FReply::Unhandled();
+	}
+
+	ProcessMouseAction(MyGeometry, MouseEvent, true);
+
+	return FReply::Handled();
+}
+
+
+int32 SColorGradingWheel::OnPaint(const FPaintArgs& Args, const FGeometry& AllottedGeometry, const FSlateRect& MyClippingRect, FSlateWindowElementList& OutDrawElements, int32 LayerId, const FWidgetStyle& InWidgetStyle, bool bParentEnabled) const
+{
+	const bool bIsEnabled = ShouldBeEnabled(bParentEnabled);
+	const uint32 DrawEffects = bIsEnabled ? ESlateDrawEffect::None : ESlateDrawEffect::DisabledEffect;
+	const FVector2D& SelectorSize = SelectorImage->ImageSize;
+	FVector2D CircleSize = AllottedGeometry.Size - SelectorSize;
+	FVector2D AllottedGeometrySize = AllottedGeometry.Size;
+	if (DesiredWheelSize.IsSet())
+	{
+		int32 CachedDesiredWheelSize = DesiredWheelSize.Get();
+		CircleSize.X = (float)CachedDesiredWheelSize - SelectorSize.X;
+		CircleSize.Y = (float)CachedDesiredWheelSize - SelectorSize.Y;
+		AllottedGeometrySize.X = CachedDesiredWheelSize;
+		AllottedGeometrySize.Y = CachedDesiredWheelSize;
+	}
+	
+	FSlateDrawElement::MakeBox(
+		OutDrawElements,
+		LayerId,
+		AllottedGeometry.ToPaintGeometry(0.5 * SelectorSize, CircleSize),
+		Image,
+		MyClippingRect,
+		DrawEffects,
+		InWidgetStyle.GetColorAndOpacityTint() * Image->GetTint(InWidgetStyle)
+	);
+
+	FSlateDrawElement::MakeBox(
+		OutDrawElements,
+		LayerId + 1,
+		AllottedGeometry.ToPaintGeometry(0.5f * (AllottedGeometrySize + CalcRelativePositionFromCenter() * CircleSize - SelectorSize), SelectorSize),
+		SelectorImage,
+		MyClippingRect,
+		DrawEffects,
+		InWidgetStyle.GetColorAndOpacityTint() * SelectorImage->GetTint(InWidgetStyle)
+	);
+
+	return LayerId + 1;
+}
+
+
+/* SColorGradingWheel implementation
+ *****************************************************************************/
+
+FVector2D SColorGradingWheel::CalcRelativePositionFromCenter() const
+{
+	float Hue = SelectedColor.Get().R;
+	float Saturation = SelectedColor.Get().G;
+	float Angle = Hue / 180.0f * PI;
+	float Radius = Saturation;
+
+	return FVector2D(FMath::Cos(Angle), FMath::Sin(Angle)) * Radius;
+}
+
+
+bool SColorGradingWheel::ProcessMouseAction(const FGeometry& MyGeometry, const FPointerEvent& MouseEvent, bool bProcessWhenOutsideColorWheel)
+{
+	FVector2D GeometrySize = MyGeometry.Size;
+	if (DesiredWheelSize.IsSet())
+	{
+		int32 CachedDesiredWheelSize = DesiredWheelSize.Get();
+		GeometrySize.X = CachedDesiredWheelSize;
+		GeometrySize.Y = CachedDesiredWheelSize;
+	}
+
+	const FVector2D LocalMouseCoordinate = MyGeometry.AbsoluteToLocal(MouseEvent.GetScreenSpacePosition());
+	const FVector2D RelativePositionFromCenter = (2.0f * LocalMouseCoordinate - GeometrySize) / (GeometrySize - SelectorImage->ImageSize);
+	const float RelativeRadius = RelativePositionFromCenter.Size();
+
+	if (RelativeRadius <= 1.0f || bProcessWhenOutsideColorWheel)
+	{
+		float Angle = FMath::Atan2(RelativePositionFromCenter.Y, RelativePositionFromCenter.X);
+
+		if (Angle < 0.0f)
+		{
+			Angle += 2.0f * PI;
+		}
+
+		FLinearColor NewColor = SelectedColor.Get();
+		{
+			NewColor.R = Angle * 180.0f * INV_PI;
+			float LinearRadius = FMath::Min(RelativeRadius, 1.0f);
+			if (ExponentDisplacement.IsSet() && ExponentDisplacement.Get() != 1.0f)
+			{
+				//Use log curve to set the distance G value
+				float LogDistance = FMath::Pow(LinearRadius, ExponentDisplacement.Get());
+				NewColor.G = LogDistance;
+			}
+			else
+			{
+				NewColor.G = LinearRadius;
+			}
+		}
+
+		OnValueChanged.ExecuteIfBound(NewColor);
+	}
+
+	return (RelativeRadius <= 1.0f);
+}
diff --git a/Engine/Source/Runtime/Slate/Private/Widgets/Input/SButton.cpp b/Engine/Source/Runtime/Slate/Private/Widgets/Input/SButton.cpp
--- a/Engine/Source/Runtime/Slate/Private/Widgets/Input/SButton.cpp
+++ b/Engine/Source/Runtime/Slate/Private/Widgets/Input/SButton.cpp
@@ -177,7 +177,7 @@ FReply SButton::OnKeyUp(const FGeometry& MyGeometry, const FKeyEvent& InKeyEvent
 	FReply Reply = FReply::Unhandled();
 
 
-	if (IsEnabled() && (/* InKeyEvent.GetKey() == EKeys::Enter || InKeyEvent.GetKey() == EKeys::SpaceBar || */ InKeyEvent.GetKey() == EKeys::Gamepad_FaceButton_Bottom))
+	if (IsEnabled() && (InKeyEvent.GetKey() == EKeys::Enter || InKeyEvent.GetKey() == EKeys::SpaceBar || InKeyEvent.GetKey() == EKeys::Gamepad_FaceButton_Bottom))
 	{
 		const bool bWasPressed = bIsPressed;
 
diff --git a/Engine/Source/Runtime/Slate/Private/Widgets/Input/SEditableTextBox.cpp b/Engine/Source/Runtime/Slate/Private/Widgets/Input/SEditableTextBox.cpp
--- a/Engine/Source/Runtime/Slate/Private/Widgets/Input/SEditableTextBox.cpp
+++ b/Engine/Source/Runtime/Slate/Private/Widgets/Input/SEditableTextBox.cpp
@@ -22,16 +22,10 @@ void SEditableTextBox::Construct( const FArguments& InArgs )
 	BackgroundColorOverride = InArgs._BackgroundColor;
 	ReadOnlyForegroundColorOverride = InArgs._ReadOnlyForegroundColor;
 
-	TAttribute<FMargin> Padding = PaddingOverride.IsSet() ? PaddingOverride : InArgs._Style->Padding;
-	TAttribute<FSlateFontInfo> Font = FontOverride.IsSet() ? FontOverride : InArgs._Style->Font;
-	TAttribute<FSlateColor> BorderForegroundColor = ForegroundColorOverride.IsSet() ? ForegroundColorOverride : InArgs._Style->ForegroundColor;
-	TAttribute<FSlateColor> BackgroundColor = BackgroundColorOverride.IsSet() ? BackgroundColorOverride : InArgs._Style->BackgroundColor;
-	ReadOnlyForegroundColor = ReadOnlyForegroundColorOverride.IsSet() ? ReadOnlyForegroundColorOverride : InArgs._Style->ReadOnlyForegroundColor;
-
 	SBorder::Construct( SBorder::FArguments()
 		.BorderImage( this, &SEditableTextBox::GetBorderImage )
-		.BorderBackgroundColor( BackgroundColor )
-		.ForegroundColor( BorderForegroundColor )
+		.BorderBackgroundColor( this, &SEditableTextBox::DetermineBackgroundColor )
+		.ForegroundColor( this, &SEditableTextBox::DetermineForegroundColor )
 		.Padding( 0 )
 		[
 			SAssignNew( Box, SHorizontalBox)
@@ -42,13 +36,13 @@ void SEditableTextBox::Construct( const FArguments& InArgs )
 			.FillWidth(1)
 			[
 				SAssignNew(PaddingBox, SBox)
-				.Padding( Padding )
+				.Padding( this, &SEditableTextBox::DeterminePadding )
 				.VAlign(VAlign_Center)
 				[
 					SAssignNew( EditableText, SEditableText )
 					.Text( InArgs._Text )
 					.HintText( InArgs._HintText )
-					.Font( Font )
+					.Font( this, &SEditableTextBox::DetermineFont )
 					.IsReadOnly( InArgs._IsReadOnly )
 					.IsPassword( InArgs._IsPassword )
 					.IsCaretMovedWhenGainFocus( InArgs._IsCaretMovedWhenGainFocus )
@@ -83,7 +77,6 @@ void SEditableTextBox::Construct( const FArguments& InArgs )
 
 }
 
-
 void SEditableTextBox::SetStyle(const FEditableTextBoxStyle* InStyle)
 {
 	Style = InStyle;
@@ -96,31 +89,6 @@ void SEditableTextBox::SetStyle(const FEditableTextBoxStyle* InStyle)
 
 	check(Style);
 
-	if (!PaddingOverride.IsSet() && PaddingBox.IsValid())
-	{
-		PaddingBox->SetPadding(Style->Padding);
-	}
-
-	if (!FontOverride.IsSet() && EditableText.IsValid())
-	{
-		EditableText->SetFont(Style->Font);
-	}
-
-	if (!ForegroundColorOverride.IsSet())
-	{
-		SetForegroundColor(Style->ForegroundColor);
-	}
-
-	if (!BackgroundColorOverride.IsSet())
-	{
-		SetBorderBackgroundColor(Style->BackgroundColor);
-	}
-
-	if (!ReadOnlyForegroundColorOverride.IsSet())
-	{
-		ReadOnlyForegroundColor = Style->ReadOnlyForegroundColor;
-	}
-
 	BorderImageNormal = &Style->BackgroundImageNormal;
 	BorderImageHovered = &Style->BackgroundImageHovered;
 	BorderImageFocused = &Style->BackgroundImageFocused;
@@ -271,6 +239,19 @@ FReply SEditableTextBox::OnKeyDown(const FGeometry& MyGeometry, const FKeyEvent&
 	return FReply::Unhandled();
 }
 
+FSlateColor SEditableTextBox::DetermineForegroundColor() const
+{
+	check(Style);  
+	
+	if (EditableText->IsTextReadOnly())
+	{
+		return ReadOnlyForegroundColorOverride.IsSet() ? ReadOnlyForegroundColorOverride.Get() : Style->ReadOnlyForegroundColor;
+	}
+	else
+	{
+		return ForegroundColorOverride.IsSet() ? ForegroundColorOverride.Get() : Style->ForegroundColor;
+	}
+}
 
 void SEditableTextBox::SetHintText(const TAttribute< FText >& InHintText)
 {
@@ -293,31 +274,22 @@ void SEditableTextBox::SetIsPassword(TAttribute< bool > InIsPassword)
 void SEditableTextBox::SetFont(const TAttribute<FSlateFontInfo>& InFont)
 {
 	FontOverride = InFont;
-
-	EditableText->SetFont(FontOverride.IsSet() ? FontOverride : Style->Font);
 }
 
-
 void SEditableTextBox::SetTextBoxForegroundColor(const TAttribute<FSlateColor>& InForegroundColor)
 {
 	ForegroundColorOverride = InForegroundColor;
-
-	SetForegroundColor(ForegroundColorOverride.IsSet() ? ForegroundColorOverride : Style->ForegroundColor);
 }
 
 void SEditableTextBox::SetTextBoxBackgroundColor(const TAttribute<FSlateColor>& InBackgroundColor)
 {
 	BackgroundColorOverride = InBackgroundColor;
-
-	SetBorderBackgroundColor(BackgroundColorOverride.IsSet() ? BackgroundColorOverride : Style->BackgroundColor);
 }
 
 
 void SEditableTextBox::SetReadOnlyForegroundColor(const TAttribute<FSlateColor>& InReadOnlyForegroundColor)
 {
 	ReadOnlyForegroundColorOverride = InReadOnlyForegroundColor;
-
-	ReadOnlyForegroundColor = ReadOnlyForegroundColorOverride.IsSet() ? ReadOnlyForegroundColorOverride : Style->ReadOnlyForegroundColor;
 }
 
 
diff --git a/Engine/Source/Runtime/Slate/Private/Widgets/Input/SMenuAnchor.cpp b/Engine/Source/Runtime/Slate/Private/Widgets/Input/SMenuAnchor.cpp
--- a/Engine/Source/Runtime/Slate/Private/Widgets/Input/SMenuAnchor.cpp
+++ b/Engine/Source/Runtime/Slate/Private/Widgets/Input/SMenuAnchor.cpp
@@ -133,41 +133,22 @@ void SMenuAnchor::Tick( const FGeometry& AllottedGeometry, const double InCurren
 		const FVector2D PopupContentDesiredSize = PopupWindow->GetContent()->GetDesiredSize();
 		FGeometry PopupGeometry = ComputeMenuPlacement( AllottedGeometry, PopupContentDesiredSize, Placement.Get() );
 		const FVector2D NewPosition = PopupGeometry.LocalToAbsolute(FVector2D::ZeroVector);
-		const FVector2D NewSize = PopupGeometry.GetDrawSize();
+		// NOTE: In order to get the right size of the window, we need to take whatever the incoming scale of the menu anchor,
+		// then divide out the popup window's DPI scale.  Finally, we need to divide that remainder to the draw size.
+		// The idea here is to divide out any "extra" scale that's not associated with the DPI scale, since ComputeMenuPlacement
+		// makes a child transform for the new window, based on the geometry of it, which if the menu anchor is inside a zoom panel
+		// that would translate to a menu that had a smaller size window, if the scale was tiny, which we don't want - we only want
+		// the DPI scale if any, to be factored into the size.
+		// NOTE: We only do this for "New Window" popups.  Because games use "Current Window", we can't do this same trick, as they have
+		// to be concerned with the viewport scale as well, which Slate knows nothing about.  Perhaps DPI Scale of the viewports should be
+		// passed down in FGeometry, along with the window DPI Scale, as one extra value code can take into account if it needs to.
+		const FVector2D NewSize = PopupGeometry.GetDrawSize() / ( AllottedGeometry.GetAccumulatedLayoutTransform().GetScale() / PopupWindow->GetLocalToWindowTransform().GetScale() );
 
 		const FSlateRect NewShape = FSlateRect( NewPosition.X, NewPosition.Y, NewPosition.X + NewSize.X, NewPosition.Y + NewSize.Y );
 
 		// We made a window for showing the popup.
 		// Update the window's position!
-		if (false /*PopupWindow->IsMorphing()*/ )
-		{
-			if( NewShape != PopupWindow->GetMorphTargetShape() )
-			{
-				// Update the target shape
-				PopupWindow->UpdateMorphTargetShape( NewShape );
-				// Set size immediately if not morphing size
-				if(!PopupWindow->IsMorphingSize())
-				{
-					PopupWindow->ReshapeWindow( PopupWindow->GetPositionInScreen(), NewSize );
-				}
-			}
-		}
-		else
-		{
-			const FVector2D WindowPosition = PopupWindow->GetPositionInScreen();
-			const FVector2D WindowSize = PopupWindow->GetSizeInScreen();
-			if ( NewPosition != WindowPosition || NewSize != WindowSize )
-			{
-#if PLATFORM_LINUX
-				// @FIXME: for some reason, popups reshaped here do not trigger OnWindowMoved() callback,
-				// so we manually set cached position to where we expect it to be. Note the order of operations - (before Reshape) - 
-				// still giving the callback a chance to change it.
-				// This needs to be investigated (tracked as TTP #347674).
-				PopupWindow->SetCachedScreenPosition( NewPosition );
-#endif // PLATFORM_LINUX
-				PopupWindow->ReshapeWindow( NewShape );
-			}
-		}
+		PopupWindow->ReshapeWindow(NewShape);
 	}
 	else if (PopupWindow.IsValid() && IsOpenAndReusingWindow())
 	{
diff --git a/Engine/Source/Runtime/Slate/Private/Widgets/Input/SMultiLineEditableTextBox.cpp b/Engine/Source/Runtime/Slate/Private/Widgets/Input/SMultiLineEditableTextBox.cpp
--- a/Engine/Source/Runtime/Slate/Private/Widgets/Input/SMultiLineEditableTextBox.cpp
+++ b/Engine/Source/Runtime/Slate/Private/Widgets/Input/SMultiLineEditableTextBox.cpp
@@ -48,14 +48,6 @@ void SMultiLineEditableTextBox::Construct( const FArguments& InArgs )
 	BackgroundColorOverride = InArgs._BackgroundColor;
 	ReadOnlyForegroundColorOverride = InArgs._ReadOnlyForegroundColor;
 
-	TAttribute<FMargin> Padding = InArgs._Padding.IsSet() ? InArgs._Padding : InArgs._Style->Padding;
-	TAttribute<FMargin> HScrollBarPadding = InArgs._HScrollBarPadding.IsSet() ? InArgs._HScrollBarPadding : InArgs._Style->HScrollBarPadding;
-	TAttribute<FMargin> VScrollBarPadding = InArgs._VScrollBarPadding.IsSet() ? InArgs._VScrollBarPadding : InArgs._Style->VScrollBarPadding;
-	TAttribute<FSlateFontInfo> Font = InArgs._Font.IsSet() ? InArgs._Font : InArgs._Style->Font;
-	TAttribute<FSlateColor> BorderForegroundColor = InArgs._ForegroundColor.IsSet() ? InArgs._ForegroundColor : InArgs._Style->ForegroundColor;
-	TAttribute<FSlateColor> BackgroundColor = InArgs._BackgroundColor.IsSet() ? InArgs._BackgroundColor : InArgs._Style->BackgroundColor;
-	ReadOnlyForegroundColor = InArgs._ReadOnlyForegroundColor.IsSet() ? InArgs._ReadOnlyForegroundColor : InArgs._Style->ReadOnlyForegroundColor;
-
 	bHasExternalHScrollBar = InArgs._HScrollBar.IsValid();
 	HScrollBar = InArgs._HScrollBar;
 	if (!HScrollBar.IsValid())
@@ -82,9 +74,9 @@ void SMultiLineEditableTextBox::Construct( const FArguments& InArgs )
 
 	SBorder::Construct( SBorder::FArguments()
 		.BorderImage( this, &SMultiLineEditableTextBox::GetBorderImage )
-		.BorderBackgroundColor( BackgroundColor )
-		.ForegroundColor( BorderForegroundColor )
-		.Padding( Padding )
+		.BorderBackgroundColor( this, &SMultiLineEditableTextBox::DetermineBackgroundColor )
+		.ForegroundColor( this, &SMultiLineEditableTextBox::DetermineForegroundColor )
+		.Padding( this, &SMultiLineEditableTextBox::DeterminePadding )
 		[
 			SAssignNew( Box, SHorizontalBox)
 			+SHorizontalBox::Slot()
@@ -104,7 +96,7 @@ void SMultiLineEditableTextBox::Construct( const FArguments& InArgs )
 					.HintText( InArgs._HintText )
 					.TextStyle( InArgs._TextStyle )
 					.Marshaller( InArgs._Marshaller )
-					.Font( Font )
+					.Font( this, &SMultiLineEditableTextBox::DetermineFont )
 					.IsReadOnly( InArgs._IsReadOnly )
 					.OnContextMenuOpening( InArgs._OnContextMenuOpening )
 					.OnTextChanged( InArgs._OnTextChanged )
@@ -137,7 +129,7 @@ void SMultiLineEditableTextBox::Construct( const FArguments& InArgs )
 				.AutoHeight()
 				[
 					SAssignNew(HScrollBarPaddingBox, SBox)
-					.Padding(HScrollBarPadding)
+					.Padding( this, &SMultiLineEditableTextBox::DetermineHScrollBarPadding )
 					[
 						AsWidgetRef( HScrollBar )
 					]
@@ -148,7 +140,7 @@ void SMultiLineEditableTextBox::Construct( const FArguments& InArgs )
 			.AutoWidth()
 			[
 				SAssignNew(VScrollBarPaddingBox, SBox)
-				.Padding(VScrollBarPadding)
+				.Padding( this, &SMultiLineEditableTextBox::DetermineVScrollBarPadding )
 				[
 					AsWidgetRef( VScrollBar )
 				]
@@ -183,41 +175,6 @@ void SMultiLineEditableTextBox::SetStyle(const FEditableTextBoxStyle* InStyle)
 
 	check(Style);
 
-	if (!PaddingOverride.IsSet())
-	{
-		SetPadding(Style->Padding);
-	}
-
-	if (!HScrollBarPaddingOverride.IsSet() && HScrollBarPaddingBox.IsValid())
-	{
-		HScrollBarPaddingBox->SetPadding(Style->HScrollBarPadding);
-	}
-
-	if (!VScrollBarPaddingOverride.IsSet() && VScrollBarPaddingBox.IsValid())
-	{
-		VScrollBarPaddingBox->SetPadding(Style->VScrollBarPadding);
-	}
-
-	if (!FontOverride.IsSet() && EditableText.IsValid())
-	{
-		EditableText->SetFont(Style->Font);
-	}
-
-	if (!ForegroundColorOverride.IsSet())
-	{
-		SetForegroundColor(Style->ForegroundColor);
-	}
-
-	if (!BackgroundColorOverride.IsSet())
-	{
-		SetBorderBackgroundColor(Style->BackgroundColor);
-	}
-
-	if (!ReadOnlyForegroundColorOverride.IsSet())
-	{
-		ReadOnlyForegroundColor = Style->ReadOnlyForegroundColor;
-	}
-
 	if (!bHasExternalHScrollBar && HScrollBar.IsValid())
 	{
 		HScrollBar->SetStyle(&Style->ScrollBarStyle);
@@ -234,6 +191,20 @@ void SMultiLineEditableTextBox::SetStyle(const FEditableTextBoxStyle* InStyle)
 	BorderImageReadOnly = &Style->BackgroundImageReadOnly;
 }
 
+FSlateColor SMultiLineEditableTextBox::DetermineForegroundColor() const
+{
+	check(Style);
+
+	if ( EditableText->IsTextReadOnly() )
+	{
+		return ReadOnlyForegroundColorOverride.IsSet() ? ReadOnlyForegroundColorOverride.Get() : Style->ReadOnlyForegroundColor;
+	}
+	else
+	{
+		return ForegroundColorOverride.IsSet() ? ForegroundColorOverride.Get() : Style->ForegroundColor;
+	}
+}
+
 void SMultiLineEditableTextBox::SetText( const TAttribute< FText >& InNewText )
 {
 	EditableText->SetText( InNewText );
@@ -247,22 +218,16 @@ void SMultiLineEditableTextBox::SetHintText( const TAttribute< FText >& InHintTe
 void SMultiLineEditableTextBox::SetTextBoxForegroundColor(const TAttribute<FSlateColor>& InForegroundColor)
 {
 	ForegroundColorOverride = InForegroundColor;
-
-	SetForegroundColor(ForegroundColorOverride.IsSet() ? ForegroundColorOverride : Style->ForegroundColor);
 }
 
 void SMultiLineEditableTextBox::SetTextBoxBackgroundColor(const TAttribute<FSlateColor>& InBackgroundColor)
 {
 	BackgroundColorOverride = InBackgroundColor;
-
-	SetBorderBackgroundColor(BackgroundColorOverride.IsSet() ? BackgroundColorOverride : Style->BackgroundColor);
 }
 
 void SMultiLineEditableTextBox::SetReadOnlyForegroundColor(const TAttribute<FSlateColor>& InReadOnlyForegroundColor)
 {
 	ReadOnlyForegroundColorOverride = InReadOnlyForegroundColor;
-
-	ReadOnlyForegroundColor = ReadOnlyForegroundColorOverride.IsSet() ? ReadOnlyForegroundColorOverride : Style->ReadOnlyForegroundColor;
 }
 
 void SMultiLineEditableTextBox::SetTextShapingMethod(const TOptional<ETextShapingMethod>& InTextShapingMethod)
diff --git a/Engine/Source/Runtime/Slate/Private/Widgets/Layout/SBackgroundBlur.cpp b/Engine/Source/Runtime/Slate/Private/Widgets/Layout/SBackgroundBlur.cpp
new file mode 100644
--- /dev/null
+++ b/Engine/Source/Runtime/Slate/Private/Widgets/Layout/SBackgroundBlur.cpp
@@ -0,0 +1,166 @@
+// Copyright 1998-2016 Epic Games, Inc. All Rights Reserved.
+
+#include "SBackgroundBlur.h"
+#include "DrawElements.h"
+#include "IConsoleManager.h"
+
+// This feature has not been tested on es2
+#if PLATFORM_USES_ES2
+static int32 bAllowBackgroundBlur = 1;
+#else
+static int32 bAllowBackgroundBlur = 1;
+#endif
+
+static FAutoConsoleVariableRef CVarSlateAllowBackgroundBlurWidgets(TEXT("Slate.AllowBackgroundBlurWidgets"), bAllowBackgroundBlur, TEXT("If 0, no background blur widgets will be rendered"));
+
+static int32 MaxKernelSize = 255;
+static FAutoConsoleVariableRef CVarSlateMaxKernelSize(TEXT("Slate.BackgroundBlurMaxKernelSize"), MaxKernelSize, TEXT("The maximum allowed kernel size.  Note: Very large numbers can cause a huge decrease in performance"));
+
+static int32 bDownsampleForBlur = 1;
+static FAutoConsoleVariableRef CVarDownsampleForBlur(TEXT("Slate.BackgroundBlurDownsample"), bDownsampleForBlur, TEXT(""), ECVF_Cheat);
+
+static int32 bForceLowQualityBrushFallback = 0;
+static FAutoConsoleVariableRef CVarForceLowQualityBackgroundBlurOverride(TEXT("Slate.ForceBackgroundBlurLowQualityOverride"), bForceLowQualityBrushFallback, TEXT("Whether or not to force a slate brush to be used instead of actually blurring the background"), ECVF_Scalability);
+
+
+void SBackgroundBlur::Construct(const FArguments& InArgs)
+{
+	bApplyAlphaToBlur = InArgs._bApplyAlphaToBlur;
+	LowQualityFallbackBrush = InArgs._LowQualityFallbackBrush;
+	BlurStrength = InArgs._BlurStrength;
+	BlurRadius = InArgs._BlurRadius;
+
+	ChildSlot
+		.HAlign(InArgs._HAlign)
+		.VAlign(InArgs._VAlign)
+		.Padding(InArgs._Padding)
+	[
+		InArgs._Content.Widget
+	];
+}
+
+void SBackgroundBlur::SetContent(const TSharedRef<SWidget>& InContent)
+{
+	ChildSlot.AttachWidget(InContent);
+}
+
+void SBackgroundBlur::SetApplyAlphaToBlur(bool bInApplyAlphaToBlur)
+{
+	bApplyAlphaToBlur = bInApplyAlphaToBlur;
+	Invalidate(EInvalidateWidget::Layout);
+}
+
+void SBackgroundBlur::SetBlurRadius(const TAttribute<TOptional<int32>>& InBlurRadius)
+{
+	BlurRadius = InBlurRadius;
+	Invalidate(EInvalidateWidget::Layout);
+}
+
+void SBackgroundBlur::SetBlurStrength(const TAttribute<float>& InStrength)
+{
+	BlurStrength = InStrength;
+	Invalidate(EInvalidateWidget::Layout);
+}
+
+void SBackgroundBlur::SetLowQualityBackgroundBrush(const FSlateBrush* InBrush)
+{
+	LowQualityFallbackBrush = InBrush;
+	Invalidate(EInvalidateWidget::Layout);
+}
+
+void SBackgroundBlur::SetHAlign(EHorizontalAlignment HAlign)
+{
+	ChildSlot.HAlignment = HAlign;
+}
+
+void SBackgroundBlur::SetVAlign(EVerticalAlignment VAlign)
+{
+	ChildSlot.VAlignment = VAlign;
+}
+
+void SBackgroundBlur::SetPadding(const TAttribute<FMargin>& InPadding)
+{
+	ChildSlot.SlotPadding = InPadding;
+}
+
+int32 SBackgroundBlur::OnPaint(const FPaintArgs& Args, const FGeometry& AllottedGeometry, const FSlateRect& MyClippingRect, FSlateWindowElementList& OutDrawElements, int32 LayerId, const FWidgetStyle& InWidgetStyle, bool bParentEnabled) const
+{
+	int32 PostFXLayerId = LayerId;
+	if (bAllowBackgroundBlur)
+	{
+		if (!bForceLowQualityBrushFallback)
+		{
+			// Modulate blur strength by the widget alpha
+			const float Strength = BlurStrength.Get() * (bApplyAlphaToBlur ? InWidgetStyle.GetColorAndOpacityTint().A : 1.f);
+			if (Strength > 0.f)
+			{
+				FPaintGeometry PaintGeometry = AllottedGeometry.ToPaintGeometry();
+
+				// extract the layout transform from the draw element
+				FSlateLayoutTransform InverseLayoutTransform(Inverse(FSlateLayoutTransform(PaintGeometry.DrawScale, PaintGeometry.DrawPosition)));
+				// The clip rect is NOT subject to the rotations specified by MakeRotatedBox.
+				FSlateRotatedClipRectType RenderClipRect = FSlateRotatedClipRectType::MakeSnappedRotatedRect(MyClippingRect, InverseLayoutTransform, AllottedGeometry.GetAccumulatedRenderTransform());
+
+				float OffsetX = PaintGeometry.DrawPosition.X - FMath::TruncToFloat(PaintGeometry.DrawPosition.X);
+				float OffsetY = PaintGeometry.DrawPosition.Y - FMath::TruncToFloat(PaintGeometry.DrawPosition.Y);
+
+				int32 RenderTargetWidth = FMath::RoundToInt(RenderClipRect.ExtentX.X);
+				int32 RenderTargetHeight = FMath::RoundToInt(RenderClipRect.ExtentY.Y);
+
+				int32 KernelSize = 0;
+				int32 DownsampleAmount = 0;
+				ComputeEffectiveKernelSize(Strength, KernelSize, DownsampleAmount);
+
+				float ComputedStrength = FMath::Max(.5f, Strength);
+
+				if (DownsampleAmount > 0)
+				{
+					RenderTargetWidth = FMath::DivideAndRoundUp(RenderTargetWidth, DownsampleAmount);
+					RenderTargetHeight = FMath::DivideAndRoundUp(RenderTargetHeight, DownsampleAmount);
+					ComputedStrength /= DownsampleAmount;
+				}
+
+				if (RenderTargetWidth > 0 && RenderTargetHeight > 0)
+				{
+					FSlateDrawElement::MakePostProcessPass(OutDrawElements, LayerId, PaintGeometry, MyClippingRect, FVector4(KernelSize, ComputedStrength, RenderTargetWidth, RenderTargetHeight), DownsampleAmount);
+				}
+
+
+				++PostFXLayerId;
+			}
+		}
+		else if (bAllowBackgroundBlur && bForceLowQualityBrushFallback && LowQualityFallbackBrush && LowQualityFallbackBrush->DrawAs != ESlateBrushDrawType::NoDrawType)
+		{
+			const bool bIsEnabled = ShouldBeEnabled(bParentEnabled);
+			const uint32 DrawEffects = bIsEnabled ? ESlateDrawEffect::None : ESlateDrawEffect::DisabledEffect;
+
+			const FLinearColor FinalColorAndOpacity(InWidgetStyle.GetColorAndOpacityTint() * LowQualityFallbackBrush->GetTint(InWidgetStyle));
+
+			FSlateDrawElement::MakeBox(OutDrawElements, PostFXLayerId, AllottedGeometry.ToPaintGeometry(), LowQualityFallbackBrush, MyClippingRect, DrawEffects, FinalColorAndOpacity);
+			++PostFXLayerId;
+		}
+	}
+
+	return SCompoundWidget::OnPaint(Args, AllottedGeometry, MyClippingRect, OutDrawElements, PostFXLayerId, InWidgetStyle, bParentEnabled);
+}
+
+void SBackgroundBlur::ComputeEffectiveKernelSize(float Strength, int32& OutKernelSize, int32& OutDownsampleAmount) const
+{
+	// If the radius isn't set, auto-compute it based on the strength
+	OutKernelSize = BlurRadius.Get().Get(FMath::RoundToInt(Strength * 3.f));
+
+	// Downsample if needed
+	if (bDownsampleForBlur && OutKernelSize > 9)
+	{
+		OutDownsampleAmount = OutKernelSize >= 31 ? 4 : 2;
+		OutKernelSize /= OutDownsampleAmount;
+	}
+
+	// Kernel sizes must be odd
+	if (OutKernelSize % 2 == 0)
+	{
+		++OutKernelSize;
+	}
+
+	OutKernelSize = FMath::Clamp(OutKernelSize, 3, MaxKernelSize);
+}
\ No newline at end of file
diff --git a/Engine/Source/Runtime/Slate/Private/Widgets/Layout/SConstraintCanvas.cpp b/Engine/Source/Runtime/Slate/Private/Widgets/Layout/SConstraintCanvas.cpp
--- a/Engine/Source/Runtime/Slate/Private/Widgets/Layout/SConstraintCanvas.cpp
+++ b/Engine/Source/Runtime/Slate/Private/Widgets/Layout/SConstraintCanvas.cpp
@@ -3,7 +3,7 @@
 #include "Widgets/Layout/SConstraintCanvas.h"
 #include "Types/PaintArgs.h"
 #include "Layout/ArrangedChildren.h"
-
+#include "SlateSettings.h"
 
 
 /* SConstraintCanvas interface
@@ -68,26 +68,42 @@ struct FSortSlotsByZOrder
  *****************************************************************************/
 
 void SConstraintCanvas::OnArrangeChildren( const FGeometry& AllottedGeometry, FArrangedChildren& ArrangedChildren ) const
+{
+	FArrangedChildLayers ChildLayers;
+	ArrangeLayeredChildren(AllottedGeometry, ArrangedChildren, ChildLayers);
+}
+
+void SConstraintCanvas::ArrangeLayeredChildren(const FGeometry& AllottedGeometry, FArrangedChildren& ArrangedChildren, FArrangedChildLayers& ArrangedChildLayers) const
 {
 	CachedGeometry = AllottedGeometry;
 
 	if (Children.Num() > 0)
 	{
+
+		// Using a Project setting here to decide whether we automatically put children in front of all previous children
+		// or allow the explicit ZOrder value to place children in the same layer. This will allow users to have non-touching
+		// children render into the same layer and have a chance of being batched by the Slate renderer for better performance.
+#if WITH_EDITOR
+		bool bExplicitChildZOrder = GetDefault<USlateSettings>()->bExplicitCanvasChildZOrder;
+#else
+		static bool bExplicitChildZOrder = GetDefault<USlateSettings>()->bExplicitCanvasChildZOrder;
+#endif
 		// Sort the children based on zorder.
 		TArray< FChildZOrder, TInlineAllocator<64> > SlotOrder;
 		SlotOrder.Reserve(Children.Num());
 
-		for ( int32 ChildIndex = 0; ChildIndex < Children.Num(); ++ChildIndex )
+		for (int32 ChildIndex = 0; ChildIndex < Children.Num(); ++ChildIndex)
 		{
 			const SConstraintCanvas::FSlot& CurChild = Children[ChildIndex];
 
 			FChildZOrder Order;
 			Order.ChildIndex = ChildIndex;
 			Order.ZOrder = CurChild.ZOrderAttr.Get();
-			SlotOrder.Add( Order );
+			SlotOrder.Add(Order);
 		}
 
 		SlotOrder.Sort(FSortSlotsByZOrder());
+		float LastZOrder = -FLT_MAX;
 
 		// Arrange the children now in their proper z-order.
 		for (int32 ChildIndex = 0; ChildIndex < Children.Num(); ++ChildIndex)
@@ -97,7 +113,7 @@ void SConstraintCanvas::OnArrangeChildren( const FGeometry& AllottedGeometry, FA
 			const TSharedRef<SWidget>& CurWidget = CurChild.GetWidget();
 
 			const EVisibility ChildVisibility = CurWidget->GetVisibility();
-			if ( ArrangedChildren.Accepts(ChildVisibility) )
+			if (ArrangedChildren.Accepts(ChildVisibility))
 			{
 				const FMargin Offset = CurChild.OffsetAttr.Get();
 				const FVector2D Alignment = CurChild.AlignmentAttr.Get();
@@ -126,7 +142,7 @@ void SConstraintCanvas::OnArrangeChildren( const FGeometry& AllottedGeometry, FA
 				FVector2D LocalPosition, LocalSize;
 
 				// Calculate the position and size based on the horizontal stretch or non-stretch
-				if ( bIsHorizontalStretch )
+				if (bIsHorizontalStretch)
 				{
 					LocalPosition.X = AnchorPixels.Left + Offset.Left;
 					LocalSize.X = AnchorPixels.Right - LocalPosition.X - Offset.Right;
@@ -138,7 +154,7 @@ void SConstraintCanvas::OnArrangeChildren( const FGeometry& AllottedGeometry, FA
 				}
 
 				// Calculate the position and size based on the vertical stretch or non-stretch
-				if ( bIsVerticalStretch )
+				if (bIsVerticalStretch)
 				{
 					LocalPosition.Y = AnchorPixels.Top + Offset.Top;
 					LocalSize.Y = AnchorPixels.Bottom - LocalPosition.Y - Offset.Bottom;
@@ -157,7 +173,21 @@ void SConstraintCanvas::OnArrangeChildren( const FGeometry& AllottedGeometry, FA
 					LocalPosition,
 					// Child's size
 					LocalSize
-					));
+				));
+
+				bool bNewLayer = true;
+				if (bExplicitChildZOrder)
+				{
+					// Split children into discrete layers for the paint method
+					bNewLayer = false;
+					if (CurSlot.ZOrder > LastZOrder + DELTA)
+					{
+						bNewLayer = true;
+						LastZOrder = CurSlot.ZOrder;
+					}
+
+				}
+				ArrangedChildLayers.Add(bNewLayer);
 			}
 		}
 	}
@@ -168,13 +198,15 @@ int32 SConstraintCanvas::OnPaint( const FPaintArgs& Args, const FGeometry& Allot
 	//FPlatformMisc::BeginNamedEvent(FColor::Orange, "SConstraintCanvas");
 
 	FArrangedChildren ArrangedChildren(EVisibility::Visible);
-	ArrangeChildren(AllottedGeometry, ArrangedChildren);
+	FArrangedChildLayers ChildLayers;
+	ArrangeLayeredChildren(AllottedGeometry, ArrangedChildren, ChildLayers);
 
 	bool bForwardedEnabled = ShouldBeEnabled(bParentEnabled);
 
 	// Because we paint multiple children, we must track the maximum layer id that they produced in case one of our parents
 	// wants to an overlay for all of its contents.
 	int32 MaxLayerId = LayerId;
+	int32 ChildLayerId = LayerId + 1;
 
 	const FPaintArgs NewArgs = Args.WithNewParent(this);
 
@@ -187,7 +219,10 @@ int32 SConstraintCanvas::OnPaint( const FPaintArgs& Args, const FGeometry& Allot
 
 		if ( bWereOverlapping )
 		{
-			const int32 CurWidgetsMaxLayerId = CurWidget.Widget->Paint(NewArgs, CurWidget.Geometry, ChildClipRect, OutDrawElements, MaxLayerId + 1, InWidgetStyle, bForwardedEnabled);
+			// Bools in ChildLayers tell us whether to paint the next child in front of all previous
+			ChildLayerId = ChildLayers[ChildIndex] ? MaxLayerId + 1 : ChildLayerId;
+
+			const int32 CurWidgetsMaxLayerId = CurWidget.Widget->Paint(NewArgs, CurWidget.Geometry, ChildClipRect, OutDrawElements, ChildLayerId, InWidgetStyle, bForwardedEnabled);
 
 			MaxLayerId = FMath::Max(MaxLayerId, CurWidgetsMaxLayerId);
 		}
diff --git a/Engine/Source/Runtime/Slate/Private/Widgets/Layout/SDPIScaler.cpp b/Engine/Source/Runtime/Slate/Private/Widgets/Layout/SDPIScaler.cpp
--- a/Engine/Source/Runtime/Slate/Private/Widgets/Layout/SDPIScaler.cpp
+++ b/Engine/Source/Runtime/Slate/Private/Widgets/Layout/SDPIScaler.cpp
@@ -60,7 +60,7 @@ void SDPIScaler::SetDPIScale(TAttribute<float> InDPIScale)
 	DPIScale = InDPIScale;
 }
 
-float SDPIScaler::GetRelativeLayoutScale(const FSlotBase& Child) const
+float SDPIScaler::GetRelativeLayoutScale(const FSlotBase& Child, float LayoutScaleMultiplier) const
 {
 	return DPIScale.Get();
 }
diff --git a/Engine/Source/Runtime/Slate/Private/Widgets/Layout/SScaleBox.cpp b/Engine/Source/Runtime/Slate/Private/Widgets/Layout/SScaleBox.cpp
--- a/Engine/Source/Runtime/Slate/Private/Widgets/Layout/SScaleBox.cpp
+++ b/Engine/Source/Runtime/Slate/Private/Widgets/Layout/SScaleBox.cpp
@@ -188,8 +188,13 @@ FVector2D SScaleBox::ComputeDesiredSize(float InScale) const
 	return LayoutScale * SCompoundWidget::ComputeDesiredSize(InScale);
 }
 
-float SScaleBox::GetRelativeLayoutScale(const FSlotBase& Child) const
+float SScaleBox::GetRelativeLayoutScale(const FSlotBase& Child, float LayoutScaleMultiplier) const
 {
+	if ( IgnoreInheritedScale.Get(false) )
+	{
+		return GetLayoutScale() / LayoutScaleMultiplier;
+	}
+
 	return GetLayoutScale();
 }
 
diff --git a/Engine/Source/Runtime/Slate/Private/Widgets/SViewport.cpp b/Engine/Source/Runtime/Slate/Private/Widgets/SViewport.cpp
--- a/Engine/Source/Runtime/Slate/Private/Widgets/SViewport.cpp
+++ b/Engine/Source/Runtime/Slate/Private/Widgets/SViewport.cpp
@@ -137,16 +137,24 @@ int32 SViewport::OnPaint( const FPaintArgs& Args, const FGeometry& AllottedGeome
 			Brush = FCoreStyle::Get().GetBrush(TEXT("SoftwareCursor_CardinalCross"));
 		}
 
+		FVector2D CursorPositionLocalSpace = ViewportInterfacePin->GetSoftwareCursorPosition() / AllottedGeometry.Scale;
+
 		LayerId++;
 		FSlateDrawElement::MakeBox(
 			OutDrawElements,
 			LayerId,
-			AllottedGeometry.ToPaintGeometry( ViewportInterfacePin->GetSoftwareCursorPosition() - ( Brush->ImageSize / 2 ), Brush->ImageSize ),
+			AllottedGeometry.ToPaintGeometry( CursorPositionLocalSpace - ( Brush->ImageSize / 2 ), Brush->ImageSize ),
 			Brush,
 			MyClippingRect
 		);
 	}
 
+	// If there are any custom hit testable widgets in the 3D world we need to register their custom hit test path here.
+	if ( CustomHitTestPath.IsValid() )
+	{
+		Args.InsertCustomHitTestPath(CustomHitTestPath.ToSharedRef(), LastHitTestIndex);
+	}
+
 	return Layer;
 }
 
@@ -260,6 +268,16 @@ void SViewport::SetContent( TSharedPtr<SWidget> InContent )
 	];
 }
 
+void SViewport::SetCustomHitTestPath( TSharedPtr<ICustomHitTestPath> InCustomHitTestPath )
+{
+	CustomHitTestPath = InCustomHitTestPath;
+}
+
+TSharedPtr<ICustomHitTestPath> SViewport::GetCustomHitTestPath()
+{
+	return CustomHitTestPath;
+}
+
 void SViewport::OnWindowClosed( const TSharedRef<SWindow>& WindowBeingClosed )
 {
 	if (ViewportInterface.IsValid())
@@ -333,6 +351,25 @@ void SViewport::OnFinishedPointerInput()
 	}
 }
 
+void SViewport::OnArrangeChildren( const FGeometry& AllottedGeometry, FArrangedChildren& ArrangedChildren ) const
+{
+	SCompoundWidget::OnArrangeChildren(AllottedGeometry, ArrangedChildren);
+	if( ArrangedChildren.Allows3DWidgets() && CustomHitTestPath.IsValid() )
+	{
+		CustomHitTestPath->ArrangeChildren( ArrangedChildren );
+	}
+}
+
+TSharedPtr<FVirtualPointerPosition> SViewport::TranslateMouseCoordinateFor3DChild(const TSharedRef<SWidget>& ChildWidget, const FGeometry& MyGeometry, const FVector2D& ScreenSpaceMouseCoordinate, const FVector2D& LastScreenSpaceMouseCoordinate) const
+{
+	if( CustomHitTestPath.IsValid() )
+	{
+		return CustomHitTestPath->TranslateMouseCoordinateFor3DChild( ChildWidget, MyGeometry, ScreenSpaceMouseCoordinate, LastScreenSpaceMouseCoordinate );
+	}
+
+	return nullptr;
+}
+
 FNavigationReply SViewport::OnNavigation(const FGeometry& MyGeometry, const FNavigationEvent& InNavigationEvent)
 {
 	return ViewportInterface.IsValid() ? ViewportInterface.Pin()->OnNavigation(MyGeometry, InNavigationEvent) : FNavigationReply::Stop();
diff --git a/Engine/Source/Runtime/Slate/Private/Widgets/Text/SlateEditableTextLayout.cpp b/Engine/Source/Runtime/Slate/Private/Widgets/Text/SlateEditableTextLayout.cpp
--- a/Engine/Source/Runtime/Slate/Private/Widgets/Text/SlateEditableTextLayout.cpp
+++ b/Engine/Source/Runtime/Slate/Private/Widgets/Text/SlateEditableTextLayout.cpp
@@ -203,6 +203,9 @@ FSlateEditableTextLayout::~FSlateEditableTextLayout()
 		TSharedRef<FTextInputMethodContext> TextInputMethodContextRef = TextInputMethodContext.ToSharedRef();
 		if (TextInputMethodSystem->IsActiveContext(TextInputMethodContextRef))
 		{
+			// Make sure that the composition is aborted to avoid any IME calls to EndComposition from trying to mutate our dying owner widget
+			TextInputMethodContextRef->AbortComposition();
+
 			// This can happen if an entire tree of widgets is culled, as Slate isn't notified of the focus loss, the widget is just deleted
 			TextInputMethodSystem->DeactivateContext(TextInputMethodContextRef);
 		}
@@ -1984,7 +1987,6 @@ bool FSlateEditableTextLayout::MoveCursor(const FMoveCursor& InArgs)
 		if (TextInputMethodSystem)
 		{
 			TextInputMethodSystem->DeactivateContext(TextInputMethodContext.ToSharedRef());
-			TextInputMethodContext->CacheWindow();
 			TextInputMethodSystem->ActivateContext(TextInputMethodContext.ToSharedRef());
 		}
 	}
@@ -2383,19 +2385,22 @@ void FSlateEditableTextLayout::UpdateCursorHighlight()
 		const FTextLocation CursorPosition = CursorInfo.GetCursorLocation();
 
 		const TArray< FTextLayout::FLineModel >& Lines = TextLayout->GetLineModels();
-		const int32 LineTextLength = Lines[CursorPosition.GetLineIndex()].Text->Len();
-
-		if (LineTextLength == 0)
-		{
-			ActiveLineHighlights.Add(FTextLineHighlight(CursorPosition.GetLineIndex(), FTextRange(0, 0), CursorZOrder, CursorLineHighlighter.ToSharedRef()));
-		}
-		else if (CursorPosition.GetOffset() == LineTextLength)
+		if (Lines.IsValidIndex(CursorPosition.GetLineIndex()))
 		{
-			ActiveLineHighlights.Add(FTextLineHighlight(CursorPosition.GetLineIndex(), FTextRange(LineTextLength - 1, LineTextLength), CursorZOrder, CursorLineHighlighter.ToSharedRef()));
-		}
-		else
-		{
-			ActiveLineHighlights.Add(FTextLineHighlight(CursorPosition.GetLineIndex(), FTextRange(CursorPosition.GetOffset(), CursorPosition.GetOffset() + 1), CursorZOrder, CursorLineHighlighter.ToSharedRef()));
+			const int32 LineTextLength = Lines[CursorPosition.GetLineIndex()].Text->Len();
+
+			if (LineTextLength == 0)
+			{
+				ActiveLineHighlights.Add(FTextLineHighlight(CursorPosition.GetLineIndex(), FTextRange(0, 0), CursorZOrder, CursorLineHighlighter.ToSharedRef()));
+			}
+			else if (CursorPosition.GetOffset() == LineTextLength)
+			{
+				ActiveLineHighlights.Add(FTextLineHighlight(CursorPosition.GetLineIndex(), FTextRange(LineTextLength - 1, LineTextLength), CursorZOrder, CursorLineHighlighter.ToSharedRef()));
+			}
+			else
+			{
+				ActiveLineHighlights.Add(FTextLineHighlight(CursorPosition.GetLineIndex(), FTextRange(CursorPosition.GetOffset(), CursorPosition.GetOffset() + 1), CursorZOrder, CursorLineHighlighter.ToSharedRef()));
+			}
 		}
 	}
 
@@ -3396,8 +3401,10 @@ void FSlateEditableTextLayout::FTextInputMethodContext::GetSelectionRange(uint32
 
 void FSlateEditableTextLayout::FTextInputMethodContext::SetSelectionRange(const uint32 BeginIndex, const uint32 Length, const ECaretPosition InCaretPosition)
 {
-	const uint32 MinIndex = BeginIndex;
-	const uint32 MaxIndex = MinIndex + Length;
+	const uint32 TextLength = GetTextLength();
+
+	const uint32 MinIndex = FMath::Min(BeginIndex, TextLength);
+	const uint32 MaxIndex = FMath::Min(MinIndex + Length, TextLength);
 
 	FTextLayout::FTextOffsetLocations OffsetLocations;
 	OwnerLayout->TextLayout->GetTextOffsetLocations(OffsetLocations);
diff --git a/Engine/Source/Runtime/Slate/Private/Widgets/Views/STableViewBase.cpp b/Engine/Source/Runtime/Slate/Private/Widgets/Views/STableViewBase.cpp
--- a/Engine/Source/Runtime/Slate/Private/Widgets/Views/STableViewBase.cpp
+++ b/Engine/Source/Runtime/Slate/Private/Widgets/Views/STableViewBase.cpp
@@ -693,6 +693,11 @@ float STableViewBase::ScrollTo( float InScrollOffset )
 	return AmountScrolled;
 }
 
+float STableViewBase::GetScrollOffset() const
+{
+	return ScrollOffset;
+}
+
 void STableViewBase::SetScrollOffset( const float InScrollOffset )
 {
 	if ( ScrollOffset != InScrollOffset )
diff --git a/Engine/Source/Runtime/Slate/Public/Framework/Application/SlateApplication.h b/Engine/Source/Runtime/Slate/Public/Framework/Application/SlateApplication.h
--- a/Engine/Source/Runtime/Slate/Public/Framework/Application/SlateApplication.h
+++ b/Engine/Source/Runtime/Slate/Public/Framework/Application/SlateApplication.h
@@ -578,9 +578,6 @@ public:
 	 */
 	void SetAllUserFocus(const TSharedPtr<SWidget>& WidgetToFocus, EFocusCause ReasonFocusIsChanging = EFocusCause::SetDirectly);
 
-	/** @return a pointer to the Widget that currently has the users focus; Empty pointer when the user has no focus. */
-	TSharedPtr< SWidget > GetUserFocusedWidget(uint32 UserIndex) const;
-
 	DEPRECATED(4.6, "FSlateApplication::GetJoystickCaptor() is deprecated, use FSlateApplication::GetUserFocusedWidget() instead.")
 	TSharedPtr< SWidget > GetJoystickCaptor(uint32 UserIndex) const;
 
@@ -926,18 +923,33 @@ protected:
 	 */
 	void RegisterUser(TSharedRef<FSlateUser> User);
 
+	/**
+	 * Gets the user at the given index, null if the user does not exist.
+	 */
 	FORCEINLINE const FSlateUser* GetUser(int32 UserIndex) const
 	{
-		return UserIndex < Users.Num() ? Users[UserIndex].Get() : nullptr;
+		return (UserIndex >= 0 && UserIndex < Users.Num()) ? Users[UserIndex].Get() : nullptr;
 	}
 
+	/**
+	 * Gets the user at the given index, null if the user does not exist.
+	 */
 	FORCEINLINE FSlateUser* GetUser(int32 UserIndex)
 	{
-		return UserIndex < Users.Num() ? Users[UserIndex].Get() : nullptr;
+		return ( UserIndex >= 0 && UserIndex < Users.Num() ) ? Users[UserIndex].Get() : nullptr;
 	}
 
+	/**
+	 * Locates the SlateUser object corresponding to the index, if one can't be found, it will create a slate user at
+	 * the provided index.  If the index is less than 0, null is returned.
+	 */
 	FORCEINLINE FSlateUser* GetOrCreateUser(int32 UserIndex)
 	{
+		if ( UserIndex < 0 )
+		{
+			return nullptr;
+		}
+
 		if ( FSlateUser* User = GetUser(UserIndex) )
 		{
 			return User;
@@ -1317,6 +1329,7 @@ public:
 	virtual bool SetUserFocus(const uint32 InUserIndex, const FWidgetPath& InFocusPath, const EFocusCause InCause) override;
 	virtual void SetAllUserFocus(const FWidgetPath& InFocusPath, const EFocusCause InCause) override;
 	virtual void SetAllUserFocusAllowingDescendantFocus(const FWidgetPath& InFocusPath, const EFocusCause InCause) override;
+	virtual TSharedPtr<SWidget> GetUserFocusedWidget(uint32 UserIndex) const override;
 
 	DECLARE_EVENT_OneParam(FSlateApplication, FApplicationActivationStateChangedEvent, const bool /*IsActive*/)
 	virtual FApplicationActivationStateChangedEvent& OnApplicationActivationStateChanged() { return ApplicationActivationStateChangedEvent; }
@@ -1430,12 +1443,12 @@ public:
 	int32 GetUserIndexForController(int32 ControllerId) const;
 
 	/**
-	* Register for a notification when the window action occurs.
-	*
-	* @param Notification          The notification to invoke.
-	*
-	* @return Handle to the registered delegate.
-	*/
+	 * Register for a notification when the window action occurs.
+	 *
+	 * @param Notification          The notification to invoke.
+	 *
+	 * @return Handle to the registered delegate.
+	 */
 	FDelegateHandle RegisterOnWindowActionNotification(const FOnWindowAction& Notification);
 
 
diff --git a/Engine/Source/Runtime/Slate/Public/Framework/Commands/Commands.h b/Engine/Source/Runtime/Slate/Public/Framework/Commands/Commands.h
--- a/Engine/Source/Runtime/Slate/Public/Framework/Commands/Commands.h
+++ b/Engine/Source/Runtime/Slate/Public/Framework/Commands/Commands.h
@@ -51,7 +51,7 @@ public:
 				NewInstance->RegisterCommands();
 
 				// Notify that new commands have been registered
-				CommandsChanged.Broadcast();
+				CommandsChanged.Broadcast(*NewInstance);
 			}
 		}
 	}
@@ -73,12 +73,14 @@ public:
 		// The instance may not be valid if it was never used.
 		if( Instance.IsValid() )
 		{
-			FInputBindingManager::Get().RemoveContextByName(Instance.Pin()->GetContextName());
+			auto InstancePtr = Instance.Pin();
+			FInputBindingManager::Get().RemoveContextByName(InstancePtr->GetContextName());
 			
 			// Notify that new commands have been unregistered
-			CommandsChanged.Broadcast();
+			CommandsChanged.Broadcast(*InstancePtr);
 
-			check(!Instance.IsValid());
+			check(InstancePtr.IsUnique());
+			InstancePtr.Reset();
 		}
 	}
 
diff --git a/Engine/Source/Runtime/Slate/Public/Framework/Commands/UICommandInfo.h b/Engine/Source/Runtime/Slate/Public/Framework/Commands/UICommandInfo.h
--- a/Engine/Source/Runtime/Slate/Public/Framework/Commands/UICommandInfo.h
+++ b/Engine/Source/Runtime/Slate/Public/Framework/Commands/UICommandInfo.h
@@ -44,8 +44,6 @@ class SLATE_API FUICommandInfoDecl
 
 public:
 
-	DEPRECATED(4.8, "Use DefaultChord instead of DefaultGesture")
-	FUICommandInfoDecl& DefaultGesture( const FInputChord& InDefaultGesture );
 	FUICommandInfoDecl& DefaultChord( const FInputChord& InDefaultChord );
 	FUICommandInfoDecl& UserInterfaceType( EUserInterfaceActionType::Type InType );
 	FUICommandInfoDecl& Icon( const FSlateIcon& InIcon );
@@ -64,6 +62,7 @@ private:
 	const TSharedRef<FBindingContext>& Context;
 };
 
+DECLARE_MULTICAST_DELEGATE_OneParam(FOnBindingContextChanged, const FBindingContext&);
 
 /**
  * Represents a context in which input bindings are valid
@@ -132,7 +131,7 @@ public:
 	}
 
 	/** A delegate that is called when commands are registered or unregistered with a binding context */
-	static FSimpleMulticastDelegate CommandsChanged; 
+	static FOnBindingContextChanged CommandsChanged;
 
 private:
 
@@ -179,12 +178,8 @@ public:
 	/**
 	 * @return	Returns the active chord for this command
 	 */
-	DEPRECATED(4.8, "Use GetActiveChord instead of GetActiveGesture")
-	const TSharedRef<const FInputChord> GetActiveGesture() const { return GetActiveChord(); }
 	const TSharedRef<const FInputChord> GetActiveChord() const { return ActiveChord; }
 
-	DEPRECATED(4.8, "Use GetDefaultChord instead of GetDefaultGesture")
-	const FInputChord& GetDefaultGesture() const { return GetDefaultChord(); }
 	const FInputChord& GetDefaultChord() const { return DefaultChord; }
 
 	/** Utility function to make an FUICommandInfo */
@@ -209,13 +204,9 @@ public:
 	FName GetBindingContext() const { return BindingContext; }
 
 	/** Sets the new active chord for this command */
-	DEPRECATED(4.8, "Use SetActiveChord instead of SetActiveGesture")
-	void SetActiveGesture( const FInputChord& NewGesture ) { SetActiveChord(NewGesture); }
 	void SetActiveChord( const FInputChord& NewChord );
 
 	/** Removes the active chord from this command */
-	DEPRECATED(4.8, "Use RemoveActiveChord instead of RemoveActiveGesture")
-	void RemoveActiveGesture() { RemoveActiveChord(); }
 	void RemoveActiveChord();
 
 	/** 
diff --git a/Engine/Source/Runtime/Slate/Public/Framework/Commands/UICommandList.h b/Engine/Source/Runtime/Slate/Public/Framework/Commands/UICommandList.h
--- a/Engine/Source/Runtime/Slate/Public/Framework/Commands/UICommandList.h
+++ b/Engine/Source/Runtime/Slate/Public/Framework/Commands/UICommandList.h
@@ -104,6 +104,13 @@ public:
 	 */
 	void UnmapAction( const TSharedPtr< const FUICommandInfo > InUICommandInfo );
 
+	/**
+	 * Checks if an action is already mapped
+	 * 
+	 * @return True if an action is already mapped or false otherwise
+	 */
+	bool IsActionMapped( const TSharedPtr< const FUICommandInfo > InUICommandInfo ) const;
+
 	/**
 	 * Executes the action associated with the provided command info
 	 * Note: It is assumed at this point that CanExecuteAction was already checked
diff --git a/Engine/Source/Runtime/Slate/Public/Framework/SlateDelegates.h b/Engine/Source/Runtime/Slate/Public/Framework/SlateDelegates.h
--- a/Engine/Source/Runtime/Slate/Public/Framework/SlateDelegates.h
+++ b/Engine/Source/Runtime/Slate/Public/Framework/SlateDelegates.h
@@ -118,6 +118,9 @@ DECLARE_DELEGATE_TwoParams( FOnInt32ValueCommitted, int32, ETextCommit::Type);
 /** Notification for FLinearColor value change */
 DECLARE_DELEGATE_OneParam( FOnLinearColorValueChanged, FLinearColor )
 
+/** Notification for FVector4 value change */
+DECLARE_DELEGATE_OneParam(FOnVector4ValueChanged, FVector4&)
+
 
 template< typename ArgumentType >
 class TSlateDelegates
diff --git a/Engine/Source/Runtime/Slate/Public/SlateSettings.h b/Engine/Source/Runtime/Slate/Public/SlateSettings.h
--- a/Engine/Source/Runtime/Slate/Public/SlateSettings.h
+++ b/Engine/Source/Runtime/Slate/Public/SlateSettings.h
@@ -21,4 +21,8 @@ class SLATE_API USlateSettings : public UObject
 	/** Display the virtual keyboard when it receives non-mouse focus. */
 	UPROPERTY(config, EditAnywhere, Category="VirtualKeyboard")
 	bool bVirtualKeyboardDisplayOnFocus;
+
+	/** Allow children of SConstraintCanvas to share render layers. Children must set explicit ZOrder on their slots to control render order. */
+	UPROPERTY(config, EditAnywhere, Category="ConstraintCanvas")
+	bool bExplicitCanvasChildZOrder;
 };
diff --git a/Engine/Source/Runtime/Slate/Public/Widgets/Colors/SColorGradingWheel.h b/Engine/Source/Runtime/Slate/Public/Widgets/Colors/SColorGradingWheel.h
new file mode 100644
--- /dev/null
+++ b/Engine/Source/Runtime/Slate/Public/Widgets/Colors/SColorGradingWheel.h
@@ -0,0 +1,106 @@
+// Copyright 1998-2016 Epic Games, Inc. All Rights Reserved.
+
+#pragma once
+
+#include "CoreMinimal.h"
+#include "Misc/Attribute.h"
+#include "Input/Reply.h"
+#include "Widgets/DeclarativeSyntaxSupport.h"
+#include "Widgets/SLeafWidget.h"
+#include "Framework/SlateDelegates.h"
+
+/**
+ * Implements the color wheel widget.
+ */
+class SLATE_API SColorGradingWheel
+	: public SLeafWidget
+{
+public:
+
+	SLATE_BEGIN_ARGS(SColorGradingWheel)
+		: _SelectedColor()
+		, _DesiredWheelSize()
+		, _ExponentDisplacement()
+		, _OnMouseCaptureBegin()
+		, _OnMouseCaptureEnd()
+		, _OnValueChanged()
+	{ }
+	
+		/** The current color selected by the user. */
+		SLATE_ATTRIBUTE(FLinearColor, SelectedColor)
+		
+		/** The current color selected by the user. */
+		SLATE_ATTRIBUTE(int32, DesiredWheelSize)
+
+		/** The current color selected by the user. */
+		SLATE_ATTRIBUTE(float, ExponentDisplacement)
+		
+		/** Invoked when the mouse is pressed and a capture begins. */
+		SLATE_EVENT(FSimpleDelegate, OnMouseCaptureBegin)
+
+		/** Invoked when the mouse is released and a capture ends. */
+		SLATE_EVENT(FSimpleDelegate, OnMouseCaptureEnd)
+
+		/** Invoked when a new value is selected on the color wheel. */
+		SLATE_EVENT(FOnLinearColorValueChanged, OnValueChanged)
+
+	SLATE_END_ARGS()
+	
+public:
+
+	/**
+	 * Construct this widget.
+	 *
+	 * @param InArgs The declaration data for this widget.
+	 */
+	void Construct(const FArguments& InArgs);
+
+public:
+
+	// SWidget overrides
+
+	virtual FVector2D ComputeDesiredSize(float) const override;
+	virtual FReply OnMouseButtonDoubleClick(const FGeometry& InMyGeometry, const FPointerEvent& InMouseEvent) override;
+	virtual FReply OnMouseButtonDown(const FGeometry& MyGeometry, const FPointerEvent& MouseEvent) override;
+	virtual FReply OnMouseButtonUp(const FGeometry& MyGeometry, const FPointerEvent& MouseEvent) override;
+	virtual FReply OnMouseMove(const FGeometry& MyGeometry, const FPointerEvent& MouseEvent) override;
+	virtual int32 OnPaint(const FPaintArgs& Args, const FGeometry& AllottedGeometry, const FSlateRect& MyClippingRect, FSlateWindowElementList& OutDrawElements, int32 LayerId, const FWidgetStyle& InWidgetStyle, bool bParentEnabled) const override;
+	
+protected:
+
+	/**
+	 * Calculates the position of the color selection indicator.
+	 *
+	 * @return The position relative to the widget.
+	 */
+	FVector2D CalcRelativePositionFromCenter() const;
+
+	/**
+	 * Performs actions according to mouse click / move
+	 *
+	 * @return	True if the mouse action occurred within the color wheel radius
+	 */
+	bool ProcessMouseAction(const FGeometry& MyGeometry, const FPointerEvent& MouseEvent, bool bProcessWhenOutsideColorWheel);
+
+	/** The color wheel image to show. */
+	const FSlateBrush* Image;
+	
+	/** The current color selected by the user. */
+	TAttribute< FLinearColor > SelectedColor;
+	
+	TAttribute< int32 > DesiredWheelSize;
+
+	TAttribute<float> ExponentDisplacement;
+
+	/** The color selector image to show. */
+	const FSlateBrush* SelectorImage;
+
+	/** Invoked when the mouse is pressed and a capture begins. */
+	FSimpleDelegate OnMouseCaptureBegin;
+
+	/** Invoked when the mouse is let up and a capture ends. */
+	FSimpleDelegate OnMouseCaptureEnd;
+
+	/** Invoked when a new value is selected on the color wheel. */
+	FOnLinearColorValueChanged OnValueChanged;
+};
diff --git a/Engine/Source/Runtime/Slate/Public/Widgets/Input/SComboBox.h b/Engine/Source/Runtime/Slate/Public/Widgets/Input/SComboBox.h
--- a/Engine/Source/Runtime/Slate/Public/Widgets/Input/SComboBox.h
+++ b/Engine/Source/Runtime/Slate/Public/Widgets/Input/SComboBox.h
@@ -108,6 +108,7 @@ public:
 		, _MaxListHeight(450.0f)
 		, _HasDownArrow( true )
 		, _EnableGamepadNavigationMode(false)
+		, _IsFocusable( true )
 		{}
 		
 		/** Slot for this button's content (optional) */
@@ -150,11 +151,14 @@ public:
 		 */
 		SLATE_ARGUMENT( bool, HasDownArrow )
 
-		/** 
-		 *  When false, directional keys will change the selection. When true, ComboBox 
+		/**
+		 *  When false, directional keys will change the selection. When true, ComboBox
 		 *	must be activated and will only capture arrow input while activated.
 		*/
 		SLATE_ARGUMENT(bool, EnableGamepadNavigationMode)
+
+		/** When true, allows the combo box to receive keyboard focus */
+		SLATE_ARGUMENT( bool, IsFocusable )
 				
 	SLATE_END_ARGS()
 
@@ -219,7 +223,7 @@ public:
 			.ContentPadding( InArgs._ContentPadding )
 			.ForegroundColor( InArgs._ForegroundColor )
 			.OnMenuOpenChanged(this, &SComboBox< OptionType >::OnMenuOpenChanged)
-
+			.IsFocusable(InArgs._IsFocusable)
 		);
 		SetMenuContentWidgetToFocus(ComboListView);
 
@@ -365,7 +369,7 @@ protected:
 
 	virtual bool SupportsKeyboardFocus() const override
 	{
-		return true;
+		return bIsFocusable;
 	}
 
 	virtual bool IsInteractable() const
diff --git a/Engine/Source/Runtime/Slate/Public/Widgets/Input/SEditableTextBox.h b/Engine/Source/Runtime/Slate/Public/Widgets/Input/SEditableTextBox.h
--- a/Engine/Source/Runtime/Slate/Public/Widgets/Input/SEditableTextBox.h
+++ b/Engine/Source/Runtime/Slate/Public/Widgets/Input/SEditableTextBox.h
@@ -132,7 +132,7 @@ public:
 	 * @param	InArgs	The declaration data for this widget
 	 */
 	void Construct( const FArguments& InArgs );
-	
+
 	/**
 	 * Returns the text string
 	 *
@@ -299,9 +299,6 @@ protected:
 	/** Read-only foreground color (overrides style) */
 	TAttribute<FSlateColor> ReadOnlyForegroundColorOverride;
 
-	/** Read-only foreground color */
-	TAttribute<FSlateColor> ReadOnlyForegroundColor;
-
 	/** Allows for inserting additional widgets that extend the functionality of the text box */
 	TSharedPtr<SHorizontalBox> Box;
 
@@ -310,6 +307,12 @@ protected:
 
 private:
 
+	FMargin FORCEINLINE DeterminePadding() const { check(Style);  return PaddingOverride.IsSet() ? PaddingOverride.Get() : Style->Padding; }
+	FSlateFontInfo FORCEINLINE DetermineFont() const { check(Style);  return FontOverride.IsSet() ? FontOverride.Get() : Style->Font;  }
+	FSlateColor FORCEINLINE DetermineBackgroundColor() const { check(Style);  return BackgroundColorOverride.IsSet() ? BackgroundColorOverride.Get() : Style->BackgroundColor; }
+	
+	FSlateColor DetermineForegroundColor() const;
+
 	/** Styling: border image to draw when not hovered or focused */
 	const FSlateBrush* BorderImageNormal;
 	/** Styling: border image to draw when hovered */
diff --git a/Engine/Source/Runtime/Slate/Public/Widgets/Input/SMultiLineEditableTextBox.h b/Engine/Source/Runtime/Slate/Public/Widgets/Input/SMultiLineEditableTextBox.h
--- a/Engine/Source/Runtime/Slate/Public/Widgets/Input/SMultiLineEditableTextBox.h
+++ b/Engine/Source/Runtime/Slate/Public/Widgets/Input/SMultiLineEditableTextBox.h
@@ -389,9 +389,6 @@ protected:
 	/** Read-only foreground color (overrides style) */
 	TAttribute<FSlateColor> ReadOnlyForegroundColorOverride;
 
-	/** Read-only foreground color */
-	TAttribute<FSlateColor> ReadOnlyForegroundColor;
-
 	/** Whether to disable the context menu */
 	TAttribute< bool > AllowContextMenu;
 
@@ -424,6 +421,14 @@ private:
 
 	const FEditableTextBoxStyle* Style;
 
+	FMargin FORCEINLINE DeterminePadding() const { check(Style);  return PaddingOverride.IsSet() ? PaddingOverride.Get() : Style->Padding; }
+	FMargin FORCEINLINE DetermineHScrollBarPadding() const { check(Style);  return HScrollBarPaddingOverride.IsSet() ? HScrollBarPaddingOverride.Get() : Style->HScrollBarPadding; }
+	FMargin FORCEINLINE DetermineVScrollBarPadding() const { check(Style);  return VScrollBarPaddingOverride.IsSet() ? VScrollBarPaddingOverride.Get() : Style->VScrollBarPadding; }
+	FSlateFontInfo FORCEINLINE DetermineFont() const { check(Style);  return FontOverride.IsSet() ? FontOverride.Get() : Style->Font; }
+	FSlateColor FORCEINLINE DetermineBackgroundColor() const { check(Style);  return BackgroundColorOverride.IsSet() ? BackgroundColorOverride.Get() : Style->BackgroundColor; }
+
+	FSlateColor DetermineForegroundColor() const;
+
 	/** Styling: border image to draw when not hovered or focused */
 	const FSlateBrush* BorderImageNormal;
 	/** Styling: border image to draw when hovered */
diff --git a/Engine/Source/Runtime/Slate/Public/Widgets/Input/SNumericEntryBox.h b/Engine/Source/Runtime/Slate/Public/Widgets/Input/SNumericEntryBox.h
--- a/Engine/Source/Runtime/Slate/Public/Widgets/Input/SNumericEntryBox.h
+++ b/Engine/Source/Runtime/Slate/Public/Widgets/Input/SNumericEntryBox.h
@@ -23,6 +23,7 @@
 #include "Widgets/Text/STextBlock.h"
 #include "Widgets/Input/SEditableText.h"
 #include "Widgets/Input/SSpinBox.h"
+#include "Widgets/Colors/SColorBlock.h"
 
 /**
  * Implementation for a box that only accepts a numeric value or that can display an undetermined value via a string
@@ -62,13 +63,18 @@ public:
 		, _BorderBackgroundColor(FLinearColor::White)
 		, _UndeterminedString( SNumericEntryBox<NumericType>::DefaultUndeterminedString )
 		, _AllowSpin(false)
+		, _UseDarkStyle(false)
+		, _ColorIndex(-1)
+		, _ShiftMouseMovePixelPerDelta(1)
 		, _Delta(0)
 		, _MinValue(TNumericLimits<NumericType>::Lowest())
 		, _MaxValue(TNumericLimits<NumericType>::Max())
 		, _MinSliderValue(0)				
 		, _MaxSliderValue(100)
 		, _SliderExponent(1.f)
+		, _SliderExponentNeutralValue(0)
 		, _MinDesiredValueWidth(0)
+		
 	{}		
 
 		/** Style to use for the editable text box within this widget */
@@ -92,6 +98,12 @@ public:
 		SLATE_ATTRIBUTE( FSlateFontInfo, Font )
 		/** Whether or not the user should be able to change the value by dragging with the mouse cursor */
 		SLATE_ARGUMENT( bool, AllowSpin )
+		/** Whether or not the user should be able to change the value by dragging with the mouse cursor */
+		SLATE_ARGUMENT(bool, UseDarkStyle)
+		/** The color index to draw the line at the bottom of the numeric box (-1= no Line, 0=Red, 1=Green 2=Blue 3=White) */
+		SLATE_ARGUMENT( int32, ColorIndex )
+		/** How many pixel the mouse must move to change the value of the delta step (only use if there is a spinbox allow) */
+		SLATE_ARGUMENT( int32, ShiftMouseMovePixelPerDelta )
 		/** Delta to increment the value as the slider moves.  If not specified will determine automatically */
 		SLATE_ATTRIBUTE( NumericType, Delta )
 		/** The minimum value that can be entered into the text edit box */
@@ -104,6 +116,8 @@ public:
 		SLATE_ATTRIBUTE( TOptional< NumericType >, MaxSliderValue )
 		/** Use exponential scale for the slider */
 		SLATE_ATTRIBUTE( float, SliderExponent )
+		/** When using exponential scale specify a neutral value where we want the maximum precision (by default it is the smallest slider value)*/
+		SLATE_ATTRIBUTE(NumericType, SliderExponentNeutralValue )
 		/** The minimum desired width for the value portion of the control. */
 		SLATE_ATTRIBUTE( float, MinDesiredValueWidth )
 		/** The text margin to use if overridden. */
@@ -145,6 +159,7 @@ public:
 		BorderImageHovered = &InArgs._EditableTextBoxStyle->BackgroundImageHovered;
 		BorderImageFocused = &InArgs._EditableTextBoxStyle->BackgroundImageFocused;
 		Interface = InArgs._TypeInterface.IsValid() ? InArgs._TypeInterface : MakeShareable( new TDefaultNumericTypeInterface<NumericType> );
+		int32 InColorIndex = InArgs._ColorIndex;
 
 		TAttribute<FMargin> TextMargin = InArgs._OverrideTextMargin.IsSet() ? InArgs._OverrideTextMargin : InArgs._EditableTextBoxStyle->Padding;
 		const bool bAllowSpin = InArgs._AllowSpin;
@@ -153,11 +168,12 @@ public:
 		if( bAllowSpin )
 		{
 			SAssignNew( SpinBox, SSpinBox<NumericType> )
-				.Style( FCoreStyle::Get(), "NumericEntrySpinBox" )
+				.Style( FCoreStyle::Get(), InArgs._UseDarkStyle ? "NumericEntrySpinBox_Dark" : "NumericEntrySpinBox" )
 				.Font( InArgs._Font.IsSet() ? InArgs._Font : InArgs._EditableTextBoxStyle->Font )
 				.ContentPadding( TextMargin )
 				.Value( this, &SNumericEntryBox<NumericType>::OnGetValueForSpinBox )
 				.Delta( InArgs._Delta )
+				.ShiftMouseMovePixelPerDelta(InArgs._ShiftMouseMovePixelPerDelta)
 				.OnValueChanged( OnValueChanged )
 				.OnValueCommitted( OnValueCommitted )
 				.MinSliderValue(InArgs._MinSliderValue)
@@ -165,6 +181,7 @@ public:
 				.MaxValue(InArgs._MaxValue)
 				.MinValue(InArgs._MinValue)
 				.SliderExponent(InArgs._SliderExponent)
+				.SliderExponentNeutralValue(InArgs._SliderExponentNeutralValue)
 				.OnBeginSliderMovement(InArgs._OnBeginSliderMovement)
 				.OnEndSliderMovement(InArgs._OnEndSliderMovement)
 				.MinDesiredWidth(InArgs._MinDesiredValueWidth)
@@ -184,7 +201,32 @@ public:
 			.ContextMenuExtender( InArgs._ContextMenuExtender )
 			.MinDesiredWidth(InArgs._MinDesiredValueWidth);
 
-		TSharedRef<SHorizontalBox> HorizontalBox = SNew( SHorizontalBox );
+		TSharedRef<SHorizontalBox> HorizontalBox = SNew(SHorizontalBox);
+		TSharedRef<SWidget> BoxToAdd = HorizontalBox;
+
+		if (InColorIndex != INDEX_NONE)
+		{
+			TSharedRef<SVerticalBox> VerticalBox = SNew(SVerticalBox);
+			FColor Color[4] = { FColor::Red, FColor::Green, FColor::Blue, FColor::White };
+			//Add the horizontal box
+			VerticalBox->AddSlot()
+				.FillHeight(1)
+				.AutoHeight()
+				[
+					HorizontalBox
+				];
+			//Add the color line
+			VerticalBox->AddSlot()
+				.AutoHeight()
+				.MaxHeight(2.0f)
+				.Padding(3.0f, 0.0f)
+				[
+					SNew(SColorBlock)
+					.Color(Color[InColorIndex])
+				];
+			BoxToAdd = VerticalBox;
+		}
+		
 		if( InArgs._Label.Widget != SNullWidget::NullWidget )
 		{
 			HorizontalBox->AddSlot()
@@ -226,7 +268,7 @@ public:
 				.ForegroundColor( InArgs._BorderForegroundColor )
 				.Padding(0)
 				[
-					HorizontalBox
+					BoxToAdd
 				]
 			];
 	}
