Copying //UE4/Dev-Editor to //UE4/Dev-Main (Source: //UE4/Dev-Editor @ 3228984)

diff --git a/Engine/Source/Runtime/Slate/Public/Widgets/Input/SSpinBox.h b/Engine/Source/Runtime/Slate/Public/Widgets/Input/SSpinBox.h
--- a/Engine/Source/Runtime/Slate/Public/Widgets/Input/SSpinBox.h
+++ b/Engine/Source/Runtime/Slate/Public/Widgets/Input/SSpinBox.h
@@ -25,6 +25,25 @@
 #include "Widgets/Images/SImage.h"
 #include "Templates/IsIntegral.h"
 
+/*
+ * This function compute a slider position by simulating two log on both side of the neutral value
+ * Example a slider going from 0.0 to 2.0 with a neutral value of 1.0, the user will have a lot of precision around the neutral value
+ * on both side.
+ |
+ ||                              |
+ | -_                          _-
+ |   --__                  __--
+ |       ----__________----
+ ----------------------------------
+  0              1               2
+
+  The function return a float representing the slider fraction used to position the slider handle
+  FractionFilled: this is the value slider position with no exponent
+  StartFractionFilled: this is the neutral value slider position with no exponent
+  SliderExponent: this is the slider exponent
+*/
+SLATE_API float SpinBoxComputeExponentSliderFraction(float FractionFilled, float StartFractionFilled, float SliderExponent);
+
 /**
  * A Slate SpinBox resembles traditional spin boxes in that it is a widget that provides
  * keyboard-based and mouse-based manipulation of a numeric value.
@@ -48,6 +67,7 @@ public:
 		, _MinValue(0)
 		, _MaxValue(10)
 		, _Delta(0)
+		, _ShiftMouseMovePixelPerDelta(1)
 		, _SliderExponent(1)
 		, _Font( FCoreStyle::Get().GetFontStyle( TEXT( "NormalFont" ) ) )
 		, _ContentPadding(  FMargin( 2.0f, 1.0f) )
@@ -73,8 +93,12 @@ public:
 		SLATE_ATTRIBUTE( TOptional< NumericType >, MaxSliderValue )
 		/** Delta to increment the value as the slider moves.  If not specified will determine automatically */
 		SLATE_ATTRIBUTE( NumericType, Delta )
+		/** How many pixel the mouse must move to change the value of the delta step */
+		SLATE_ATTRIBUTE( int32, ShiftMouseMovePixelPerDelta )
 		/** Use exponential scale for the slider */
 		SLATE_ATTRIBUTE( float, SliderExponent )
+		/** When use exponential scale for the slider which is the neutral value */
+		SLATE_ATTRIBUTE( NumericType, SliderExponentNeutralValue )
 		/** Font used to display text in the slider */
 		SLATE_ATTRIBUTE( FSlateFontInfo, Font )
 		/** Padding to add around this widget and its internal widgets */
@@ -132,10 +156,13 @@ public:
 	
 		SliderExponent = InArgs._SliderExponent;
 
+		SliderExponentNeutralValue = InArgs._SliderExponentNeutralValue;
+
 		DistanceDragged = 0.0f;
 		PreDragValue = 0;
 
 		Delta = InArgs._Delta;
+		ShiftMouseMovePixelPerDelta = InArgs._ShiftMouseMovePixelPerDelta;
 	
 		BackgroundHoveredBrush = &InArgs._Style->HoveredBackgroundBrush;
 		BackgroundBrush = &InArgs._Style->BackgroundBrush;
@@ -144,9 +171,13 @@ public:
 		const FMargin& TextMargin = InArgs._Style->TextPadding;
 
 		bDragging = false;
+		PointerDraggingSliderIndex = INDEX_NONE;
+
 		CachedExternalValue = ValueAttribute.Get();
 		InternalValue = ValueAttribute.Get();
 
+		bIsTextChanging = false;
+
 		this->ChildSlot
 		.Padding( InArgs._ContentPadding )
 		[
@@ -176,6 +207,7 @@ public:
 				.SelectAllTextWhenFocused( true )
 				.Text( this, &SSpinBox<NumericType>::GetValueAsText )
 				.OnIsTypedCharValid(this, &SSpinBox<NumericType>::IsCharacterValid)
+				.OnTextChanged( this, &SSpinBox<NumericType>::TextField_OnTextChanged )
 				.OnTextCommitted( this, &SSpinBox<NumericType>::TextField_OnTextCommitted )
 				.ClearKeyboardFocusOnCommit( InArgs._ClearKeyboardFocusOnCommit )
 				.SelectAllTextOnCommit( InArgs._SelectAllTextOnCommit )
@@ -238,7 +270,16 @@ public:
 			const float CachedSliderExponent = SliderExponent.Get();
 			if (CachedSliderExponent != 1)
 			{
-				FractionFilled = 1.0f - FMath::Pow( 1.0f - FractionFilled, CachedSliderExponent);
+				if (SliderExponentNeutralValue.IsSet())
+				{
+					//Compute a log curve on both side of the neutral value
+					float StartFractionFilled = Fraction(SliderExponentNeutralValue.Get(), GetMinSliderValue(), GetMaxSliderValue());
+					FractionFilled = SpinBoxComputeExponentSliderFraction(FractionFilled, StartFractionFilled, CachedSliderExponent);
+				}
+				else
+				{
+					FractionFilled = 1.0f - FMath::Pow( 1.0f - FractionFilled, CachedSliderExponent);
+				}
 			}
 			const FVector2D FillSize( AllottedGeometry.Size.X * FractionFilled, AllottedGeometry.Size.Y );
 
@@ -268,10 +309,11 @@ public:
 	 */
 	virtual FReply OnMouseButtonDown( const FGeometry& MyGeometry, const FPointerEvent& MouseEvent ) override
 	{
-		if ( MouseEvent.GetEffectingButton() == EKeys::LeftMouseButton )
+		if ( MouseEvent.GetEffectingButton() == EKeys::LeftMouseButton && PointerDraggingSliderIndex == INDEX_NONE )
 		{
 			DistanceDragged = 0;
 			PreDragValue = InternalValue;
+			PointerDraggingSliderIndex = MouseEvent.GetPointerIndex();
 			CachedMousePosition = MouseEvent.GetScreenSpacePosition().IntPoint();
 			return FReply::Handled().CaptureMouse(SharedThis(this)).UseHighPrecisionMouseMovement(SharedThis(this)).SetUserFocus(SharedThis(this), EFocusCause::Mouse);
 		}
@@ -291,7 +333,7 @@ public:
 	 */
 	virtual FReply OnMouseButtonUp( const FGeometry& MyGeometry, const FPointerEvent& MouseEvent ) override
 	{
-		if ( MouseEvent.GetEffectingButton() == EKeys::LeftMouseButton && this->HasMouseCapture() )
+		if ( MouseEvent.GetEffectingButton() == EKeys::LeftMouseButton && PointerDraggingSliderIndex == MouseEvent.GetPointerIndex() && this->HasMouseCapture() )
 		{
 			if( bDragging )
 			{
@@ -299,14 +341,22 @@ public:
 			}
 
 			bDragging = false;
+			PointerDraggingSliderIndex = INDEX_NONE;
+
+			FReply Reply = FReply::Handled().ReleaseMouseCapture();
+
+			if ( !MouseEvent.IsTouchEvent() )
+			{
+				Reply.SetMousePos(CachedMousePosition);
+			}
+
 			if ( DistanceDragged < FSlateApplication::Get().GetDragTriggerDistance() )
 			{
 				EnterTextMode();
-				return FReply::Handled().ReleaseMouseCapture().SetUserFocus(EditableText.ToSharedRef(), EFocusCause::SetDirectly).SetMousePos(CachedMousePosition);
+				Reply.SetUserFocus(EditableText.ToSharedRef(), EFocusCause::SetDirectly);
 			}
 
-			return FReply::Handled().ReleaseMouseCapture().SetMousePos(CachedMousePosition);
-			
+			return Reply;
 		}
 		else
 		{
@@ -323,7 +373,7 @@ public:
 	 */
 	virtual FReply OnMouseMove( const FGeometry& MyGeometry, const FPointerEvent& MouseEvent ) override
 	{
-		if ( this->HasMouseCapture() )
+		if ( PointerDraggingSliderIndex == MouseEvent.GetPointerIndex() && this->HasMouseCapture() )
 		{
 			if (!bDragging)
 			{
@@ -346,7 +396,13 @@ public:
 
 				// A minimum slider width to use for calculating deltas in the slider-range space
 				const float MinSliderWidth = 100.f;
-				const float SliderWidthInSlateUnits = FMath::Max(MyGeometry.GetDrawSize().X, MinSliderWidth);
+				float SliderWidthInSlateUnits = FMath::Max(MyGeometry.GetDrawSize().X, MinSliderWidth);
+				
+				const int32 CachedShiftMouseMovePixelPerDelta = ShiftMouseMovePixelPerDelta.Get();
+				if (CachedShiftMouseMovePixelPerDelta > 1 && MouseEvent.IsShiftDown())
+				{
+					SliderWidthInSlateUnits *= CachedShiftMouseMovePixelPerDelta;
+				}
 				
 				//if we have a range to draw in
 				if ( !bUnlimitedSpinRange) 
@@ -357,7 +413,16 @@ public:
 						
 					if (CachedSliderExponent != 1)
 					{
-						FractionFilled = 1.0f - FMath::Pow( 1.0f - FractionFilled, CachedSliderExponent);
+						if (SliderExponentNeutralValue.IsSet())
+						{
+							//Compute a log curve on both side of the neutral value
+							float StartFractionFilled = Fraction(SliderExponentNeutralValue.Get(), GetMinSliderValue(), GetMaxSliderValue());
+							FractionFilled = SpinBoxComputeExponentSliderFraction(FractionFilled, StartFractionFilled, CachedSliderExponent);
+						}
+						else
+						{
+							FractionFilled = 1.0f - FMath::Pow( 1.0f - FractionFilled, CachedSliderExponent);
+						}
 					}
 					FractionFilled *= SliderWidthInSlateUnits;
 
@@ -372,7 +437,18 @@ public:
 					if (CachedSliderExponent != 1)
 					{
 						// Have to convert the percent to the proper value due to the exponent component to the spin.
-						Percent = 1.0f - FMath::Pow(1.0f - Percent, 1.0 / CachedSliderExponent);
+						if (SliderExponentNeutralValue.IsSet())
+						{
+							//Compute a log curve on both side of the neutral value
+							float StartFractionFilled = Fraction(SliderExponentNeutralValue.Get(), GetMinSliderValue(), GetMaxSliderValue());
+							Percent = SpinBoxComputeExponentSliderFraction(Percent, StartFractionFilled, 1.0/CachedSliderExponent);
+						}
+						else
+						{
+							Percent = 1.0f - FMath::Pow(1.0f - Percent, 1.0 / CachedSliderExponent);
+						}
+						
+						
 					}
 
 					NewValue = FMath::LerpStable<double>(GetMinSliderValue(), GetMaxSliderValue(), Percent);
@@ -435,18 +511,21 @@ public:
 		}
 		else if ( Key == EKeys::Up || Key == EKeys::Right )
 		{
-			CommitValue( InternalValue + Delta.Get(), CommittedViaSpin, ETextCommit::OnEnter );
+			InternalValue = ValueAttribute.Get();
+			CommitValue( InternalValue + Delta.Get(), CommittedViaArrowKey, ETextCommit::OnEnter );
 			ExitTextMode();
 			return FReply::Handled();	
 		}
 		else if ( Key == EKeys::Down || Key == EKeys::Left )
 		{
-			CommitValue( InternalValue - Delta.Get(), CommittedViaSpin, ETextCommit::OnEnter );
+			InternalValue = ValueAttribute.Get();
+			CommitValue( InternalValue - Delta.Get(), CommittedViaArrowKey, ETextCommit::OnEnter );
 			ExitTextMode();
 			return FReply::Handled();
 		}
 		else if ( Key == EKeys::Enter )
 		{
+			InternalValue = ValueAttribute.Get();
 			EnterTextMode();
 			return FReply::Handled().SetUserFocus(EditableText.ToSharedRef(), EFocusCause::Navigation);
 		}
@@ -547,6 +626,38 @@ protected:
 	{
 		return FText::FromString(GetValueAsString());
 	}
+
+	/**
+	 * Invoked when the text in the text field changes
+	 *
+	 * @param NewText		The value of the text in the text field
+	 */
+	void TextField_OnTextChanged( const FText& NewText)
+	{
+		if (!bIsTextChanging)
+		{
+			TGuardValue<bool> TextChangedGuard(bIsTextChanging, true);
+
+			// Validate the text on change, and only accept text up until the first invalid character
+			FString Data = NewText.ToString();
+			int32 NumValidChars = Data.Len();
+
+			for (int32 Index = 0; Index < Data.Len(); ++Index)
+			{
+				if (!Interface->IsCharacterValid(Data[Index]))
+				{
+					NumValidChars = Index;
+					break;
+				}
+			}
+
+			if (NumValidChars < Data.Len())
+			{
+				FString ValidData = NumValidChars > 0 ? Data.Left(NumValidChars) : GetValueAsString();
+				EditableText->SetText(FText::FromString(ValidData));
+			}
+		}
+	}
 	
 	/**
 	 * Invoked when the text field commits its text.
@@ -573,7 +684,8 @@ protected:
 	enum ECommitMethod
 	{
 		CommittedViaSpin,
-		CommittedViaTypeIn	
+		CommittedViaTypeIn,
+		CommittedViaArrowKey
 	};
 
 	/**
@@ -585,7 +697,7 @@ protected:
 	 */
 	void CommitValue( double NewValue, ECommitMethod CommitMethod, ETextCommit::Type OriginalCommitInfo )
 	{
-		if( CommitMethod == CommittedViaSpin )
+		if( CommitMethod == CommittedViaSpin || CommitMethod == CommittedViaArrowKey )
 		{
 			NewValue = FMath::Clamp<double>( NewValue, GetMinSliderValue(), GetMaxSliderValue() );
 		}
@@ -613,7 +725,7 @@ protected:
 		InternalValue = NewValue;
 
 		// If needed, round this value to the delta. Internally the value is not held to the Delta but externally it appears to be.
-		if ( CommitMethod == CommittedViaSpin )
+		if ( CommitMethod == CommittedViaSpin || CommitMethod == CommittedViaArrowKey )
 		{
 			NumericType CurrentDelta = Delta.Get();
 			if( CurrentDelta != 0 )
@@ -622,7 +734,7 @@ protected:
 			}
 		}
 
-		if( CommitMethod == CommittedViaTypeIn )
+		if( CommitMethod == CommittedViaTypeIn || CommitMethod == CommittedViaArrowKey )
 		{
 			OnValueCommitted.ExecuteIfBound( RoundIfIntegerValue( NewValue ), OriginalCommitInfo );
 		}
@@ -699,7 +811,9 @@ private:
 
 	float DistanceDragged;
 	TAttribute<NumericType> Delta;
+	TAttribute<int32> ShiftMouseMovePixelPerDelta;
 	TAttribute<float> SliderExponent;
+	TAttribute<NumericType> SliderExponentNeutralValue;
 	TAttribute< TOptional<NumericType> > MinValue;
 	TAttribute< TOptional<NumericType> > MaxValue;
 	TAttribute< TOptional<NumericType> > MinSliderValue;
@@ -728,6 +842,9 @@ private:
 
 	/** Whether the user is dragging the slider */
 	bool bDragging;
+
+	/** Tracks which cursor is currently dragging the slider (e.g., the mouse cursor or a specific finger) */
+	int32 PointerDraggingSliderIndex;
 	
 	/** Cached mouse position to restore after scrolling. */
 	FIntPoint CachedMousePosition;
@@ -743,4 +860,7 @@ private:
 	/** This is the cached value the user believes it to be (usually different due to truncation to an int). Used for identifying 
 		external forces on the spinbox and syncing the internal value to them. Synced when a value is committed to the spinbox. */
 	NumericType CachedExternalValue;
+
+	/** Re-entrant guard for the text changed handler */
+	bool bIsTextChanging;
 };
diff --git a/Engine/Source/Runtime/Slate/Public/Widgets/Layout/SBackgroundBlur.h b/Engine/Source/Runtime/Slate/Public/Widgets/Layout/SBackgroundBlur.h
new file mode 100644
--- /dev/null
+++ b/Engine/Source/Runtime/Slate/Public/Widgets/Layout/SBackgroundBlur.h
@@ -0,0 +1,65 @@
+// Copyright 1998-2016 Epic Games, Inc. All Rights Reserved.
+
+#pragma once
+
+#include "SlateDelegates.h"
+#include "Widgets/DeclarativeSyntaxSupport.h"
+#include "Widgets/SWidget.h"
+#include "Layout/Margin.h"
+#include "Widgets/SCompoundWidget.h"
+/**
+* A BackgroundBlur is similar to a Border in that it can be used to contain other widgets.
+* However, instead of surrounding the content with an image, it applies a post-process blur to
+* everything (both actors and widgets) that is underneath it.
+* 
+* Note: For low-spec machines where the blur effect is too computationally expensive, 
+* a user-specified fallback image is used instead (effectively turning this into a Border)
+*/
+class SLATE_API SBackgroundBlur : public SCompoundWidget
+{
+public:
+	SLATE_BEGIN_ARGS(SBackgroundBlur)
+		: _HAlign(HAlign_Fill)
+		, _VAlign(VAlign_Fill)
+		, _Padding(FMargin(2.0f))
+		, _bApplyAlphaToBlur(true)
+		, _BlurStrength(0.f)
+		, _BlurRadius()
+		, _LowQualityFallbackBrush(nullptr)
+	{
+		_Visibility = EVisibility::SelfHitTestInvisible;
+	}
+		SLATE_DEFAULT_SLOT(FArguments, Content)
+		
+		SLATE_ARGUMENT(EHorizontalAlignment, HAlign)
+		SLATE_ARGUMENT(EVerticalAlignment, VAlign)
+		SLATE_ATTRIBUTE(FMargin, Padding)
+
+		SLATE_ARGUMENT(bool, bApplyAlphaToBlur)
+		SLATE_ATTRIBUTE(float, BlurStrength)
+		SLATE_ATTRIBUTE(TOptional<int32>, BlurRadius)
+		SLATE_ARGUMENT(const FSlateBrush*, LowQualityFallbackBrush)
+	SLATE_END_ARGS()
+
+public:
+	void Construct(const FArguments& InArgs);
+	virtual int32 OnPaint(const FPaintArgs& Args, const FGeometry& AllottedGeometry, const FSlateRect& MyClippingRect, FSlateWindowElementList& OutDrawElements, int32 LayerId, const FWidgetStyle& InWidgetStyle, bool bParentEnabled) const override;
+
+	void SetContent(const TSharedRef<SWidget>& InContent);
+	void SetApplyAlphaToBlur(bool bInApplyAlphaToBlur);
+	void SetBlurRadius(const TAttribute<TOptional<int32>>& InBlurRadius);
+	void SetBlurStrength(const TAttribute<float>& InStrength);
+	void SetLowQualityBackgroundBrush(const FSlateBrush* InBrush);
+	
+	void SetHAlign(EHorizontalAlignment HAlign);
+	void SetVAlign(EVerticalAlignment VAlign);
+	void SetPadding(const TAttribute<FMargin>& InPadding);
+
+private:
+	void ComputeEffectiveKernelSize(float Strength, int32& OutKernelSize, int32& OutDownsampleAmount) const;
+
+	bool bApplyAlphaToBlur;
+	TAttribute<float> BlurStrength;
+	TAttribute<TOptional<int32>> BlurRadius;
+	const FSlateBrush* LowQualityFallbackBrush;
+};
\ No newline at end of file
diff --git a/Engine/Source/Runtime/Slate/Public/Widgets/Layout/SConstraintCanvas.h b/Engine/Source/Runtime/Slate/Public/Widgets/Layout/SConstraintCanvas.h
--- a/Engine/Source/Runtime/Slate/Public/Widgets/Layout/SConstraintCanvas.h
+++ b/Engine/Source/Runtime/Slate/Public/Widgets/Layout/SConstraintCanvas.h
@@ -160,6 +160,14 @@ public:
 
 	// End SWidget overrides
 
+private:
+
+	/** An array matching the length and order of ArrangedChildren. True means the child must be placed in a layer in front of all previous children. */
+	typedef TArray<bool, TInlineAllocator<16>> FArrangedChildLayers;
+
+	/** Like ArrangeChildren but also generates an array of layering information (see FArrangedChildLayers). */
+	void ArrangeLayeredChildren(const FGeometry& AllottedGeometry, FArrangedChildren& ArrangedChildren, FArrangedChildLayers& ArrangedChildLayers) const;
+
 protected:
 
 	/** The ConstraintCanvas widget's children. */
diff --git a/Engine/Source/Runtime/Slate/Public/Widgets/Layout/SDPIScaler.h b/Engine/Source/Runtime/Slate/Public/Widgets/Layout/SDPIScaler.h
--- a/Engine/Source/Runtime/Slate/Public/Widgets/Layout/SDPIScaler.h
+++ b/Engine/Source/Runtime/Slate/Public/Widgets/Layout/SDPIScaler.h
@@ -54,7 +54,7 @@ public:
 
 protected:
 
-	virtual float GetRelativeLayoutScale(const FSlotBase& Child) const override;
+	virtual float GetRelativeLayoutScale(const FSlotBase& Child, float LayoutScaleMultiplier) const override;
 
 	/**
 	 * The scaling factor from 1:1 pixel to slate unit (SU) ratio.
diff --git a/Engine/Source/Runtime/Slate/Public/Widgets/Layout/SScaleBox.h b/Engine/Source/Runtime/Slate/Public/Widgets/Layout/SScaleBox.h
--- a/Engine/Source/Runtime/Slate/Public/Widgets/Layout/SScaleBox.h
+++ b/Engine/Source/Runtime/Slate/Public/Widgets/Layout/SScaleBox.h
@@ -143,7 +143,7 @@ public:
 	void SetIgnoreInheritedScale(bool InIgnoreInheritedScale);
 	
 protected:
-	virtual float GetRelativeLayoutScale(const FSlotBase& Child) const override;
+	virtual float GetRelativeLayoutScale(const FSlotBase& Child, float LayoutScaleMultiplier) const override;
 
 	float GetLayoutScale() const;
 	void RefreshSafeZoneScale();
diff --git a/Engine/Source/Runtime/Slate/Public/Widgets/Text/SlateEditableTextLayout.h b/Engine/Source/Runtime/Slate/Public/Widgets/Text/SlateEditableTextLayout.h
--- a/Engine/Source/Runtime/Slate/Public/Widgets/Text/SlateEditableTextLayout.h
+++ b/Engine/Source/Runtime/Slate/Public/Widgets/Text/SlateEditableTextLayout.h
@@ -409,6 +409,11 @@ private:
 			return bIsComposing;
 		}
 
+		FORCEINLINE void AbortComposition()
+		{
+			bIsComposing = false;
+		}
+
 		FORCEINLINE FTextRange GetCompositionRange() const
 		{
 			return FTextRange(CompositionBeginIndex, CompositionBeginIndex + CompositionLength);
diff --git a/Engine/Source/Runtime/Slate/Public/Widgets/Views/STableViewBase.h b/Engine/Source/Runtime/Slate/Public/Widgets/Views/STableViewBase.h
--- a/Engine/Source/Runtime/Slate/Public/Widgets/Views/STableViewBase.h
+++ b/Engine/Source/Runtime/Slate/Public/Widgets/Views/STableViewBase.h
@@ -100,6 +100,9 @@ public:
 	/** Scrolls the view to the bottom */
 	void ScrollToBottom();
 
+	/** Gets the scroll offset of this view (in items) */
+	float GetScrollOffset() const;
+
 	/** Set the scroll offset of this view (in items) */
 	void SetScrollOffset( const float InScrollOffset );
 
diff --git a/Engine/Source/Runtime/SlateCore/Private/Fonts/CompositeFont.cpp b/Engine/Source/Runtime/SlateCore/Private/Fonts/CompositeFont.cpp
--- a/Engine/Source/Runtime/SlateCore/Private/Fonts/CompositeFont.cpp
+++ b/Engine/Source/Runtime/SlateCore/Private/Fonts/CompositeFont.cpp
@@ -1,57 +1,72 @@
 // Copyright 1998-2017 Epic Games, Inc. All Rights Reserved.
 
 #include "Fonts/CompositeFont.h"
+#include "SlateGlobals.h"
 #include "UObject/EditorObjectVersion.h"
 #include "Fonts/FontFaceInterface.h"
 #include "Templates/Casts.h"
 #include "Fonts/FontBulkData.h"
 
+// The total true type memory we are using for resident font faces
+DECLARE_MEMORY_STAT(TEXT("Resident Font Memory (TTF/OTF)"), STAT_SlateRawFontDataMemory, STATGROUP_SlateMemory);
+
+void FFontFaceData::TrackMemoryUsage() const
+{
+	INC_DWORD_STAT_BY(STAT_SlateRawFontDataMemory, Data.GetAllocatedSize());
+}
+
+void FFontFaceData::UntrackMemoryUsage() const
+{
+	DEC_DWORD_STAT_BY(STAT_SlateRawFontDataMemory, Data.GetAllocatedSize());
+}
+
 FFontData::FFontData()
 	: FontFilename()
 	, Hinting(EFontHinting::Default)
-	, LoadingPolicy(EFontLoadingPolicy::PreLoad)
-#if WITH_EDITORONLY_DATA
+	, LoadingPolicy(EFontLoadingPolicy::LazyLoad)
 	, FontFaceAsset(nullptr)
+#if WITH_EDITORONLY_DATA
 	, BulkDataPtr_DEPRECATED(nullptr)
 	, FontData_DEPRECATED()
 #endif // WITH_EDITORONLY_DATA
 {
 }
 
-#if WITH_EDITORONLY_DATA
 FFontData::FFontData(const UObject* const InFontFaceAsset)
 	: FontFilename()
 	, Hinting(EFontHinting::Default)
-	, LoadingPolicy(EFontLoadingPolicy::PreLoad)
+	, LoadingPolicy(EFontLoadingPolicy::LazyLoad)
 	, FontFaceAsset(InFontFaceAsset)
+#if WITH_EDITORONLY_DATA
 	, BulkDataPtr_DEPRECATED(nullptr)
 	, FontData_DEPRECATED()
+#endif // WITH_EDITORONLY_DATA
 {
 	if (FontFaceAsset)
 	{
 		CastChecked<const IFontFaceInterface>(FontFaceAsset);
 	}
 }
-#endif // WITH_EDITORONLY_DATA
 
 FFontData::FFontData(FString InFontFilename, const EFontHinting InHinting, const EFontLoadingPolicy InLoadingPolicy)
 	: FontFilename(MoveTemp(InFontFilename))
 	, Hinting(InHinting)
 	, LoadingPolicy(InLoadingPolicy)
-#if WITH_EDITORONLY_DATA
 	, FontFaceAsset(nullptr)
+#if WITH_EDITORONLY_DATA
 	, BulkDataPtr_DEPRECATED(nullptr)
 	, FontData_DEPRECATED()
 #endif // WITH_EDITORONLY_DATA
 {
+	check(InLoadingPolicy != EFontLoadingPolicy::Inline);
 }
 
 FFontData::FFontData(FString InFontFilename, const UFontBulkData* const InBulkData, const EFontHinting InHinting)
 	: FontFilename(MoveTemp(InFontFilename))
 	, Hinting(InHinting)
-	, LoadingPolicy(EFontLoadingPolicy::PreLoad)
-#if WITH_EDITORONLY_DATA
+	, LoadingPolicy(EFontLoadingPolicy::LazyLoad)
 	, FontFaceAsset(nullptr)
+#if WITH_EDITORONLY_DATA
 	, BulkDataPtr_DEPRECATED(nullptr)
 	, FontData_DEPRECATED()
 #endif // WITH_EDITORONLY_DATA
@@ -60,50 +75,42 @@ FFontData::FFontData(FString InFontFilename, const UFontBulkData* const InBulkDa
 
 const FString& FFontData::GetFontFilename() const
 {
-#if WITH_EDITORONLY_DATA
 	if (FontFaceAsset)
 	{
 		const IFontFaceInterface* FontFace = CastChecked<const IFontFaceInterface>(FontFaceAsset);
 		return FontFace->GetFontFilename();
 	}
-#endif // WITH_EDITORONLY_DATA
 	return FontFilename;
 }
 
 EFontHinting FFontData::GetHinting() const
 {
-#if WITH_EDITORONLY_DATA
 	if (FontFaceAsset)
 	{
 		const IFontFaceInterface* FontFace = CastChecked<const IFontFaceInterface>(FontFaceAsset);
 		return FontFace->GetHinting();
 	}
-#endif // WITH_EDITORONLY_DATA
 	return Hinting;
 }
 
 EFontLoadingPolicy FFontData::GetLoadingPolicy() const
 {
-#if WITH_EDITORONLY_DATA
 	if (FontFaceAsset)
 	{
 		const IFontFaceInterface* FontFace = CastChecked<const IFontFaceInterface>(FontFaceAsset);
 		return FontFace->GetLoadingPolicy();
 	}
-#endif // WITH_EDITORONLY_DATA
 	return LoadingPolicy;
 }
 
-#if WITH_EDITORONLY_DATA
-bool FFontData::GetFontFaceData(TArray<uint8>& OutFontFaceData) const
+FFontFaceDataConstPtr FFontData::GetFontFaceData() const
 {
 	if (FontFaceAsset)
 	{
 		const IFontFaceInterface* FontFace = CastChecked<const IFontFaceInterface>(FontFaceAsset);
-		OutFontFaceData = FontFace->GetFontFaceData();
-		return true;
+		return FontFace->GetFontFaceData();
 	}
-	return false;
+	return nullptr;
 }
 
 const UObject* FFontData::GetFontFaceAsset() const
@@ -111,6 +118,7 @@ const UObject* FFontData::GetFontFaceAsset() const
 	return FontFaceAsset;
 }
 
+#if WITH_EDITORONLY_DATA
 bool FFontData::HasLegacyData() const
 {
 	return FontData_DEPRECATED.Num() > 0 || BulkDataPtr_DEPRECATED;
@@ -150,7 +158,6 @@ void FFontData::ConditionalUpgradeBulkDataToFontFace(UObject* InOuter, UClass* I
 
 bool FFontData::operator==(const FFontData& Other) const
 {
-#if WITH_EDITORONLY_DATA
 	if (FontFaceAsset != Other.FontFaceAsset)
 	{
 		// Using different assets
@@ -161,7 +168,6 @@ bool FFontData::operator==(const FFontData& Other) const
 		// Using the same asset
 		return true;
 	}
-#endif // WITH_EDITORONLY_DATA
 
 	// Compare inline properties
 	return FontFilename == Other.FontFilename
@@ -189,26 +195,12 @@ bool FFontData::Serialize(FArchive& Ar)
 
 	if (bIsCooked)
 	{
-		// Cooked data uses a more compact format, and doesn't persist the font face asset when loading
-		UObject* LocalFontFaceAsset = nullptr;
-#if WITH_EDITORONLY_DATA
-		LocalFontFaceAsset = const_cast<UObject*>(FontFaceAsset);
-#endif // WITH_EDITORONLY_DATA
+		// Cooked data uses a more compact format
+		UObject* LocalFontFaceAsset = const_cast<UObject*>(FontFaceAsset);
 		Ar << LocalFontFaceAsset;
+		FontFaceAsset = LocalFontFaceAsset;
 		
-		if (LocalFontFaceAsset)
-		{
-			if (Ar.IsLoading())
-			{
-				IFontFaceInterface* LocalFontFace = CastChecked<IFontFaceInterface>(LocalFontFaceAsset);
-
-				// Set the other properties from the font face asset
-				FontFilename = LocalFontFace->GetCookedFilename();
-				Hinting = LocalFontFace->GetHinting();
-				LoadingPolicy = LocalFontFace->GetLoadingPolicy();
-			}
-		}
-		else
+		if (!FontFaceAsset)
 		{
 			// Only need to serialize the other properties when we lack a font face asset
 			Ar << FontFilename;
@@ -218,15 +210,6 @@ bool FFontData::Serialize(FArchive& Ar)
 	}
 	else
 	{
-#if WITH_EDITORONLY_DATA
-		// Need to make sure the font face asset is tracked as a reference for cooking, despite being in an editor-only property
-		if (Ar.IsObjectReferenceCollector())
-		{
-			UObject* LocalFontFaceAsset = const_cast<UObject*>(FontFaceAsset);
-			Ar << LocalFontFaceAsset;
-		}
-#endif // WITH_EDITORONLY_DATA
-
 		// Uncooked data uses the standard struct serialization to avoid versioning on editor-only data
 		UScriptStruct* FontDataStruct = StaticStruct();
 		if (FontDataStruct->UseBinarySerialization(Ar))
@@ -244,8 +227,8 @@ bool FFontData::Serialize(FArchive& Ar)
 
 void FFontData::AddReferencedObjects(FReferenceCollector& Collector)
 {
-#if WITH_EDITORONLY_DATA
 	Collector.AddReferencedObject(FontFaceAsset);
+#if WITH_EDITORONLY_DATA
 	Collector.AddReferencedObject(BulkDataPtr_DEPRECATED);
 #endif // WITH_EDITORONLY_DATA
 }
diff --git a/Engine/Source/Runtime/SlateCore/Private/Fonts/FontCache.cpp b/Engine/Source/Runtime/SlateCore/Private/Fonts/FontCache.cpp
--- a/Engine/Source/Runtime/SlateCore/Private/Fonts/FontCache.cpp
+++ b/Engine/Source/Runtime/SlateCore/Private/Fonts/FontCache.cpp
@@ -935,9 +935,6 @@ FCharacterList& FSlateFontCache::GetCharacterList( const FSlateFontInfo &InFontI
 	// Create a key for looking up each character
 	const FSlateFontKey FontKey( InFontInfo, InOutlineSettings, FontScale );
 
-	//@HSL_BEGIN - Chance.Lyon - A critical section to help make this class thread-safe */
-	FScopeLock ScopeLock(&CacheCriticalSection);
-	//@HSL_END
 	TSharedRef< class FCharacterList >* CachedCharacterList = FontToCharacterListCache.Find( FontKey );
 
 	if( CachedCharacterList )
@@ -1056,27 +1053,11 @@ void FSlateFontCache::RequestFlushCache()
 
 void FSlateFontCache::FlushObject( const UObject* const InObject )
 {
-	if( !InObject )
+	if (InObject)
 	{
-		return;
-	}
-
-	bool bHasRemovedEntries = false;
-	//@HSL_BEGIN - Chance.Lyon - A critical section to help make this class thread-safe */
-	FScopeLock ScopeLock(&CacheCriticalSection);
-	//@HSL_END
-	for( auto It = FontToCharacterListCache.CreateIterator(); It; ++It )
-	{
-		if( It.Key().GetFontInfo().FontObject == InObject)
-		{
-			bHasRemovedEntries = true;
-			It.RemoveCurrent();
-		}
-	}
-
-	if( bHasRemovedEntries )
-	{
-		FlushData();
+		// Add it to the list of pending objects to flush
+		FScopeLock ScopeLock(&FontObjectsToFlushCS);
+		FontObjectsToFlush.AddUnique(InObject);
 	}
 }
 
@@ -1088,13 +1069,19 @@ void FSlateFontCache::FlushCompositeFont(const FCompositeFont& InCompositeFont)
 bool FSlateFontCache::ConditionalFlushCache()
 {
 	bool bFlushed = false;
-	if( bFlushRequested )
+	if (bFlushRequested)
 	{
 		bFlushRequested = false;
 		FlushCache();
 		bFlushed = !bFlushRequested;
 	}
 
+	if (!bFlushed && DoesThreadOwnSlateRendering())
+	{
+		// Only bother calling this if we didn't do a full flush
+		FlushFontObjects();
+	}
+
 	return bFlushed;
 }
 
@@ -1136,6 +1123,11 @@ void FSlateFontCache::FlushCache()
 		NonAtlasedTextures.Empty();
 		AllFontTextures.Empty();
 
+		{
+			FScopeLock ScopeLock(&FontObjectsToFlushCS);
+			FontObjectsToFlush.Empty();
+		}
+
 		UE_LOG(LogSlate, Verbose, TEXT("Slate font cache was flushed"));
 	}
 	else
@@ -1158,6 +1150,35 @@ void FSlateFontCache::FlushData()
 	ShapedGlyphToAtlasData.Empty();
 }
 
+void FSlateFontCache::FlushFontObjects()
+{
+	check(DoesThreadOwnSlateRendering());
+
+	bool bHasRemovedEntries = false;
+	{
+		FScopeLock ScopeLock(&FontObjectsToFlushCS);
+
+		if (FontObjectsToFlush.Num() > 0)
+		{
+			for (auto It = FontToCharacterListCache.CreateIterator(); It; ++It)
+			{
+				if (FontObjectsToFlush.Contains(It.Key().GetFontInfo().FontObject))
+				{
+					bHasRemovedEntries = true;
+					It.RemoveCurrent();
+				}
+			}
+
+			FontObjectsToFlush.Empty();
+		}
+	}
+
+	if (bHasRemovedEntries)
+	{
+		FlushData();
+	}
+}
+
 void FSlateFontCache::HandleCultureChanged()
 {
 	// The culture has changed, so request the font cache be flushed once it is safe to do so
diff --git a/Engine/Source/Runtime/SlateCore/Private/Fonts/FontCacheCompositeFont.cpp b/Engine/Source/Runtime/SlateCore/Private/Fonts/FontCacheCompositeFont.cpp
--- a/Engine/Source/Runtime/SlateCore/Private/Fonts/FontCacheCompositeFont.cpp
+++ b/Engine/Source/Runtime/SlateCore/Private/Fonts/FontCacheCompositeFont.cpp
@@ -273,34 +273,40 @@ TSharedPtr<FFreeTypeFace> FCompositeFontCache::GetFontFace(const FFontData& InFo
 	TSharedPtr<FFreeTypeFace> FaceAndMemory = FontFaceMap.FindRef(InFontData);
 	if (!FaceAndMemory.IsValid())
 	{
-#if WITH_EDITORONLY_DATA
+		// IMPORTANT: Do not log from this function until the new font has been added to the FontFaceMap, as it may be the Output Log font being loaded, which would cause an infinite recursion!
+		FString LoadLogMessage;
+
 		{
 			// If this font data is referencing an asset, we just need to load it from memory
-			TArray<uint8> FontFaceData;
-			if (InFontData.GetFontFaceData(FontFaceData))
+			FFontFaceDataConstPtr FontFaceData = InFontData.GetFontFaceData();
+			if (FontFaceData.IsValid() && FontFaceData->HasData())
 			{
-				FaceAndMemory = MakeShareable(new FFreeTypeFace(FTLibrary, MoveTemp(FontFaceData)));
+				FaceAndMemory = MakeShared<FFreeTypeFace>(FTLibrary, FontFaceData.ToSharedRef());
 			}
 		}
-#endif // WITH_EDITORONLY_DATA
 
 		// If no asset was loaded, then we go through the normal font loading process
 		if (!FaceAndMemory.IsValid())
 		{
 			switch (InFontData.GetLoadingPolicy())
 			{
-			case EFontLoadingPolicy::PreLoad:
+			case EFontLoadingPolicy::LazyLoad:
 				{
+					const double FontDataLoadStartTime = FPlatformTime::Seconds();
+
 					TArray<uint8> FontFaceData;
 					if (FFileHelper::LoadFileToArray(FontFaceData, *InFontData.GetFontFilename()))
 					{
-						FaceAndMemory = MakeShareable(new FFreeTypeFace(FTLibrary, MoveTemp(FontFaceData)));
+						const int32 FontDataSizeKB = (FontFaceData.Num() + 1023) / 1024;
+						LoadLogMessage = FString::Printf(TEXT("Took %f seconds to synchronously load lazily loaded font '%s' (%dK)"), float(FPlatformTime::Seconds() - FontDataLoadStartTime), *InFontData.GetFontFilename(), FontDataSizeKB);
+
+						FaceAndMemory = MakeShared<FFreeTypeFace>(FTLibrary, FFontFaceData::MakeFontFaceData(MoveTemp(FontFaceData)));
 					}
 				}
 				break;
 			case EFontLoadingPolicy::Stream:
 				{
-					FaceAndMemory = MakeShareable(new FFreeTypeFace(FTLibrary, InFontData.GetFontFilename()));
+					FaceAndMemory = MakeShared<FFreeTypeFace>(FTLibrary, InFontData.GetFontFilename());
 				}
 				break;
 			default:
@@ -312,6 +318,19 @@ TSharedPtr<FFreeTypeFace> FCompositeFontCache::GetFontFace(const FFontData& InFo
 		if (FaceAndMemory.IsValid() && FaceAndMemory->IsValid())
 		{
 			FontFaceMap.Add(InFontData, FaceAndMemory);
+
+			if (!LoadLogMessage.IsEmpty())
+			{
+				const bool bLogLoadAsWarning = GIsRunning && !GIsEditor;
+				if (bLogLoadAsWarning)
+				{
+					UE_LOG(LogSlate, Log, TEXT("%s"), *LoadLogMessage);
+				}
+				else
+				{
+					UE_LOG(LogSlate, Log, TEXT("%s"), *LoadLogMessage);
+				}
+			}
 		}
 		else
 		{
diff --git a/Engine/Source/Runtime/SlateCore/Private/Fonts/FontCacheFreeType.cpp b/Engine/Source/Runtime/SlateCore/Private/Fonts/FontCacheFreeType.cpp
--- a/Engine/Source/Runtime/SlateCore/Private/Fonts/FontCacheFreeType.cpp
+++ b/Engine/Source/Runtime/SlateCore/Private/Fonts/FontCacheFreeType.cpp
@@ -10,9 +10,6 @@
 // The total amount of memory freetype allocates internally
 DECLARE_MEMORY_STAT(TEXT("FreeType Total Allocated Memory"), STAT_SlateFreetypeAllocatedMemory, STATGROUP_SlateMemory);
 
-// The total true type memory we are using for resident font faces
-DECLARE_MEMORY_STAT(TEXT("Resident Font Memory (TTF/OTF)"), STAT_SlateRawFontDataMemory, STATGROUP_SlateMemory);
-
 // The active counts of resident and streaming fonts
 DECLARE_DWORD_ACCUMULATOR_STAT(TEXT("Num Resident Fonts"), STAT_SlateResidentFontCount, STATGROUP_SlateMemory);
 DECLARE_DWORD_ACCUMULATOR_STAT(TEXT("Num Streaming Fonts"), STAT_SlateStreamingFontCount, STATGROUP_SlateMemory);
@@ -140,20 +137,19 @@ FFreeTypeLibrary::~FFreeTypeLibrary()
 }
 
 
-FFreeTypeFace::FFreeTypeFace(const FFreeTypeLibrary* InFTLibrary, TArray<uint8>&& InMemory)
+FFreeTypeFace::FFreeTypeFace(const FFreeTypeLibrary* InFTLibrary, FFontFaceDataConstRef InMemory)
 #if WITH_FREETYPE
 	: FTFace(nullptr)
 	, Memory(MoveTemp(InMemory))
 #endif // WITH_FREETYPE
 {
 #if WITH_FREETYPE
-	FT_New_Memory_Face(InFTLibrary->GetLibrary(), Memory.GetData(), static_cast<FT_Long>(Memory.Num()), 0, &FTFace);
+	FT_New_Memory_Face(InFTLibrary->GetLibrary(), Memory->GetData().GetData(), static_cast<FT_Long>(Memory->GetData().Num()), 0, &FTFace);
 
 	ParseAttributes();
 
-	if (Memory.Num() > 0)
+	if (Memory->HasData())
 	{
-		INC_DWORD_STAT_BY(STAT_SlateRawFontDataMemory, Memory.GetAllocatedSize());
 		INC_DWORD_STAT_BY(STAT_SlateResidentFontCount, 1);
 	}
 #endif // WITH_FREETYPE
@@ -189,9 +185,8 @@ FFreeTypeFace::~FFreeTypeFace()
 #if WITH_FREETYPE
 	if (FTFace)
 	{
-		if (Memory.Num() > 0)
+		if (Memory->HasData())
 		{
-			DEC_DWORD_STAT_BY(STAT_SlateRawFontDataMemory, Memory.GetAllocatedSize());
 			DEC_DWORD_STAT_BY(STAT_SlateResidentFontCount, 1);
 		}
 		else
@@ -201,7 +196,6 @@ FFreeTypeFace::~FFreeTypeFace()
 
 		FT_Done_Face(FTFace);
 
-		Memory.Empty();
 		FTStreamHandler = FFTStreamHandler();
 	}
 #endif // WITH_FREETYPE
diff --git a/Engine/Source/Runtime/SlateCore/Private/Fonts/FontCacheFreeType.h b/Engine/Source/Runtime/SlateCore/Private/Fonts/FontCacheFreeType.h
--- a/Engine/Source/Runtime/SlateCore/Private/Fonts/FontCacheFreeType.h
+++ b/Engine/Source/Runtime/SlateCore/Private/Fonts/FontCacheFreeType.h
@@ -115,7 +115,7 @@ private:
 class FFreeTypeFace
 {
 public:
-	FFreeTypeFace(const FFreeTypeLibrary* InFTLibrary, TArray<uint8>&& InMemory);
+	FFreeTypeFace(const FFreeTypeLibrary* InFTLibrary, FFontFaceDataConstRef InMemory);
 	FFreeTypeFace(const FFreeTypeLibrary* InFTLibrary, const FString& InFilename);
 	~FFreeTypeFace();
 
@@ -151,7 +151,7 @@ private:
 
 #if WITH_FREETYPE
 	FT_Face FTFace;
-	TArray<uint8> Memory;
+	FFontFaceDataConstPtr Memory;
 
 	/** Custom FreeType stream handler for reading font data via the Unreal File System */
 	struct FFTStreamHandler
diff --git a/Engine/Source/Runtime/SlateCore/Private/Fonts/LegacySlateFontInfoCache.cpp b/Engine/Source/Runtime/SlateCore/Private/Fonts/LegacySlateFontInfoCache.cpp
--- a/Engine/Source/Runtime/SlateCore/Private/Fonts/LegacySlateFontInfoCache.cpp
+++ b/Engine/Source/Runtime/SlateCore/Private/Fonts/LegacySlateFontInfoCache.cpp
@@ -64,7 +64,7 @@ TSharedPtr<const FCompositeFont> FLegacySlateFontInfoCache::GetCompositeFont(con
 	// Don't allow GC while we perform this allocation
 	FGCScopeGuard GCGuard;
 
-	TSharedRef<const FCompositeFont> NewCompositeFont = MakeShareable(new FStandaloneCompositeFont(NAME_None, LegacyFontPath, InLegacyFontHinting, EFontLoadingPolicy::PreLoad));
+	TSharedRef<const FCompositeFont> NewCompositeFont = MakeShareable(new FStandaloneCompositeFont(NAME_None, LegacyFontPath, InLegacyFontHinting, EFontLoadingPolicy::LazyLoad));
 	LegacyFontNameToCompositeFont.Add(LegacyFontKey, NewCompositeFont);
 	return NewCompositeFont;
 }
@@ -82,7 +82,7 @@ TSharedPtr<const FCompositeFont> FLegacySlateFontInfoCache::GetSystemFont()
 			const FString FontFilename = FPaths::EngineIntermediateDir() / TEXT("DefaultSystemFont.ttf");
 			if (FFileHelper::SaveArrayToFile(FontBytes, *FontFilename))
 			{
-				SystemFont = MakeShareable(new FStandaloneCompositeFont(NAME_None, FontFilename, EFontHinting::Default, EFontLoadingPolicy::PreLoad));
+				SystemFont = MakeShareable(new FStandaloneCompositeFont(NAME_None, FontFilename, EFontHinting::Default, EFontLoadingPolicy::LazyLoad));
 			}
 		}
 	}
@@ -116,7 +116,7 @@ const FFontData& FLegacySlateFontInfoCache::GetLocalizedFallbackFontData()
 
 		if (!LocalizedFallbackFontData.IsValid())
 		{
-			LocalizedFallbackFontData = MakeShareable(new FFontData(FallbackFontPath, EFontHinting::Default, EFontLoadingPolicy::PreLoad));
+			LocalizedFallbackFontData = MakeShareable(new FFontData(FallbackFontPath, EFontHinting::Default, EFontLoadingPolicy::LazyLoad));
 			AllLocalizedFallbackFontData.Add(FallbackFontPath, LocalizedFallbackFontData);
 		}
 
@@ -160,7 +160,7 @@ const FFontData& FLegacySlateFontInfoCache::GetLastResortFontData()
 		FGCScopeGuard GCGuard;
 
 		const FString LastResortFontPath = FPaths::EngineContentDir() / TEXT("SlateDebug/Fonts/LastResort.ttf");
-		LastResortFontData = MakeShareable(new FFontData(LastResortFontPath, EFontHinting::Default, EFontLoadingPolicy::PreLoad));
+		LastResortFontData = MakeShareable(new FFontData(LastResortFontPath, EFontHinting::Default, EFontLoadingPolicy::LazyLoad));
 	}
 
 	return *LastResortFontData;
diff --git a/Engine/Source/Runtime/SlateCore/Private/Fonts/SlateFontInfo.cpp b/Engine/Source/Runtime/SlateCore/Private/Fonts/SlateFontInfo.cpp
--- a/Engine/Source/Runtime/SlateCore/Private/Fonts/SlateFontInfo.cpp
+++ b/Engine/Source/Runtime/SlateCore/Private/Fonts/SlateFontInfo.cpp
@@ -25,9 +25,10 @@ FSlateFontInfo::FSlateFontInfo( )
 }
 
 
-FSlateFontInfo::FSlateFontInfo( TSharedPtr<const FCompositeFont> InCompositeFont, const int32 InSize, const FName& InTypefaceFontName )
+FSlateFontInfo::FSlateFontInfo( TSharedPtr<const FCompositeFont> InCompositeFont, const int32 InSize, const FName& InTypefaceFontName, const FFontOutlineSettings& InOutlineSettings )
 	: FontObject(nullptr)
 	, FontMaterial(nullptr)
+	, OutlineSettings(InOutlineSettings)
 	, CompositeFont(InCompositeFont)
 	, TypefaceFontName(InTypefaceFontName)
 	, Size(InSize)
@@ -42,9 +43,10 @@ FSlateFontInfo::FSlateFontInfo( TSharedPtr<const FCompositeFont> InCompositeFont
 }
 
 
-FSlateFontInfo::FSlateFontInfo( const UObject* InFontObject, const int32 InSize, const FName& InTypefaceFontName )
+FSlateFontInfo::FSlateFontInfo( const UObject* InFontObject, const int32 InSize, const FName& InTypefaceFontName, const FFontOutlineSettings& InOutlineSettings )
 	: FontObject(InFontObject)
 	, FontMaterial(nullptr)
+	, OutlineSettings(InOutlineSettings)
 	, CompositeFont()
 	, TypefaceFontName(InTypefaceFontName)
 	, Size(InSize)
@@ -67,9 +69,10 @@ FSlateFontInfo::FSlateFontInfo( const UObject* InFontObject, const int32 InSize,
 }
 
 
-FSlateFontInfo::FSlateFontInfo( const FString& InFontName, uint16 InSize, EFontHinting InHinting )
+FSlateFontInfo::FSlateFontInfo( const FString& InFontName, uint16 InSize, EFontHinting InHinting, const FFontOutlineSettings& InOutlineSettings)
 	: FontObject(nullptr)
 	, FontMaterial(nullptr)
+	, OutlineSettings(InOutlineSettings)
 	, CompositeFont()
 	, TypefaceFontName()
 	, Size(InSize)
diff --git a/Engine/Source/Runtime/SlateCore/Private/Fonts/SlateTextShaper.cpp b/Engine/Source/Runtime/SlateCore/Private/Fonts/SlateTextShaper.cpp
--- a/Engine/Source/Runtime/SlateCore/Private/Fonts/SlateTextShaper.cpp
+++ b/Engine/Source/Runtime/SlateCore/Private/Fonts/SlateTextShaper.cpp
@@ -596,8 +596,8 @@ bool FSlateTextShaper::InsertSubstituteGlyphs(const TCHAR* InText, const int32 I
 			ShapedGlyphEntry.XOffset = 0;
 			ShapedGlyphEntry.YOffset = 0;
 			ShapedGlyphEntry.Kerning = 0;
-			ShapedGlyphEntry.NumCharactersInGlyph = 1;
-			ShapedGlyphEntry.NumGraphemeClustersInGlyph = 1;
+			ShapedGlyphEntry.NumCharactersInGlyph = (SpaceIndex == 0) ? 1 : 0;
+			ShapedGlyphEntry.NumGraphemeClustersInGlyph = (SpaceIndex == 0) ? 1 : 0;
 			ShapedGlyphEntry.TextDirection = TextBiDi::ETextDirection::LeftToRight;
 			ShapedGlyphEntry.bIsVisible = false;
 		}
diff --git a/Engine/Source/Runtime/SlateCore/Private/Rendering/DrawElements.cpp b/Engine/Source/Runtime/SlateCore/Private/Rendering/DrawElements.cpp
--- a/Engine/Source/Runtime/SlateCore/Private/Rendering/DrawElements.cpp
+++ b/Engine/Source/Runtime/SlateCore/Private/Rendering/DrawElements.cpp
@@ -45,6 +45,13 @@ void FSlateDrawElement::Init(uint32 InLayer, const FPaintGeometry& PaintGeometry
 	extern SLATECORE_API TOptional<FShortRect> GSlateScissorRect;
 	ScissorRect = GSlateScissorRect;
 
+	// Like GSlateScissorRect above, this is a workaround because we want to keep track of the various Scenes 
+	// in use throughout the UI. We keep a synchronized set with the render thread on the SlateRenderer and 
+	// use indices to synchronize between them.
+	TSharedPtr<FSlateRenderer> Renderer = FSlateApplicationBase::Get().GetRenderer();
+	check(Renderer.IsValid());
+	SceneIndex = Renderer->GetCurrentSceneIndex();
+
 	DataPayload.BatchFlags = ESlateBatchDrawFlag::None;
 	if ( InDrawEffects & ESlateDrawEffect::NoBlending )
 	{
diff --git a/Engine/Source/Runtime/SlateCore/Private/Rendering/ElementBatcher.cpp b/Engine/Source/Runtime/SlateCore/Private/Rendering/ElementBatcher.cpp
--- a/Engine/Source/Runtime/SlateCore/Private/Rendering/ElementBatcher.cpp
+++ b/Engine/Source/Runtime/SlateCore/Private/Rendering/ElementBatcher.cpp
@@ -229,7 +229,7 @@ void FSlateElementBatcher::AddQuadElement( const FSlateDrawElement& DrawElement,
 	FSlateLayoutTransform InverseLayoutTransform(Inverse(FSlateLayoutTransform(DrawElement.GetScale(), DrawElement.GetPosition())));
 	FSlateRotatedClipRectType RenderClipRect = FSlateRotatedClipRectType::MakeSnappedRotatedRect(InClippingRect, InverseLayoutTransform, RenderTransform);
 
-	FSlateElementBatch& ElementBatch = FindBatchForElement( Layer, FShaderParams(), nullptr, ESlateDrawPrimitive::TriangleList, ESlateShader::Default, ESlateDrawEffect::None, ESlateBatchDrawFlag::Wireframe|ESlateBatchDrawFlag::NoBlending, DrawElement.GetScissorRect() );
+	FSlateElementBatch& ElementBatch = FindBatchForElement( Layer, FShaderParams(), nullptr, ESlateDrawPrimitive::TriangleList, ESlateShader::Default, ESlateDrawEffect::None, ESlateBatchDrawFlag::Wireframe|ESlateBatchDrawFlag::NoBlending, DrawElement.GetScissorRect(), DrawElement.GetSceneIndex());
 	FSlateVertexArray& BatchVertices = BatchData->GetBatchVertexList(ElementBatch);
 	FSlateIndexArray& BatchIndices = BatchData->GetBatchIndexList(ElementBatch);
 
@@ -369,7 +369,7 @@ void FSlateElementBatcher::AddBoxElement( const FSlateDrawElement& DrawElement )
 		ESlateBatchDrawFlag::Type DrawFlags = InPayload.BatchFlags;
 		DrawFlags |= ( ( bTileHorizontal ? ESlateBatchDrawFlag::TileU : 0 ) | ( bTileVertical ? ESlateBatchDrawFlag::TileV : 0 ) );
 
-		FSlateElementBatch& ElementBatch = FindBatchForElement( Layer, FShaderParams(), Resource, ESlateDrawPrimitive::TriangleList, ESlateShader::Default, InDrawEffects, DrawFlags, DrawElement.GetScissorRect() );
+		FSlateElementBatch& ElementBatch = FindBatchForElement( Layer, FShaderParams(), Resource, ESlateDrawPrimitive::TriangleList, ESlateShader::Default, InDrawEffects, DrawFlags, DrawElement.GetScissorRect(), DrawElement.GetSceneIndex());
 		FSlateVertexArray& BatchVertices = BatchData->GetBatchVertexList(ElementBatch);
 		FSlateIndexArray& BatchIndices = BatchData->GetBatchIndexList(ElementBatch);
 
@@ -706,7 +706,7 @@ void FSlateElementBatcher::AddTextElement(const FSlateDrawElement& DrawElement)
 					FontShaderResource = ResourceManager.GetFontShaderResource( FontTextureIndex, FontAtlasTexture, InPayload.FontInfo.FontMaterial );
 					check(FontShaderResource);
 
-					ElementBatch = &FindBatchForElement(InLayer, FShaderParams(), FontShaderResource, ESlateDrawPrimitive::TriangleList, ESlateShader::Font, InDrawEffects, ESlateBatchDrawFlag::None, DrawElement.GetScissorRect() );
+					ElementBatch = &FindBatchForElement(InLayer, FShaderParams(), FontShaderResource, ESlateDrawPrimitive::TriangleList, ESlateShader::Font, InDrawEffects, ESlateBatchDrawFlag::None, DrawElement.GetScissorRect(), DrawElement.GetSceneIndex());
 
 					BatchVertices = &BatchData->GetBatchVertexList(*ElementBatch);
 					BatchIndices = &BatchData->GetBatchIndexList(*ElementBatch);
@@ -919,7 +919,7 @@ void FSlateElementBatcher::AddShapedTextElement( const FSlateDrawElement& DrawEl
 						FontShaderResource = ResourceManager.GetFontShaderResource(FontTextureIndex, FontAtlasTexture, FontMaterial);
 						check(FontShaderResource);
 
-						ElementBatch = &FindBatchForElement(InLayer, FShaderParams(), FontShaderResource, ESlateDrawPrimitive::TriangleList, ESlateShader::Font, InDrawEffects, ESlateBatchDrawFlag::None, DrawElement.GetScissorRect());
+						ElementBatch = &FindBatchForElement(InLayer, FShaderParams(), FontShaderResource, ESlateDrawPrimitive::TriangleList, ESlateShader::Font, InDrawEffects, ESlateBatchDrawFlag::None, DrawElement.GetScissorRect(), DrawElement.GetSceneIndex());
 
 						BatchVertices = &BatchData->GetBatchVertexList(*ElementBatch);
 						BatchIndices = &BatchData->GetBatchIndexList(*ElementBatch);
@@ -1043,7 +1043,8 @@ void FSlateElementBatcher::AddGradientElement( const FSlateDrawElement& DrawElem
 			ESlateShader::Default,
 			InDrawEffects,
 			InPayload.BatchFlags,
-			DrawElement.GetScissorRect() );
+			DrawElement.GetScissorRect(),
+			DrawElement.GetSceneIndex());
 
 	FSlateVertexArray& BatchVertices = BatchData->GetBatchVertexList(ElementBatch);
 	FSlateIndexArray& BatchIndices = BatchData->GetBatchIndexList(ElementBatch);
@@ -1197,7 +1198,7 @@ void FSlateElementBatcher::AddSplineElement( const FSlateDrawElement& DrawElemen
 	const float HalfThickness = LineThickness * .5f + Radius;
 
 	// Find a batch for the element
-	FSlateElementBatch& ElementBatch = FindBatchForElement( Layer, FShaderParams::MakePixelShaderParams( FVector4( InPayload.Thickness,Radius,0,0) ), nullptr, ESlateDrawPrimitive::TriangleList, ESlateShader::LineSegment, InDrawEffects, ESlateBatchDrawFlag::None, DrawElement.GetScissorRect() );
+	FSlateElementBatch& ElementBatch = FindBatchForElement( Layer, FShaderParams::MakePixelShaderParams( FVector4( InPayload.Thickness,Radius,0,0) ), nullptr, ESlateDrawPrimitive::TriangleList, ESlateShader::LineSegment, InDrawEffects, ESlateBatchDrawFlag::None, DrawElement.GetScissorRect(), DrawElement.GetSceneIndex());
 	FSlateVertexArray& BatchVertices = BatchData->GetBatchVertexList(ElementBatch);
 	FSlateIndexArray& BatchIndices = BatchData->GetBatchIndexList(ElementBatch);
 
@@ -1344,7 +1345,7 @@ void FSlateElementBatcher::AddLineElement( const FSlateDrawElement& DrawElement
 		const float HalfThickness = LineThickness * .5f + Radius;
 
 		// Find a batch for the element
-		FSlateElementBatch& ElementBatch = FindBatchForElement( Layer, FShaderParams::MakePixelShaderParams( FVector4(RequestedThickness,Radius,0,0) ), nullptr, ESlateDrawPrimitive::TriangleList, ESlateShader::LineSegment, DrawEffects, ESlateBatchDrawFlag::None, DrawElement.GetScissorRect() );
+		FSlateElementBatch& ElementBatch = FindBatchForElement( Layer, FShaderParams::MakePixelShaderParams( FVector4(RequestedThickness,Radius,0,0) ), nullptr, ESlateDrawPrimitive::TriangleList, ESlateShader::LineSegment, DrawEffects, ESlateBatchDrawFlag::None, DrawElement.GetScissorRect(), DrawElement.GetSceneIndex());
 		FSlateVertexArray& BatchVertices = BatchData->GetBatchVertexList(ElementBatch);
 		FSlateIndexArray& BatchIndices = BatchData->GetBatchIndexList(ElementBatch);
 
@@ -1454,7 +1455,7 @@ void FSlateElementBatcher::AddLineElement( const FSlateDrawElement& DrawElement
 	else
 	{
 		// Find a batch for the element
-		FSlateElementBatch& ElementBatch = FindBatchForElement( Layer, FShaderParams(), nullptr, ESlateDrawPrimitive::LineList, ESlateShader::Default, DrawEffects, ESlateBatchDrawFlag::None, DrawElement.GetScissorRect() );
+		FSlateElementBatch& ElementBatch = FindBatchForElement( Layer, FShaderParams(), nullptr, ESlateDrawPrimitive::LineList, ESlateShader::Default, DrawEffects, ESlateBatchDrawFlag::None, DrawElement.GetScissorRect(), DrawElement.GetSceneIndex());
 		FSlateVertexArray& BatchVertices = BatchData->GetBatchVertexList(ElementBatch);
 		FSlateIndexArray& BatchIndices = BatchData->GetBatchIndexList(ElementBatch);
 
@@ -1505,7 +1506,7 @@ void FSlateElementBatcher::AddViewportElement( const FSlateDrawElement& DrawElem
 	}
 
 
-	FSlateElementBatch& ElementBatch = FindBatchForElement( Layer, FShaderParams(), ViewportResource, ESlateDrawPrimitive::TriangleList, ShaderType, InDrawEffects, DrawFlags, DrawElement.GetScissorRect() );
+	FSlateElementBatch& ElementBatch = FindBatchForElement( Layer, FShaderParams(), ViewportResource, ESlateDrawPrimitive::TriangleList, ShaderType, InDrawEffects, DrawFlags, DrawElement.GetScissorRect(), DrawElement.GetSceneIndex());
 	FSlateVertexArray& BatchVertices = BatchData->GetBatchVertexList(ElementBatch);
 	FSlateIndexArray& BatchIndices = BatchData->GetBatchIndexList(ElementBatch);
 
@@ -1640,7 +1641,7 @@ void FSlateElementBatcher::AddBorderElement( const FSlateDrawElement& DrawElemen
 	// Pass the tiling information as a flag so we can pick the correct texture addressing mode
 	ESlateBatchDrawFlag::Type DrawFlags = (ESlateBatchDrawFlag::TileU|ESlateBatchDrawFlag::TileV);
 
-	FSlateElementBatch& ElementBatch = FindBatchForElement( Layer, ShaderParams, Resource, ESlateDrawPrimitive::TriangleList, ESlateShader::Border, InDrawEffects, DrawFlags, DrawElement.GetScissorRect() );
+	FSlateElementBatch& ElementBatch = FindBatchForElement( Layer, ShaderParams, Resource, ESlateDrawPrimitive::TriangleList, ESlateShader::Border, InDrawEffects, DrawFlags, DrawElement.GetScissorRect(), DrawElement.GetSceneIndex());
 	FSlateVertexArray& BatchVertices = BatchData->GetBatchVertexList(ElementBatch);
 	FSlateIndexArray& BatchIndices = BatchData->GetBatchIndexList(ElementBatch);
 
@@ -1815,7 +1816,8 @@ void FSlateElementBatcher::AddCustomVerts(const FSlateDrawElement& DrawElement)
 			DrawElement.GetScissorRect(),
 			InPayload.NumInstances,
 			InPayload.InstanceOffset,
-			InPayload.InstanceData
+			InPayload.InstanceData,
+			DrawElement.GetSceneIndex()
 		);
 
 		int32 Index = (*ElementBatches)->Add(NewBatch);
@@ -1949,7 +1951,8 @@ FSlateElementBatch& FSlateElementBatcher::FindBatchForElement(
 	ESlateShader::Type ShaderType, 
 	ESlateDrawEffect::Type DrawEffects, 
 	ESlateBatchDrawFlag::Type DrawFlags,
-	const TOptional<FShortRect>& ScissorRect)
+	const TOptional<FShortRect>& ScissorRect,
+	int32 SceneIndex)
 {
 	SLATE_CYCLE_COUNTER_SCOPE_DETAILED(SLATE_STATS_DETAIL_LEVEL_HI, GSlateFindBatchTime);
 
@@ -1967,7 +1970,7 @@ FSlateElementBatch& FSlateElementBatcher::FindBatchForElement(
 	checkSlow( ElementBatches );
 
 	// Create a temp batch so we can use it as our key to find if the same batch already exists
-	FSlateElementBatch TempBatch( InTexture, ShaderParams, ShaderType, PrimitiveType, DrawEffects, DrawFlags, ScissorRect );
+	FSlateElementBatch TempBatch( InTexture, ShaderParams, ShaderType, PrimitiveType, DrawEffects, DrawFlags, ScissorRect, 0, 0, nullptr, SceneIndex );
 
 	FSlateElementBatch* ElementBatch = (*ElementBatches)->FindByKey( TempBatch );
 	if( !ElementBatch )
diff --git a/Engine/Source/Runtime/SlateCore/Private/Styling/CoreStyle.cpp b/Engine/Source/Runtime/SlateCore/Private/Styling/CoreStyle.cpp
--- a/Engine/Source/Runtime/SlateCore/Private/Styling/CoreStyle.cpp
+++ b/Engine/Source/Runtime/SlateCore/Private/Styling/CoreStyle.cpp
@@ -471,6 +471,17 @@ TSharedRef<ISlateStyle> FCoreStyle::Create( const FName& InStyleSetName )
 		// "NormalFont".
 	}
 
+	const FEditableTextBoxStyle DarkEditableTextBoxStyle = FEditableTextBoxStyle()
+		.SetBackgroundImageNormal(BOX_BRUSH("Common/TextBox_Dark", FMargin(4.0f / 16.0f)))
+		.SetBackgroundImageHovered(BOX_BRUSH("Common/TextBox_Hovered_Dark", FMargin(4.0f / 16.0f)))
+		.SetBackgroundImageFocused(BOX_BRUSH("Common/TextBox_Hovered_Dark", FMargin(4.0f / 16.0f)))
+		.SetBackgroundImageReadOnly(BOX_BRUSH("Common/TextBox_ReadOnly", FMargin(4.0f / 16.0f)))
+		.SetScrollBarStyle(ScrollBar);
+	{
+		Style->Set("DarkEditableTextBox", DarkEditableTextBoxStyle);
+		// "NormalFont".
+	}
+
 	// SSpinBox defaults...
 	{
 		// "NormalFont".
@@ -651,6 +662,17 @@ TSharedRef<ISlateStyle> FCoreStyle::Create( const FName& InStyleSetName )
 			.SetTextPadding( FMargin(0.0f) )
 			.SetForegroundColor( InvertedForeground )
 			);
+
+		Style->Set("NumericEntrySpinBox_Dark", FSpinBoxStyle()
+			.SetBackgroundBrush(FSlateNoResource())
+			.SetHoveredBackgroundBrush(FSlateNoResource())
+			.SetActiveFillBrush(BOX_BRUSH("Common/Spinbox_Fill_Hovered_Dark", FMargin(4.0f / 16.0f)))
+			.SetInactiveFillBrush(BOX_BRUSH("Common/Spinbox_Fill_Dark", FMargin(4.0f / 16.0f, 4.0f / 16.0f, 8.0f / 16.0f, 4.0f / 16.0f)))
+			.SetArrowsImage(IMAGE_BRUSH("Common/SpinArrows", Icon12x12))
+			.SetTextPadding(FMargin(0.0f))
+			.SetForegroundColor(InvertedForeground)
+			);
+
 		Style->Set( "NumericEntrySpinBox.Decorator", new BOX_BRUSH( "Common/TextBoxLabelBorder", FMargin(5.0f/16.0f) ) );
 
 		Style->Set( "NumericEntrySpinBox.NarrowDecorator", new BOX_BRUSH( "Common/TextBoxLabelBorder", FMargin(2.0f/16.0f, 4.0f/16.0f, 2.0f/16.0f, 4.0f/16.0f) ) );
@@ -688,6 +710,12 @@ TSharedRef<ISlateStyle> FCoreStyle::Create( const FName& InStyleSetName )
 		Style->Set( "ColorWheel.Selector", new IMAGE_BRUSH( "Common/Circle", FVector2D(8, 8) ) );
 	}
 
+	// SColorGradingWheel defaults...
+	{
+		Style->Set("ColorGradingWheel.HueValueCircle", new IMAGE_BRUSH("Common/ColorGradingWheel", FVector2D(192, 192)));
+		Style->Set("ColorGradingWheel.Selector", new IMAGE_BRUSH("Common/Circle", FVector2D(8, 8)));
+	}
+
 	// SSplitter
 	{
 		Style->Set( "Splitter", FSplitterStyle()
diff --git a/Engine/Source/Runtime/SlateCore/Private/Widgets/SWidget.cpp b/Engine/Source/Runtime/SlateCore/Private/Widgets/SWidget.cpp
--- a/Engine/Source/Runtime/SlateCore/Private/Widgets/SWidget.cpp
+++ b/Engine/Source/Runtime/SlateCore/Private/Widgets/SWidget.cpp
@@ -79,6 +79,11 @@ void SWidget::CreateStatID() const
 #endif
 }
 
+FName NAME_MouseButtonDown(TEXT("MouseButtonDown"));
+FName NAME_MouseButtonUp(TEXT("MouseButtonUp"));
+FName NAME_MouseMove(TEXT("MouseMove"));
+FName NAME_MouseDoubleClick(TEXT("MouseDoubleClick"));
+
 SWidget::SWidget()
 	: Cursor( TOptional<EMouseCursor::Type>() )
 	, EnabledState( true )
@@ -219,40 +224,48 @@ FReply SWidget::OnPreviewMouseButtonDown( const FGeometry& MyGeometry, const FPo
 
 FReply SWidget::OnMouseButtonDown(const FGeometry& MyGeometry, const FPointerEvent& MouseEvent)
 {
-	if (MouseButtonDownHandler.IsBound())
+	if ( FPointerEventHandler* Event = PointerEvents.Find(NAME_MouseButtonDown) )
 	{
-		// If a handler is assigned, call it.
-		return MouseButtonDownHandler.Execute(MyGeometry, MouseEvent);
+		if ( Event->IsBound() )
+		{
+			return Event->Execute(MyGeometry, MouseEvent);
+		}
 	}
 	return FReply::Unhandled();
 }
 
 FReply SWidget::OnMouseButtonUp(const FGeometry& MyGeometry, const FPointerEvent& MouseEvent)
 {
-	if (MouseButtonUpHandler.IsBound())
+	if ( FPointerEventHandler* Event = PointerEvents.Find(NAME_MouseButtonUp) )
 	{
-		// If a handler is assigned, call it.
-		return MouseButtonUpHandler.Execute(MyGeometry, MouseEvent);
+		if ( Event->IsBound() )
+		{
+			return Event->Execute(MyGeometry, MouseEvent);
+		}
 	}
 	return FReply::Unhandled();
 }
 
 FReply SWidget::OnMouseMove(const FGeometry& MyGeometry, const FPointerEvent& MouseEvent)
 {
-	if (MouseMoveHandler.IsBound())
+	if ( FPointerEventHandler* Event = PointerEvents.Find(NAME_MouseMove) )
 	{
-		// A valid handler is assigned for mouse move; let it handle the event.
-		return MouseMoveHandler.Execute(MyGeometry, MouseEvent);
+		if ( Event->IsBound() )
+		{
+			return Event->Execute(MyGeometry, MouseEvent);
+		}
 	}
 	return FReply::Unhandled();
 }
 
 FReply SWidget::OnMouseButtonDoubleClick(const FGeometry& MyGeometry, const FPointerEvent& MouseEvent)
 {
-	if (MouseDoubleClickHandler.IsBound())
+	if ( FPointerEventHandler* Event = PointerEvents.Find(NAME_MouseDoubleClick) )
 	{
-		// A valid handler is assigned; let it handle the event.
-		return MouseDoubleClickHandler.Execute(MyGeometry, MouseEvent);
+		if ( Event->IsBound() )
+		{
+			return Event->Execute(MyGeometry, MouseEvent);
+		}
 	}
 	return FReply::Unhandled();
 }
@@ -330,21 +343,6 @@ FReply SWidget::OnDrop( const FGeometry& MyGeometry, const FDragDropEvent& DragD
 	return FReply::Unhandled();
 }
 
-FReply SWidget::OnControllerButtonPressed( const FGeometry& MyGeometry, const FControllerEvent& ControllerEvent )
-{
-	return FReply::Unhandled();
-}
-
-FReply SWidget::OnControllerButtonReleased( const FGeometry& MyGeometry, const FControllerEvent& ControllerEvent )
-{
-	return FReply::Unhandled();
-}
-
-FReply SWidget::OnControllerAnalogValueChanged( const FGeometry& MyGeometry, const FControllerEvent& ControllerEvent )
-{
-	return FReply::Unhandled();
-}
-
 FReply SWidget::OnTouchGesture( const FGeometry& MyGeometry, const FPointerEvent& GestureEvent )
 {
 	return FReply::Unhandled();
@@ -470,9 +468,9 @@ void SWidget::SlatePrepass(float LayoutScaleMultiplier)
 
 			if ( Child->Visibility.Get() != EVisibility::Collapsed )
 			{
-				const float ChildLayoutScaleMultiplier = GetRelativeLayoutScale(MyChildren->GetSlotAt(ChildIndex));
+				const float ChildLayoutScaleMultiplier = GetRelativeLayoutScale(MyChildren->GetSlotAt(ChildIndex), LayoutScaleMultiplier);
 				// Recur: Descend down the widget tree.
-				Child->SlatePrepass(LayoutScaleMultiplier*ChildLayoutScaleMultiplier);
+				Child->SlatePrepass(LayoutScaleMultiplier * ChildLayoutScaleMultiplier);
 			}
 		}
 	}
@@ -658,12 +656,20 @@ FName SWidget::GetType() const
 
 FString SWidget::GetReadableLocation() const
 {
+#if !UE_BUILD_SHIPPING
 	return FString::Printf(TEXT("%s(%d)"), *FPaths::GetCleanFilename(this->CreatedInLocation.GetPlainNameString()), this->CreatedInLocation.GetNumber());
+#else
+	return FString();
+#endif
 }
 
 FName SWidget::GetCreatedInLocation() const
 {
+#if !UE_BUILD_SHIPPING
 	return CreatedInLocation;
+#else
+	return NAME_None;
+#endif
 }
 
 FName SWidget::GetTag() const
@@ -719,8 +725,11 @@ void SWidget::SetCursor( const TAttribute< TOptional<EMouseCursor::Type> >& InCu
 void SWidget::SetDebugInfo( const ANSICHAR* InType, const ANSICHAR* InFile, int32 OnLine )
 {
 	TypeOfWidget = InType;
+
+#if !UE_BUILD_SHIPPING
 	CreatedInLocation = FName( InFile );
 	CreatedInLocation.SetNumber(OnLine);
+#endif
 }
 
 int32 SWidget::Paint(const FPaintArgs& Args, const FGeometry& AllottedGeometry, const FSlateRect& MyClippingRect, FSlateWindowElementList& OutDrawElements, int32 LayerId, const FWidgetStyle& InWidgetStyle, bool bParentEnabled) const
@@ -751,14 +760,16 @@ int32 SWidget::Paint(const FPaintArgs& Args, const FGeometry& AllottedGeometry,
 		return VolatileLayerId;
 	}
 
+	// Cache the geometry for tick to allow external users to get the last geometry that was used,
+	// or would have been used to tick the Widget.
+	CachedGeometry = AllottedGeometry;
+	CachedGeometry.AppendTransform(FSlateLayoutTransform(Args.GetWindowToDesktopTransform()));
+
 	if ( bFoldTick && bCanTick )
 	{
-		FGeometry TickGeometry = AllottedGeometry;
-		TickGeometry.AppendTransform( FSlateLayoutTransform(Args.GetWindowToDesktopTransform()) );
-
 		SWidget* MutableThis = const_cast<SWidget*>(this);
 		MutableThis->ExecuteActiveTimers( Args.GetCurrentTime(), Args.GetDeltaTime() );
-		MutableThis->Tick( TickGeometry, Args.GetCurrentTime(), Args.GetDeltaTime() );
+		MutableThis->Tick( CachedGeometry, Args.GetCurrentTime(), Args.GetDeltaTime() );
 	}
 
 	// Record hit test geometry, but only if we're not caching.
@@ -785,7 +796,7 @@ int32 SWidget::Paint(const FPaintArgs& Args, const FGeometry& AllottedGeometry,
 					MyClippingRect,
 					ESlateDrawEffect::None,
 					BrushResource->GetTint(InWidgetStyle)
-					);
+				);
 			}
 		}
 	}
@@ -798,7 +809,7 @@ int32 SWidget::Paint(const FPaintArgs& Args, const FGeometry& AllottedGeometry,
 	return NewLayerID;
 }
 
-float SWidget::GetRelativeLayoutScale(const FSlotBase& Child) const
+float SWidget::GetRelativeLayoutScale(const FSlotBase& Child, float LayoutScaleMultiplier) const
 {
 	return 1.0f;
 }
@@ -844,22 +855,22 @@ void SWidget::ExecuteActiveTimers(double CurrentTime, float DeltaTime)
 
 void SWidget::SetOnMouseButtonDown(FPointerEventHandler EventHandler)
 {
-	MouseButtonDownHandler = EventHandler;
+	PointerEvents.Add(NAME_MouseButtonDown, EventHandler);
 }
 
 void SWidget::SetOnMouseButtonUp(FPointerEventHandler EventHandler)
 {
-	MouseButtonUpHandler = EventHandler;
+	PointerEvents.Add(NAME_MouseButtonUp, EventHandler);
 }
 
 void SWidget::SetOnMouseMove(FPointerEventHandler EventHandler)
 {
-	MouseMoveHandler = EventHandler;
+	PointerEvents.Add(NAME_MouseMove, EventHandler);
 }
 
 void SWidget::SetOnMouseDoubleClick(FPointerEventHandler EventHandler)
 {
-	MouseDoubleClickHandler = EventHandler;
+	PointerEvents.Add(NAME_MouseDoubleClick, EventHandler);
 }
 
 void SWidget::SetOnMouseEnter(FNoReplyPointerEventHandler EventHandler)
diff --git a/Engine/Source/Runtime/SlateCore/Private/Widgets/SWindow.cpp b/Engine/Source/Runtime/SlateCore/Private/Widgets/SWindow.cpp
--- a/Engine/Source/Runtime/SlateCore/Private/Widgets/SWindow.cpp
+++ b/Engine/Source/Runtime/SlateCore/Private/Widgets/SWindow.cpp
@@ -826,26 +826,32 @@ void SWindow::MoveWindowTo( FVector2D NewPosition )
 
 void SWindow::ReshapeWindow( FVector2D NewPosition, FVector2D NewSize )
 {
-	if (NativeWindow.IsValid())
-	{
-#if 1//PLATFORM_LINUX
-		// Slate code often expects cached screen position to be accurate immediately after the move.
-		// This expectation is generally invalid (see UE-1308) as there may be a delay before the OS reports it back.
-		// This hack sets the position speculatively, keeping Slate happy while also giving the OS chance to report it
-		// correctly after or even during the actual call.
-		FVector2D SpeculativeScreenPosition(FMath::TruncToInt(NewPosition.X), FMath::TruncToInt(NewPosition.Y));
-		SetCachedScreenPosition(SpeculativeScreenPosition);
-#endif // PLATFORM_LINUX
+	const FVector2D CurrentPosition = GetPositionInScreen();
+	const FVector2D CurrentSize = GetSizeInScreen();
 
-		NativeWindow->ReshapeWindow( FMath::TruncToInt(NewPosition.X), FMath::TruncToInt(NewPosition.Y), FMath::RoundToInt(NewSize.X), FMath::RoundToInt(NewSize.Y) );
-	}
-	else
+	const FVector2D NewPositionTruncated = FVector2D(FMath::TruncToInt(NewPosition.X), FMath::TruncToInt(NewPosition.Y));
+	const FVector2D NewSizeRounded = FVector2D(FMath::RoundToInt(NewSize.X), FMath::RoundToInt(NewSize.Y));
+
+	if ( CurrentPosition != NewPositionTruncated || CurrentSize != NewSizeRounded )
 	{
-		InitialDesiredScreenPosition = NewPosition;
-		InitialDesiredSize = NewSize;
-	}
+		if ( NativeWindow.IsValid() )
+		{
+			// Slate code often expects cached screen position to be accurate immediately after the move.
+			// This expectation is generally invalid (see UE-1308) as there may be a delay before the OS reports it back.
+			// This hack sets the position speculatively, keeping Slate happy while also giving the OS chance to report it
+			// correctly after or even during the actual call.
+			SetCachedScreenPosition(NewPositionTruncated);
 
-	SetCachedSize( NewSize );
+			NativeWindow->ReshapeWindow(NewPositionTruncated.X, NewPositionTruncated.Y, NewSizeRounded.X, NewSizeRounded.Y);
+		}
+		else
+		{
+			InitialDesiredScreenPosition = NewPosition;
+			InitialDesiredSize = NewSize;
+		}
+
+		SetCachedSize(NewSize);
+	}
 }
 
 void SWindow::ReshapeWindow( const FSlateRect& InNewShape )
@@ -1457,6 +1463,15 @@ bool SWindow::OnIsActiveChanged( const FWindowActivateEvent& ActivateEvent )
 		{
 			WidgetToFocusOnActivate.Reset();
 		}
+		else if (SupportsKeyboardFocus())
+		{
+			// If we have no specific widget to focus then cache the currently focused widget so we can restore its focus when we regain focus
+			WidgetToFocusOnActivate = FSlateApplicationBase::Get().GetKeyboardFocusedWidget();
+			if (!WidgetToFocusOnActivate.IsValid())
+			{
+				WidgetToFocusOnActivate = FSlateApplicationBase::Get().GetUserFocusedWidget(0);
+			}
+		}
 	}
 	else
 	{
diff --git a/Engine/Source/Runtime/SlateCore/Public/Application/SlateApplicationBase.h b/Engine/Source/Runtime/SlateCore/Public/Application/SlateApplicationBase.h
--- a/Engine/Source/Runtime/SlateCore/Public/Application/SlateApplicationBase.h
+++ b/Engine/Source/Runtime/SlateCore/Public/Application/SlateApplicationBase.h
@@ -382,6 +382,11 @@ public:
 	 */
 	virtual void SetAllUserFocusAllowingDescendantFocus(const FWidgetPath& InFocusPath, const EFocusCause InCause) = 0;
 
+	/**
+	 * @return a pointer to the Widget that currently has the users focus; Empty pointer when the user has no focus. 
+	 */
+	virtual TSharedPtr<SWidget> GetUserFocusedWidget(uint32 UserIndex) const = 0;
+
 	/**
 	 * Gets a delegate that is invoked when a global invalidate of all widgets should occur
 	 */
diff --git a/Engine/Source/Runtime/SlateCore/Public/Fonts/CompositeFont.h b/Engine/Source/Runtime/SlateCore/Public/Fonts/CompositeFont.h
--- a/Engine/Source/Runtime/SlateCore/Public/Fonts/CompositeFont.h
+++ b/Engine/Source/Runtime/SlateCore/Public/Fonts/CompositeFont.h
@@ -9,6 +9,7 @@
 #include "CompositeFont.generated.h"
 
 class UFontBulkData;
+struct FFontFaceData;
 
 UENUM()
 enum class EFontHinting : uint8
@@ -28,10 +29,97 @@ enum class EFontHinting : uint8
 UENUM()
 enum class EFontLoadingPolicy : uint8
 {
-	/** Pre-load the entire font into memory. This will consume more memory, however there will be zero file-IO when rendering glyphs within the font. */
-	PreLoad,
-	/** Stream the font from disk. This will consume less memory, however there will be file-IO when rendering glyphs, which may cause hitches under certain circumstances or on certain platforms. */
+	/** Lazy load the entire font into memory. This will consume more memory than Streaming, however there will be zero file-IO when rendering glyphs within the font, although the initial load may cause a hitch. */
+	LazyLoad,
+	/** Stream the font from disk. This will consume less memory than LazyLoad or Inline, however there will be file-IO when rendering glyphs, which may cause hitches under certain circumstances or on certain platforms. */
 	Stream,
+	/** Embed the font data within the asset. This will consume more memory than Streaming, however it is guaranteed to be hitch free (only valid for font data within a Font Face asset). */
+	Inline,
+};
+
+typedef TSharedPtr<FFontFaceData, ESPMode::ThreadSafe> FFontFaceDataPtr;
+typedef TSharedRef<FFontFaceData, ESPMode::ThreadSafe> FFontFaceDataRef;
+typedef TSharedPtr<const FFontFaceData, ESPMode::ThreadSafe> FFontFaceDataConstPtr;
+typedef TSharedRef<const FFontFaceData, ESPMode::ThreadSafe> FFontFaceDataConstRef;
+
+/** 
+ * Raw font data for a Font Face asset.
+ * @note Exists as a struct so it can be shared between its owner asset and the font cache without worrying about UObject lifetimes. 
+ */
+struct SLATECORE_API FFontFaceData
+{
+public:
+	/** Default constructor */
+	FFontFaceData()
+	{
+	}
+
+	/** Construct from an existing block of font data */
+	FFontFaceData(TArray<uint8>&& InData)
+		: Data(MoveTemp(InData))
+	{
+		TrackMemoryUsage();
+	}
+
+	/** Destructor */
+	~FFontFaceData()
+	{
+		UntrackMemoryUsage();
+	}
+
+	/** @return true if this instance has data set on it */
+	bool HasData() const
+	{
+		return Data.Num() > 0;
+	}
+
+	/** Set from an existing block of font data */
+	void SetData(TArray<uint8>&& InData)
+	{
+		UntrackMemoryUsage();
+		Data = MoveTemp(InData);
+		TrackMemoryUsage();
+	}
+
+	/** @return The data set on this instance */
+	const TArray<uint8>& GetData() const
+	{
+		return Data;
+	}
+
+	/** Serialize the data on this instance to/from an archive */
+	void Serialize(FArchive& Ar)
+	{
+		if (Ar.IsLoading())
+		{
+			UntrackMemoryUsage();
+		}
+		Ar << Data;
+		if (Ar.IsLoading())
+		{
+			TrackMemoryUsage();
+		}
+	}
+
+	/** Helper function to construct a default font face data */
+	static FFontFaceDataRef MakeFontFaceData()
+	{
+		return MakeShared<FFontFaceData, ESPMode::ThreadSafe>();
+	}
+
+	/** Helper function to construct a font face data from an existing block of font data */
+	static FFontFaceDataRef MakeFontFaceData(TArray<uint8>&& InData)
+	{
+		return MakeShared<FFontFaceData, ESPMode::ThreadSafe>(MoveTemp(InData));
+	}
+
+private:
+	/** Memory stat tracking */
+	void TrackMemoryUsage() const;
+	void UntrackMemoryUsage() const;
+
+	/** Internal data */
+	TArray<uint8> Data;
 };
 
 /** Payload data describing an individual font in a typeface. Keep this lean as it's also used as a key! */
@@ -43,10 +131,8 @@ struct SLATECORE_API FFontData
 	/** Default constructor */
 	FFontData();
 
-#if WITH_EDITORONLY_DATA
 	/** Construct the raw data from a font face asset */
 	explicit FFontData(const UObject* const InFontFaceAsset);
-#endif // WITH_EDITORONLY_DATA
 
 	/** Construct the raw data from a filename and the font data attributes */
 	FFontData(FString InFontFilename, const EFontHinting InHinting, const EFontLoadingPolicy InLoadingPolicy);
@@ -64,13 +150,13 @@ struct SLATECORE_API FFontData
 	/** Get the enum controlling how this font should be loaded at runtime. */
 	EFontLoadingPolicy GetLoadingPolicy() const;
 
-#if WITH_EDITORONLY_DATA
 	/** Get the data buffer containing the data for the current font face. */
-	bool GetFontFaceData(TArray<uint8>& OutFontFaceData) const;
+	FFontFaceDataConstPtr GetFontFaceData() const;
 
 	/** Get the font face asset used by this data (if any). */
 	const UObject* GetFontFaceAsset() const;
 
+#if WITH_EDITORONLY_DATA
 	/** True if this object contains any legacy data that needs to be upgraded PostLoad by calling the functions below (in order). */
 	bool HasLegacyData() const;
 
@@ -92,13 +178,11 @@ struct SLATECORE_API FFontData
 	{
 		uint32 KeyHash = 0;
 
-#if WITH_EDITORONLY_DATA
 		if (Key.FontFaceAsset)
 		{
 			KeyHash = HashCombine(KeyHash, GetTypeHash(Key.FontFaceAsset));
 		}
 		else
-#endif
 		{
 			KeyHash = HashCombine(KeyHash, GetTypeHash(Key.FontFilename));
 			KeyHash = HashCombine(KeyHash, GetTypeHash(Key.Hinting));
@@ -141,13 +225,13 @@ private:
 	UPROPERTY()
 	EFontLoadingPolicy LoadingPolicy;
 
-#if WITH_EDITORONLY_DATA
 	/**
 	 * Font data v3. This points to a font face asset.
 	 */
 	UPROPERTY()
 	const UObject* FontFaceAsset;
 
+#if WITH_EDITORONLY_DATA
 	/**
 	 * Legacy font data v2. This used to be where font data was stored prior to font face assets. 
 	 * This can be removed once we no longer support loading packages older than FEditorObjectVersion::AddedFontFaceAssets (as can UFontBulkData itself).
diff --git a/Engine/Source/Runtime/SlateCore/Public/Fonts/FontCache.h b/Engine/Source/Runtime/SlateCore/Public/Fonts/FontCache.h
--- a/Engine/Source/Runtime/SlateCore/Public/Fonts/FontCache.h
+++ b/Engine/Source/Runtime/SlateCore/Public/Fonts/FontCache.h
@@ -864,6 +864,11 @@ private:
 	 */
 	void FlushCache();
 
+	/**
+	 * Clears out any pending UFont objects that were requested to be flushed
+	 */
+	void FlushFontObjects();
+
 	/** Called after the active culture has changed */
 	void HandleCultureChanged();
 
@@ -920,8 +925,9 @@ private:
 	/** The frame counter the last time the font cache was asked to be flushed */
 	uint64 FrameCounterLastFlushRequest;
 
-//@HSL_BEGIN - Chance.Lyon - A critical section to help make this class thread-safe */
-	/** Critical section for thread synchronization for the font cache */
-	mutable FCriticalSection CacheCriticalSection;
-//@HSL_END
+	/** Critical section preventing concurrent access to FontObjectsToFlush */
+	mutable FCriticalSection FontObjectsToFlushCS;
+
+	/** Array of UFont objects that the font cache has been requested to flush. Since GC can happen while the loading screen is running, the request may be deferred until the next call to ConditionalFlushCache */
+	TArray<const UObject*> FontObjectsToFlush;
 };
diff --git a/Engine/Source/Runtime/SlateCore/Public/Fonts/FontFaceInterface.h b/Engine/Source/Runtime/SlateCore/Public/Fonts/FontFaceInterface.h
--- a/Engine/Source/Runtime/SlateCore/Public/Fonts/FontFaceInterface.h
+++ b/Engine/Source/Runtime/SlateCore/Public/Fonts/FontFaceInterface.h
@@ -32,10 +32,8 @@ class IFontFaceInterface
 	/** Get the enum controlling how this font should be loaded at runtime. */
 	virtual EFontLoadingPolicy GetLoadingPolicy() const = 0;
 
-#if WITH_EDITORONLY_DATA
 	/** Get the data buffer containing the data for the current font face. */
-	virtual const TArray<uint8>& GetFontFaceData() const = 0;
-#endif // WITH_EDITORONLY_DATA
+	virtual FFontFaceDataConstRef GetFontFaceData() const = 0;
 
 	/** Get the filename of the font to use when cooked. This file will be generated during cook. */
 	virtual FString GetCookedFilename() const = 0;
diff --git a/Engine/Source/Runtime/SlateCore/Public/Fonts/SlateFontInfo.h b/Engine/Source/Runtime/SlateCore/Public/Fonts/SlateFontInfo.h
--- a/Engine/Source/Runtime/SlateCore/Public/Fonts/SlateFontInfo.h
+++ b/Engine/Source/Runtime/SlateCore/Public/Fonts/SlateFontInfo.h
@@ -57,7 +57,14 @@ struct SLATECORE_API FFontOutlineSettings
 	FFontOutlineSettings()
 		: OutlineSize(0)
 		, OutlineMaterial(nullptr)
-		, OutlineColor(FColor::Black)
+		, OutlineColor(FLinearColor::Black)
+		, bSeparateFillAlpha(false)
+	{}
+
+	FFontOutlineSettings(int32 InOutlineSize, FLinearColor InColor = FLinearColor::Black)
+		: OutlineSize(InOutlineSize)
+		, OutlineMaterial(nullptr)
+		, OutlineColor(InColor)
 		, bSeparateFillAlpha(false)
 	{}
 
@@ -140,7 +147,7 @@ public:
 	 * @param InSize The size of the font.
 	 * @param InTypefaceFontName The name of the font to use from the default typeface (None will use the first entry)
 	 */
-	FSlateFontInfo( TSharedPtr<const FCompositeFont> InCompositeFont, const int32 InSize, const FName& InTypefaceFontName = NAME_None );
+	FSlateFontInfo( TSharedPtr<const FCompositeFont> InCompositeFont, const int32 InSize, const FName& InTypefaceFontName = NAME_None, const FFontOutlineSettings& InOutlineSettings = FFontOutlineSettings());
 
 	/**
 	 * Creates and initializes a new instance with the specified font, size, and emphasis.
@@ -149,7 +156,7 @@ public:
 	 * @param InSize The size of the font.
 	 * @param InFamilyFontName The name of the font to use from the default typeface (None will use the first entry)
 	 */
-	FSlateFontInfo( const UObject* InFontObject, const int32 InSize, const FName& InTypefaceFontName = NAME_None );
+	FSlateFontInfo( const UObject* InFontObject, const int32 InSize, const FName& InTypefaceFontName = NAME_None, const FFontOutlineSettings& InOutlineSettings = FFontOutlineSettings());
 
 	/**
 	 * DEPRECATED - Creates and initializes a new instance with the specified font name and size.
@@ -158,7 +165,7 @@ public:
 	 * @param InSize The size of the font.
 	 * @param InHinting The type of hinting to use for the font.
 	 */
-	FSlateFontInfo( const FString& InFontName, uint16 InSize, EFontHinting InHinting = EFontHinting::Default );
+	FSlateFontInfo( const FString& InFontName, uint16 InSize, EFontHinting InHinting = EFontHinting::Default, const FFontOutlineSettings& InOutlineSettings = FFontOutlineSettings());
 
 	/**
 	 * DEPRECATED - Creates and initializes a new instance with the specified font name and size.
diff --git a/Engine/Source/Runtime/SlateCore/Public/Rendering/DrawElements.h b/Engine/Source/Runtime/SlateCore/Public/Rendering/DrawElements.h
--- a/Engine/Source/Runtime/SlateCore/Public/Rendering/DrawElements.h
+++ b/Engine/Source/Runtime/SlateCore/Public/Rendering/DrawElements.h
@@ -472,6 +472,7 @@ public:
 	FORCEINLINE const FSlateDataPayload& GetDataPayload() const { return DataPayload; }
 	FORCEINLINE uint32 GetDrawEffects() const { return DrawEffects; }
 	FORCEINLINE const TOptional<FShortRect>& GetScissorRect() const { return ScissorRect; }
+	FORCEINLINE const int32 GetSceneIndex() const { return SceneIndex; }
 
 private:
 	void Init(uint32 InLayer, const FPaintGeometry& PaintGeometry, const FSlateRect& InClippingRect, ESlateDrawEffect::Type InDrawEffects);
@@ -490,6 +491,7 @@ private:
 	uint32 DrawEffects;
 	EElementType ElementType;
 	TOptional<FShortRect> ScissorRect;
+	int32 SceneIndex;
 };
 
 
@@ -565,9 +567,9 @@ private:
  */
 class FSlateElementBatch
 {
-public: 
-	FSlateElementBatch(const FSlateShaderResource* InShaderResource, const FShaderParams& InShaderParams, ESlateShader::Type ShaderType, ESlateDrawPrimitive::Type PrimitiveType, ESlateDrawEffect::Type DrawEffects, ESlateBatchDrawFlag::Type DrawFlags, const TOptional<FShortRect>& ScissorRect, int32 InstanceCount = 0, uint32 InstanceOffset = 0, ISlateUpdatableInstanceBuffer* InstanceData = nullptr)
-		: BatchKey(InShaderParams, ShaderType, PrimitiveType, DrawEffects, DrawFlags, ScissorRect, InstanceCount, InstanceOffset, InstanceData)
+public:
+	FSlateElementBatch(const FSlateShaderResource* InShaderResource, const FShaderParams& InShaderParams, ESlateShader::Type ShaderType, ESlateDrawPrimitive::Type PrimitiveType, ESlateDrawEffect::Type DrawEffects, ESlateBatchDrawFlag::Type DrawFlags, const TOptional<FShortRect>& ScissorRect, int32 InstanceCount = 0, uint32 InstanceOffset = 0, ISlateUpdatableInstanceBuffer* InstanceData = nullptr, int32 SceneIndex = -1)
+		: BatchKey(InShaderParams, ShaderType, PrimitiveType, DrawEffects, DrawFlags, ScissorRect, InstanceCount, InstanceOffset, InstanceData, SceneIndex)
 		, ShaderResource(InShaderResource)
 		, NumElementsInBatch(0)
 		, VertexArrayIndex(INDEX_NONE)
@@ -622,7 +624,7 @@ public:
 	int32 GetInstanceCount() const{ return BatchKey.InstanceCount; }
 	uint32 GetInstanceOffset() const{ return BatchKey.InstanceOffset; }
 	const ISlateUpdatableInstanceBuffer* GetInstanceData() const { return BatchKey.InstanceData; }
-
+	int32 GetSceneIndex() const { return BatchKey.SceneIndex; }
 private:
 	struct FBatchKey
 	{
@@ -639,6 +641,21 @@ private:
 		const int32 InstanceCount;
 		const uint32 InstanceOffset;
 		const ISlateUpdatableInstanceBuffer* InstanceData;
+		const int32 SceneIndex;
+
+		FBatchKey(const FShaderParams& InShaderParams, ESlateShader::Type InShaderType, ESlateDrawPrimitive::Type InDrawPrimitiveType, ESlateDrawEffect::Type InDrawEffects, ESlateBatchDrawFlag::Type InDrawFlags, const TOptional<FShortRect>& InScissorRect, int32 InInstanceCount, uint32 InInstanceOffset, ISlateUpdatableInstanceBuffer* InInstanceBuffer, int32 InSceneIndex)
+			: ShaderParams(InShaderParams)
+			, DrawFlags(InDrawFlags)
+			, ShaderType(InShaderType)
+			, DrawPrimitiveType(InDrawPrimitiveType)
+			, DrawEffects(InDrawEffects)
+			, ScissorRect(InScissorRect)
+			, InstanceCount(InInstanceCount)
+			, InstanceOffset(InInstanceOffset)
+			, InstanceData(InInstanceBuffer)
+			, SceneIndex(InSceneIndex)
+		{
+		}
 
 		FBatchKey( const FShaderParams& InShaderParams, ESlateShader::Type InShaderType, ESlateDrawPrimitive::Type InDrawPrimitiveType, ESlateDrawEffect::Type InDrawEffects, ESlateBatchDrawFlag::Type InDrawFlags, const TOptional<FShortRect>& InScissorRect, int32 InInstanceCount, uint32 InInstanceOffset, ISlateUpdatableInstanceBuffer* InInstanceBuffer)
 			: ShaderParams( InShaderParams )
@@ -649,7 +666,8 @@ private:
 			, ScissorRect( InScissorRect )
 			, InstanceCount( InInstanceCount )
 			, InstanceOffset( InInstanceOffset )
-			, InstanceData( InInstanceBuffer )
+			, InstanceData(InInstanceBuffer)
+			, SceneIndex(-1)
 		{
 		}
 
@@ -663,6 +681,7 @@ private:
 			, ScissorRect( InScissorRect )
 			, InstanceCount(0)
 			, InstanceOffset(0)
+			, SceneIndex(-1)
 		{}
 
 		FBatchKey( TSharedPtr<FSlateRenderDataHandle, ESPMode::ThreadSafe> InCachedRenderHandle, FVector2D InCachedRenderDataOffset, const TOptional<FShortRect>& InScissorRect)
@@ -676,6 +695,7 @@ private:
 			, ScissorRect(InScissorRect)
 			, InstanceCount(0)
 			, InstanceOffset(0)
+			, SceneIndex(-1)
 		{}
 
 		FBatchKey(TSharedPtr<FSlateDrawLayerHandle, ESPMode::ThreadSafe> InLayerHandle, const TOptional<FShortRect>& InScissorRect)
@@ -688,6 +708,7 @@ private:
 			, ScissorRect(InScissorRect)
 			, InstanceCount(0)
 			, InstanceOffset(0)
+			, SceneIndex(-1)
 		{}
 
 		bool operator==( const FBatchKey& Other ) const
@@ -703,7 +724,8 @@ private:
 				&& LayerHandle == Other.LayerHandle
 				&& InstanceCount == Other.InstanceCount
 				&& InstanceOffset == Other.InstanceOffset
-				&& InstanceData == Other.InstanceData;
+				&& InstanceData == Other.InstanceData
+				&& SceneIndex == Other.SceneIndex;
 		}
 
 		/** Compute an efficient hash for this type for use in hash containers. */
@@ -720,6 +742,7 @@ private:
 			RunningHash = bHasInstances ? HashCombine(InBatchKey.InstanceCount, RunningHash) : RunningHash;
 			RunningHash = bHasInstances ? HashCombine(InBatchKey.InstanceOffset, RunningHash) : RunningHash;
 			RunningHash = InBatchKey.InstanceData ? HashCombine( PointerHash(InBatchKey.InstanceData), RunningHash ) : RunningHash;
+			RunningHash = HashCombine(InBatchKey.SceneIndex, RunningHash);
 
 			return RunningHash;
 			//return FCrc::MemCrc32(&InBatchKey.ShaderParams, sizeof(FShaderParams)) ^ ((InBatchKey.ShaderType << 16) | (InBatchKey.DrawFlags+InBatchKey.ShaderType+InBatchKey.DrawPrimitiveType+InBatchKey.DrawEffects));
@@ -764,6 +787,7 @@ public:
 		, IndexOffset( InIndexOffset )
 		, NumVertices( InNumVertices )
 		, NumIndices( InNumIndices )
+		, SceneIndex(InBatch.GetSceneIndex())
 	{}
 
 public:
@@ -812,6 +836,8 @@ public:
 	const uint32 NumVertices;
 	/** Number of indices in the batch */
 	const uint32 NumIndices;
+
+	const int32 SceneIndex;
 };
 
 
diff --git a/Engine/Source/Runtime/SlateCore/Public/Rendering/ElementBatcher.h b/Engine/Source/Runtime/SlateCore/Public/Rendering/ElementBatcher.h
--- a/Engine/Source/Runtime/SlateCore/Public/Rendering/ElementBatcher.h
+++ b/Engine/Source/Runtime/SlateCore/Public/Rendering/ElementBatcher.h
@@ -124,6 +124,8 @@ private:
 	 * @param PrimitiveType	The primitive type( triangles, lines ) to use when drawing the batch
 	 * @param ShaderType	The shader to use when rendering this batch
 	 * @param DrawFlags		Any optional draw flags for this batch
+	 * @param ScissorRect   Optional scissor rectangle for this batch
+	 * @param SceneIndex    Index in the slate renderer's scenes array associated with this element.
 	 */	
 	FSlateElementBatch& FindBatchForElement( uint32 Layer, 
 											 const FShaderParams& ShaderParams, 
@@ -132,7 +134,8 @@ private:
 											 ESlateShader::Type ShaderType, 
 											 ESlateDrawEffect::Type DrawEffects, 
 											 ESlateBatchDrawFlag::Type DrawFlags,
-											 const TOptional<FShortRect>& ScissorRect);
+											 const TOptional<FShortRect>& ScissorRect,
+											 int32 SceneIndex = -1);
 private:
 	/** Batch data currently being filled in */
 	FSlateBatchData* BatchData;
diff --git a/Engine/Source/Runtime/SlateCore/Public/Rendering/RenderingPolicy.h b/Engine/Source/Runtime/SlateCore/Public/Rendering/RenderingPolicy.h
--- a/Engine/Source/Runtime/SlateCore/Public/Rendering/RenderingPolicy.h
+++ b/Engine/Source/Runtime/SlateCore/Public/Rendering/RenderingPolicy.h
@@ -4,7 +4,13 @@
 
 #include "CoreMinimal.h"
 
+
 class FSlateFontServices;
+class FSlateElementBatch;
+class FSlateShaderResource;
+struct FSlateVertex;
+class FSceneInterface;
+
 
 /**
  * Abstract base class for Slate rendering policies.
@@ -18,13 +24,13 @@ public:
 	 *
 	 * @param InPixelCenterOffset
 	 */
-	FSlateRenderingPolicy( const TSharedRef<class FSlateFontServices>& InFontServices, float InPixelCenterOffset )
-		: FontServices( InFontServices )
-		, PixelCenterOffset( InPixelCenterOffset )
+	FSlateRenderingPolicy(const TSharedRef<class FSlateFontServices>& InFontServices, float InPixelCenterOffset)
+		: FontServices(InFontServices)
+		, PixelCenterOffset(InPixelCenterOffset)
 	{ }
 
 	/** Virtual destructor. */
-	virtual ~FSlateRenderingPolicy( ) { }
+	virtual ~FSlateRenderingPolicy() { }
 
 	TSharedRef<class FSlateFontCache> GetFontCache() const;
 	TSharedRef<class FSlateFontServices> GetFontServices() const;
@@ -38,6 +44,10 @@ public:
 		return PixelCenterOffset;
 	}
 
+
+	virtual void AddSceneAt(FSceneInterface* Scene, int32 Index){};
+	virtual void ClearScenes() {};
+
 private:
 
 	FSlateRenderingPolicy(const FSlateRenderingPolicy&);
diff --git a/Engine/Source/Runtime/SlateCore/Public/Rendering/SlateRenderer.h b/Engine/Source/Runtime/SlateCore/Public/Rendering/SlateRenderer.h
--- a/Engine/Source/Runtime/SlateCore/Public/Rendering/SlateRenderer.h
+++ b/Engine/Source/Runtime/SlateCore/Public/Rendering/SlateRenderer.h
@@ -15,6 +15,8 @@ class ISlateAtlasProvider;
 class ISlateStyle;
 class SWindow;
 struct Rect;
+class FSceneInterface;
+struct FSlateBrush;
 
 typedef FRHITexture2D* FTexture2DRHIParamRef;
 
@@ -432,6 +434,15 @@ public:
 	 */
 	FCriticalSection* GetResourceCriticalSection() { return &ResourceCriticalSection; }
 
+	/** Register the active scene pointer with the renderer. This will return the scene internal index that will be used for all subsequent elements drawn. */
+	virtual int32 RegisterCurrentScene(FSceneInterface* Scene) = 0;
+
+	/** Get the currently registered scene index (set by RegisterCurrentScene)*/
+	virtual int32 GetCurrentSceneIndex() const  = 0;
+
+	/** Reset the internal Scene tracking.*/
+	virtual void ClearScenes() = 0;
+
 private:
 
 	// Non-copyable
diff --git a/Engine/Source/Runtime/SlateCore/Public/Widgets/SWidget.h b/Engine/Source/Runtime/SlateCore/Public/Widgets/SWidget.h
--- a/Engine/Source/Runtime/SlateCore/Public/Widgets/SWidget.h
+++ b/Engine/Source/Runtime/SlateCore/Public/Widgets/SWidget.h
@@ -225,9 +225,6 @@ public:
 	 */
 	int32 Paint(const FPaintArgs& Args, const FGeometry& AllottedGeometry, const FSlateRect& MyClippingRect, FSlateWindowElementList& OutDrawElements, int32 LayerId, const FWidgetStyle& InWidgetStyle, bool bParentEnabled) const;
 
-	DEPRECATED(4.4, "Paint() now requires an extra FPaintArgs parameter. When calling paint on a child widget, use Args.WithNewParent(this).")
-	int32 Paint(const FGeometry& AllottedGeometry, const FSlateRect& MyClippingRect, FSlateWindowElementList& OutDrawElements, int32 LayerId, const FWidgetStyle& InWidgetStyle, bool bParentEnabled) const { return 0; }
-
 	/**
 	 * Ticks this widget with Geometry.  Override in derived classes, but always call the parent implementation.
 	 *
@@ -237,13 +234,6 @@ public:
 	 */
 	virtual void Tick( const FGeometry& AllottedGeometry, const double InCurrentTime, const float InDeltaTime );
 
-	/**
-	 * DEPRECATED: This function caused a lot of confusion, and we found that all known uses of it were unnecessary.
-	 * We will provide a less confusing API for supporting non-rectangular hittesting and pixel-perfect hit testing.
-	 */
-	DEPRECATED(4.5, "OnHitTest is no longer called. It will be replaced by more robust hit testing logic in the near future.")
-	virtual bool OnHitTest( const FGeometry& MyGeometry, FVector2D InAbsoluteCursorPosition ){return false;}
-
 	//
 	// KEY INPUT
 	//
@@ -302,9 +292,6 @@ public:
 	 */
 	virtual FReply OnKeyDown(const FGeometry& MyGeometry, const FKeyEvent& InKeyEvent);
 
-	DEPRECATED(4.6, "SWidget::OnControllerButtonPressed() is deprecated, SWidget::OnKeyDown() now handles controller input as well as controller.")
-	virtual FReply OnControllerButtonPressed(const FGeometry& MyGeometry, const FControllerEvent& ControllerEvent);
-
 	/**
 	 * Called after a key is released when this widget has focus
 	 *
@@ -314,9 +301,6 @@ public:
 	 */
 	virtual FReply OnKeyUp( const FGeometry& MyGeometry, const FKeyEvent& InKeyEvent );
 
-	DEPRECATED(4.6, "SWidget::OnControllerButtonReleased() is deprecated, SWidget::OnKeyUp() now handles controller input as well as controller.")
-	virtual FReply OnControllerButtonReleased(const FGeometry& MyGeometry, const FControllerEvent& ControllerEvent);
-
 	/**
 	 * Called when an analog value changes on a button that supports analog
 	 *
@@ -326,9 +310,6 @@ public:
 	 */
 	virtual FReply OnAnalogValueChanged( const FGeometry& MyGeometry, const FAnalogInputEvent& InAnalogInputEvent );
 
-	DEPRECATED(4.6, "SWidget::OnControllerAnalogValueChanged() is deprecated, SWidget::OnAnalogValueChanged() handles keyboard and controller input.")
-	virtual FReply OnControllerAnalogValueChanged(const FGeometry& MyGeometry, const FControllerEvent& ControllerEvent);
-
 	//
 	// MOUSE INPUT
 	//
@@ -656,7 +637,11 @@ public:
 	}
 
 	/** What is the Child's scale relative to this widget. */
-	virtual float GetRelativeLayoutScale( const FSlotBase& Child ) const;
+	DEPRECATED(4.15, "This version is no longer used, please use the new version which also provides the incoming prepass scale, in case that has bearing on the relative layout scale.")
+	virtual float GetRelativeLayoutScale(const FSlotBase& Child) const { return 1.0f; }
+
+	/** What is the Child's scale relative to this widget. */
+	virtual float GetRelativeLayoutScale(const FSlotBase& Child, float LayoutScaleMultiplier) const;
 
 	/**
 	 * Non-virtual entry point for arrange children. ensures common work is executed before calling the virtual
@@ -765,12 +750,12 @@ public:
 	}
 
 	/** @return Whether or not this widget is enabled */
-	bool IsEnabled() const 
+	FORCEINLINE bool IsEnabled() const 
 	{
 		return EnabledState.Get();
 	}
 
-	/** @return Is this widget interactable or not? Defaults to false */
+	/** @return Is this widget interactive or not? Defaults to false */
 	virtual bool IsInteractable() const
 	{
 		return false;
@@ -1033,6 +1018,17 @@ public:
 	/** @return the Foreground color that this widget sets; unset options if the widget does not set a foreground color */
 	virtual FSlateColor GetForegroundColor() const;
 
+	/**
+	 * Gets the last geometry used to Tick the widget.  This data may not exist yet if this call happens prior to 
+	 * the widget having been ticked/painted, or it may be out of date, or a frame behind.
+	 *
+	 * We recommend not to use this data unless there's no other way to solve your problem.  Normally in Slate we
+	 * try and handle these issues by making a dependent widget part of the hierarchy, as to avoid frame behind
+	 * or what are referred to as hysteresis problems, both caused by depending on geometry from the previous frame
+	 * being used to advise how to layout a dependent object the current frame.
+	 */
+	const FGeometry& GetCachedGeometry() const { return CachedGeometry; }
+
 protected:
 
 	/**
@@ -1131,9 +1127,6 @@ private:
 	 */
 	virtual int32 OnPaint(const FPaintArgs& Args, const FGeometry& AllottedGeometry, const FSlateRect& MyClippingRect, FSlateWindowElementList& OutDrawElements, int32 LayerId, const FWidgetStyle& InWidgetStyle, bool bParentEnabled) const = 0;
 
-	DEPRECATED(4.4, "OnPaint() now requires an extra FPaintArgs parameter. When calling Superclass::OnPaint() you can simply pass through the existing Args.")
-	virtual int32 OnPaint(const FGeometry& AllottedGeometry, const FSlateRect& MyClippingRect, FSlateWindowElementList& OutDrawElements, int32 LayerId, const FWidgetStyle& InWidgetStyle, bool bParentEnabled) const { return 0; }
-
 	/**
 	 * Compute the Geometry of all the children and add populate the ArrangedChildren list with their values.
 	 * Each type of Layout panel should arrange children based on desired behavior.
@@ -1198,9 +1191,14 @@ protected:
 	//	DEBUG INFORMATION
 	// @todo Slate: Should compile out in final release builds?
 	FName TypeOfWidget;
+
+#if !UE_BUILD_SHIPPING
+
 	/** Full file path (and line) in which this widget was created */
 	FName CreatedInLocation;
 
+#endif
+
 	/** Tag for this widget */
 	FName Tag;
 
@@ -1231,6 +1229,13 @@ private:
 	 */
 	FVector2D DesiredSize;
 
+	/**
+	 * Stores the cached Tick Geometry of the widget.  This information can and will be outdated, that's the
+	 * nature of it.  However, users were found to often need access to the geometry at times inconvenient to always
+	 * need to be located in Widget Tick.
+	 */
+	mutable FGeometry CachedGeometry;
+
 	/** Tool tip content for this widget */
 	TSharedPtr<IToolTip> ToolTip;
 
@@ -1238,6 +1243,11 @@ private:
 	mutable TWeakPtr<ILayoutCache> LayoutCache;
 
 	STAT(mutable TStatId				StatID;)
+
+private:
+	// Events
+	TMap<FName, FPointerEventHandler> PointerEvents;
+
 protected:
 	/** Is this widget hovered? */
 	bool bIsHovered : 1;
@@ -1270,11 +1280,6 @@ private:
 
 	/** If we're owned by a volatile widget, we need inherit that volatility and use as part of our volatility, but don't cache it. */
 	mutable bool bInheritedVolatility : 1;
-
-	FPointerEventHandler MouseButtonDownHandler;
-	FPointerEventHandler MouseButtonUpHandler;
-	FPointerEventHandler MouseMoveHandler;
-	FPointerEventHandler MouseDoubleClickHandler;
 	FNoReplyPointerEventHandler MouseEnterHandler;
 	FSimpleNoReplyPointerEventHandler MouseLeaveHandler;
 };
diff --git a/Engine/Source/Runtime/SlateNullRenderer/Private/SlateNullRenderer.cpp b/Engine/Source/Runtime/SlateNullRenderer/Private/SlateNullRenderer.cpp
--- a/Engine/Source/Runtime/SlateNullRenderer/Private/SlateNullRenderer.cpp
+++ b/Engine/Source/Runtime/SlateNullRenderer/Private/SlateNullRenderer.cpp
@@ -89,3 +89,20 @@ void FSlateNullRenderer::ReleaseUpdatableTexture(FSlateUpdatableTexture* Texture
 void FSlateNullRenderer::RequestResize( const TSharedPtr<SWindow>& Window, uint32 NewWidth, uint32 NewHeight )
 {
 }
+
+int32 FSlateNullRenderer::RegisterCurrentScene(FSceneInterface* Scene) 
+{
+	// This is a no-op
+	return -1;
+}
+
+int32 FSlateNullRenderer::GetCurrentSceneIndex() const
+{
+	// This is a no-op
+	return -1;
+}
+
+void FSlateNullRenderer::ClearScenes() 
+{
+	// This is a no-op
+}
\ No newline at end of file
diff --git a/Engine/Source/Runtime/SlateRHIRenderer/Private/SlatePostProcessResource.cpp b/Engine/Source/Runtime/SlateRHIRenderer/Private/SlatePostProcessResource.cpp
--- a/Engine/Source/Runtime/SlateRHIRenderer/Private/SlatePostProcessResource.cpp
+++ b/Engine/Source/Runtime/SlateRHIRenderer/Private/SlatePostProcessResource.cpp
@@ -9,7 +9,6 @@ FSlatePostProcessResource::FSlatePostProcessResource(int32 InRenderTargetCount)
 	: RenderTargetSize(FIntPoint::ZeroValue)
 	, RenderTargetCount(InRenderTargetCount)
 {
-	
 }
 
 FSlatePostProcessResource::~FSlatePostProcessResource()
@@ -17,29 +16,28 @@ FSlatePostProcessResource::~FSlatePostProcessResource()
 
 }
 
-void FSlatePostProcessResource::Resize(const FIntPoint& NewSize)
+void FSlatePostProcessResource::Update(const FIntPoint& NewSize)
 {
-	check(IsInRenderingThread());
-	if(RenderTargetSize != NewSize)
+	if(NewSize.X > RenderTargetSize.X || NewSize.Y > RenderTargetSize.Y || RenderTargetSize == FIntPoint::ZeroValue || RenderTargets.Num() == 0 )
 	{
-		RenderTargetSize = NewSize;
-		UpdateRHI();
+		if(!IsInitialized())
+		{
+			InitResource();
+		}
+	
+		ResizeTargets(NewSize);
 	}
 }
 
-void FSlatePostProcessResource::CleanUp()
+void FSlatePostProcessResource::ResizeTargets(const FIntPoint& NewSize)
 {
-	BeginReleaseResource(this);
-
-	BeginCleanup(this);
-}
+	check(IsInRenderingThread());
 
-void FSlatePostProcessResource::InitDynamicRHI()
-{
 	RenderTargets.Empty();
 
+	RenderTargetSize = NewSize;
 	PixelFormat = PF_B8G8R8A8;
-	if(RenderTargetSize.X > 0 && RenderTargetSize.Y > 0)
+	if (RenderTargetSize.X > 0 && RenderTargetSize.Y > 0)
 	{
 		for (int32 TexIndex = 0; TexIndex < RenderTargetCount; ++TexIndex)
 		{
@@ -67,10 +65,23 @@ void FSlatePostProcessResource::InitDynamicRHI()
 	SET_MEMORY_STAT(STAT_SLATEPPRenderTargetMem, TotalMemory);
 }
 
+void FSlatePostProcessResource::CleanUp()
+{
+	BeginReleaseResource(this);
+
+	BeginCleanup(this);
+}
+
+void FSlatePostProcessResource::InitDynamicRHI()
+{
+}
+
 void FSlatePostProcessResource::ReleaseDynamicRHI()
 {
 	SET_MEMORY_STAT(STAT_SLATEPPRenderTargetMem, 0);
 
+	RenderTargetSize = FIntPoint::ZeroValue;
+
 	RenderTargets.Empty();
 }
 
diff --git a/Engine/Source/Runtime/SlateRHIRenderer/Private/SlatePostProcessResource.h b/Engine/Source/Runtime/SlateRHIRenderer/Private/SlatePostProcessResource.h
--- a/Engine/Source/Runtime/SlateRHIRenderer/Private/SlatePostProcessResource.h
+++ b/Engine/Source/Runtime/SlateRHIRenderer/Private/SlatePostProcessResource.h
@@ -21,7 +21,8 @@ public:
 		return RenderTargets[Index]; 
 	}
 
-	void Resize(const FIntPoint& NewSize);
+	/** Performs per frame updates to this resource */
+	void Update(const FIntPoint& NewSize);
 
 	void CleanUp();
 
@@ -36,8 +37,10 @@ public:
 
 
 private:
+	/** Resizes targets to the new size */
+	void ResizeTargets(const FIntPoint& NewSize);
 
-
+	/** FDeferredCleanupInterface */
 	virtual void FinishCleanup();
 
 private:
diff --git a/Engine/Source/Runtime/SlateRHIRenderer/Private/SlatePostProcessor.cpp b/Engine/Source/Runtime/SlateRHIRenderer/Private/SlatePostProcessor.cpp
--- a/Engine/Source/Runtime/SlateRHIRenderer/Private/SlatePostProcessor.cpp
+++ b/Engine/Source/Runtime/SlateRHIRenderer/Private/SlatePostProcessor.cpp
@@ -31,12 +31,10 @@ void FSlatePostProcessor::BlurRect(FRHICommandListImmediate& RHICmdList, IRender
 	TArray<FVector4> WeightsAndOffsets;
 	const int32 SampleCount = ComputeBlurWeights(Params.KernelSize, Params.Strength, WeightsAndOffsets);
 
-	SCOPED_DRAW_EVENTF(RHICmdList, SlatePostProcess, TEXT("Slate Post Process Blur Background %dx%d"), SampleCount, SampleCount);
 
 	const bool bDownsample = Params.DownsampleAmount > 0;
 
 	FIntPoint DestRectSize = RectParams.DestRect.GetSize().IntPoint();
-	const FIntPoint CurrentSize = FIntPoint(IntermediateTargets->GetWidth(), IntermediateTargets->GetHeight());
 	FIntPoint RequiredSize = bDownsample
 										? FIntPoint(FMath::DivideAndRoundUp(DestRectSize.X, Params.DownsampleAmount), FMath::DivideAndRoundUp(DestRectSize.Y, Params.DownsampleAmount))
 										: DestRectSize;
@@ -45,11 +43,12 @@ void FSlatePostProcessor::BlurRect(FRHICommandListImmediate& RHICmdList, IRender
 	RequiredSize.X = FMath::Min(RequiredSize.X, RectParams.SourceTextureSize.X);
 	RequiredSize.Y = FMath::Min(RequiredSize.Y, RectParams.SourceTextureSize.Y);
 	
+	SCOPED_DRAW_EVENTF(RHICmdList, SlatePostProcess, TEXT("Slate Post Process Blur Background Kernel: %dx%d Size: %dx%d"), SampleCount, SampleCount, RequiredSize.X, RequiredSize.Y);
+
+
 	const FIntPoint DownsampleSize = RequiredSize;
-	if(RequiredSize.X > CurrentSize.X || RequiredSize.Y > CurrentSize.Y)
-	{		
-		IntermediateTargets->Resize(RequiredSize);
-	}
+
+	IntermediateTargets->Update(RequiredSize);
 
 	if(bDownsample)
 	{
@@ -185,6 +184,13 @@ void FSlatePostProcessor::BlurRect(FRHICommandListImmediate& RHICmdList, IRender
 	UpsampleRect(RHICmdList, RendererModule, RectParams, DownsampleSize);
 }
 
+void FSlatePostProcessor::ReleaseRenderTargets()
+{
+	check(IsInGameThread());
+	// Only release the resource not delete it.  Deleting it could cause issues on any RHI thread
+	BeginReleaseResource(IntermediateTargets);
+}
+
 void FSlatePostProcessor::DownsampleRect(FRHICommandListImmediate& RHICmdList, IRendererModule& RendererModule, const FPostProcessRectParams& Params, const FIntPoint& DownsampleSize)
 {
 	SCOPED_DRAW_EVENT(RHICmdList, SlatePostProcessDownsample);
