[NO NON-PII COMMIT MESSAGE]

diff --git a/Engine/Source/Runtime/Slate/Public/Widgets/Input/SNumericEntryBox.h b/Engine/Source/Runtime/Slate/Public/Widgets/Input/SNumericEntryBox.h
--- a/Engine/Source/Runtime/Slate/Public/Widgets/Input/SNumericEntryBox.h
+++ b/Engine/Source/Runtime/Slate/Public/Widgets/Input/SNumericEntryBox.h
@@ -72,7 +72,6 @@ public:
 		, _MinSliderValue(0)				
 		, _MaxSliderValue(100)
 		, _SliderExponent(1.f)
-		, _SliderExponentNeutralValue(0)
 		, _MinDesiredValueWidth(0)
 		
 	{}		
diff --git a/Engine/Source/Runtime/Slate/Public/Widgets/Input/SSpinBox.h b/Engine/Source/Runtime/Slate/Public/Widgets/Input/SSpinBox.h
--- a/Engine/Source/Runtime/Slate/Public/Widgets/Input/SSpinBox.h
+++ b/Engine/Source/Runtime/Slate/Public/Widgets/Input/SSpinBox.h
@@ -270,7 +270,7 @@ public:
 			const float CachedSliderExponent = SliderExponent.Get();
 			if (CachedSliderExponent != 1)
 			{
-				if (SliderExponentNeutralValue.IsSet())
+				if (SliderExponentNeutralValue.IsSet() && SliderExponentNeutralValue.Get() > GetMinSliderValue() && SliderExponentNeutralValue.Get() < GetMaxSliderValue())
 				{
 					//Compute a log curve on both side of the neutral value
 					float StartFractionFilled = Fraction(SliderExponentNeutralValue.Get(), GetMinSliderValue(), GetMaxSliderValue());
@@ -407,13 +407,15 @@ public:
 				//if we have a range to draw in
 				if ( !bUnlimitedSpinRange) 
 				{
+					bool HasValidExponentNeutralValue = SliderExponentNeutralValue.IsSet() && SliderExponentNeutralValue.Get() > GetMinSliderValue() && SliderExponentNeutralValue.Get() < GetMaxSliderValue();
+
 					const float CachedSliderExponent = SliderExponent.Get();
 					// The amount currently filled in the spinbox, needs to be calculated to do deltas correctly.
 					float FractionFilled = Fraction(InternalValue, GetMinSliderValue(), GetMaxSliderValue());
 						
 					if (CachedSliderExponent != 1)
 					{
-						if (SliderExponentNeutralValue.IsSet())
+						if (HasValidExponentNeutralValue)
 						{
 							//Compute a log curve on both side of the neutral value
 							float StartFractionFilled = Fraction(SliderExponentNeutralValue.Get(), GetMinSliderValue(), GetMaxSliderValue());
@@ -437,7 +439,7 @@ public:
 					if (CachedSliderExponent != 1)
 					{
 						// Have to convert the percent to the proper value due to the exponent component to the spin.
-						if (SliderExponentNeutralValue.IsSet())
+						if (HasValidExponentNeutralValue)
 						{
 							//Compute a log curve on both side of the neutral value
 							float StartFractionFilled = Fraction(SliderExponentNeutralValue.Get(), GetMinSliderValue(), GetMaxSliderValue());
