Copying //UE4/Dev-Blueprints to //UE4/Dev-Main (Source: //UE4/Dev-Blueprints @ 3235800)

diff --git a/Engine/Source/Runtime/Engine/Private/UserDefinedStruct.cpp b/Engine/Source/Runtime/Engine/Private/UserDefinedStruct.cpp
--- a/Engine/Source/Runtime/Engine/Private/UserDefinedStruct.cpp
+++ b/Engine/Source/Runtime/Engine/Private/UserDefinedStruct.cpp
@@ -226,6 +226,55 @@ void UUserDefinedStruct::SerializeTaggedProperties(FArchive& Ar, uint8* Data, US
 }
 
 uint32 UUserDefinedStruct::GetStructTypeHash(const void* Src) const
+{
+	return GetUserDefinedStructTypeHash(Src, this);
+}
+
+void UUserDefinedStruct::RecursivelyPreload()
+{
+	FLinkerLoad* Linker = GetLinker();
+	if( Linker && (NULL == PropertyLink) )
+	{
+		TArray<UObject*> AllChildMembers;
+		GetObjectsWithOuter(this, AllChildMembers);
+		for (int32 Index = 0; Index < AllChildMembers.Num(); ++Index)
+		{
+			UObject* Member = AllChildMembers[Index];
+			Linker->Preload(Member);
+		}
+
+		Linker->Preload(this);
+		if (NULL == PropertyLink)
+		{
+			StaticLink(true);
+		}
+	}
+}
+
+FGuid UUserDefinedStruct::GetCustomGuid() const
+{
+	return Guid;
+}
+
+FString GetPathPostfix(const UObject* ForObject)
+{
+	FString FullAssetName = ForObject->GetOutermost()->GetPathName();
+	if (FullAssetName.StartsWith(TEXT("/Temp/__TEMP_BP__"), ESearchCase::CaseSensitive))
+	{
+		FullAssetName.RemoveFromStart(TEXT("/Temp/__TEMP_BP__"), ESearchCase::CaseSensitive);
+	}
+	FString AssetName = FPackageName::GetLongPackageAssetName(FullAssetName);
+	// append a hash of the path, this uniquely identifies assets with the same name, but different folders:
+	FullAssetName.RemoveFromEnd(AssetName);
+	return FString::Printf(TEXT("%u"), FCrc::MemCrc32(*FullAssetName, FullAssetName.Len()*sizeof(TCHAR)));
+}
+
+FString UUserDefinedStruct::GetStructCPPName() const
+{
+	return ::UnicodeToCPPIdentifier(*GetName(), false, GetPrefixCPP()) + GetPathPostfix(this);
+}
+
+uint32 UUserDefinedStruct::GetUserDefinedStructTypeHash(const void* Src, const UScriptStruct* Type)
 {
 	// Ugliness to maximize entropy on first call to HashCombine... Alternatively we could
 	// do special logic on the first iteration of the below loops (unwind loops or similar):
@@ -249,10 +298,10 @@ uint32 UUserDefinedStruct::GetStructTypeHash(const void* Src) const
 	bool bHasBoolValues = false;
 	uint64 BoolValues = 0;
 	// for blueprint defined structs we can just loop and hash the individual properties:
-	for (TFieldIterator<UProperty> It(this); It; ++It)
+	for (TFieldIterator<const UProperty> It(Type); It; ++It)
 	{
 		uint32 CurrentHash = 0;
-		if (const UBoolProperty* BoolProperty = Cast<UBoolProperty>(*It))
+		if (const UBoolProperty* BoolProperty = Cast<const UBoolProperty>(*It))
 		{
 			for (int32 I = 0; I < It->ArrayDim; ++I)
 			{
@@ -279,28 +328,3 @@ uint32 UUserDefinedStruct::GetStructTypeHash(const void* Src) const
 	return ValueHash;
 }
 
-void UUserDefinedStruct::RecursivelyPreload()
-{
-	FLinkerLoad* Linker = GetLinker();
-	if( Linker && (NULL == PropertyLink) )
-	{
-		TArray<UObject*> AllChildMembers;
-		GetObjectsWithOuter(this, AllChildMembers);
-		for (int32 Index = 0; Index < AllChildMembers.Num(); ++Index)
-		{
-			UObject* Member = AllChildMembers[Index];
-			Linker->Preload(Member);
-		}
-
-		Linker->Preload(this);
-		if (NULL == PropertyLink)
-		{
-			StaticLink(true);
-		}
-	}
-}
-
-FGuid UUserDefinedStruct::GetCustomGuid() const
-{
-	return Guid;
-}
diff --git a/Engine/Source/Runtime/Engine/Private/World.cpp b/Engine/Source/Runtime/Engine/Private/World.cpp
--- a/Engine/Source/Runtime/Engine/Private/World.cpp
+++ b/Engine/Source/Runtime/Engine/Private/World.cpp
@@ -1848,7 +1848,17 @@ void UWorld::TransferBlueprintDebugReferences(UWorld* NewWorld)
 			}
 		}
 	}
-
+	// Ensure the level script actor debug references are transferred to the new world.
+	if (NewWorld)
+	{
+		if (ALevelScriptActor* LevelScript = NewWorld->GetLevelScriptActor())
+		{
+			if (UBlueprint* LevelScriptBlueprint = Cast<UBlueprint>(LevelScript->GetClass()->ClassGeneratedBy))
+			{
+				LevelScriptBlueprint->SetObjectBeingDebugged(LevelScript);
+			}
+		}
+	}
 	// Empty the map, anything useful got moved over the map in the new world
 	BlueprintObjectsBeingDebugged.Empty();
 #endif	//#if WITH_EDITOR
diff --git a/Engine/Source/Runtime/Engine/Public/GeneratedCodeHelpers.h b/Engine/Source/Runtime/Engine/Public/GeneratedCodeHelpers.h
--- a/Engine/Source/Runtime/Engine/Public/GeneratedCodeHelpers.h
+++ b/Engine/Source/Runtime/Engine/Public/GeneratedCodeHelpers.h
@@ -330,6 +330,154 @@ public:
 	{
 		return TargetArray.IsValidIndex(Index);
 	}
+	
+	//Replacements for CustomThunk functions from UBlueprintSetLibrary
+	template<typename T, typename U>
+	static bool Set_Add(TSet<T>& TargetSet, const U& NewItem)
+	{
+		if(T* Result = TargetSet.Find(NewItem))
+		{
+			return false;
+		}
+		TargetSet.Add(NewItem);
+		return true;
+	}
+	
+	template<typename T, typename U>
+	static void Set_AddItems(TSet<T>& TargetSet, const TArray<U>& NewItems)
+	{
+		for(const auto& Entry : NewItems)
+		{
+			TargetSet.Add(Entry);
+		}
+	}
+	
+	template<typename T, typename U>
+	static bool Set_Remove(TSet<T>& TargetSet, const U& Item)
+	{
+		return TargetSet.Remove(Item) > 0;
+	}
+
+	template<typename T, typename U>
+	static void Set_RemoveItems(TSet<T>& TargetSet, const TArray<U>& Items)
+	{
+		for(const U& Entry : Items)
+		{
+			TargetSet.Remove(Entry);
+		}
+	}
+	
+	template<typename T, typename U>
+	static void Set_ToArray(const TSet<T>& A, TArray<U>& Result)
+	{
+		ensure(Result.Num() == 0);
+		for(const T& Entry : A)
+		{
+			Result.Add(Entry);
+		}
+	}
+	
+	template<typename T>
+	static void Set_Clear(TSet<T>& TargetSet)
+	{
+		TargetSet.Empty();
+	}
+	
+	template<typename T>
+	static int32 Set_Length(const TSet<T>& TargetSet)
+	{
+		return TargetSet.Num();
+	}
+	
+	template<typename T, typename U>
+	static bool Set_Contains(const TSet<T>& TargetSet, const U& ItemToFind)
+	{
+		return TargetSet.Find(ItemToFind) != nullptr;
+	}
+	
+	template<typename T>
+	static void Set_Intersection(const TSet<T>& A, const TSet<T>& B, TSet<T>& Result )
+	{
+		Result = A.Intersect(B);
+	}
+	
+	template<typename T>
+	static void Set_Union(const TSet<T>& A, const TSet<T>& B, TSet<T>& Result )
+	{
+		Result = A.Union(B);
+	}
+	
+	template<typename T>
+	static void Set_Difference(const TSet<T>& A, const TSet<T>& B, TSet<T>& Result )
+	{
+		Result = A.Difference(B);
+	}
+
+	//Replacements for CustomThunk functions from UBlueprintMapLibrary
+	template<typename T, typename U, typename V, typename W>
+	static bool Map_Add(TMap<T, U>& TargetMap, const V& Key, const W& Value)
+	{
+		if(U* CurrentValue = TargetMap.Find(Key))
+		{
+			*CurrentValue = Value;
+			return false;
+		}
+		else
+		{
+			TargetMap.Add(Key, Value);
+			return true;
+		}
+	}
+	
+	template<typename T, typename U, typename V>
+	static bool Map_Remove(TMap<T, U>& TargetMap, const V& Key)
+	{
+		return TargetMap.Remove(Key) > 0;
+	}
+	
+	template<typename T, typename U, typename V, typename W>
+	static bool Map_Find(const TMap<T, U>& TargetMap, const V& Key, W& Value)
+	{
+		if(const U* CurrentValue = TargetMap.Find(Key))
+		{
+			Value = *CurrentValue;
+			return true;
+		}
+		return false;
+	}
+	
+	template<typename T, typename U, typename V>
+	static bool Map_Contains(const TMap<T, U>& TargetMap, const V& Key)
+	{
+		return TargetMap.Find(Key) != nullptr;
+	}
+	
+	template<typename T, typename U, typename V>
+	static void Map_Keys(const TMap<T, U>& TargetMap, TArray<V>& Keys)
+	{
+		TargetMap.GetKeys(Keys);
+	}
+
+	template<typename T, typename U, typename V>
+	static void Map_Values(const TMap<T, U>& TargetMap, TArray<V>& Values)
+	{
+		for (typename TMap<T, U>::TConstIterator Iterator(TargetMap); Iterator; ++Iterator)
+		{
+			Values.Add(Iterator.Value());
+		}
+	}
+	
+	template<typename T, typename U>
+	static int32 Map_Length(const TMap<T, U>& TargetMap)
+	{
+		return TargetMap.Num();
+	}
+	
+	template<typename T, typename U>
+	static void Map_Clear(TMap<T, U>& TargetMap)
+	{
+		TargetMap.Empty();
+	}
 
 	//Replacements for CustomThunk functions from UDataTableFunctionLibrary
 
diff --git a/Engine/Source/Runtime/Slate/Public/Framework/SlateDelegates.h b/Engine/Source/Runtime/Slate/Public/Framework/SlateDelegates.h
--- a/Engine/Source/Runtime/Slate/Public/Framework/SlateDelegates.h
+++ b/Engine/Source/Runtime/Slate/Public/Framework/SlateDelegates.h
@@ -121,11 +121,23 @@ DECLARE_DELEGATE_OneParam( FOnLinearColorValueChanged, FLinearColor )
 /** Notification for FVector4 value change */
 DECLARE_DELEGATE_OneParam(FOnVector4ValueChanged, FVector4&)
 
+DECLARE_DELEGATE(FOnTableViewBadState);
 
 template< typename ArgumentType >
 class TSlateDelegates
 {
 public:
+	/**
+	 * Useful for debugging issues with table views.
+	 */
+	DECLARE_DELEGATE_RetVal_OneParam (
+	/** return: String representation of the item for debugging purposes. */
+	FString,
+		FOnItemToString_Debug,
+		/** param: An item to turn into a string for debug purposes. */
+		ArgumentType);
+
+
 	/** A delegate type for OnGenerateWidget handler. Given a data item, the handler should return a Widget visualizing that item */
 	DECLARE_DELEGATE_RetVal_OneParam (
 	/** return: The Widget visualization of the item */
diff --git a/Engine/Source/Runtime/Slate/Public/Widgets/Views/SListView.h b/Engine/Source/Runtime/Slate/Public/Widgets/Views/SListView.h
--- a/Engine/Source/Runtime/Slate/Public/Widgets/Views/SListView.h
+++ b/Engine/Source/Runtime/Slate/Public/Widgets/Views/SListView.h
@@ -57,13 +57,15 @@ public:
 	typedef typename TSlateDelegates< ItemType >::FOnMouseButtonClick FOnMouseButtonClick;
 	typedef typename TSlateDelegates< ItemType >::FOnMouseButtonDoubleClick FOnMouseButtonDoubleClick;
 
+	typedef typename TSlateDelegates< ItemType >::FOnItemToString_Debug FOnItemToString_Debug; 
+
 	DECLARE_DELEGATE_OneParam( FOnWidgetToBeRemoved, const TSharedRef<ITableRow>& );
 
 public:
 	SLATE_BEGIN_ARGS( SListView<ItemType> )
 		: _OnGenerateRow()
 		, _OnRowReleased()
-		, _ListItemsSource( static_cast<const TArray<ItemType>*>(nullptr) ) //@todo Slate Syntax: Initializing from nullptr without a cast
+		, _ListItemsSource()
 		, _ItemHeight(16)
 		, _OnContextMenuOpening()
 		, _OnMouseButtonClick()
@@ -76,6 +78,8 @@ public:
 		, _ConsumeMouseWheel( EConsumeMouseWheel::WhenScrollingPossible )
 		, _WheelScrollMultiplier(GetGlobalScrollAmount())
 		, _HandleGamepadEvents( true )
+		, _OnItemToString_Debug()
+		, _OnEnteredBadState()
 		{ }
 
 		SLATE_EVENT( FOnGenerateRow, OnGenerateRow )
@@ -116,14 +120,19 @@ public:
 
 		SLATE_ARGUMENT( bool, HandleGamepadEvents );
 
+		/** Assign this to get more diagnostics from the list view. */
+		SLATE_EVENT(FOnItemToString_Debug, OnItemToString_Debug)
+
+		SLATE_EVENT(FOnTableViewBadState, OnEnteredBadState);
+
 	SLATE_END_ARGS()
 
 	/**
-		* Construct this widget
-		*
-		* @param	InArgs	The declaration data for this widget
-		*/
-	void Construct( const typename SListView<ItemType>::FArguments& InArgs )
+	 * Construct this widget
+	 *
+	 * @param	InArgs	The declaration data for this widget
+	 */
+	void Construct(const typename SListView<ItemType>::FArguments& InArgs)
 	{
 		this->OnGenerateRow = InArgs._OnGenerateRow;
 		this->OnRowReleased = InArgs._OnRowReleased;
@@ -144,6 +153,12 @@ public:
 		this->WheelScrollMultiplier = InArgs._WheelScrollMultiplier;
 		this->bHandleGamepadEvents = InArgs._HandleGamepadEvents;
 
+		this->OnItemToString_Debug =
+			InArgs._OnItemToString_Debug.IsBound()
+			? InArgs._OnItemToString_Debug
+			: GetDefaultDebugDelegate();
+		OnEnteredBadState = InArgs._OnEnteredBadState;
+
 		// Check for any parameters that the coder forgot to specify.
 		FString ErrorString;
 		{
@@ -417,6 +432,7 @@ public:
 
 private:
 
+	friend class FWidgetGenerator;
 	/**
 	 * A WidgetGenerator is a component responsible for creating widgets from data items.
 	 * It also provides mapping from currently generated widgets to the data items which they
@@ -510,13 +526,7 @@ private:
 				}				
 			}
 
-			checkf( ItemToWidgetMap.Num() == WidgetMapToItem.Num(),
-			        TEXT( "ItemToWidgetMap length (%d) does not match WidgetMapToItem length (%d).  %s" ), ItemToWidgetMap.Num(),
-			        WidgetMapToItem.Num(), OwnerList ? *OwnerList->ToString() : TEXT( "null" ) );
-
-			checkf( WidgetMapToItem.Num() == ItemsWithGeneratedWidgets.Num(),
-			        TEXT( "WidgetMapToItem length (%d) does not match ItemsWithGeneratedWidgets length (%d). This is often because the same item is in the list more than once.  %s" ),
-			        WidgetMapToItem.Num(), ItemsWithGeneratedWidgets.Num(), OwnerList ? *OwnerList->ToString() : TEXT( "null" ) );
+			ValidateWidgetGeneration();
 
 			ItemsToBeCleanedUp.Reset();
 		}
@@ -545,6 +555,64 @@ private:
 			ItemsToBeCleanedUp.Reset();
 		}
 
+		void ValidateWidgetGeneration()
+		{
+			const bool bMapsMismatch = ItemToWidgetMap.Num() != WidgetMapToItem.Num();
+			const bool bGeneratedWidgetsSizeMismatch = WidgetMapToItem.Num() != ItemsWithGeneratedWidgets.Num();
+			if (bMapsMismatch)
+			{
+				UE_LOG(LogSlate, Warning, TEXT("ItemToWidgetMap length (%d) does not match WidgetMapToItem length (%d) in %s. Diagnostics follow. "),
+					ItemToWidgetMap.Num(),
+					WidgetMapToItem.Num(),
+					OwnerList ? *OwnerList->ToString() : TEXT("null"));
+			}
+			
+			if (bGeneratedWidgetsSizeMismatch)
+			{
+				UE_LOG(LogSlate, Warning,
+					TEXT("WidgetMapToItem length (%d) does not match ItemsWithGeneratedWidgets length (%d). This is often because the same item is in the list more than once in %s. Diagnostics follow."),
+					WidgetMapToItem.Num(),
+					ItemsWithGeneratedWidgets.Num(),
+					OwnerList ? *OwnerList->ToString() : TEXT("null"));
+			}
+
+			if (bMapsMismatch || bGeneratedWidgetsSizeMismatch)
+			{
+				if (OwnerList->OnItemToString_Debug.IsBound())
+				{
+					UE_LOG(LogSlate, Warning, TEXT(""));
+					UE_LOG(LogSlate, Warning, TEXT("ItemToWidgetMap :"));
+					for (auto ItemWidgetPair = ItemToWidgetMap.CreateConstIterator(); ItemWidgetPair; ++ItemWidgetPair)
+					{
+						const TSharedRef<SWidget> RowAsWidget = ItemWidgetPair.Value()->AsWidget();
+						UE_LOG(LogSlate, Warning, TEXT("%s -> 0x%08x @ %s"), *OwnerList->OnItemToString_Debug.Execute(ItemWidgetPair.Key()), &RowAsWidget.Get(), *RowAsWidget->ToString() );
+					}
+
+					UE_LOG(LogSlate, Warning, TEXT(""));
+					UE_LOG(LogSlate, Warning, TEXT("WidgetMapToItem:"))
+					for (auto WidgetItemPair = WidgetMapToItem.CreateConstIterator(); WidgetItemPair; ++WidgetItemPair)
+					{
+						UE_LOG(LogSlate, Warning, TEXT("0x%08x -> %s"), WidgetItemPair.Key(), *OwnerList->OnItemToString_Debug.Execute(WidgetItemPair.Value()) );
+					}
+
+					UE_LOG(LogSlate, Warning, TEXT(""));
+					UE_LOG(LogSlate, Warning, TEXT("ItemsWithGeneratedWidgets:"));
+					for (int i = 0; i < ItemsWithGeneratedWidgets.Num(); ++i)
+					{
+						UE_LOG(LogSlate, Warning, TEXT("[%d] %s"),i, *OwnerList->OnItemToString_Debug.Execute(ItemsWithGeneratedWidgets[i]));
+					}
+				}
+				else
+				{
+					UE_LOG(LogSlate, Warning, TEXT("Provide custom 'OnItemToString_Debug' for diagnostics dump."));
+				}
+
+				OwnerList->OnEnteredBadState.ExecuteIfBound();
+
+				checkf( false, TEXT("%s detected a critical error. See diagnostic dump above. Provide a custom 'OnItemToString_Debug' for more detailed diagnostics."), *OwnerList->ToString() );
+			}			
+		}
+
 		/** We store a pointer to the owner list for error purposes, so when asserts occur we can report which list it happened for. */
 		SListView<ItemType>* OwnerList;
 
@@ -1160,6 +1228,22 @@ public:
 
 protected:
 
+	FOnItemToString_Debug GetDefaultDebugDelegate()
+	{
+		return
+		FOnItemToString_Debug::CreateLambda([](ItemType InItem)
+		{
+			if (TListTypeTraits<ItemType>::IsPtrValid(InItem))
+			{
+				return FString::Printf(TEXT("0x%08x"), &(*InItem));
+			}
+			else
+			{
+				return FString(TEXT("nullptr"));
+			}
+		});
+	}
+
 	/**
 	 * If there is a pending request to scroll an item into view, do so.
 	 * 
@@ -1425,6 +1509,12 @@ protected:
 	/** Delegate to be invoked when the list needs to generate a new widget from a data item. */
 	FOnGenerateRow OnGenerateRow;
 
+	/** Assign this to get more diagnostics from the list view. */
+	FOnItemToString_Debug OnItemToString_Debug;
+
+	/** Invoked when the tree enters a bad state. */
+	FOnTableViewBadState OnEnteredBadState;
+
 	/**/
 	FOnWidgetToBeRemoved OnRowReleased;
 
diff --git a/Engine/Source/Runtime/Slate/Public/Widgets/Views/STileView.h b/Engine/Source/Runtime/Slate/Public/Widgets/Views/STileView.h
--- a/Engine/Source/Runtime/Slate/Public/Widgets/Views/STileView.h
+++ b/Engine/Source/Runtime/Slate/Public/Widgets/Views/STileView.h
@@ -31,14 +31,16 @@ public:
 	typedef typename TSlateDelegates< ItemType >::FOnMouseButtonDoubleClick FOnMouseButtonDoubleClick;
 	typedef typename TSlateDelegates< NullableItemType >::FOnSelectionChanged FOnSelectionChanged;
 
+	typedef typename TSlateDelegates< ItemType >::FOnItemToString_Debug FOnItemToString_Debug; 
+
 	using FOnWidgetToBeRemoved = typename SListView<ItemType>::FOnWidgetToBeRemoved;
 
 public:
 
-	SLATE_BEGIN_ARGS( STileView<ItemType> )
+	SLATE_BEGIN_ARGS(STileView<ItemType>)
 		: _OnGenerateTile()
 		, _OnTileReleased()
-		, _ListItemsSource( static_cast<TArray<ItemType>*>(nullptr) ) //@todo Slate Syntax: Initializing from nullptr without a cast
+		, _ListItemsSource(static_cast<TArray<ItemType>*>(nullptr)) //@todo Slate Syntax: Initializing from nullptr without a cast
 		, _ItemHeight(128)
 		, _ItemWidth(128)
 		, _ItemAlignment(EListItemAlignment::EvenlyDistributed)
@@ -51,9 +53,11 @@ public:
 		, _ExternalScrollbar()
 		, _ScrollbarVisibility(EVisibility::Visible)
 		, _AllowOverscroll(EAllowOverscroll::Yes)
-		, _ConsumeMouseWheel( EConsumeMouseWheel::WhenScrollingPossible )
+		, _ConsumeMouseWheel(EConsumeMouseWheel::WhenScrollingPossible)
 		, _WheelScrollMultiplier(GetGlobalScrollAmount())
-		, _HandleGamepadEvents( true )
+		, _HandleGamepadEvents(true)
+		, _OnItemToString_Debug()
+		, _OnEnteredBadState()
 		{}
 
 		SLATE_EVENT( FOnGenerateRow, OnGenerateTile )
@@ -96,6 +100,11 @@ public:
 
 		SLATE_ARGUMENT( bool, HandleGamepadEvents );
 
+		/** Assign this to get more diagnostics from the list view. */
+		SLATE_EVENT(FOnItemToString_Debug, OnItemToString_Debug)
+
+		SLATE_EVENT(FOnTableViewBadState, OnEnteredBadState);
+
 	SLATE_END_ARGS()
 
 	/**
@@ -125,6 +134,11 @@ public:
 
 		this->bHandleGamepadEvents = InArgs._HandleGamepadEvents;
 
+		this->OnItemToString_Debug = InArgs._OnItemToString_Debug.IsBound()
+			? InArgs._OnItemToString_Debug
+			: SListView< ItemType >::GetDefaultDebugDelegate();
+		this->OnEnteredBadState = InArgs._OnEnteredBadState;
+
 		// Check for any parameters that the coder forgot to specify.
 		FString ErrorString;
 		{
diff --git a/Engine/Source/Runtime/Slate/Public/Widgets/Views/STreeView.h b/Engine/Source/Runtime/Slate/Public/Widgets/Views/STreeView.h
--- a/Engine/Source/Runtime/Slate/Public/Widgets/Views/STreeView.h
+++ b/Engine/Source/Runtime/Slate/Public/Widgets/Views/STreeView.h
@@ -84,6 +84,8 @@ public:
 	typedef typename TSlateDelegates< ItemType >::FOnMouseButtonDoubleClick FOnMouseButtonDoubleClick;
 	typedef typename TSlateDelegates< ItemType >::FOnExpansionChanged FOnExpansionChanged;
 
+	typedef typename TSlateDelegates< ItemType >::FOnItemToString_Debug FOnItemToString_Debug; 
+
 	using FOnWidgetToBeRemoved = typename SListView<ItemType>::FOnWidgetToBeRemoved;
 
 public:
@@ -104,6 +106,8 @@ public:
 		, _ConsumeMouseWheel( EConsumeMouseWheel::WhenScrollingPossible )
 		, _AllowOverscroll(EAllowOverscroll::Yes)
 		, _WheelScrollMultiplier(GetGlobalScrollAmount())
+		, _OnItemToString_Debug()
+		, _OnEnteredBadState()
 		{}
 
 		SLATE_EVENT( FOnGenerateRow, OnGenerateRow )
@@ -144,6 +148,11 @@ public:
 
 		SLATE_ARGUMENT( float, WheelScrollMultiplier );
 
+		/** Assign this to get more diagnostics from the list view. */
+		SLATE_EVENT(FOnItemToString_Debug, OnItemToString_Debug)
+
+		SLATE_EVENT(FOnTableViewBadState, OnEnteredBadState);
+
 	SLATE_END_ARGS()
 
 		
@@ -173,6 +182,11 @@ public:
 
 		this->WheelScrollMultiplier = InArgs._WheelScrollMultiplier;
 
+		this->OnItemToString_Debug = InArgs._OnItemToString_Debug.IsBound()
+			? InArgs._OnItemToString_Debug
+			: SListView< ItemType >::GetDefaultDebugDelegate();
+		this->OnEnteredBadState = InArgs._OnEnteredBadState;
+
 		// Check for any parameters that the coder forgot to specify.
 		FString ErrorString;
 		{
