Moving SVirtualWindow into Slate module

diff --git a/Engine/Source/Runtime/Slate/Private/Widgets/SVirtualWindow.cpp b/Engine/Source/Runtime/Slate/Private/Widgets/SVirtualWindow.cpp
new file mode 100644
--- /dev/null
+++ b/Engine/Source/Runtime/Slate/Private/Widgets/SVirtualWindow.cpp
@@ -0,0 +1,100 @@
+// Copyright 1998-2017 Epic Games, Inc. All Rights Reserved.
+
+#include "Widgets/SVirtualWindow.h"
+#include "Layout/ArrangedChildren.h"
+#include "Misc/App.h"
+#include "Modules/ModuleManager.h"
+#include "Rendering/SlateDrawBuffer.h"
+
+#include "Widgets/LayerManager/STooltipPresenter.h"
+#include "Widgets/Layout/SPopup.h"
+
+void SVirtualWindow::Construct(const FArguments& InArgs)
+{
+	bIsPopupWindow = true;
+	bVirtualWindow = true;
+	bIsFocusable = false;
+	bShouldResolveDeferred = true;
+	SetCachedSize(InArgs._Size);
+	SetNativeWindow(MakeShareable(new FGenericWindow()));
+
+	ConstructWindowInternals();
+
+	WindowOverlay->AddSlot()
+		[
+			SNew(SPopup)
+			[
+				SAssignNew(TooltipPresenter, STooltipPresenter)
+			]
+		];
+
+	SetContent(SNullWidget::NullWidget);
+}
+
+FPopupMethodReply SVirtualWindow::OnQueryPopupMethod() const
+{
+	return FPopupMethodReply::UseMethod(EPopupMethod::UseCurrentWindow)
+		.SetShouldThrottle(EShouldThrottle::No);
+}
+
+bool SVirtualWindow::OnVisualizeTooltip(const TSharedPtr<SWidget>& TooltipContent)
+{
+	TooltipPresenter->SetContent(TooltipContent.IsValid() ? TooltipContent.ToSharedRef() : SNullWidget::NullWidget);
+
+	return true;
+}
+
+void SVirtualWindow::SetShouldResolveDeferred(bool bResolve)
+{
+	bShouldResolveDeferred = bResolve;
+}
+
+void SVirtualWindow::SetIsFocusable(bool bFocusable)
+{
+	bIsFocusable = bFocusable;
+}
+
+bool SVirtualWindow::SupportsKeyboardFocus() const
+{
+	return bIsFocusable;
+}
+
+FVector2D SVirtualWindow::ComputeDesiredSize(float LayoutScaleMultiplier) const
+{
+	return SCompoundWidget::ComputeDesiredSize(LayoutScaleMultiplier);
+}
+
+int32 SVirtualWindow::OnPaint(const FPaintArgs& Args, const FGeometry& AllottedGeometry, const FSlateRect& MyClippingRect, FSlateWindowElementList& OutDrawElements, int32 LayerId, const FWidgetStyle& InWidgetStyle, bool bParentEnabled) const
+{
+	if (bShouldResolveDeferred)
+	{
+		OutDrawElements.BeginDeferredGroup();
+	}
+
+	// We intentionally skip SWindow's OnPaint, since we are going to do our own handling of deferred groups.
+	int32 MaxLayer = SCompoundWidget::OnPaint(Args, AllottedGeometry, MyClippingRect, OutDrawElements, LayerId, InWidgetStyle, bParentEnabled);
+
+	if (bShouldResolveDeferred)
+	{
+		OutDrawElements.EndDeferredGroup();
+	}
+
+	return MaxLayer;
+}
+
+void SVirtualWindow::OnArrangeChildren(const FGeometry& AllottedGeometry, FArrangedChildren& ArrangedChildren) const
+{
+	SWindow::OnArrangeChildren(AllottedGeometry, ArrangedChildren);
+
+	// @HACK VREDITOR - otherwise popup layers don't work in nested child windows, in tab managers and such.
+	if (ArrangedChildren.Allows3DWidgets())
+	{
+		const TArray< TSharedRef<SWindow> >& WindowChildren = GetChildWindows();
+		for (int32 ChildIndex = 0; ChildIndex < WindowChildren.Num(); ++ChildIndex)
+		{
+			const TSharedRef<SWindow>& ChildWindow = WindowChildren[ChildIndex];
+			FGeometry ChildWindowGeometry = ChildWindow->GetWindowGeometryInWindow();
+			ChildWindow->ArrangeChildren(ChildWindowGeometry, ArrangedChildren);
+		}
+	}
+}
\ No newline at end of file
diff --git a/Engine/Source/Runtime/Slate/Public/Widgets/SVirtualWindow.h b/Engine/Source/Runtime/Slate/Public/Widgets/SVirtualWindow.h
new file mode 100644
--- /dev/null
+++ b/Engine/Source/Runtime/Slate/Public/Widgets/SVirtualWindow.h
@@ -0,0 +1,58 @@
+// Copyright 1998-2017 Epic Games, Inc. All Rights Reserved.
+
+#pragma once
+
+
+#include "CoreMinimal.h"
+#include "Layout/SlateRect.h"
+#include "Layout/Geometry.h"
+#include "Input/PopupMethodReply.h"
+#include "Widgets/SWidget.h"
+#include "Widgets/DeclarativeSyntaxSupport.h"
+#include "Input/HittestGrid.h"
+#include "Rendering/DrawElements.h"
+#include "Widgets/SWindow.h"
+
+class FArrangedChildren;
+class FSlateDrawBuffer;
+class ISlate3DRenderer;
+class STooltipPresenter;
+class UTextureRenderTarget2D;
+
+class SLATE_API SVirtualWindow : public SWindow
+{
+	SLATE_BEGIN_ARGS(SVirtualWindow)
+		: _Size(FVector2D(100, 100))
+	{}
+
+	SLATE_ARGUMENT(FVector2D, Size)
+
+		SLATE_END_ARGS()
+
+public:
+	void Construct(const FArguments& InArgs);
+
+	/**
+	* We allow users to control the resolve policy for deferred content.  When a virtual window is used in a retainer, we don't want it trying to resolve
+	* the deferred content, we instead need it to pass that content up to the non-retained window drawing the retainer.  If a virtual window is used in
+	* a case like the WidgetComponent, we always want to resolve the deferred content because there won't be another opportunity since the 3D widget has
+	* no owner window that can draw it.
+	*/
+	void SetShouldResolveDeferred(bool bResolve);
+
+	/**  */
+	void SetIsFocusable(bool bFocusable);
+
+	virtual FPopupMethodReply OnQueryPopupMethod() const override;
+	virtual bool OnVisualizeTooltip(const TSharedPtr<SWidget>& TooltipContent) override;
+	virtual int32 OnPaint(const FPaintArgs& Args, const FGeometry& AllottedGeometry, const FSlateRect& MyClippingRect, FSlateWindowElementList& OutDrawElements, int32 LayerId, const FWidgetStyle& InWidgetStyle, bool bParentEnabled) const override;
+	virtual void OnArrangeChildren(const FGeometry& AllottedGeometry, FArrangedChildren& ArrangedChildren) const override;
+	virtual bool SupportsKeyboardFocus() const override;
+	virtual FVector2D ComputeDesiredSize(float LayoutScaleMultiplier) const override;
+
+private:
+	TSharedPtr<class STooltipPresenter> TooltipPresenter;
+	bool bIsFocusable;
+	bool bShouldResolveDeferred;
+};
+
diff --git a/Engine/Source/Runtime/UMG/Private/Slate/WidgetRenderer.cpp b/Engine/Source/Runtime/UMG/Private/Slate/WidgetRenderer.cpp
--- a/Engine/Source/Runtime/UMG/Private/Slate/WidgetRenderer.cpp
+++ b/Engine/Source/Runtime/UMG/Private/Slate/WidgetRenderer.cpp
@@ -17,96 +17,6 @@
 
 extern SLATECORE_API int32 bFoldTick;
 
-void SVirtualWindow::Construct(const FArguments& InArgs)
-{
-	bIsPopupWindow = true;
-	bVirtualWindow = true;
-	bIsFocusable = false;
-	bShouldResolveDeferred = true;
-	SetCachedSize(InArgs._Size);
-	SetNativeWindow(MakeShareable(new FGenericWindow()));
-
-	ConstructWindowInternals();
-
-	WindowOverlay->AddSlot()
-	[
-		SNew(SPopup)
-		[
-			SAssignNew(TooltipPresenter, STooltipPresenter)
-		]
-	];
-
-	SetContent(SNullWidget::NullWidget);
-}
-
-FPopupMethodReply SVirtualWindow::OnQueryPopupMethod() const
-{
-	return FPopupMethodReply::UseMethod(EPopupMethod::UseCurrentWindow)
-		.SetShouldThrottle(EShouldThrottle::No);
-}
-
-bool SVirtualWindow::OnVisualizeTooltip(const TSharedPtr<SWidget>& TooltipContent)
-{
-	TooltipPresenter->SetContent(TooltipContent.IsValid() ? TooltipContent.ToSharedRef() : SNullWidget::NullWidget);
-
-	return true;
-}
-
-void SVirtualWindow::SetShouldResolveDeferred(bool bResolve)
-{
-	bShouldResolveDeferred = bResolve;
-}
-
-void SVirtualWindow::SetIsFocusable(bool bFocusable)
-{
-	bIsFocusable = bFocusable;
-}
-
-bool SVirtualWindow::SupportsKeyboardFocus() const
-{
-	return bIsFocusable;
-}
-
-FVector2D SVirtualWindow::ComputeDesiredSize(float LayoutScaleMultiplier) const
-{
-	return SCompoundWidget::ComputeDesiredSize(LayoutScaleMultiplier);
-}
-
-int32 SVirtualWindow::OnPaint(const FPaintArgs& Args, const FGeometry& AllottedGeometry, const FSlateRect& MyClippingRect, FSlateWindowElementList& OutDrawElements, int32 LayerId, const FWidgetStyle& InWidgetStyle, bool bParentEnabled) const
-{
-	if ( bShouldResolveDeferred )
-	{
-		OutDrawElements.BeginDeferredGroup();
-	}
-	
-	// We intentionally skip SWindow's OnPaint, since we are going to do our own handling of deferred groups.
-	int32 MaxLayer = SCompoundWidget::OnPaint(Args, AllottedGeometry, MyClippingRect, OutDrawElements, LayerId, InWidgetStyle, bParentEnabled);
-
-	if ( bShouldResolveDeferred )
-	{
-		OutDrawElements.EndDeferredGroup();
-	}
-
-	return MaxLayer;
-}
-
-void SVirtualWindow::OnArrangeChildren(const FGeometry& AllottedGeometry, FArrangedChildren& ArrangedChildren) const
-{
-	SWindow::OnArrangeChildren(AllottedGeometry, ArrangedChildren);
-
-	// @HACK VREDITOR - otherwise popup layers don't work in nested child windows, in tab managers and such.
-	if ( ArrangedChildren.Allows3DWidgets() )
-	{
-		const TArray< TSharedRef<SWindow> >& WindowChildren = GetChildWindows();
-		for ( int32 ChildIndex=0; ChildIndex < WindowChildren.Num(); ++ChildIndex )
-		{
-			const TSharedRef<SWindow>& ChildWindow = WindowChildren[ChildIndex];
-			FGeometry ChildWindowGeometry = ChildWindow->GetWindowGeometryInWindow();
-			ChildWindow->ArrangeChildren(ChildWindowGeometry, ArrangedChildren);
-		}
-	}
-}
-
 FWidgetRenderer::FWidgetRenderer(bool bUseGammaCorrection, bool bInClearTarget)
 	: bPrepassNeeded(true)
 	, bUseGammaSpace(bUseGammaCorrection)
diff --git a/Engine/Source/Runtime/UMG/Public/Slate/WidgetRenderer.h b/Engine/Source/Runtime/UMG/Public/Slate/WidgetRenderer.h
--- a/Engine/Source/Runtime/UMG/Public/Slate/WidgetRenderer.h
+++ b/Engine/Source/Runtime/UMG/Public/Slate/WidgetRenderer.h
@@ -12,6 +12,7 @@
 #include "Rendering/DrawElements.h"
 #include "Engine/Texture.h"
 #include "Widgets/SWindow.h"
+#include "Widgets/SVirtualWindow.h"
 
 class FArrangedChildren;
 class FSlateDrawBuffer;
@@ -19,43 +20,6 @@ class ISlate3DRenderer;
 class STooltipPresenter;
 class UTextureRenderTarget2D;
 
-class UMG_API SVirtualWindow : public SWindow
-{
-	SLATE_BEGIN_ARGS(SVirtualWindow)
-		: _Size(FVector2D(100, 100))
-	{}
-
-	SLATE_ARGUMENT(FVector2D, Size)
-	
-	SLATE_END_ARGS()
-
-public:
-	void Construct(const FArguments& InArgs);
-
-	/**
-	 * We allow users to control the resolve policy for deferred content.  When a virtual window is used in a retainer, we don't want it trying to resolve
-	 * the deferred content, we instead need it to pass that content up to the non-retained window drawing the retainer.  If a virtual window is used in 
-	 * a case like the WidgetComponent, we always want to resolve the deferred content because there won't be another opportunity since the 3D widget has
-	 * no owner window that can draw it.
-	 */
-	void SetShouldResolveDeferred(bool bResolve);
-
-	/**  */
-	void SetIsFocusable(bool bFocusable);
-	
-	virtual FPopupMethodReply OnQueryPopupMethod() const override;
-	virtual bool OnVisualizeTooltip(const TSharedPtr<SWidget>& TooltipContent) override;
-	virtual int32 OnPaint(const FPaintArgs& Args, const FGeometry& AllottedGeometry, const FSlateRect& MyClippingRect, FSlateWindowElementList& OutDrawElements, int32 LayerId, const FWidgetStyle& InWidgetStyle, bool bParentEnabled) const override;
-	virtual void OnArrangeChildren(const FGeometry& AllottedGeometry, FArrangedChildren& ArrangedChildren) const override;
-	virtual bool SupportsKeyboardFocus() const override;
-	virtual FVector2D ComputeDesiredSize(float LayoutScaleMultiplier) const override;
-
-private:
-	TSharedPtr<class STooltipPresenter> TooltipPresenter;
-	bool bIsFocusable;
-	bool bShouldResolveDeferred;
-};
-
 /**
  * 
  */
