Fix non-unity builds.

diff --git a/Engine/Source/Runtime/Slate/Public/Framework/Views/TreeFilterHandler.h b/Engine/Source/Runtime/Slate/Public/Framework/Views/TreeFilterHandler.h
--- a/Engine/Source/Runtime/Slate/Public/Framework/Views/TreeFilterHandler.h
+++ b/Engine/Source/Runtime/Slate/Public/Framework/Views/TreeFilterHandler.h
@@ -2,6 +2,9 @@
 
 #pragma once
 
+#include "CoreMinimal.h"
+#include "Widgets/Views/STreeView.h"
+
 /** Handles filtering and expanding a TreeView using an IFilter implementation. */
 template<typename ItemType>
 class TreeFilterHandler
