Copying //UE4/Dev-VREditor to //UE4/Dev-Main (Source: //UE4/Dev-VREditor @ 3352034)

diff --git a/Engine/Source/Runtime/Core/Private/Misc/FbxErrors.cpp b/Engine/Source/Runtime/Core/Private/Misc/FbxErrors.cpp
--- a/Engine/Source/Runtime/Core/Private/Misc/FbxErrors.cpp
+++ b/Engine/Source/Runtime/Core/Private/Misc/FbxErrors.cpp
@@ -79,8 +79,6 @@ FName FFbxErrors::SkeletalMesh_RestoreSortingForSectionNumber(TEXT("SkeletalMesh
 
 FName FFbxErrors::SkeletalMesh_NoMeshFoundOnRoot(TEXT("SkeletalMesh_NoMeshFoundOnRoot"));
 
-FName FFbxErrors::SkeletalMesh_ImportSubDSurface(TEXT("SkeletalMesh_ImportSubDSurface"));
-
 FName FFbxErrors::SkeletalMesh_InvalidRoot(TEXT("SkeletalMesh_InvalidRoot"));
 
 FName FFbxErrors::SkeletalMesh_InvalidBone(TEXT("SkeletalMesh_InvalidBone"));
diff --git a/Engine/Source/Runtime/Core/Private/UObject/DevObjectVersion.cpp b/Engine/Source/Runtime/Core/Private/UObject/DevObjectVersion.cpp
--- a/Engine/Source/Runtime/Core/Private/UObject/DevObjectVersion.cpp
+++ b/Engine/Source/Runtime/Core/Private/UObject/DevObjectVersion.cpp
@@ -14,6 +14,7 @@
 #include "UObject/RenderingObjectVersion.h"
 #include "UObject/SequencerObjectVersion.h"
 #include "UObject/VRObjectVersion.h"
+#include "UObject/GeometryObjectVersion.h"
 
 DEFINE_LOG_CATEGORY_STATIC(LogDevObjectVersion, Log, All);
 
@@ -108,3 +109,10 @@ FDevVersionRegistration GRegisterVRObjectVersion(FVRObjectVersion::GUID, FVRObje
 const FGuid FLoadTimesObjectVersion::GUID(0xC2A15278, 0xBFE74AFE, 0x6C1790FF, 0x531DF755);
 // Register LoadTimes custom version with Core
 FDevVersionRegistration GRegisterLoadTimesObjectVersion(FLoadTimesObjectVersion::GUID, FLoadTimesObjectVersion::LatestVersion, TEXT("Dev-LoadTimes"));
+
+
+// Unique Geometry Object version id
+const FGuid FGeometryObjectVersion::GUID(0x6EACA3D4, 0x40EC4CC1, 0xb7868BED, 0x9428FC5);
+// Register Geometry custom version with Core
+FDevVersionRegistration GRegisterGeometryObjectVersion(FGeometryObjectVersion::GUID, FGeometryObjectVersion::LatestVersion, TEXT("Private-Geometry"));
+
diff --git a/Engine/Source/Runtime/Core/Public/Misc/FbxErrors.h b/Engine/Source/Runtime/Core/Public/Misc/FbxErrors.h
--- a/Engine/Source/Runtime/Core/Public/Misc/FbxErrors.h
+++ b/Engine/Source/Runtime/Core/Public/Misc/FbxErrors.h
@@ -91,8 +91,6 @@ struct CORE_API FFbxErrors
 
 	static FName SkeletalMesh_NoMeshFoundOnRoot;
 	
-	static FName SkeletalMesh_ImportSubDSurface;
-
 	static FName SkeletalMesh_InvalidRoot;
 
 	static FName SkeletalMesh_InvalidBone;
diff --git a/Engine/Source/Runtime/Engine/Classes/Engine/StaticMesh.h b/Engine/Source/Runtime/Engine/Classes/Engine/StaticMesh.h
--- a/Engine/Source/Runtime/Engine/Classes/Engine/StaticMesh.h
+++ b/Engine/Source/Runtime/Engine/Classes/Engine/StaticMesh.h
@@ -430,7 +430,7 @@ class UStaticMesh : public UObject, public IInterface_CollisionDataProvider, pub
 	FMeshSectionInfoMap SectionInfoMap;
 
 	/** The LOD group to which this mesh belongs. */
-	UPROPERTY(AssetRegistrySearchable)
+	UPROPERTY(EditAnywhere, AssetRegistrySearchable, Category=StaticMesh)
 	FName LODGroup;
 
 	/** If true, the screen sizees at which LODs swap are computed automatically. */
@@ -612,7 +612,7 @@ public:
 	ENGINE_API virtual void PreEditChange(UProperty* PropertyAboutToChange) override;
 	ENGINE_API virtual void PostEditChangeProperty(FPropertyChangedEvent& PropertyChangedEvent) override;
 	ENGINE_API virtual void GetAssetRegistryTagMetadata(TMap<FName, FAssetRegistryTagMetadata>& OutMetadata) const override;
-	ENGINE_API void SetLODGroup(FName NewGroup);
+	ENGINE_API void SetLODGroup(FName NewGroup, bool bRebuildImmediately = true);
 	ENGINE_API void BroadcastNavCollisionChange();
 #endif // WITH_EDITOR
 	ENGINE_API virtual void Serialize(FArchive& Ar) override;
diff --git a/Engine/Source/Runtime/Engine/Engine.Build.cs b/Engine/Source/Runtime/Engine/Engine.Build.cs
--- a/Engine/Source/Runtime/Engine/Engine.Build.cs
+++ b/Engine/Source/Runtime/Engine/Engine.Build.cs
@@ -101,7 +101,7 @@ public class Engine : ModuleRules
                 "MaterialShaderQualitySettings",
                 "CinematicCamera",
 				"Analytics",
-				"AnalyticsET",
+				"AnalyticsET"
             }
 		);
 
@@ -358,8 +358,7 @@ public class Engine : ModuleRules
 				"UEOgg",
 				"Vorbis",
 				"VorbisFile",
-				"libOpus",
-			    "OpenSubdiv"
+				"libOpus"
 				);
 
             // Head Mounted Display support
diff --git a/Engine/Source/Runtime/Engine/Private/DrawDebugHelpers.cpp b/Engine/Source/Runtime/Engine/Private/DrawDebugHelpers.cpp
--- a/Engine/Source/Runtime/Engine/Private/DrawDebugHelpers.cpp
+++ b/Engine/Source/Runtime/Engine/Private/DrawDebugHelpers.cpp
@@ -436,8 +436,8 @@ void DrawDebug2DDonut(const UWorld* InWorld, const FMatrix& TransformMatrix, flo
 
 			// Need at least 4 segments
 			Segments = FMath::Max((Segments - 4) / 2, 4);
-			InternalDrawDebugCircle(InWorld, TransformMatrix, InnerRadius, Segments, Color, bPersistentLines, LifeTime, DepthPriority);
-			InternalDrawDebugCircle(InWorld, TransformMatrix, OuterRadius, Segments, Color, bPersistentLines, LifeTime, DepthPriority);
+			InternalDrawDebugCircle(InWorld, TransformMatrix, InnerRadius, Segments, Color, bPersistentLines, LifeTime, DepthPriority, Thickness);
+			InternalDrawDebugCircle(InWorld, TransformMatrix, OuterRadius, Segments, Color, bPersistentLines, LifeTime, DepthPriority, Thickness );
 		
 			const FVector Center = TransformMatrix.GetOrigin();
 			const FVector AxisY = TransformMatrix.GetScaledAxis( EAxis::Y );
diff --git a/Engine/Source/Runtime/Engine/Private/InstancedStaticMesh.cpp b/Engine/Source/Runtime/Engine/Private/InstancedStaticMesh.cpp
--- a/Engine/Source/Runtime/Engine/Private/InstancedStaticMesh.cpp
+++ b/Engine/Source/Runtime/Engine/Private/InstancedStaticMesh.cpp
@@ -2048,6 +2048,7 @@ void FInstancedStaticMeshVertexFactoryShaderParameters::SetMesh( FRHICommandList
 			for (int32 SampleIndex = 0; SampleIndex < 2; SampleIndex++)
 			{
 				FVector4& InstancingViewZCompare(SampleIndex ? InstancingViewZCompareOne : InstancingViewZCompareZero);
+				float Fac = View.GetTemporalLODDistanceFactor(SampleIndex) * SphereRadius * LODScale;
 
 				float FinalCull = MAX_flt;
 				if (MinSize > 0.0)
@@ -2058,6 +2059,7 @@ void FInstancedStaticMeshVertexFactoryShaderParameters::SetMesh( FRHICommandList
 				{
 					FinalCull = FMath::Min(FinalCull, InstancingUserData->EndCullDistance * MaxDrawDistanceScale);
 				}
+				FinalCull *= MaxDrawDistanceScale;
 
 				InstancingViewZCompare.Z = FinalCull;
 				if (BatchElement.InstancedLODIndex < InstancingUserData->MeshRenderData->LODResources.Num() - 1)
diff --git a/Engine/Source/Runtime/Engine/Private/SceneManagement.cpp b/Engine/Source/Runtime/Engine/Private/SceneManagement.cpp
--- a/Engine/Source/Runtime/Engine/Private/SceneManagement.cpp
+++ b/Engine/Source/Runtime/Engine/Private/SceneManagement.cpp
@@ -395,7 +395,7 @@ int8 ComputeStaticMeshLOD( const FStaticMeshRenderData* RenderData, const FVecto
 	{
 		const int32 NumLODs = MAX_STATIC_MESH_LODS;
 
-		const float ScreenRadiusSquared = ComputeBoundsScreenRadiusSquared(Origin, SphereRadius, View) * FactorScale * FactorScale * View.LODDistanceFactor * View.LODDistanceFactor;
+		const float ScreenRadiusSquared = ComputeBoundsScreenRadiusSquared(Origin, SphereRadius, View) * FactorScale * FactorScale;
 
 		// Walk backwards and return the first matching LOD
 		for (int32 LODIndex = NumLODs - 1; LODIndex >= 0; --LODIndex)
diff --git a/Engine/Source/Runtime/Engine/Private/Slate/SceneViewport.cpp b/Engine/Source/Runtime/Engine/Private/Slate/SceneViewport.cpp
--- a/Engine/Source/Runtime/Engine/Private/Slate/SceneViewport.cpp
+++ b/Engine/Source/Runtime/Engine/Private/Slate/SceneViewport.cpp
@@ -1570,6 +1570,20 @@ void FSceneViewport::Tick( const FGeometry& AllottedGeometry, double InCurrentTi
 
 void FSceneViewport::OnPlayWorldViewportSwapped( const FSceneViewport& OtherViewport )
 {
+	// We need to call WindowRenderTargetUpdate() to make sure the Slate renderer is updated to render
+	// to the viewport client we'll be using for PIE/SIE.  Otherwise if stereo rendering is enabled, Slate
+	// could render the HMD mirror to a game viewport client which is not visible on screen!
+	TSharedPtr<SWidget> PinnedViewport = ViewportWidget.Pin();
+	if( PinnedViewport.IsValid() )
+	{
+		TSharedPtr<FSlateRenderer> Renderer = FSlateApplication::Get().GetRenderer();
+
+		FWidgetPath WidgetPath;
+		TSharedPtr<SWindow> Window = FSlateApplication::Get().FindWidgetWindow( PinnedViewport.ToSharedRef(), WidgetPath );
+
+		WindowRenderTargetUpdate( Renderer.Get(), Window.Get() );
+	}
+
 	// Play world viewports should always be the same size.  Resize to other viewports size
 	if( GetSizeXY() != OtherViewport.GetSizeXY() )
 	{
diff --git a/Engine/Source/Runtime/Engine/Private/StaticMesh.cpp b/Engine/Source/Runtime/Engine/Private/StaticMesh.cpp
--- a/Engine/Source/Runtime/Engine/Private/StaticMesh.cpp
+++ b/Engine/Source/Runtime/Engine/Private/StaticMesh.cpp
@@ -1671,6 +1671,14 @@ void UStaticMesh::PostEditChangeProperty(FPropertyChangedEvent& PropertyChangedE
 	UProperty* PropertyThatChanged = PropertyChangedEvent.Property;
 	const FName PropertyName = PropertyThatChanged ? PropertyThatChanged->GetFName() : NAME_None;
 	
+	if (PropertyName == GET_MEMBER_NAME_CHECKED(UStaticMesh, LODGroup))
+	{
+		// Force an update of LOD group settings
+
+		// Dont rebuild inside here.  We're doing that below.
+		bool bRebuild = false;
+		SetLODGroup(LODGroup, bRebuild);
+	}
 	LightMapResolution = FMath::Max(LightMapResolution, 0);
 
 	if (PropertyChangedEvent.MemberProperty 
@@ -1714,7 +1722,7 @@ void UStaticMesh::PostEditChangeProperty(FPropertyChangedEvent& PropertyChangedE
 	Super::PostEditChangeProperty(PropertyChangedEvent);
 }
 
-void UStaticMesh::SetLODGroup(FName NewGroup)
+void UStaticMesh::SetLODGroup(FName NewGroup, bool bRebuildImmediately)
 {
 #if WITH_EDITORONLY_DATA
 	const bool bBeforeDerivedDataCached = (RenderData == nullptr);
@@ -1754,6 +1762,11 @@ void UStaticMesh::SetLODGroup(FName NewGroup)
 	if (!bBeforeDerivedDataCached)
 	{
 		bAutoComputeLODScreenSize = true;
+		
+	}
+
+	if (bRebuildImmediately && !bBeforeDerivedDataCached)
+	{
 		PostEditChange();
 	}
 #endif
@@ -2502,6 +2515,26 @@ void UStaticMesh::PostLoad()
 	{
 		CalculateExtendedBounds();
 	}
+
+	if (SectionInfoMap.Map.Num() == 0)
+	{
+		// Before this serialization issue was fixed, some assets were resaved and permanently lost their section info map.
+		// This attempts to recreate it based on the render data.
+		SectionInfoMap.Clear();
+		for (int32 LODResourceIndex = 0; LODResourceIndex < RenderData->LODResources.Num(); ++LODResourceIndex)
+		{
+			FStaticMeshLODResources& LOD = RenderData->LODResources[LODResourceIndex];
+			const int32 NumSections = LOD.Sections.Num();
+			for (int32 SectionIndex = 0; SectionIndex < NumSections; ++SectionIndex)
+			{
+				const int32 MaterialIndex = LOD.Sections[SectionIndex].MaterialIndex;
+				if (StaticMaterials.IsValidIndex(MaterialIndex))
+				{
+					SectionInfoMap.Set(LODResourceIndex, SectionIndex, FMeshSectionInfo(MaterialIndex));
+				}
+			}
+		}
+	}
 #endif // #if WITH_EDITOR
 
 	// We want to always have a BodySetup, its used for per-poly collision as well
@@ -2510,6 +2543,7 @@ void UStaticMesh::PostLoad()
 		CreateBodySetup();
 	}
 
+
 	CreateNavCollision();
 }
 
diff --git a/Engine/Source/Runtime/Engine/Public/Engine.h b/Engine/Source/Runtime/Engine/Public/Engine.h
--- a/Engine/Source/Runtime/Engine/Public/Engine.h
+++ b/Engine/Source/Runtime/Engine/Public/Engine.h
@@ -110,7 +110,6 @@ MONOLITHIC_HEADER_BOILERPLATE()
 #include "Curves/CurveFloat.h"
 #include "GameFramework/ForceFeedbackEffect.h"
 #include "Engine/SubsurfaceProfile.h"
-#include "Engine/SubDSurface.h"
 #include "Engine/DebugCameraController.h"
 #include "Engine/DecalActor.h"
 #include "CanvasItem.h"
diff --git a/Engine/Source/Runtime/GameLiveStreaming/GameLiveStreaming.cpp b/Engine/Source/Runtime/GameLiveStreaming/GameLiveStreaming.cpp
--- a/Engine/Source/Runtime/GameLiveStreaming/GameLiveStreaming.cpp
+++ b/Engine/Source/Runtime/GameLiveStreaming/GameLiveStreaming.cpp
@@ -448,7 +448,7 @@ void FGameLiveStreaming::StartCopyingNextGameVideoFrame( const FViewportRHIRef&
 		SetRenderTarget(RHICmdList, DestRenderTarget.TargetableTexture, FTextureRHIRef());
 		RHICmdList.SetViewport(0, 0, 0.0f, Context.ResizeTo.X, Context.ResizeTo.Y, 1.0f);
 
-		RHICmdList.SetBlendState(TStaticBlendState<>::GetRHI());
+		RHICmdList.SetBlendState(TStaticBlendState<CW_RGB>::GetRHI());
 		RHICmdList.SetRasterizerState(TStaticRasterizerState<>::GetRHI());
 		RHICmdList.SetDepthStencilState(TStaticDepthStencilState<false,CF_Always>::GetRHI());
 
diff --git a/Engine/Source/Runtime/GeometryCache/Classes/GeometryCacheMeshData.h b/Engine/Source/Runtime/GeometryCache/Classes/GeometryCacheMeshData.h
--- a/Engine/Source/Runtime/GeometryCache/Classes/GeometryCacheMeshData.h
+++ b/Engine/Source/Runtime/GeometryCache/Classes/GeometryCacheMeshData.h
@@ -5,7 +5,7 @@
 #include "UObject/ObjectMacros.h"
 #include "ProfilingDebugging/ResourceSize.h"
 #include "DynamicMeshBuilder.h"
-
+#include "UObject/GeometryObjectVersion.h"
 #include "GeometryCacheMeshData.generated.h"
 
 /** Stores per-batch data used for rendering */
@@ -56,6 +56,8 @@ struct FGeometryCacheMeshData
 	/** Serialization for FVertexAnimationSample. */
 	friend FArchive& operator<<(FArchive& Ar, FGeometryCacheMeshData& Mesh)
 	{
+		Ar.UsingCustomVersion(FGeometryObjectVersion::GUID);
+
 		int32 NumVertices = 0;
 		
 		if (Ar.IsSaving())
@@ -69,19 +71,26 @@ struct FGeometryCacheMeshData
 			Mesh.Vertices.AddUninitialized(NumVertices);
 		}
 
-		for (int32 VertexIndex = 0; VertexIndex < NumVertices; ++VertexIndex)
+		if (Ar.CustomVer(FGeometryObjectVersion::GUID) < FGeometryObjectVersion::CompressGeometryCache)
+		{
+			for (int32 VertexIndex = 0; VertexIndex < NumVertices; ++VertexIndex)
+			{
+				FDynamicMeshVertex& Vertex = Mesh.Vertices[VertexIndex];
+				Ar << Vertex.Position;
+				Ar << Vertex.TextureCoordinate;
+				Ar << Vertex.TangentX;
+				Ar << Vertex.TangentZ;
+				Ar << Vertex.Color;
+			}
+		}
+		else if(NumVertices)
 		{
-			FDynamicMeshVertex& Vertex = Mesh.Vertices[VertexIndex];
-			Ar << Vertex.Position;
-			Ar << Vertex.TextureCoordinate;
-			Ar << Vertex.TangentX;
-			Ar << Vertex.TangentZ;
-			Ar << Vertex.Color;
+			Ar.SerializeCompressed(&Mesh.Vertices[0], Mesh.Vertices.Num()*Mesh.Vertices.GetTypeSize(), COMPRESS_ZLIB);
 		}
 
+	
 		Ar << Mesh.BoundingBox;
 		Ar << Mesh.BatchesInfo;
-
 		Ar << Mesh.Indices;	
 
 		return Ar;
diff --git a/Engine/Source/Runtime/HeadMountedDisplay/Private/HeadMountedDisplayTypes.cpp b/Engine/Source/Runtime/HeadMountedDisplay/Private/HeadMountedDisplayTypes.cpp
--- a/Engine/Source/Runtime/HeadMountedDisplay/Private/HeadMountedDisplayTypes.cpp
+++ b/Engine/Source/Runtime/HeadMountedDisplay/Private/HeadMountedDisplayTypes.cpp
@@ -26,7 +26,17 @@ FHMDViewMesh::~FHMDViewMesh()
 
 void FHMDViewMesh::BuildMesh(const FVector2D Positions[], uint32 VertexCount, EHMDMeshType MeshType)
 {
-	check(pVertices == nullptr);
+	if( pVertices != nullptr )
+	{
+		delete[] pVertices;
+		pVertices = nullptr;
+	}
+	if( pIndices != nullptr )
+	{
+		delete[] pIndices;
+		pIndices = nullptr;
+	}
+
 	check(VertexCount > 2 && VertexCount % 3 == 0);
 
 	NumVertices = VertexCount;
diff --git a/Engine/Source/Runtime/Launch/Private/LaunchEngineLoop.cpp b/Engine/Source/Runtime/Launch/Private/LaunchEngineLoop.cpp
--- a/Engine/Source/Runtime/Launch/Private/LaunchEngineLoop.cpp
+++ b/Engine/Source/Runtime/Launch/Private/LaunchEngineLoop.cpp
@@ -3034,15 +3034,28 @@ void FEngineLoop::Tick()
 
 		// @todo vreditor urgent: Temporary hack to allow world-to-meters to be set before
 		// input is polled for motion controller devices each frame.
-#if WITH_ENGINE
 		extern ENGINE_API float GNewWorldToMetersScale;
-		UWorld* HackGWorld = GWorld;
-		if( GNewWorldToMetersScale != 0.0f && HackGWorld != nullptr )
+		if( GNewWorldToMetersScale != 0.0f  )
 		{
-			if( GNewWorldToMetersScale != HackGWorld->GetWorldSettings()->WorldToMeters )
+#if WITH_ENGINE
+			UWorld* WorldToScale = GWorld;
+
+#if WITH_EDITOR
+			if( GIsEditor && GEditor->PlayWorld != nullptr && GEditor->bIsSimulatingInEditor )
 			{
-				HackGWorld->GetWorldSettings()->WorldToMeters = GNewWorldToMetersScale;
+				WorldToScale = GEditor->PlayWorld;
 			}
+#endif	// WITH_EDITOR
+
+			if( WorldToScale != nullptr )
+			{
+				if( GNewWorldToMetersScale != WorldToScale->GetWorldSettings()->WorldToMeters )
+				{
+					WorldToScale->GetWorldSettings()->WorldToMeters = GNewWorldToMetersScale;
+				}
+			}
+
+			GNewWorldToMetersScale = 0.0f;
 		}
 #endif	// WITH_ENGINE
 
diff --git a/Engine/Source/Runtime/Renderer/Private/PostProcess/PostProcessing.cpp b/Engine/Source/Runtime/Renderer/Private/PostProcess/PostProcessing.cpp
--- a/Engine/Source/Runtime/Renderer/Private/PostProcess/PostProcessing.cpp
+++ b/Engine/Source/Runtime/Renderer/Private/PostProcess/PostProcessing.cpp
@@ -1782,7 +1782,6 @@ void FPostProcessing::Process(FRHICommandListImmediate& RHICmdList, const FViewI
 		if ( GIsEditor
 			&& View.Family->EngineShowFlags.SelectionOutline
 			&& !(View.Family->EngineShowFlags.Wireframe)
-			&& ( !GIsDemoMode || ( GIsDemoMode && !View.Family->EngineShowFlags.Game ) ) 
 			&& !bVisualizeBloom
 			&& !View.Family->EngineShowFlags.VisualizeHDR)
 		{
@@ -2410,7 +2409,6 @@ void FPostProcessing::ProcessES2(FRHICommandListImmediate& RHICmdList, const FVi
 		if ( GIsEditor
 			&& View.Family->EngineShowFlags.SelectionOutline
 			&& !(View.Family->EngineShowFlags.Wireframe)
-			&& ( !GIsDemoMode || ( GIsDemoMode && !View.Family->EngineShowFlags.Game ) ) 
 			)
 		{
 			// Editor selection outline
diff --git a/Engine/Source/Runtime/Renderer/Private/VelocityRendering.cpp b/Engine/Source/Runtime/Renderer/Private/VelocityRendering.cpp
--- a/Engine/Source/Runtime/Renderer/Private/VelocityRendering.cpp
+++ b/Engine/Source/Runtime/Renderer/Private/VelocityRendering.cpp
@@ -39,6 +39,11 @@ static TAutoConsoleVariable<int32> CVarRHICmdVelocityPassDeferredContexts(
 	1,
 	TEXT("True to use deferred contexts to parallelize velocity pass command list execution."));
 
+RENDERER_API TAutoConsoleVariable<int32> CVarAllowMotionBlurInVR(
+	TEXT("vr.AllowMotionBlurInVR"),
+	0,
+	TEXT("For projects with motion blur enabled, this allows motion blur to be enabled even while in VR."));
+
 DECLARE_FLOAT_COUNTER_STAT(TEXT("Render Velocities"), Stat_GPU_RenderVelocities, STATGROUP_GPU);
 
 bool IsParallelVelocity()
@@ -607,7 +612,7 @@ bool IsMotionBlurEnabled(const FViewInfo& View)
 		&& View.Family->bRealtimeUpdate
 		&& MotionBlurQuality > 0
 		&& !View.bIsSceneCapture
-		&& !(View.Family->Views.Num() > 1);
+		&& (CVarAllowMotionBlurInVR->GetInt() != 0 || !(View.Family->Views.Num() > 1));
 }
 
 void FDeferredShadingSceneRenderer::RenderDynamicVelocitiesMeshElementsInner(FRHICommandList& RHICmdList, const FViewInfo& View, const FDrawingPolicyRenderState& DrawRenderState, int32 FirstIndex, int32 LastIndex)
diff --git a/Engine/Source/Runtime/Slate/Private/Framework/Application/SlateApplication.cpp b/Engine/Source/Runtime/Slate/Private/Framework/Application/SlateApplication.cpp
--- a/Engine/Source/Runtime/Slate/Private/Framework/Application/SlateApplication.cpp
+++ b/Engine/Source/Runtime/Slate/Private/Framework/Application/SlateApplication.cpp
@@ -2969,7 +2969,7 @@ void FSlateApplication::ProcessReply( const FWidgetPath& CurrentEventPath, const
 	// are all operations that we shouldn't do if our application isn't Active (The OS ignores half of it, and we'd be in a half state)
 	// We do allow the release of capture and lock when deactivated, this is innocuous of some platforms but required on others when 
 	// the Application deactivated before the window. (Mac is an example of this)
-	if (bAppIsActive)
+	if (bAppIsActive || bIsVirtualInteraction)
 	{
 		TSharedPtr<SWidget> RequestedMouseCaptor = TheReply.GetMouseCaptor();
 		// Do not capture the mouse if we are also starting a drag and drop.
@@ -4816,6 +4816,12 @@ bool FSlateApplication::ProcessMouseButtonDownEvent( const TSharedPtr< FGenericW
 	}
 	PressedMouseButtons.Add( MouseEvent.GetEffectingButton() );
 
+	// Input preprocessor gets first chance at the input
+	if( InputPreProcessor.IsValid() && InputPreProcessor->HandleMouseButtonDownEvent( *this, MouseEvent ) )
+	{
+		return true;
+	}
+
 	bool bInGame = false;
 
 	// Only process mouse down messages if we are not drag/dropping
@@ -4993,6 +4999,15 @@ FReply FSlateApplication::RoutePointerUpEvent(FWidgetPath& WidgetsUnderPointer,
 	TSharedPtr<SWindow> TopLevelWindow;
 #endif
 
+	if ( DragDetector.DetectDragForWidget.IsValid() )
+	{
+		if ( PointerEvent.GetEffectingButton() == DragDetector.DetectDragButton && PointerEvent.GetPointerIndex() == DragDetector.DetectDragPointerIndex )
+		{
+			// The user has released the button (or finger) that was supposed to start the drag; stop detecting it.
+			DragDetector = FDragDetector();
+		}
+	}
+
 	if (MouseCaptor.HasCaptureForPointerIndex(PointerEvent.GetUserIndex(), PointerEvent.GetPointerIndex()))
 	{
 		//FWidgetPath MouseCaptorPath = MouseCaptor.ToWidgetPath(PointerEvent.GetPointerIndex());
@@ -5101,7 +5116,14 @@ bool FSlateApplication::RoutePointerMoveEvent(const FWidgetPath& WidgetsUnderPoi
 
 	// User asked us to detect a drag.
 	bool bDragDetected = false;
-	if( DragDetector.DetectDragForWidget.IsValid() && !bIsSynthetic )
+	bool bShouldStartDetectingDrag = true;
+
+	//@TODO VREDITOR - Remove and move to interaction component
+	if (OnDragDropCheckOverride.IsBound())
+	{
+		bShouldStartDetectingDrag = OnDragDropCheckOverride.Execute();
+	}
+	if( DragDetector.DetectDragForWidget.IsValid() && !bIsSynthetic && bShouldStartDetectingDrag)
 	{	
 		const FVector2D DragDelta = (DragDetector.DetectDragStartLocation - PointerEvent.GetScreenSpacePosition());
 		bDragDetected = ( DragDelta.SizeSquared() > FMath::Square(FSlateApplication::Get().GetDragTriggerDistance()) );
@@ -5307,6 +5329,13 @@ bool FSlateApplication::RoutePointerMoveEvent(const FWidgetPath& WidgetsUnderPoi
 	if (IsDragDropping())
 	{
 		FDragDropEvent DragDropEvent( PointerEvent, DragDropContent );
+		//@TODO VREDITOR - Remove and move to interaction component
+		if (OnDragDropCheckOverride.IsBound() && DragDropEvent.GetOperation().IsValid())
+		{
+			DragDropEvent.GetOperation()->SetDecoratorVisibility(false);
+			DragDropEvent.GetOperation()->SetCursorOverride(EMouseCursor::None);
+			DragDropContent->SetCursorOverride(EMouseCursor::None);
+		}
 		FScopedSwitchWorldHack SwitchWorld( WidgetsUnderPointer );
 		DragDropContent->OnDragged( DragDropEvent );
 
@@ -5458,13 +5487,10 @@ bool FSlateApplication::ProcessMouseButtonUpEvent( FPointerEvent& MouseEvent )
 	LastUserInteractionTimeForThrottling = LastUserInteractionTime;
 	PressedMouseButtons.Remove( MouseEvent.GetEffectingButton() );
 
-	if ( DragDetector.DetectDragForWidget.IsValid() )
+	// Input preprocessor gets first chance at the input
+	if( InputPreProcessor.IsValid() && InputPreProcessor->HandleMouseButtonUpEvent( *this, MouseEvent ) )
 	{
-		if ( MouseEvent.GetEffectingButton() == DragDetector.DetectDragButton && MouseEvent.GetPointerIndex() == DragDetector.DetectDragPointerIndex )
-		{
-			// The user has released the button (or finger) that was supposed to start the drag; stop detecting it.
-			DragDetector = FDragDetector();
-		}
+		return true;
 	}
 
 	// An empty widget path is passed in.  As an optimization, one will be generated only if a captured mouse event isn't routed
diff --git a/Engine/Source/Runtime/Slate/Private/Framework/Docking/TabManager.cpp b/Engine/Source/Runtime/Slate/Private/Framework/Docking/TabManager.cpp
--- a/Engine/Source/Runtime/Slate/Private/Framework/Docking/TabManager.cpp
+++ b/Engine/Source/Runtime/Slate/Private/Framework/Docking/TabManager.cpp
@@ -1153,11 +1153,10 @@ TSharedRef<SDockTab> FTabManager::InvokeTab_Internal( const FTabId& TabId )
 			->AddTab( TabId, ETabState::OpenedTab )
 		);
 
-		TSharedRef<SDockingArea> DockingArea = RestoreArea( NewAreaForTab, GetPrivateApi().GetParentWindow() );
-				
+		TSharedRef<SDockingArea> DockingArea = RestoreArea(NewAreaForTab, GetPrivateApi().GetParentWindow());
 		const TSharedPtr<SDockTab> NewlyOpenedTab = DockingArea->GetAllChildTabs()[0];
 		check(NewlyOpenedTab.IsValid());
-		return NewlyOpenedTab.ToSharedRef();
+		return NewlyOpenedTab.ToSharedRef();	
 	}
 }
 
@@ -1913,7 +1912,10 @@ void FGlobalTabmanager::DrawAttentionToTabManager( const TSharedRef<FTabManager>
 		// #HACK VREDITOR
 		if ( ProxyTabManager.IsValid() )
 		{
-			ProxyTabManager->DrawAttention(Tab.ToSharedRef());
+			if( ProxyTabManager->IsTabSupported( Tab->GetLayoutIdentifier() ) )
+			{
+				ProxyTabManager->DrawAttention(Tab.ToSharedRef());
+			}
 		}
 	}
 }
@@ -2121,7 +2123,7 @@ void FGlobalTabmanager::UpdateStats()
 
 void FGlobalTabmanager::OpenUnmanagedTab(FName PlaceholderId, const FSearchPreference& SearchPreference, const TSharedRef<SDockTab>& UnmanagedTab)
 {
-	if ( ProxyTabManager.IsValid() )
+	if ( ProxyTabManager.IsValid() && ProxyTabManager->IsTabSupported( UnmanagedTab->GetLayoutIdentifier() ) )
 	{
 		ProxyTabManager->OpenUnmanagedTab(PlaceholderId, SearchPreference, UnmanagedTab);
 	}
@@ -2136,29 +2138,43 @@ void FGlobalTabmanager::SetProxyTabManager(TSharedPtr<FProxyTabmanager> InProxyT
 	ProxyTabManager = InProxyTabManager;
 }
 
-void FProxyTabmanager::OpenUnmanagedTab(FName PlaceholderId, const FSearchPreference& SearchPreference, const TSharedRef<SDockTab>& UnmanagedTab)
+bool FProxyTabmanager::IsTabSupported( const FTabId TabId ) const
 {
-	// No layout info about this tab found; start 
-	TSharedRef<FArea> NewAreaForTab = FTabManager::NewPrimaryArea()
-		->Split
-		(
-			FTabManager::NewStack()
-			->AddTab(UnmanagedTab->GetLayoutIdentifier(), ETabState::OpenedTab)
-		);
+	bool bIsTabSupported = true;
+	if( OnIsTabSupported.IsBound() )
+	{
+		OnIsTabSupported.Broadcast( TabId, /* In/Out */ bIsTabSupported );
+	}
+
+	return bIsTabSupported;
+}
 
+void FProxyTabmanager::OpenUnmanagedTab(FName PlaceholderId, const FSearchPreference& SearchPreference, const TSharedRef<SDockTab>& UnmanagedTab)
+{
 	TSharedPtr<SWindow> ParentWindowPtr = ParentWindow.Pin();
-	TSharedRef<SDockingArea> DockingArea = RestoreArea(NewAreaForTab, ParentWindowPtr, false);
-	ParentWindowPtr->SetContent(DockingArea);
+	if( ensure( ParentWindowPtr.IsValid() ) )
+	{
+		// No layout info about this tab found; start 
+		TSharedRef<FArea> NewAreaForTab = FTabManager::NewPrimaryArea()
+			->Split
+			(
+				FTabManager::NewStack()
+				->AddTab( UnmanagedTab->GetLayoutIdentifier(), ETabState::OpenedTab )
+			);
+
+		TSharedRef<SDockingArea> DockingArea = RestoreArea(NewAreaForTab, ParentWindowPtr, false);
+		ParentWindowPtr->SetContent(DockingArea);
 
-	const TSharedPtr<SDockTab> NewlyOpenedTab = DockingArea->GetAllChildTabs()[0];
-	check(NewlyOpenedTab.IsValid());
+		const TSharedPtr<SDockTab> NewlyOpenedTab = DockingArea->GetAllChildTabs()[0];
+		check(NewlyOpenedTab.IsValid());
 		
-	NewlyOpenedTab->GetParent()->GetParentDockTabStack()->OpenTab(UnmanagedTab);
-	NewlyOpenedTab->RequestCloseTab();
+		NewlyOpenedTab->GetParent()->GetParentDockTabStack()->OpenTab(UnmanagedTab);
+		NewlyOpenedTab->RequestCloseTab();
 
-	MainNonCloseableTab = UnmanagedTab;
+		MainNonCloseableTab = UnmanagedTab;
 
-	OnTabOpened.Broadcast(UnmanagedTab);
+		OnTabOpened.Broadcast(UnmanagedTab);
+	}
 }
 
 void FProxyTabmanager::DrawAttention(const TSharedRef<SDockTab>& TabToHighlight)
@@ -2169,4 +2185,9 @@ void FProxyTabmanager::DrawAttention(const TSharedRef<SDockTab>& TabToHighlight)
 }
 
 
+void FProxyTabmanager::SetParentWindow(TSharedRef<SWindow> InParentWindow)
+{
+	ParentWindow = InParentWindow;
+}
+
 #undef LOCTEXT_NAMESPACE
diff --git a/Engine/Source/Runtime/Slate/Private/Framework/MultiBox/SMenuEntryBlock.cpp b/Engine/Source/Runtime/Slate/Private/Framework/MultiBox/SMenuEntryBlock.cpp
--- a/Engine/Source/Runtime/Slate/Private/Framework/MultiBox/SMenuEntryBlock.cpp
+++ b/Engine/Source/Runtime/Slate/Private/Framework/MultiBox/SMenuEntryBlock.cpp
@@ -493,7 +493,8 @@ TSharedRef< SWidget > SMenuEntryBlock::BuildMenuEntryWidget( const FMenuEntryBui
 				EVisibility::Hidden;
 
 	// Collapse (rather than Hide) the checkbox when using a custom menu widget with a button action, otherwise we add additional padding around the user defined widget
-	if (MenuEntryBlock->EntryWidget.IsValid() && UserInterfaceType == EUserInterfaceActionType::Button)
+	if ((MenuEntryBlock->EntryWidget.IsValid() && UserInterfaceType == EUserInterfaceActionType::Button) ||
+		UserInterfaceType == EUserInterfaceActionType::CollapsedButton)
 	{
 		CheckBoxVisibility = EVisibility::Collapsed;
 	}
diff --git a/Engine/Source/Runtime/Slate/Private/Framework/Notifications/NotificationManager.cpp b/Engine/Source/Runtime/Slate/Private/Framework/Notifications/NotificationManager.cpp
--- a/Engine/Source/Runtime/Slate/Private/Framework/Notifications/NotificationManager.cpp
+++ b/Engine/Source/Runtime/Slate/Private/Framework/Notifications/NotificationManager.cpp
@@ -63,6 +63,7 @@ void FSlateNotificationManager::FRegionalNotificationList::Arrange()
 }
 
 FSlateNotificationManager::FSlateNotificationManager()
+	: bAllowNotifications( true )
 {
 }
 
diff --git a/Engine/Source/Runtime/Slate/Private/Widgets/Layout/SScissorRectBox.cpp b/Engine/Source/Runtime/Slate/Private/Widgets/Layout/SScissorRectBox.cpp
--- a/Engine/Source/Runtime/Slate/Private/Widgets/Layout/SScissorRectBox.cpp
+++ b/Engine/Source/Runtime/Slate/Private/Widgets/Layout/SScissorRectBox.cpp
@@ -4,6 +4,12 @@
 #include "Types/PaintArgs.h"
 #include "Layout/ArrangedChildren.h"
 #include "Rendering/RenderingCommon.h"
+#include "IConsoleManager.h"
+
+namespace Slate
+{
+	static FAutoConsoleVariable OverrideScissorRect( TEXT( "Slate.OverrideScissorRect" ), 1, TEXT( "Whether to allow Slate to apply a scissor clip to UI elements to prevent certain artifacts.  You might need to disable this if you are drawing UI over a mirror backbuffer while in VR." ) );
+}
 
 SScissorRectBox::SScissorRectBox()
 {
@@ -56,19 +62,26 @@ int32 SScissorRectBox::OnPaint(const FPaintArgs& Args, const FGeometry& Allotted
 		check( ArrangedChildren.Num() == 1 );
 		FArrangedWidget& TheChild = ArrangedChildren[0];
 
-		// Child clipping rect is made effectively infinite to ensure that any rotated children do not get prematurely cut off. We let the scissor rect do all the work.
-		const FSlateRect ChildClippingRect = AllottedGeometry.GetClippingRect().IntersectionWith(MyClippingRect);
-
 		// set the hacky, back door global scissor rect settings.
 		// This will get picked up by the element batcher and added to the draw element.
 		// We do it this way because this scissor rect is a real hack until better solution is made available.
 		// In the meantime, we don't want to modify the public interface of all element drawing to ensure all widgets pick up this scissor rect.
 		extern SLATECORE_API TOptional<FShortRect> GSlateScissorRect;
 		TOptional<FShortRect> OldRect = GSlateScissorRect;
-		GSlateScissorRect = FShortRect(ChildClippingRect);
+		const bool bOverrideScissorRect = Slate::OverrideScissorRect->GetInt() != 0;
+		if( bOverrideScissorRect )
+		{
+			// Child clipping rect is made effectively infinite to ensure that any rotated children do not get prematurely cut off. We let the scissor rect do all the work.
+			const FSlateRect ChildClippingRect = AllottedGeometry.GetClippingRect().IntersectionWith( MyClippingRect );
+			GSlateScissorRect = FShortRect(ChildClippingRect);
+		}
 		int32 Result = TheChild.Widget->Paint( Args.WithNewParent(this), TheChild.Geometry, FSlateRect(FVector2D(-FLT_MAX, -FLT_MAX), FVector2D(FLT_MAX, FLT_MAX)), OutDrawElements, LayerId, InWidgetStyle, ShouldBeEnabled( bParentEnabled ) );
-		// We've added all the draw elements, so put the scissor rect back
-		GSlateScissorRect = OldRect;
+
+		if( bOverrideScissorRect )
+		{
+			// We've added all the draw elements, so put the scissor rect back
+			GSlateScissorRect = OldRect;
+		}
 		
 		return Result;
 	}
diff --git a/Engine/Source/Runtime/Slate/Private/Widgets/Notifications/SNotificationList.cpp b/Engine/Source/Runtime/Slate/Private/Widgets/Notifications/SNotificationList.cpp
--- a/Engine/Source/Runtime/Slate/Private/Widgets/Notifications/SNotificationList.cpp
+++ b/Engine/Source/Runtime/Slate/Private/Widgets/Notifications/SNotificationList.cpp
@@ -15,6 +15,7 @@
 #include "Widgets/Input/SHyperlink.h"
 #include "Widgets/Images/SThrobber.h"
 #include "Widgets/Notifications/INotificationWidget.h"
+#include "Framework/Notifications/NotificationManager.h"
 
 /////////////////////////////////////////////////
 // SNotificationExtendable
@@ -682,54 +683,64 @@ TSharedRef<SNotificationItem> SNotificationList::AddNotification(const FNotifica
 {
 	TSharedPtr<SNotificationExtendable> NewItem;
 
-	if (Info.ContentWidget.IsValid())
+	if (FSlateNotificationManager::Get().AreNotificationsAllowed())
 	{
-		NewItem = SNew(SNotificationItemExternalImpl)
-			.ContentWidget(Info.ContentWidget)
-			.FadeInDuration(Info.FadeInDuration)
-			.ExpireDuration(Info.ExpireDuration)
-			.FadeOutDuration(Info.FadeOutDuration);
-	}
-	else
-	{
-		static const FSlateBrush* CachedImage = FCoreStyle::Get().GetBrush("NotificationList.DefaultMessage");
-
-		// Create notification.
-		NewItem = SNew(SNotificationItemImpl)
-			.Text(Info.Text)
-			.Font(Font)
-			.ButtonDetails(Info.ButtonDetails)
-			.Image((Info.Image != nullptr) ? Info.Image : CachedImage)
-			.FadeInDuration(Info.FadeInDuration)
-			.ExpireDuration(Info.ExpireDuration)
-			.FadeOutDuration(Info.FadeOutDuration)
-			.bUseThrobber(Info.bUseThrobber)
-			.bUseSuccessFailIcons(Info.bUseSuccessFailIcons)
-			.bUseLargeFont(Info.bUseLargeFont)
-			.WidthOverride(Info.WidthOverride)
-			.CheckBoxState(Info.CheckBoxState)
-			.CheckBoxStateChanged(Info.CheckBoxStateChanged)
-			.CheckBoxText(Info.CheckBoxText)
-			.Hyperlink(Info.Hyperlink)
-			.HyperlinkText(Info.HyperlinkText);
-	}
+		if (Info.ContentWidget.IsValid())
+		{
+			NewItem = SNew(SNotificationItemExternalImpl)
+				.ContentWidget(Info.ContentWidget)
+				.FadeInDuration(Info.FadeInDuration)
+				.ExpireDuration(Info.ExpireDuration)
+				.FadeOutDuration(Info.FadeOutDuration);
+		}
+		else
+		{
+			static const FSlateBrush* CachedImage = FCoreStyle::Get().GetBrush("NotificationList.DefaultMessage");
+
+			// Create notification.
+			NewItem = SNew(SNotificationItemImpl)
+				.Text(Info.Text)
+				.Font(Font)
+				.ButtonDetails(Info.ButtonDetails)
+				.Image((Info.Image != nullptr) ? Info.Image : CachedImage)
+				.FadeInDuration(Info.FadeInDuration)
+				.ExpireDuration(Info.ExpireDuration)
+				.FadeOutDuration(Info.FadeOutDuration)
+				.bUseThrobber(Info.bUseThrobber)
+				.bUseSuccessFailIcons(Info.bUseSuccessFailIcons)
+				.bUseLargeFont(Info.bUseLargeFont)
+				.WidthOverride(Info.WidthOverride)
+				.CheckBoxState(Info.CheckBoxState)
+				.CheckBoxStateChanged(Info.CheckBoxStateChanged)
+				.CheckBoxText(Info.CheckBoxText)
+				.Hyperlink(Info.Hyperlink)
+				.HyperlinkText(Info.HyperlinkText);
+		}
 		
-	NewItem->MyList = SharedThis(this);
+		NewItem->MyList = SharedThis(this);
 
-	MessageItemBoxPtr->AddSlot()
-		.AutoHeight()
-		.HAlign(HAlign_Right)
-		[
-			NewItem.ToSharedRef()
-		];
+		MessageItemBoxPtr->AddSlot()
+			.AutoHeight()
+			.HAlign(HAlign_Right)
+			[
+				NewItem.ToSharedRef()
+			];
 
-	NewItem->Fadein( Info.bAllowThrottleWhenFrameRateIsLow );
+		NewItem->Fadein( Info.bAllowThrottleWhenFrameRateIsLow );
 
-	if (Info.bFireAndForget)
+		if (Info.bFireAndForget)
+		{
+			NewItem->ExpireAndFadeout();
+		}
+	}
+	else
 	{
-		NewItem->ExpireAndFadeout();
+		// When notifications are not allowed we want to return an empty notification.
+		NewItem = SNew(SNotificationItemImpl);
 	}
 
+	LastNotification = NewItem;
+
 	return NewItem.ToSharedRef();
 }
 
diff --git a/Engine/Source/Runtime/Slate/Private/Widgets/SViewport.cpp b/Engine/Source/Runtime/Slate/Private/Widgets/SViewport.cpp
--- a/Engine/Source/Runtime/Slate/Private/Widgets/SViewport.cpp
+++ b/Engine/Source/Runtime/Slate/Private/Widgets/SViewport.cpp
@@ -14,6 +14,7 @@ DECLARE_CYCLE_STAT(TEXT("Game UI Tick/Paint"), STAT_ViewportUpdateTime, STATGROU
 SViewport::SViewport()
 	: bRenderDirectlyToWindow(false)
 	, bEnableGammaCorrection(true)
+	, bReverseGammaCorrection(false)
 	, bEnableStereoRendering(false)
 { }
 
@@ -26,6 +27,7 @@ void SViewport::Construct( const FArguments& InArgs )
 	ShowDisabledEffect = InArgs._ShowEffectWhenDisabled;
 	bRenderDirectlyToWindow = InArgs._RenderDirectlyToWindow;
 	bEnableGammaCorrection = InArgs._EnableGammaCorrection;
+	bReverseGammaCorrection = InArgs._ReverseGammaCorrection;
 	bEnableBlending = InArgs._EnableBlending;
 	bEnableStereoRendering = InArgs._EnableStereoRendering;
 	bIgnoreTextureAlpha = InArgs._IgnoreTextureAlpha;
@@ -80,6 +82,12 @@ int32 SViewport::OnPaint( const FPaintArgs& Args, const FGeometry& AllottedGeome
 		DrawEffects |= ESlateDrawEffect::NoGamma;
 	}
 
+	// Should we reverse gamma correction
+	if (bReverseGammaCorrection)
+	{
+		DrawEffects |= ESlateDrawEffect::ReverseGamma;
+	}
+
 	// Show we enable blending?
 	if( !bEnableBlending )
 	{
diff --git a/Engine/Source/Runtime/Slate/Public/Framework/Application/SlateApplication.h b/Engine/Source/Runtime/Slate/Public/Framework/Application/SlateApplication.h
--- a/Engine/Source/Runtime/Slate/Public/Framework/Application/SlateApplication.h
+++ b/Engine/Source/Runtime/Slate/Public/Framework/Application/SlateApplication.h
@@ -48,6 +48,8 @@ DECLARE_DELEGATE(FModalWindowStackEnded)
 /** Delegate for when window action occurs (ClickedNonClientArea, Maximize, Restore, WindowMenu). Return true if the OS layer should stop processing the action. */
 DECLARE_DELEGATE_RetVal_TwoParams(bool, FOnWindowAction, const TSharedRef<FGenericWindow>&, EWindowAction::Type);
 
+DECLARE_DELEGATE_RetVal(bool, FDragDropCheckingOverride);
+
 extern SLATE_API const FName NAME_UnrealOS;
 
 
@@ -1463,6 +1465,9 @@ public:
 	*/
 	const void* FindBestParentWindowHandleForDialogs(const TSharedPtr<SWidget>& InWidget);
 
+public:
+	FDragDropCheckingOverride OnDragDropCheckOverride;
+
 private:
 
 	TSharedRef< FGenericWindow > MakeWindow( TSharedRef<SWindow> InSlateWindow, const bool bShowImmediately );
diff --git a/Engine/Source/Runtime/Slate/Public/Framework/Docking/TabManager.h b/Engine/Source/Runtime/Slate/Public/Framework/Docking/TabManager.h
--- a/Engine/Source/Runtime/Slate/Public/Framework/Docking/TabManager.h
+++ b/Engine/Source/Runtime/Slate/Public/Framework/Docking/TabManager.h
@@ -912,6 +912,10 @@ public:
 
 	void SetProxyTabManager(TSharedPtr<FProxyTabmanager> InProxyTabManager);
 
+	DECLARE_DELEGATE(FOnOverrideDockableAreaRestore);
+	/** Used to override dockable area restoration behavior */
+	FOnOverrideDockableAreaRestore OnOverrideDockableAreaRestore_Handler;
+
 protected:
 	virtual void OnTabForegrounded( const TSharedPtr<SDockTab>& NewForegroundTab, const TSharedPtr<SDockTab>& BackgroundedTab ) override;
 	virtual void OnTabRelocated( const TSharedRef<SDockTab>& RelocatedTab, const TSharedPtr<SWindow>& NewOwnerWindow ) override;
@@ -1005,6 +1009,7 @@ private:
 
 //#HACK VREDITOR - Had to introduce the proxy tab manager to steal asset tabs.
 
+DECLARE_MULTICAST_DELEGATE_TwoParams(FIsTabSupportedEvent, FTabId /* TabId */, bool& /* bOutIsSupported */ );
 DECLARE_MULTICAST_DELEGATE_OneParam(FOnTabEvent, TSharedPtr<SDockTab>);
 
 
@@ -1021,8 +1026,12 @@ public:
 
 	virtual void OpenUnmanagedTab(FName PlaceholderId, const FSearchPreference& SearchPreference, const TSharedRef<SDockTab>& UnmanagedTab) override;
 	virtual void DrawAttention(const TSharedRef<SDockTab>& TabToHighlight) override;
+	
+	bool IsTabSupported( const FTabId TabId ) const;
+	void SetParentWindow(TSharedRef<SWindow> InParentWindow);
 
 public:
+	FIsTabSupportedEvent OnIsTabSupported;
 	FOnTabEvent OnTabOpened;
 	FOnTabEvent OnAttentionDrawnToTab;
 
diff --git a/Engine/Source/Runtime/Slate/Public/Framework/Notifications/NotificationManager.h b/Engine/Source/Runtime/Slate/Public/Framework/Notifications/NotificationManager.h
--- a/Engine/Source/Runtime/Slate/Public/Framework/Notifications/NotificationManager.h
+++ b/Engine/Source/Runtime/Slate/Public/Framework/Notifications/NotificationManager.h
@@ -57,6 +57,21 @@ public:
 	 */
 	void GetWindows(TArray< TSharedRef<SWindow> >& OutWindows) const;
 
+	/**
+	 * Sets whether notifications should be displayed at all 
+	 * @param	bShouldAllow	Whether notifications should be enabled.  It defaults to on.
+	 */
+	void SetAllowNotifications( const bool bShouldAllow )
+	{
+		this->bAllowNotifications = bShouldAllow;
+	}
+
+	/** @return	Checks whether notifications are currently enabled */
+	bool AreNotificationsAllowed() const
+	{
+		return this->bAllowNotifications;
+	}
+
 protected:
 	/** Protect constructor as this is a singleton */
 	FSlateNotificationManager();
@@ -96,4 +111,7 @@ private:
 
 	/** Thread safe queue of notifications to display */
 	TLockFreePointerListLIFO<FNotificationInfo> PendingNotifications;
+
+	/** Whether notifications should be displayed or not.  This can be used to globally suppress notification pop-ups */
+	bool bAllowNotifications;
 };
diff --git a/Engine/Source/Runtime/Slate/Public/Widgets/Notifications/SNotificationList.h b/Engine/Source/Runtime/Slate/Public/Widgets/Notifications/SNotificationList.h
--- a/Engine/Source/Runtime/Slate/Public/Widgets/Notifications/SNotificationList.h
+++ b/Engine/Source/Runtime/Slate/Public/Widgets/Notifications/SNotificationList.h
@@ -259,4 +259,10 @@ protected:
 
 	/** Flag to auto-destroy this list. */
 	bool bDone;
+
+private:
+
+	/** The last notification added to the list. To avoid weak pointer pointing to a destroyed notification when notification are not allowed. */
+	TSharedPtr<SNotificationItem> LastNotification;
+
 };
diff --git a/Engine/Source/Runtime/Slate/Public/Widgets/SViewport.h b/Engine/Source/Runtime/Slate/Public/Widgets/SViewport.h
--- a/Engine/Source/Runtime/Slate/Public/Widgets/SViewport.h
+++ b/Engine/Source/Runtime/Slate/Public/Widgets/SViewport.h
@@ -29,6 +29,7 @@ public:
 		, _ShowEffectWhenDisabled(true)
 		, _RenderDirectlyToWindow(false)
 		, _EnableGammaCorrection(true)
+		, _ReverseGammaCorrection(false)
 		, _EnableBlending(false)
 		, _EnableStereoRendering(false)
 		, _PreMultipliedAlpha(true)
@@ -51,6 +52,9 @@ public:
 		/** Whether or not to enable gamma correction. Doesn't apply when rendering directly to a backbuffer. */
 		SLATE_ARGUMENT( bool, EnableGammaCorrection )
 
+		/** Whether or not to reverse the gamma correction done to the texture in this viewport.  Ignores the bEnableGammaCorrection setting */
+		SLATE_ARGUMENT( bool, ReverseGammaCorrection )
+
 		/** Allow this viewport to blend with its background. */
 		SLATE_ARGUMENT( bool, EnableBlending )
 
@@ -257,6 +261,9 @@ private:
 	/** Whether or not to apply gamma correction on the render target supplied by the ISlateViewport. */
 	bool bEnableGammaCorrection;
 
+	/** Whether or not to reverse the gamma correction done by the texture in this viewport.  Ignores the bEnableGammaCorrection setting */
+	bool bReverseGammaCorrection;
+
 	/** Whether or not to blend this viewport with the background. */
 	bool bEnableBlending;
 
diff --git a/Engine/Source/Runtime/SlateCore/Private/Rendering/DrawElements.cpp b/Engine/Source/Runtime/SlateCore/Private/Rendering/DrawElements.cpp
--- a/Engine/Source/Runtime/SlateCore/Private/Rendering/DrawElements.cpp
+++ b/Engine/Source/Runtime/SlateCore/Private/Rendering/DrawElements.cpp
@@ -89,6 +89,10 @@ void FSlateDrawElement::Init(uint32 InLayer, const FPaintGeometry& PaintGeometry
 			DrawEffects |= ESlateDrawEffect::NoPixelSnapping;
 		}
 	}
+	if ((InDrawEffects & ESlateDrawEffect::ReverseGamma) != ESlateDrawEffect::None)
+	{
+		DataPayload.BatchFlags |= ESlateBatchDrawFlag::ReverseGamma;
+	}
 }
 
 void FSlateDrawElement::ApplyPositionOffset(FSlateDrawElement& Element, const FVector2D& InOffset)
diff --git a/Engine/Source/Runtime/SlateCore/Public/Rendering/RenderingCommon.h b/Engine/Source/Runtime/SlateCore/Public/Rendering/RenderingCommon.h
--- a/Engine/Source/Runtime/SlateCore/Public/Rendering/RenderingCommon.h
+++ b/Engine/Source/Runtime/SlateCore/Public/Rendering/RenderingCommon.h
@@ -81,6 +81,9 @@ enum class ESlateDrawEffect : uint8
 	DisabledEffect		= 1 << 5,
 	/** Advanced: Don't read from texture alpha channel */
 	IgnoreTextureAlpha	= 1 << 6,
+
+	/** Advanced: Existing Gamma correction should be reversed */
+	ReverseGamma			= 1 << 7
 };
 
 ENUM_CLASS_FLAGS(ESlateDrawEffect);
@@ -106,7 +109,9 @@ enum class ESlateBatchDrawFlag : uint8
 	/** The element should be tiled horizontally */
 	TileU				= 1 << 5,
 	/** The element should be tiled vertically */
-	TileV				= 1 << 6
+	TileV				= 1 << 6,
+	/** Reverse gamma correction */
+	ReverseGamma			 = 1 << 7
 };
 
 ENUM_CLASS_FLAGS(ESlateBatchDrawFlag);
diff --git a/Engine/Source/Runtime/SlateRHIRenderer/Private/SlateRHIRenderer.cpp b/Engine/Source/Runtime/SlateRHIRenderer/Private/SlateRHIRenderer.cpp
--- a/Engine/Source/Runtime/SlateRHIRenderer/Private/SlateRHIRenderer.cpp
+++ b/Engine/Source/Runtime/SlateRHIRenderer/Private/SlateRHIRenderer.cpp
@@ -53,6 +53,13 @@ static TAutoConsoleVariable<int32> CVarUICompositeMode(
 	TEXT("1: Shader pass to improve HDR blending\n"),
 	ECVF_RenderThreadSafe);
 
+static TAutoConsoleVariable<float> CVarDrawToVRRenderTarget(
+	TEXT("Slate.DrawToVRRenderTarget"),
+	1,
+	TEXT("If enabled while in VR. Slate UI will be drawn into the render target texture where the VR imagery for either eye was rendered, allow the viewer of the HMD to see the UI (for better or worse.)  This render target will then be cropped/scaled into the back buffer, if mirroring is enabled.  When disabled, Slate UI will be drawn on top of the backbuffer (not to the HMD) after the mirror texture has been cropped/scaled into the backbuffer."),
+	ECVF_RenderThreadSafe);
+
+
 void FSlateCrashReportResource::InitDynamicRHI()
 {
 	int32 Width = VirtualScreen.Width();
@@ -651,6 +658,13 @@ void FSlateRHIRenderer::DrawWindow_RenderThread(FRHICommandListImmediate& RHICmd
 	ViewportInfo.ColorSpaceLUTOutputDevice = HDROutputDevice;
 	ViewportInfo.ColorSpaceLUTOutputGamut = HDROutputGamut;
 	
+	bool bRenderedStereo = false;
+	if( CVarDrawToVRRenderTarget->GetInt() == 0 && GEngine && IsValidRef( ViewportInfo.GetRenderTargetTexture() ) && GEngine->StereoRenderingDevice.IsValid() )
+	{
+		GEngine->StereoRenderingDevice->RenderTexture_RenderThread( RHICmdList, RHICmdList.GetViewportBackBuffer( ViewportInfo.ViewportRHI ), ViewportInfo.GetRenderTargetTexture() );
+		bRenderedStereo = true;
+	}
+
 	{
 		SCOPED_GPU_STAT(RHICmdList, Stat_GPU_SlateUI);
 		SCOPE_CYCLE_COUNTER( STAT_SlateRenderingRTTime );
@@ -674,9 +688,8 @@ void FSlateRHIRenderer::DrawWindow_RenderThread(FRHICommandListImmediate& RHICmd
 		// should have been created by the game thread
 		check( IsValidRef(ViewportInfo.ViewportRHI) );
 
-		FTexture2DRHIRef ViewportRT = ViewportInfo.GetRenderTargetTexture();
-		FTexture2DRHIRef BackBuffer = (ViewportRT) ?
-		ViewportRT : RHICmdList.GetViewportBackBuffer(ViewportInfo.ViewportRHI);
+		FTexture2DRHIRef ViewportRT = bRenderedStereo ? nullptr : ViewportInfo.GetRenderTargetTexture();
+		FTexture2DRHIRef BackBuffer = (ViewportRT) ? ViewportRT : RHICmdList.GetViewportBackBuffer(ViewportInfo.ViewportRHI);
 		
 		const uint32 ViewportWidth = (ViewportRT) ? ViewportRT->GetSizeX() : ViewportInfo.Width;
 		const uint32 ViewportHeight = (ViewportRT) ? ViewportRT->GetSizeY() : ViewportInfo.Height;
@@ -833,7 +846,7 @@ void FSlateRHIRenderer::DrawWindow_RenderThread(FRHICommandListImmediate& RHICmd
 		}
 	}
 
-	if (GEngine && IsValidRef(ViewportInfo.GetRenderTargetTexture()) && GEngine->StereoRenderingDevice.IsValid())
+	if (!bRenderedStereo && GEngine && IsValidRef(ViewportInfo.GetRenderTargetTexture()) && GEngine->StereoRenderingDevice.IsValid())
 	{
 		GEngine->StereoRenderingDevice->RenderTexture_RenderThread(RHICmdList, RHICmdList.GetViewportBackBuffer(ViewportInfo.ViewportRHI), ViewportInfo.GetRenderTargetTexture());
 	}
diff --git a/Engine/Source/Runtime/SlateRHIRenderer/Private/SlateRHIRenderingPolicy.cpp b/Engine/Source/Runtime/SlateRHIRenderer/Private/SlateRHIRenderingPolicy.cpp
--- a/Engine/Source/Runtime/SlateRHIRenderer/Private/SlateRHIRenderingPolicy.cpp
+++ b/Engine/Source/Runtime/SlateRHIRenderer/Private/SlateRHIRenderingPolicy.cpp
@@ -246,7 +246,8 @@ void FSlateRHIRenderingPolicy::DrawElements(FRHICommandListImmediate& RHICmdList
 
 	static const FEngineShowFlags DefaultShowFlags(ESFIM_Game);
 
-	const float DisplayGamma = bGammaCorrect ? GEngine ? GEngine->GetDisplayGamma() : 2.2f : 1.0f;
+	const float EngineGamma = GEngine ? GEngine->GetDisplayGamma() : 2.2f;
+	const float DisplayGamma = bGammaCorrect ? EngineGamma : 1.0f;
 
 	// In order to support MaterialParameterCollections, we need to create multiple FSceneViews for 
 	// each possible Scene that we encounter. The following code creates these as separate arrays, where the first 
diff --git a/Engine/Source/Runtime/UMG/Private/Components/WidgetComponent.cpp b/Engine/Source/Runtime/UMG/Private/Components/WidgetComponent.cpp
--- a/Engine/Source/Runtime/UMG/Private/Components/WidgetComponent.cpp
+++ b/Engine/Source/Runtime/UMG/Private/Components/WidgetComponent.cpp
@@ -1240,11 +1240,14 @@ void UWidgetComponent::UpdateWidget()
 				NewSlateWidget = Widget->TakeWidget();
 			}
 
+			bool bNeededNewWindow = false;
 			if ( !SlateWindow.IsValid() )
 			{
 				SlateWindow = SNew(SVirtualWindow).Size(DrawSize);
 				SlateWindow->SetIsFocusable(bWindowFocusable);
 				RegisterWindow();
+
+				bNeededNewWindow = true;
 			}
 
 			if ( !HitTestGrid.IsValid() )
@@ -1256,7 +1259,7 @@ void UWidgetComponent::UpdateWidget()
 
 			if ( NewSlateWidget.IsValid() )
 			{
-				if ( NewSlateWidget != CurrentSlateWidget )
+				if ( NewSlateWidget != CurrentSlateWidget || bNeededNewWindow )
 				{
 					CurrentSlateWidget = NewSlateWidget;
 					SlateWindow->SetContent(NewSlateWidget.ToSharedRef());
@@ -1264,7 +1267,7 @@ void UWidgetComponent::UpdateWidget()
 			}
 			else if( SlateWidget.IsValid() )
 			{
-				if ( SlateWidget != CurrentSlateWidget )
+				if ( SlateWidget != CurrentSlateWidget || bNeededNewWindow )
 				{
 					CurrentSlateWidget = SlateWidget;
 					SlateWindow->SetContent(SlateWidget.ToSharedRef());
