Copying //UE4/Dev-Framework to //UE4/Dev-Main (Source: //UE4/Dev-Framework @ 3431384)

diff --git a/Engine/Source/Runtime/Engine/Private/Pawn.cpp b/Engine/Source/Runtime/Engine/Private/Pawn.cpp
--- a/Engine/Source/Runtime/Engine/Private/Pawn.cpp
+++ b/Engine/Source/Runtime/Engine/Private/Pawn.cpp
@@ -105,17 +105,19 @@ void APawn::PostInitializeComponents()
 		GetWorld()->AddPawn( this );
 
 		// Automatically add Controller to AI Pawns if we are allowed to.
-		if (AutoPossessPlayer == EAutoReceiveInput::Disabled)
+		if (AutoPossessPlayer == EAutoReceiveInput::Disabled
+			&& AutoPossessAI != EAutoPossessAI::Disabled && Controller == NULL && GetNetMode() != NM_Client
+#if WITH_EDITOR
+			&& (GIsEditor == false || GetWorld()->IsGameWorld())
+#endif // WITH_EDITOR
+			)
 		{
-			if (AutoPossessAI != EAutoPossessAI::Disabled && Controller == NULL && GetNetMode() != NM_Client)
+			const bool bPlacedInWorld = (GetWorld()->bStartup);
+			if ((AutoPossessAI == EAutoPossessAI::PlacedInWorldOrSpawned) ||
+				(AutoPossessAI == EAutoPossessAI::PlacedInWorld && bPlacedInWorld) ||
+				(AutoPossessAI == EAutoPossessAI::Spawned && !bPlacedInWorld))
 			{
-				const bool bPlacedInWorld = (GetWorld()->bStartup);
-				if ((AutoPossessAI == EAutoPossessAI::PlacedInWorldOrSpawned) ||
-					(AutoPossessAI == EAutoPossessAI::PlacedInWorld && bPlacedInWorld) ||
-					(AutoPossessAI == EAutoPossessAI::Spawned && !bPlacedInWorld))
-				{
-					SpawnDefaultController();
-				}
+				SpawnDefaultController();
 			}
 		}
 
diff --git a/Engine/Source/Runtime/Engine/Private/PhysicsEngine/BodyInstance.cpp b/Engine/Source/Runtime/Engine/Private/PhysicsEngine/BodyInstance.cpp
--- a/Engine/Source/Runtime/Engine/Private/PhysicsEngine/BodyInstance.cpp
+++ b/Engine/Source/Runtime/Engine/Private/PhysicsEngine/BodyInstance.cpp
@@ -4549,7 +4549,7 @@ bool FBodyInstance::InternalSweepPhysX(struct FHitResult& OutHit, const FVector&
 
 		UPrimitiveComponent* OwnerComponentInst = OwnerComponent.Get();
 		PxTransform PStartTM(U2PVector(Start), ShapeAdaptor.GetGeomOrientation());
-		PxTransform PCompTM(U2PTransform(OwnerComponentInst->ComponentToWorld));
+		PxTransform PCompTM(U2PTransform(OwnerComponentInst->GetComponentTransform()));
 
 		PxVec3 PDir = U2PVector(Delta/DeltaMag);
 
diff --git a/Engine/Source/Runtime/Engine/Private/PhysicsEngine/PhysAnim.cpp b/Engine/Source/Runtime/Engine/Private/PhysicsEngine/PhysAnim.cpp
--- a/Engine/Source/Runtime/Engine/Private/PhysicsEngine/PhysAnim.cpp
+++ b/Engine/Source/Runtime/Engine/Private/PhysicsEngine/PhysAnim.cpp
@@ -144,7 +144,7 @@ void USkeletalMeshComponent::PerformBlendPhysicsBones(const TArray<FBoneIndexTyp
 {
 	SCOPE_CYCLE_COUNTER(STAT_BlendInPhysics);
 	// Get drawscale from Owner (if there is one)
-	FVector TotalScale3D = ComponentToWorld.GetScale3D();
+	FVector TotalScale3D = GetComponentTransform().GetScale3D();
 	FVector RecipScale3D = TotalScale3D.Reciprocal();
 
 	UPhysicsAsset * const PhysicsAsset = GetPhysicsAsset();
@@ -172,7 +172,7 @@ void USkeletalMeshComponent::PerformBlendPhysicsBones(const TArray<FBoneIndexTyp
 	TAssetWorldBoneTMArray WorldBoneTMs;
 	WorldBoneTMs.AddZeroed(GetNumComponentSpaceTransforms());
 	
-	FTransform LocalToWorldTM = ComponentToWorld;
+	FTransform LocalToWorldTM = GetComponentTransform();
 	LocalToWorldTM.RemoveScaling();
 
 	TArray<FTransform>& EditableComponentSpaceTransforms = GetEditableComponentSpaceTransforms();
@@ -318,7 +318,7 @@ void USkeletalMeshComponent::PerformBlendPhysicsBones(const TArray<FBoneIndexTyp
 				//For now I'm juts deferring it to the end of this loop, but in general we need to move it all out of here and do it when the blend task is done
 				FBodyTMPair* BodyTMPair = new (PendingBodyTMs) FBodyTMPair;
 				BodyTMPair->BI = PhysicsAssetBodyInstance;
-				BodyTMPair->TM = EditableComponentSpaceTransforms[BoneIndex] * ComponentToWorld;
+				BodyTMPair->TM = EditableComponentSpaceTransforms[BoneIndex] * GetComponentTransform();
 			}
 #endif
 		}
@@ -487,7 +487,7 @@ void USkeletalMeshComponent::UpdateKinematicBonesToAnim(const TArray<FTransform>
 		return;
 	}
 
-	const FTransform& CurrentLocalToWorld = ComponentToWorld;
+	const FTransform& CurrentLocalToWorld = GetComponentTransform();
 
 #if !(UE_BUILD_SHIPPING)
 	// Gracefully handle NaN
diff --git a/Engine/Source/Runtime/Engine/Private/PhysicsEngine/PhysScene.cpp b/Engine/Source/Runtime/Engine/Private/PhysicsEngine/PhysScene.cpp
--- a/Engine/Source/Runtime/Engine/Private/PhysicsEngine/PhysScene.cpp
+++ b/Engine/Source/Runtime/Engine/Private/PhysicsEngine/PhysScene.cpp
@@ -1224,9 +1224,9 @@ void FPhysScene::SyncComponentsToBodies_AssumesLocked(uint32 SceneType)
 
 		// See if the transform is actually different, and if so, move the component to match physics
 		const FTransform NewTransform = BodyInstance->GetUnrealWorldTransform_AssumesLocked();
-		if (!NewTransform.EqualsNoScale(BodyInstance->OwnerComponent->ComponentToWorld))
+		if (!NewTransform.EqualsNoScale(BodyInstance->OwnerComponent->GetComponentTransform()))
 		{
-			const FVector MoveBy = NewTransform.GetLocation() - BodyInstance->OwnerComponent->ComponentToWorld.GetLocation();
+			const FVector MoveBy = NewTransform.GetLocation() - BodyInstance->OwnerComponent->GetComponentTransform().GetLocation();
 			const FQuat NewRotation = NewTransform.GetRotation();
 
 			//@warning: do not reference BodyInstance again after calling MoveComponent() - events from the move could have made it unusable (destroying the actor, SetPhysics(), etc)
diff --git a/Engine/Source/Runtime/Engine/Private/PhysicsEngine/PhysicsConstraintComponent.cpp b/Engine/Source/Runtime/Engine/Private/PhysicsEngine/PhysicsConstraintComponent.cpp
--- a/Engine/Source/Runtime/Engine/Private/PhysicsEngine/PhysicsConstraintComponent.cpp
+++ b/Engine/Source/Runtime/Engine/Private/PhysicsEngine/PhysicsConstraintComponent.cpp
@@ -137,8 +137,8 @@ FTransform UPhysicsConstraintComponent::GetBodyTransformInternal(EConstraintFram
 		return FTransform::Identity;
 	}
 	  
-	//Use ComponentToWorld by default for all components
-	FTransform ResultTM = PrimComp->ComponentToWorld;
+	//Use GetComponentTransform() by default for all components
+	FTransform ResultTM = PrimComp->GetComponentTransform();
 		
 	// Skeletal case
 	if(const USkeletalMeshComponent* SkelComp = Cast<USkeletalMeshComponent>(PrimComp))
@@ -524,8 +524,8 @@ void UPhysicsConstraintComponent::UpdateConstraintFrames()
 
 	// World ref frame
 	const FVector WPos = GetComponentLocation();
-	const FVector WPri = ComponentToWorld.GetUnitAxis( EAxis::X );
-	const FVector WOrth = ComponentToWorld.GetUnitAxis( EAxis::Y );
+	const FVector WPri = GetComponentTransform().GetUnitAxis( EAxis::X );
+	const FVector WOrth = GetComponentTransform().GetUnitAxis( EAxis::Y );
 
 	ConstraintInstance.Pos1 = A1Transform.InverseTransformPosition(WPos);
 	ConstraintInstance.PriAxis1 = A1Transform.InverseTransformVectorNoScale(WPri);
@@ -533,8 +533,8 @@ void UPhysicsConstraintComponent::UpdateConstraintFrames()
 
 	const FVector RotatedX = ConstraintInstance.AngularRotationOffset.RotateVector(FVector(1,0,0));
 	const FVector RotatedY = ConstraintInstance.AngularRotationOffset.RotateVector(FVector(0,1,0));
-	const FVector WPri2 = ComponentToWorld.TransformVectorNoScale(RotatedX);
-	const FVector WOrth2 = ComponentToWorld.TransformVectorNoScale(RotatedY);
+	const FVector WPri2 = GetComponentTransform().TransformVectorNoScale(RotatedX);
+	const FVector WOrth2 = GetComponentTransform().TransformVectorNoScale(RotatedY);
 	
 	
 	ConstraintInstance.Pos2 = A2Transform.InverseTransformPosition(WPos);
diff --git a/Engine/Source/Runtime/Engine/Private/PhysicsEngine/PhysicsSpring.cpp b/Engine/Source/Runtime/Engine/Private/PhysicsEngine/PhysicsSpring.cpp
--- a/Engine/Source/Runtime/Engine/Private/PhysicsEngine/PhysicsSpring.cpp
+++ b/Engine/Source/Runtime/Engine/Private/PhysicsEngine/PhysicsSpring.cpp
@@ -27,12 +27,12 @@ UPhysicsSpringComponent::UPhysicsSpringComponent(const FObjectInitializer& Objec
 }
 FVector UPhysicsSpringComponent::SpringPositionFromLength(float Length) const
 {
-	return ComponentToWorld.GetLocation() + Length * GetSpringDirection();
+	return GetComponentTransform().GetLocation() + Length * GetSpringDirection();
 }
 
 FVector UPhysicsSpringComponent::GetSpringDirection() const
 {
-	return ComponentToWorld.TransformVectorNoScale(FVector(1.f, 0.f, 0.f));
+	return GetComponentTransform().TransformVectorNoScale(FVector(1.f, 0.f, 0.f));
 }
 
 float UPhysicsSpringComponent::GetNormalizedCompressionScalar() const
@@ -104,7 +104,7 @@ void UPhysicsSpringComponent::TickComponent(float DeltaTime, enum ELevelTick Tic
 	{
 		if (bIsActive)
 		{
-			const FVector SpringStart = ComponentToWorld.GetLocation();
+			const FVector SpringStart = GetComponentTransform().GetLocation();
 			const FVector SpringDesiredEnd = SpringPositionFromLength(SpringLengthAtRest);
 			float CollisionTime = 1.f;
 			
diff --git a/Engine/Source/Runtime/Engine/Private/PhysicsEngine/PhysicsThruster.cpp b/Engine/Source/Runtime/Engine/Private/PhysicsEngine/PhysicsThruster.cpp
--- a/Engine/Source/Runtime/Engine/Private/PhysicsEngine/PhysicsThruster.cpp
+++ b/Engine/Source/Runtime/Engine/Private/PhysicsEngine/PhysicsThruster.cpp
@@ -45,7 +45,7 @@ void UPhysicsThrusterComponent::TickComponent(float DeltaTime, enum ELevelTick T
 	// Applied force to the base, so if we don't have one, do nothing.
 	if( bIsActive && GetAttachParent())
 	{
-		FVector WorldForce = ThrustStrength * ComponentToWorld.TransformVectorNoScale( FVector(-1.f,0.f,0.f) );
+		FVector WorldForce = ThrustStrength * GetComponentTransform().TransformVectorNoScale( FVector(-1.f,0.f,0.f) );
 
 		UPrimitiveComponent* BasePrimComp = Cast<UPrimitiveComponent>(GetAttachParent());
 		if(BasePrimComp)
diff --git a/Engine/Source/Runtime/Engine/Private/PhysicsEngine2D/Box2DIntegration.cpp b/Engine/Source/Runtime/Engine/Private/PhysicsEngine2D/Box2DIntegration.cpp
--- a/Engine/Source/Runtime/Engine/Private/PhysicsEngine2D/Box2DIntegration.cpp
+++ b/Engine/Source/Runtime/Engine/Private/PhysicsEngine2D/Box2DIntegration.cpp
@@ -243,7 +243,7 @@ void FStartPhysics2DTickFunction::ExecuteTick(float DeltaTime, enum ELevelTick T
 				{
 					// See if the transform is actually different, and if so, move the component to match physics
 					const FTransform NewTransform = UnrealBodyInstance->GetUnrealWorldTransform();
-					const FTransform& OldTransform = UnrealBodyInstance->OwnerComponent->ComponentToWorld;
+					const FTransform& OldTransform = UnrealBodyInstance->OwnerComponent->GetComponentTransform();
 					if (!NewTransform.Equals(OldTransform))
 					{
 						const FVector MoveBy = NewTransform.GetLocation() - OldTransform.GetLocation();
diff --git a/Engine/Source/Runtime/Engine/Private/PlanarReflectionSceneProxy.cpp b/Engine/Source/Runtime/Engine/Private/PlanarReflectionSceneProxy.cpp
--- a/Engine/Source/Runtime/Engine/Private/PlanarReflectionSceneProxy.cpp
+++ b/Engine/Source/Runtime/Engine/Private/PlanarReflectionSceneProxy.cpp
@@ -41,7 +41,7 @@ FPlanarReflectionSceneProxy::FPlanarReflectionSceneProxy(UPlanarReflectionCompon
 
 	OwnerName = Component->GetOwner() ? Component->GetOwner()->GetFName() : NAME_None;
 
-	UpdateTransform(Component->ComponentToWorld.ToMatrixWithScale());
+	UpdateTransform(Component->GetComponentTransform().ToMatrixWithScale());
 
 	PlanarReflectionId = Component->GetPlanarReflectionId();
 	PrefilterRoughness = Component->PrefilterRoughness;
diff --git a/Engine/Source/Runtime/Engine/Private/PreviewScene.cpp b/Engine/Source/Runtime/Engine/Private/PreviewScene.cpp
--- a/Engine/Source/Runtime/Engine/Private/PreviewScene.cpp
+++ b/Engine/Source/Runtime/Engine/Private/PreviewScene.cpp
@@ -143,7 +143,7 @@ void FPreviewScene::ClearLineBatcher()
 /** Accessor for finding the current direction of the preview scene's DirectionalLight. */
 FRotator FPreviewScene::GetLightDirection()
 {
-	return DirectionalLight->ComponentToWorld.GetUnitAxis( EAxis::X ).Rotation();
+	return DirectionalLight->GetComponentTransform().GetUnitAxis( EAxis::X ).Rotation();
 }
 
 /** Function for modifying the current direction of the preview scene's DirectionalLight. */
diff --git a/Engine/Source/Runtime/Engine/Private/PrimaryAssetLabel.cpp b/Engine/Source/Runtime/Engine/Private/PrimaryAssetLabel.cpp
--- a/Engine/Source/Runtime/Engine/Private/PrimaryAssetLabel.cpp
+++ b/Engine/Source/Runtime/Engine/Private/PrimaryAssetLabel.cpp
@@ -38,12 +38,20 @@ void UPrimaryAssetLabel::UpdateAssetBundleData()
 		TArray<FAssetData> DirectoryAssets;
 		AssetRegistry.GetAssetsByPath(PackagePath, DirectoryAssets, true);
 
+		TArray<FStringAssetReference> NewPaths;
+
 		for (const FAssetData& AssetData : DirectoryAssets)
 		{
 			FStringAssetReference AssetRef = Manager.GetAssetPathForData(AssetData);
 
-			AssetBundleData.AddBundleAsset(DirectoryBundle, AssetRef);
+			if (!AssetRef.IsNull())
+			{
+				NewPaths.Add(AssetRef);
+			}
 		}
+
+		// Fast set, destroys NewPaths
+		AssetBundleData.SetBundleAssets(DirectoryBundle, MoveTemp(NewPaths));
 	}
 	
 	// Update rules
diff --git a/Engine/Source/Runtime/Engine/Private/PrimitiveComponentPhysics.cpp b/Engine/Source/Runtime/Engine/Private/PrimitiveComponentPhysics.cpp
--- a/Engine/Source/Runtime/Engine/Private/PrimitiveComponentPhysics.cpp
+++ b/Engine/Source/Runtime/Engine/Private/PrimitiveComponentPhysics.cpp
@@ -552,9 +552,9 @@ FVector UPrimitiveComponent::GetInertiaTensor(FName BoneName /* = NAME_None */)
 FVector UPrimitiveComponent::ScaleByMomentOfInertia(FVector InputVector, FName BoneName /* = NAME_None */) const
 {
 	const FVector LocalInertiaTensor = GetInertiaTensor(BoneName);
-	const FVector InputVectorLocal = ComponentToWorld.InverseTransformVectorNoScale(InputVector);
+	const FVector InputVectorLocal = GetComponentTransform().InverseTransformVectorNoScale(InputVector);
 	const FVector LocalScaled = InputVectorLocal * LocalInertiaTensor;
-	const FVector WorldScaled = ComponentToWorld.TransformVectorNoScale(LocalScaled);
+	const FVector WorldScaled = GetComponentTransform().TransformVectorNoScale(LocalScaled);
 	return WorldScaled;
 }
 
@@ -658,9 +658,9 @@ void UPrimitiveComponent::SyncComponentToRBPhysics()
 
 	// See if the transform is actually different, and if so, move the component to match physics
 	const FTransform NewTransform = GetComponentTransformFromBodyInstance(UseBI);	
-	if(!NewTransform.EqualsNoScale(ComponentToWorld))
+	if(!NewTransform.EqualsNoScale(GetComponentTransform()))
 	{
-		const FVector MoveBy = NewTransform.GetLocation() - ComponentToWorld.GetLocation();
+		const FVector MoveBy = NewTransform.GetLocation() - GetComponentTransform().GetLocation();
 		const FRotator NewRotation = NewTransform.Rotator();
 
 		//@warning: do not reference BodyInstance again after calling MoveComponent() - events from the move could have made it unusable (destroying the actor, SetPhysics(), etc)
diff --git a/Engine/Source/Runtime/Engine/Private/SkeletalMesh.cpp b/Engine/Source/Runtime/Engine/Private/SkeletalMesh.cpp
--- a/Engine/Source/Runtime/Engine/Private/SkeletalMesh.cpp
+++ b/Engine/Source/Runtime/Engine/Private/SkeletalMesh.cpp
@@ -4807,7 +4807,7 @@ bool USkeletalMeshSocket::AttachActor(AActor* Actor, class USkeletalMeshComponen
 
 			Actor->SetActorLocation(SocketTM.GetOrigin(), false);
 			Actor->SetActorRotation(SocketTM.Rotator());
-			Actor->GetRootComponent()->SnapTo( SkelComp, SocketName );
+			Actor->GetRootComponent()->AttachToComponent(SkelComp, FAttachmentTransformRules::SnapToTargetNotIncludingScale, SocketName);
 
 #if WITH_EDITOR
 			if (GIsEditor)
diff --git a/Engine/Source/Runtime/Engine/Private/SkeletalMeshComponentPhysics.cpp b/Engine/Source/Runtime/Engine/Private/SkeletalMeshComponentPhysics.cpp
--- a/Engine/Source/Runtime/Engine/Private/SkeletalMeshComponentPhysics.cpp
+++ b/Engine/Source/Runtime/Engine/Private/SkeletalMeshComponentPhysics.cpp
@@ -534,7 +534,7 @@ void USkeletalMeshComponent::InitArticulated(FPhysScene* PhysScene)
 		return;
 	}
 
-	FVector Scale3D = ComponentToWorld.GetScale3D();
+	FVector Scale3D = GetComponentTransform().GetScale3D();
 
 	// Find root physics body
 	RootBodyData.BodyIndex = INDEX_NONE;	//Reset the root body index just in case we need to refind a new one
@@ -1315,7 +1315,7 @@ void USkeletalMeshComponent::SendRenderDebugPhysics(FPrimitiveSceneProxy* Overri
 				FPrimitiveSceneProxy::FDebugMassData& MassData = DebugMassData.Last();
 				const FTransform MassToWorld = BI->GetMassSpaceToWorldSpace();
 				const FTransform& BoneTM = GetComponentSpaceTransforms()[BoneIndex];
-				const FTransform BoneToWorld = BoneTM * ComponentToWorld;
+				const FTransform BoneToWorld = BoneTM * GetComponentTransform();
 
 				MassData.LocalCenterOfMass = BoneToWorld.InverseTransformPosition(MassToWorld.GetLocation());
 				MassData.LocalTensorOrientation = MassToWorld.GetRotation() * BoneToWorld.GetRotation().Inverse();
@@ -1845,7 +1845,7 @@ FVector USkeletalMeshComponent::GetSkinnedVertexPosition(int32 VertexIndex) cons
 			{
 				// a simulated position is in world space and convert this to local space
 				// because SkinnedMeshComponent::GetSkinnedVertexPosition() returns the position in local space
-				SimulatedPos = ComponentToWorld.InverseTransformPosition(SimulatedPos);
+				SimulatedPos = GetComponentTransform().InverseTransformPosition(SimulatedPos);
 
 				// if blend weight is 1.0, doesn't need to blend with a skinned position
 				if (ClothBlendWeight < 1.0f) 
@@ -1939,7 +1939,7 @@ bool USkeletalMeshComponent::GetClosestPointOnPhysicsAsset(const FVector& WorldP
 	{
 		const TArray<FTransform>& BoneTransforms = GetComponentSpaceTransforms();
 		const bool bHasMasterPoseComponent = MasterPoseComponent.IsValid();
-		const FVector ComponentPosition = ComponentToWorld.InverseTransformPosition(WorldPosition);
+		const FVector ComponentPosition = GetComponentTransform().InverseTransformPosition(WorldPosition);
 	
 		float CurrentClosestDistance = FLT_MAX;
 		int32 CurrentClosestBoneIndex = INDEX_NONE;
@@ -1970,7 +1970,7 @@ bool USkeletalMeshComponent::GetClosestPointOnPhysicsAsset(const FVector& WorldP
 		{
 			bSuccess = true;
 
-			const FTransform BoneTM = bHasMasterPoseComponent ? GetBoneTransform(CurrentClosestBoneIndex) : (BoneTransforms[CurrentClosestBoneIndex] * ComponentToWorld);
+			const FTransform BoneTM = bHasMasterPoseComponent ? GetBoneTransform(CurrentClosestBoneIndex) : (BoneTransforms[CurrentClosestBoneIndex] * GetComponentTransform());
 			ClosestPointOnPhysicsAsset.Distance = CurrentClosestBodySetup->GetClosestPointAndNormal(WorldPosition, BoneTM, ClosestPointOnPhysicsAsset.ClosestWorldPosition, ClosestPointOnPhysicsAsset.Normal);
 			ClosestPointOnPhysicsAsset.BoneName = CurrentClosestBodySetup->BoneName;
 		}
@@ -2091,7 +2091,7 @@ bool USkeletalMeshComponent::ComponentOverlapMultiImpl(TArray<struct FOverlapRes
 		return false;
 	}
 
-	const FTransform WorldToComponent(ComponentToWorld.Inverse());
+	const FTransform WorldToComponent(GetComponentTransform().Inverse());
 	const FCollisionResponseParams ResponseParams(GetCollisionResponseToChannels());
 
 	FComponentQueryParams ParamsWithSelf = Params;
@@ -2191,7 +2191,7 @@ void USkeletalMeshComponent::GetWindForCloth_GameThread(FVector& WindDirection,
 		// set wind
 		if(IsWindEnabled())
 		{
-			FVector Position = ComponentToWorld.GetTranslation();
+			FVector Position = GetComponentTransform().GetTranslation();
 
 			float WindSpeed;
 			float WindMinGust;
@@ -2331,9 +2331,9 @@ void USkeletalMeshComponent::ProcessClothCollisionWithEnvironment()
 
 					// Matrices required to transform shapes into sim space (component space)
 					// Transform of external component and matrix describing external component -> this component
-					FTransform Transform = Component->ComponentToWorld;
+					FTransform Transform = Component->GetComponentTransform();
 					FMatrix TransformMatrix = Transform.ToMatrixWithScale();
-					FMatrix ComponentToClothMatrix = TransformMatrix * ComponentToWorld.ToMatrixWithScale().Inverse();
+					FMatrix ComponentToClothMatrix = TransformMatrix * GetComponentTransform().ToMatrixWithScale().Inverse();
 
 					for(PxShape* Shape : AllShapes)
 					{
@@ -2561,15 +2561,17 @@ void USkeletalMeshComponent::UpdateClothTransformImp()
 #endif // WITH_CLOTH_COLLISION_DETECTION
 
 #if !(UE_BUILD_SHIPPING)
-	if (ComponentToWorld.GetRotation().ContainsNaN())
+	FTransform ComponentTransform = GetComponentTransform();
+	if (ComponentTransform.GetRotation().ContainsNaN())
 	{
-		logOrEnsureNanError(TEXT("SkeletalMeshComponent::UpdateClothTransform found NaN in ComponentToWorld.GetRotation()"));
-		ComponentToWorld.SetRotation(FQuat(0.0f, 0.0f, 0.0f, 1.0f));
+		logOrEnsureNanError(TEXT("SkeletalMeshComponent::UpdateClothTransform found NaN in GetComponentTransform().GetRotation()"));
+		ComponentTransform.SetRotation(FQuat(0.0f, 0.0f, 0.0f, 1.0f));
+		SetComponentToWorld(ComponentTransform);
 	}
-	if (ComponentToWorld.ContainsNaN())
+	if (ComponentTransform.ContainsNaN())
 	{
-		logOrEnsureNanError(TEXT("SkeletalMeshComponent::UpdateClothTransform still found NaN in ComponentToWorld (wasn't the rotation)"));
-		ComponentToWorld.SetIdentity();
+		logOrEnsureNanError(TEXT("SkeletalMeshComponent::UpdateClothTransform still found NaN in GetComponentTransform() (wasn't the rotation)"));
+		SetComponentToWorld(FTransform::Identity);
 	}
 #endif
 }
@@ -2768,7 +2770,7 @@ bool USkeletalMeshComponent::GetClothSimulatedPosition_GameThread(const FGuid& A
 
 		if(ActorData && ActorData->Positions.IsValidIndex(VertexIndex))
 		{
-			OutSimulPos = ComponentToWorld.TransformPosition(ActorData->Positions[VertexIndex]);
+			OutSimulPos = GetComponentTransform().TransformPosition(ActorData->Positions[VertexIndex]);
 
 			bSucceed = true;
 		}
@@ -2935,7 +2937,7 @@ FTransform USkeletalMeshComponent::GetComponentTransformFromBodyInstance(FBodyIn
 	}
 	else
 	{
-		return ComponentToWorld;
+		return GetComponentTransform();
 	}
 }
 #undef LOCTEXT_NAMESPACE
diff --git a/Engine/Source/Runtime/Engine/Private/SoundCue.cpp b/Engine/Source/Runtime/Engine/Private/SoundCue.cpp
--- a/Engine/Source/Runtime/Engine/Private/SoundCue.cpp
+++ b/Engine/Source/Runtime/Engine/Private/SoundCue.cpp
@@ -5,6 +5,7 @@
 #include "EngineDefines.h"
 #include "EngineGlobals.h"
 #include "Engine/Engine.h"
+#include "Misc/CoreDelegates.h"
 #include "Components/AudioComponent.h"
 #include "UObject/UObjectIterator.h"
 #include "EngineUtils.h"
@@ -121,13 +122,13 @@ void USoundCue::PostLoad()
 	}
 	else
 	{
-		OnPostEngineInitHandle = UEngine::OnPostEngineInit.AddUObject(this, &USoundCue::OnPostEngineInit);
+		OnPostEngineInitHandle = FCoreDelegates::OnPostEngineInit.AddUObject(this, &USoundCue::OnPostEngineInit);
 	}
 }
 
 void USoundCue::OnPostEngineInit()
 {
-	UEngine::OnPostEngineInit.Remove(OnPostEngineInitHandle);
+	FCoreDelegates::OnPostEngineInit.Remove(OnPostEngineInitHandle);
 	OnPostEngineInitHandle.Reset();
 
 	EvaluateNodes(true);
diff --git a/Engine/Source/Runtime/Engine/Private/StaticMesh.cpp b/Engine/Source/Runtime/Engine/Private/StaticMesh.cpp
--- a/Engine/Source/Runtime/Engine/Private/StaticMesh.cpp
+++ b/Engine/Source/Runtime/Engine/Private/StaticMesh.cpp
@@ -3540,14 +3540,14 @@ UStaticMeshSocket::UStaticMeshSocket(const FObjectInitializer& ObjectInitializer
 bool UStaticMeshSocket::GetSocketMatrix(FMatrix& OutMatrix, UStaticMeshComponent const* MeshComp) const
 {
 	check( MeshComp );
-	OutMatrix = FScaleRotationTranslationMatrix( RelativeScale, RelativeRotation, RelativeLocation ) * MeshComp->ComponentToWorld.ToMatrixWithScale();
+	OutMatrix = FScaleRotationTranslationMatrix( RelativeScale, RelativeRotation, RelativeLocation ) * MeshComp->GetComponentTransform().ToMatrixWithScale();
 	return true;
 }
 
 bool UStaticMeshSocket::GetSocketTransform(FTransform& OutTransform, class UStaticMeshComponent const* MeshComp) const
 {
 	check( MeshComp );
-	OutTransform = FTransform(RelativeRotation, RelativeLocation, RelativeScale) * MeshComp->ComponentToWorld;
+	OutTransform = FTransform(RelativeRotation, RelativeLocation, RelativeScale) * MeshComp->GetComponentTransform();
 	return true;
 }
 
@@ -3565,7 +3565,7 @@ bool UStaticMeshSocket::AttachActor(AActor* Actor,  UStaticMeshComponent* MeshCo
 
 			Actor->SetActorLocation(SocketTM.GetOrigin(), false);
 			Actor->SetActorRotation(SocketTM.Rotator());
-			Actor->GetRootComponent()->SnapTo( MeshComp, SocketName );
+			Actor->GetRootComponent()->AttachToComponent(MeshComp, FAttachmentTransformRules::SnapToTargetNotIncludingScale, SocketName);
 
 #if WITH_EDITOR
 			if (GIsEditor)
diff --git a/Engine/Source/Runtime/Engine/Private/StaticMeshLight.cpp b/Engine/Source/Runtime/Engine/Private/StaticMeshLight.cpp
--- a/Engine/Source/Runtime/Engine/Private/StaticMeshLight.cpp
+++ b/Engine/Source/Runtime/Engine/Private/StaticMeshLight.cpp
@@ -65,7 +65,7 @@ FStaticMeshStaticLightingMesh::FStaticMeshStaticLightingMesh(const UStaticMeshCo
 	StaticMesh(InPrimitive->GetStaticMesh()),
 	Primitive(InPrimitive),
 	LODRenderData(InPrimitive->GetStaticMesh()->RenderData->LODResources[InLODIndex]),
-	bReverseWinding(InPrimitive->ComponentToWorld.GetDeterminant() < 0.0f)
+	bReverseWinding(InPrimitive->GetComponentTransform().GetDeterminant() < 0.0f)
 {
 	LODIndexBuffer = LODRenderData.IndexBuffer.GetArrayView();
 
diff --git a/Engine/Source/Runtime/Engine/Private/StreamableManager.cpp b/Engine/Source/Runtime/Engine/Private/StreamableManager.cpp
--- a/Engine/Source/Runtime/Engine/Private/StreamableManager.cpp
+++ b/Engine/Source/Runtime/Engine/Private/StreamableManager.cpp
@@ -242,24 +242,9 @@ void FStreamableHandle::GetLoadedAssets(TArray<UObject *>& LoadedAssets) const
 
 void FStreamableHandle::GetLoadedCount(int32& LoadedCount, int32& RequestedCount) const
 {
-	LoadedCount = 0;
 	RequestedCount = RequestedAssets.Num();
-
-	if (HasLoadCompleted())
-	{
-		LoadedCount = RequestedAssets.Num();
-	}
-	else if (IsActive())
-	{
-		for (const FStringAssetReference& Ref : RequestedAssets)
-		{
-			if (OwningManager->IsAsyncLoadComplete(Ref))
-			{
-				LoadedCount++;
-			}
-		}
-	}
-
+	LoadedCount = RequestedCount - StreamablesLoading;
+	
 	// Check child handles
 	for (TSharedPtr<FStreamableHandle> ChildHandle : ChildHandles)
 	{
@@ -766,6 +751,7 @@ FStreamable* FStreamableManager::StreamInternal(const FStringAssetReference& InT
 			}
 
 			Existing->bAsyncLoadRequestOutstanding = true;
+			Existing->bLoadFailed = false;
 			LoadPackageAsync(Package, FLoadPackageAsyncDelegate::CreateSP(Handle, &FStreamableHandle::AsyncLoadCallbackWrapper, TargetName), Priority);
 		}
 	}
diff --git a/Engine/Source/Runtime/Engine/Private/Tests/EngineAutomationTests.cpp b/Engine/Source/Runtime/Engine/Private/Tests/EngineAutomationTests.cpp
--- a/Engine/Source/Runtime/Engine/Private/Tests/EngineAutomationTests.cpp
+++ b/Engine/Source/Runtime/Engine/Private/Tests/EngineAutomationTests.cpp
@@ -547,9 +547,7 @@ PRAGMA_ENABLE_DEPRECATION_WARNINGS
 		Test->TestTrue(FString::Printf(TEXT("Child world rotation was incorrect after attachment (was %s, should be %s)"), *ChildActor->GetActorQuat().ToString(), *LegacyExpectedChildTransforms[LocationInteger][0].GetRotation().ToString()), ChildActor->GetActorQuat().Equals(LegacyExpectedChildTransforms[LocationInteger][0].GetRotation(), KINDA_SMALL_NUMBER));
 		Test->TestTrue(FString::Printf(TEXT("Child world scale was incorrect after attachment (was %s, should be %s)"), *ChildActor->GetActorScale3D().ToString(), *LegacyExpectedChildTransforms[LocationInteger][0].GetScale3D().ToString()), ChildActor->GetActorScale3D().Equals(LegacyExpectedChildTransforms[LocationInteger][0].GetScale3D(), KINDA_SMALL_NUMBER));
 
-PRAGMA_DISABLE_DEPRECATION_WARNINGS
-		ChildActor->DetachRootComponentFromParent(true);
-PRAGMA_ENABLE_DEPRECATION_WARNINGS
+		ChildActor->DetachFromActor(FDetachmentTransformRules::KeepWorldTransform);
 
 		// check we have expected values after detachment
 		Test->TestEqual<FVector>(TEXT("Parent location was affected by detachment"), ParentActor->GetActorLocation(), AttachTestConstants::ParentLocation);
diff --git a/Engine/Source/Runtime/Engine/Private/UnrealEngine.cpp b/Engine/Source/Runtime/Engine/Private/UnrealEngine.cpp
--- a/Engine/Source/Runtime/Engine/Private/UnrealEngine.cpp
+++ b/Engine/Source/Runtime/Engine/Private/UnrealEngine.cpp
@@ -191,6 +191,7 @@
 
 #include "GeneralProjectSettings.h"
 #include "ProfilingDebugging/LoadTimeTracker.h"
+#include "AssetRegistryModule.h"
 
 DEFINE_LOG_CATEGORY(LogEngine);
 IMPLEMENT_MODULE( FEngineModule, Engine );
@@ -279,7 +280,9 @@ bool GIsTextureMemoryCorrupted = false;
 	bool GIsPrepareMapChangeBroken = false;
 #endif
 
+PRAGMA_DISABLE_DEPRECATION_WARNINGS
 FSimpleMulticastDelegate UEngine::OnPostEngineInit;
+PRAGMA_ENABLE_DEPRECATION_WARNINGS
 
 // We expose these variables to everyone as we need to access them in other files via an extern
 ENGINE_API float GAverageFPS = 0.0f;
@@ -4731,11 +4734,17 @@ bool UEngine::HandleMemCommand( const TCHAR* Cmd, FOutputDevice& Ar )
 
 	if( bDetailed || bReport)
 	{
-		Ar.CategorizedLogf( CategoryName, ELogVerbosity::Log, TEXT("Memory Stats:") );
-		Ar.CategorizedLogf( CategoryName, ELogVerbosity::Log, TEXT("FMemStack (gamethread) current size = %.2f MB"), FMemStack::Get().GetByteCount() / (1024.0f * 1024.0f));
+		Ar.CategorizedLogf(CategoryName, ELogVerbosity::Log, TEXT("Memory Stats:") );
+		Ar.CategorizedLogf(CategoryName, ELogVerbosity::Log, TEXT("FMemStack (gamethread) current size = %.2f MB"), FMemStack::Get().GetByteCount() / (1024.0f * 1024.0f));
 		Ar.CategorizedLogf(CategoryName, ELogVerbosity::Log, TEXT("FPageAllocator (all threads) allocation size [used/ unused] = [%.2f / %.2f] MB"), (FPageAllocator::BytesUsed()) / (1024.0f * 1024.0f), (FPageAllocator::BytesFree()) / (1024.0f * 1024.0f));
 		Ar.CategorizedLogf(CategoryName, ELogVerbosity::Log, TEXT("Nametable memory usage = %.2f MB"), FName::GetNameTableMemorySize() / (1024.0f * 1024.0f));
 
+		FAssetRegistryModule* AssetRegistryModule = FModuleManager::LoadModulePtr<FAssetRegistryModule>(AssetRegistryConstants::ModuleName);
+		if (AssetRegistryModule)
+		{
+			Ar.CategorizedLogf(CategoryName, ELogVerbosity::Log, TEXT("AssetRegistry memory usage = %.2f MB"), AssetRegistryModule->Get().GetAllocatedSize() / (1024.0f * 1024.0f));
+		}
+
 #if STATS
 		TArray<FStatMessage> Stats;
 		GetPermanentStats(Stats);
@@ -7429,7 +7438,7 @@ static void DrawVolumeOnCanvas(const AVolume* Volume, FCanvas* Canvas, const FVe
 {
 	if(Volume && Volume->GetBrushComponent() && Volume->GetBrushComponent()->BrushBodySetup)
 	{
-		FTransform BrushTM = Volume->GetBrushComponent()->ComponentToWorld;
+		FTransform BrushTM = Volume->GetBrushComponent()->GetComponentTransform();
 
 		// Iterate over each piece
 		for(int32 ConIdx=0; ConIdx<Volume->GetBrushComponent()->BrushBodySetup->AggGeom.ConvexElems.Num(); ConIdx++)
@@ -11610,10 +11619,9 @@ void UEngine::CopyPropertiesForUnrelatedObjects(UObject* OldObject, UObject* New
 		TArray<UObject*> Components;
 		OldObject->CollectDefaultSubobjects(Components, true);
 
-		for (int32 Index = 0; Index < Components.Num(); Index++)
+		for (UObject* OldInstance : Components)
 		{
 			FInstancedObjectRecord* pRecord = new(SavedInstances) FInstancedObjectRecord();
-			UObject* OldInstance = Components[Index];
 			pRecord->OldInstance = OldInstance;
 			OldInstanceMap.Add(OldInstance->GetPathName(OldObject), SavedInstances.Num() - 1);
 			const uint32 AdditionalPortFlags = Params.bCopyDeprecatedProperties ? PPF_UseDeprecatedProperties : PPF_None;
diff --git a/Engine/Source/Runtime/Engine/Private/UserInterface/InputSettings.cpp b/Engine/Source/Runtime/Engine/Private/UserInterface/InputSettings.cpp
--- a/Engine/Source/Runtime/Engine/Private/UserInterface/InputSettings.cpp
+++ b/Engine/Source/Runtime/Engine/Private/UserInterface/InputSettings.cpp
@@ -137,13 +137,36 @@ void UInputSettings::SaveKeyMappings()
 	SaveConfig();
 }
 
-void UInputSettings::AddActionMapping(const FInputActionKeyMapping& KeyMapping)
+UInputSettings* UInputSettings::GetInputSettings()
+{
+	return GetMutableDefault<UInputSettings>();
+}
+
+void UInputSettings::AddActionMapping(const FInputActionKeyMapping& KeyMapping, const bool bForceRebuildKeymaps)
 {
 	ActionMappings.AddUnique(KeyMapping);
-	ForceRebuildKeymaps();
+	if (bForceRebuildKeymaps)
+	{
+		ForceRebuildKeymaps();
+	}
 }
 
-void UInputSettings::RemoveActionMapping(const FInputActionKeyMapping& KeyMapping)
+void UInputSettings::GetActionMappingByName(const FName InActionName, TArray<FInputActionKeyMapping>& OutMappings) const
+{
+	if (InActionName.IsValid())
+	{
+		for (int32 ActionIndex = ActionMappings.Num() - 1; ActionIndex >= 0; --ActionIndex)
+		{
+			if (ActionMappings[ActionIndex].ActionName == InActionName)
+			{
+				OutMappings.Add(ActionMappings[ActionIndex]);
+				// we don't break because the mapping may have been in the array twice
+			}
+		}
+	}
+}
+
+void UInputSettings::RemoveActionMapping(const FInputActionKeyMapping& KeyMapping, const bool bForceRebuildKeymaps)
 {
 	for (int32 ActionIndex = ActionMappings.Num() - 1; ActionIndex >= 0; --ActionIndex)
 	{
@@ -153,16 +176,37 @@ void UInputSettings::RemoveActionMapping(const FInputActionKeyMapping& KeyMappin
 			// we don't break because the mapping may have been in the array twice
 		}
 	}
-	ForceRebuildKeymaps();
+	if (bForceRebuildKeymaps)
+	{
+		ForceRebuildKeymaps();
+	}
 }
 
-void UInputSettings::AddAxisMapping(const FInputAxisKeyMapping& KeyMapping)
+void UInputSettings::AddAxisMapping(const FInputAxisKeyMapping& KeyMapping, const bool bForceRebuildKeymaps)
 {
 	AxisMappings.AddUnique(KeyMapping);
-	ForceRebuildKeymaps();
+	if (bForceRebuildKeymaps)
+	{
+		ForceRebuildKeymaps();
+	}
+}
+
+void UInputSettings::GetAxisMappingByName(const FName InAxisName, TArray<FInputAxisKeyMapping>& OutMappings) const
+{
+	if (InAxisName.IsValid())
+	{
+		for (int32 AxisIndex = AxisMappings.Num() - 1; AxisIndex >= 0; --AxisIndex)
+		{
+			if (AxisMappings[AxisIndex].AxisName == InAxisName)
+			{
+				OutMappings.Add(AxisMappings[AxisIndex]);
+				// we don't break because the mapping may have been in the array twice
+			}
+		}
+	}
 }
 
-void UInputSettings::RemoveAxisMapping(const FInputAxisKeyMapping& InKeyMapping)
+void UInputSettings::RemoveAxisMapping(const FInputAxisKeyMapping& InKeyMapping, const bool bForceRebuildKeymaps)
 {
 	for (int32 AxisIndex = AxisMappings.Num() - 1; AxisIndex >= 0; --AxisIndex)
 	{
@@ -174,7 +218,10 @@ void UInputSettings::RemoveAxisMapping(const FInputAxisKeyMapping& InKeyMapping)
 			// we don't break because the mapping may have been in the array twice
 		}
 	}
-	ForceRebuildKeymaps();
+	if (bForceRebuildKeymaps)
+	{
+		ForceRebuildKeymaps();
+	}
 }
 
 void UInputSettings::GetActionNames(TArray<FName>& ActionNames) const
diff --git a/Engine/Source/Runtime/Engine/Private/VectorField.cpp b/Engine/Source/Runtime/Engine/Private/VectorField.cpp
--- a/Engine/Source/Runtime/Engine/Private/VectorField.cpp
+++ b/Engine/Source/Runtime/Engine/Private/VectorField.cpp
@@ -520,7 +520,7 @@ void UVectorFieldComponent::OnRegister()
 		{
 			FVectorFieldInstance* Instance = new FVectorFieldInstance();
 			VectorField->InitInstance(Instance, /*bPreviewInstance=*/ true);
-			Instance->UpdateTransforms(ComponentToWorld.ToMatrixWithScale());
+			Instance->UpdateTransforms(GetComponentTransform().ToMatrixWithScale());
 			VectorFieldInstance = Instance;
 		}
 		else
diff --git a/Engine/Source/Runtime/Engine/Private/Volume.cpp b/Engine/Source/Runtime/Engine/Private/Volume.cpp
--- a/Engine/Source/Runtime/Engine/Private/Volume.cpp
+++ b/Engine/Source/Runtime/Engine/Private/Volume.cpp
@@ -72,7 +72,7 @@ bool AVolume::EncompassesPoint(FVector Point, float SphereRadius/*=0.f*/, float*
 			return false;
 		}
 #else
-		FBoxSphereBounds Bounds = BrushComponent->CalcBounds(BrushComponent->ComponentToWorld);
+		FBoxSphereBounds Bounds = BrushComponent->CalcBounds(BrushComponent->GetComponentTransform());
 		const float DistanceSqr = Bounds.GetBox().ComputeSquaredDistanceToPoint(Point);
 #endif
 
diff --git a/Engine/Source/Runtime/Engine/Private/WindDirectionalSource.cpp b/Engine/Source/Runtime/Engine/Private/WindDirectionalSource.cpp
--- a/Engine/Source/Runtime/Engine/Private/WindDirectionalSource.cpp
+++ b/Engine/Source/Runtime/Engine/Private/WindDirectionalSource.cpp
@@ -253,7 +253,7 @@ FWindSourceSceneProxy* UWindDirectionalSourceComponent::CreateSceneProxy() const
 	if (bPointWind)
 	{
 		return new FWindSourceSceneProxy(
-			ComponentToWorld.GetLocation(),
+			GetComponentTransform().GetLocation(),
 			Strength,
 			Speed,
 			MinGustAmount,
@@ -264,7 +264,7 @@ FWindSourceSceneProxy* UWindDirectionalSourceComponent::CreateSceneProxy() const
 	else
 	{
 		return new FWindSourceSceneProxy(
-			ComponentToWorld.GetUnitAxis(EAxis::X),
+			GetComponentTransform().GetUnitAxis(EAxis::X),
 			Strength,
 			Speed,
 			MinGustAmount,
@@ -280,12 +280,12 @@ bool UWindDirectionalSourceComponent::GetWindParameters(const FVector& EvaluateP
 
 	if(bPointWind)
 	{
-		FWindSourceSceneProxy LocalProxy = FWindSourceSceneProxy(ComponentToWorld.GetLocation(), Strength, Speed, MinGustAmount, MaxGustAmount, Radius);
+		FWindSourceSceneProxy LocalProxy = FWindSourceSceneProxy(GetComponentTransform().GetLocation(), Strength, Speed, MinGustAmount, MaxGustAmount, Radius);
 		bFoundWind = LocalProxy.GetWindParameters(EvaluatePosition, OutData, Weight);
 	}
 	else
 	{
-		FWindSourceSceneProxy LocalProxy = FWindSourceSceneProxy(ComponentToWorld.GetUnitAxis(EAxis::X), Strength, Speed, MinGustAmount, MaxGustAmount);
+		FWindSourceSceneProxy LocalProxy = FWindSourceSceneProxy(GetComponentTransform().GetUnitAxis(EAxis::X), Strength, Speed, MinGustAmount, MaxGustAmount);
 		bFoundWind = LocalProxy.GetWindParameters(EvaluatePosition, OutData, Weight);
 	}
 
diff --git a/Engine/Source/Runtime/Engine/Public/AI/NavigationSystemHelpers.h b/Engine/Source/Runtime/Engine/Public/AI/NavigationSystemHelpers.h
--- a/Engine/Source/Runtime/Engine/Public/AI/NavigationSystemHelpers.h
+++ b/Engine/Source/Runtime/Engine/Public/AI/NavigationSystemHelpers.h
@@ -44,28 +44,55 @@ struct FNavigableGeometryExport
 
 namespace NavigationHelper
 {
+	struct FNavLinkOwnerData
+	{
+		const AActor* Actor;
+		FTransform LinkToWorld;
+
+		FNavLinkOwnerData() : Actor(nullptr) {}
+		FNavLinkOwnerData(const AActor& InActor);
+		FNavLinkOwnerData(const USceneComponent& InComponent);
+	};
+
 	void GatherCollision(UBodySetup* RigidBody, TNavStatArray<FVector>& OutVertexBuffer, TNavStatArray<int32>& OutIndexBuffer, const FTransform& ComponentToWorld = FTransform::Identity);
 	void GatherCollision(UBodySetup* RigidBody, UNavCollision* NavCollision);
 
 	DECLARE_DELEGATE_ThreeParams(FNavLinkProcessorDelegate, FCompositeNavModifier*, const AActor*, const TArray<FNavigationLink>&);
 	DECLARE_DELEGATE_ThreeParams(FNavLinkSegmentProcessorDelegate, FCompositeNavModifier*, const AActor*, const TArray<FNavigationSegmentLink>&);
 
+	DECLARE_DELEGATE_ThreeParams(FNavLinkProcessorDataDelegate, FCompositeNavModifier*, const FNavLinkOwnerData&, const TArray<FNavigationLink>&);
+	DECLARE_DELEGATE_ThreeParams(FNavLinkSegmentProcessorDataDelegate, FCompositeNavModifier*, const FNavLinkOwnerData&, const TArray<FNavigationSegmentLink>&);
+
 	/** Set new implementation of nav link processor, a function that will be
 	 *	be used to process/transform links before adding them to CompositeModifier.
 	 *	This function is supposed to be called once during the engine/game 
 	 *	setup phase. Not intended to be toggled at runtime */
-	ENGINE_API void SetNavLinkProcessorDelegate(const FNavLinkProcessorDelegate& NewDelegate);
-	ENGINE_API void SetNavLinkSegmentProcessorDelegate(const FNavLinkSegmentProcessorDelegate& NewDelegate);
+	ENGINE_API void SetNavLinkProcessorDelegate(const FNavLinkProcessorDataDelegate& NewDelegate);
+	ENGINE_API void SetNavLinkSegmentProcessorDelegate(const FNavLinkSegmentProcessorDataDelegate& NewDelegate);
 
 	/** called to do any necessary processing on NavLinks and put results in CompositeModifier */
 	ENGINE_API void ProcessNavLinkAndAppend(FCompositeNavModifier* OUT CompositeModifier, const AActor* Actor, const TArray<FNavigationLink>& IN NavLinks);
+	ENGINE_API void ProcessNavLinkAndAppend(FCompositeNavModifier* OUT CompositeModifier, const FNavLinkOwnerData& OwnerData, const TArray<FNavigationLink>& IN NavLinks);
 
 	/** called to do any necessary processing on NavLinks and put results in CompositeModifier */
 	ENGINE_API void ProcessNavLinkSegmentAndAppend(FCompositeNavModifier* OUT CompositeModifier, const AActor* Actor, const TArray<FNavigationSegmentLink>& IN NavLinks);
+	ENGINE_API void ProcessNavLinkSegmentAndAppend(FCompositeNavModifier* OUT CompositeModifier, const FNavLinkOwnerData& OwnerData, const TArray<FNavigationSegmentLink>& IN NavLinks);
 
-	ENGINE_API void DefaultNavLinkProcessorImpl(FCompositeNavModifier* OUT CompositeModifier, const AActor* Actor, const TArray<FNavigationLink>& IN NavLinks);
-
-	ENGINE_API void DefaultNavLinkSegmentProcessorImpl(FCompositeNavModifier* OUT CompositeModifier, const AActor* Actor, const TArray<FNavigationSegmentLink>& IN NavLinks);
+	ENGINE_API void DefaultNavLinkProcessorImpl(FCompositeNavModifier* OUT CompositeModifier, const FNavLinkOwnerData& OwnerData, const TArray<FNavigationLink>& IN NavLinks);
+	ENGINE_API void DefaultNavLinkSegmentProcessorImpl(FCompositeNavModifier* OUT CompositeModifier, const FNavLinkOwnerData& OwnerData, const TArray<FNavigationSegmentLink>& IN NavLinks);
 
 	ENGINE_API bool IsBodyNavigationRelevant(const UBodySetup& IN BodySetup);
+
+	// deprecated functions
+	DEPRECATED_FORGAME(4.17, "FNavLinkProcessorDelegate type is now deprecated, please use FNavLinkProcessorDataDelegate instead.")
+	ENGINE_API void SetNavLinkProcessorDelegate(const FNavLinkProcessorDelegate& NewDelegate);
+	
+	DEPRECATED_FORGAME(4.17, "FNavLinkSegmentProcessorDelegate type is now deprecated, please use FNavLinkSegmentProcessorDataDelegate instead.")
+	ENGINE_API void SetNavLinkSegmentProcessorDelegate(const FNavLinkSegmentProcessorDelegate& NewDelegate);
+	
+	DEPRECATED_FORGAME(4.17, "This function is now deprecated, please use override with OwnerData argument.")
+	ENGINE_API void DefaultNavLinkProcessorImpl(FCompositeNavModifier* OUT CompositeModifier, const AActor* Actor, const TArray<FNavigationLink>& IN NavLinks);
+	
+	DEPRECATED_FORGAME(4.17, "This function is now deprecated, please use override with OwnerData argument.")
+	ENGINE_API void DefaultNavLinkSegmentProcessorImpl(FCompositeNavModifier* OUT CompositeModifier, const AActor* Actor, const TArray<FNavigationSegmentLink>& IN NavLinks);
 }
diff --git a/Engine/Source/Runtime/Engine/Public/Animation/AnimCurveTypes.h b/Engine/Source/Runtime/Engine/Public/Animation/AnimCurveTypes.h
--- a/Engine/Source/Runtime/Engine/Public/Animation/AnimCurveTypes.h
+++ b/Engine/Source/Runtime/Engine/Public/Animation/AnimCurveTypes.h
@@ -40,7 +40,7 @@ struct ENGINE_API FAnimCurveParam
 {
 	GENERATED_USTRUCT_BODY()
 
-	UPROPERTY(EditAnywhere, BlueprintReadWrite, Category = FAnimCurveParam)
+	UPROPERTY(EditAnywhere, Category = FAnimCurveParam)
 	FName Name;
 
 	// name UID for fast access
diff --git a/Engine/Source/Runtime/Engine/Public/ComponentInstanceDataCache.h b/Engine/Source/Runtime/Engine/Public/ComponentInstanceDataCache.h
--- a/Engine/Source/Runtime/Engine/Public/ComponentInstanceDataCache.h
+++ b/Engine/Source/Runtime/Engine/Public/ComponentInstanceDataCache.h
@@ -54,6 +54,7 @@ protected:
 	EComponentCreationMethod SourceComponentCreationMethod;
 
 	TArray<uint8> SavedProperties;
+	TArray<UObject*> InstancedObjects;
 };
 
 /** 
diff --git a/Engine/Source/Runtime/GameplayTags/Private/GameplayTagsManager.cpp b/Engine/Source/Runtime/GameplayTags/Private/GameplayTagsManager.cpp
--- a/Engine/Source/Runtime/GameplayTags/Private/GameplayTagsManager.cpp
+++ b/Engine/Source/Runtime/GameplayTags/Private/GameplayTagsManager.cpp
@@ -15,7 +15,7 @@
 #include "GameplayTagsModule.h"
 #include "Framework/Notifications/NotificationManager.h"
 #include "Widgets/Notifications/SNotificationList.h"
-#include "Engine/Engine.h"
+#include "Misc/CoreDelegates.h"
 
 #if WITH_EDITOR
 #include "SourceControlHelpers.h"
@@ -515,7 +515,7 @@ void UGameplayTagsManager::InitializeManager()
 	SingletonManager->ConstructGameplayTagTree();
 
 	// Bind to end of engine init to be done adding native tags
-	UEngine::OnPostEngineInit.AddUObject(SingletonManager, &UGameplayTagsManager::DoneAddingNativeTags);
+	FCoreDelegates::OnPostEngineInit.AddUObject(SingletonManager, &UGameplayTagsManager::DoneAddingNativeTags);
 }
 
 void UGameplayTagsManager::PopulateTreeFromDataTable(class UDataTable* InTable)
diff --git a/Engine/Source/Runtime/Landscape/Private/LandscapeCollision.cpp b/Engine/Source/Runtime/Landscape/Private/LandscapeCollision.cpp
--- a/Engine/Source/Runtime/Landscape/Private/LandscapeCollision.cpp
+++ b/Engine/Source/Runtime/Landscape/Private/LandscapeCollision.cpp
@@ -598,7 +598,7 @@ bool ULandscapeHeightfieldCollisionComponent::CookCollisionData(const FName& For
 
 	UPhysicalMaterial* DefMaterial = Proxy->DefaultPhysMaterial ? Proxy->DefaultPhysMaterial : GEngine->DefaultPhysMaterial;
 
-	// ComponentToWorld might not be initialized at this point, so use landscape transform
+	// GetComponentTransform() might not be initialized at this point, so use landscape transform
 	const FVector LandscapeScale = Proxy->GetRootComponent()->RelativeScale3D;
 	const bool bIsMirrored = (LandscapeScale.X*LandscapeScale.Y*LandscapeScale.Z) < 0.f;
 
@@ -1291,7 +1291,7 @@ void ULandscapeHeightfieldCollisionComponent::SnapFoliageInstances(const FBox& I
 			if (InstanceSet)
 			{
 				float TraceExtentSize = Bounds.SphereRadius * 2.f + 10.f; // extend a little
-				FVector TraceVector = GetOwner()->GetRootComponent()->ComponentToWorld.GetUnitAxis(EAxis::Z) * TraceExtentSize;
+				FVector TraceVector = GetOwner()->GetRootComponent()->GetComponentTransform().GetUnitAxis(EAxis::Z) * TraceExtentSize;
 
 				bool bFirst = true;
 				TArray<int32> InstancesToRemove;
@@ -1544,7 +1544,7 @@ bool ULandscapeHeightfieldCollisionComponent::DoCustomNavigableGeometryExport(FN
 #if WITH_PHYSX
 	if (IsValidRef(HeightfieldRef) && HeightfieldRef->RBHeightfield)
 	{
-		FTransform HFToW = ComponentToWorld;
+		FTransform HFToW = GetComponentTransform();
 		if (HeightfieldRef->RBHeightfieldSimple)
 		{
 			const float SimpleCollisionScale = CollisionScale * CollisionSizeQuads / SimpleCollisionSizeQuads;
@@ -1566,7 +1566,7 @@ void ULandscapeHeightfieldCollisionComponent::GatherGeometrySlice(FNavigableGeom
 	// note that this function can get called off game thread
 	if (CachedHeightFieldSamples.IsEmpty() == false)
 	{
-		FTransform HFToW = ComponentToWorld;
+		FTransform HFToW = GetComponentTransform();
 		HFToW.MultiplyScale3D(FVector(CollisionScale, CollisionScale, LANDSCAPE_ZSCALE));
 
 		GeomExport.ExportHeightFieldSlice(CachedHeightFieldSamples, HeightfieldRowsCount, HeightfieldColumnsCount, HFToW, SliceBox);
@@ -1620,7 +1620,7 @@ bool ULandscapeMeshCollisionComponent::DoCustomNavigableGeometryExport(FNavigabl
 #if WITH_PHYSX
 	if (IsValidRef(MeshRef) && MeshRef->RBTriangleMesh != nullptr)
 	{
-		FTransform MeshToW = ComponentToWorld;
+		FTransform MeshToW = GetComponentTransform();
 		MeshToW.MultiplyScale3D(FVector(CollisionScale, CollisionScale, 1.f));
 
 		if (MeshRef->RBTriangleMesh->getTriangleMeshFlags() & PxTriangleMeshFlag::e16_BIT_INDICES)
diff --git a/Engine/Source/Runtime/Landscape/Private/LandscapeDataAccess.cpp b/Engine/Source/Runtime/Landscape/Private/LandscapeDataAccess.cpp
--- a/Engine/Source/Runtime/Landscape/Private/LandscapeDataAccess.cpp
+++ b/Engine/Source/Runtime/Landscape/Private/LandscapeDataAccess.cpp
@@ -142,7 +142,7 @@ LANDSCAPE_API FVector FLandscapeComponentDataInterface::GetLocalVertex(int32 Loc
 
 LANDSCAPE_API FVector FLandscapeComponentDataInterface::GetWorldVertex(int32 LocalX, int32 LocalY) const
 {
-	return Component->ComponentToWorld.TransformPosition(GetLocalVertex(LocalX, LocalY));
+	return Component->GetComponentTransform().TransformPosition(GetLocalVertex(LocalX, LocalY));
 }
 
 LANDSCAPE_API void FLandscapeComponentDataInterface::GetWorldTangentVectors(int32 LocalX, int32 LocalY, FVector& WorldTangentX, FVector& WorldTangentY, FVector& WorldTangentZ) const
@@ -154,9 +154,9 @@ LANDSCAPE_API void FLandscapeComponentDataInterface::GetWorldTangentVectors(int3
 	WorldTangentX = FVector(-WorldTangentZ.Z, 0.f, WorldTangentZ.X);
 	WorldTangentY = FVector(0.f, WorldTangentZ.Z, -WorldTangentZ.Y);
 
-	WorldTangentX = Component->ComponentToWorld.TransformVectorNoScale(WorldTangentX);
-	WorldTangentY = Component->ComponentToWorld.TransformVectorNoScale(WorldTangentY);
-	WorldTangentZ = Component->ComponentToWorld.TransformVectorNoScale(WorldTangentZ);
+	WorldTangentX = Component->GetComponentTransform().TransformVectorNoScale(WorldTangentX);
+	WorldTangentY = Component->GetComponentTransform().TransformVectorNoScale(WorldTangentY);
+	WorldTangentZ = Component->GetComponentTransform().TransformVectorNoScale(WorldTangentZ);
 }
 
 LANDSCAPE_API void FLandscapeComponentDataInterface::GetWorldPositionTangents(int32 LocalX, int32 LocalY, FVector& WorldPos, FVector& WorldTangentX, FVector& WorldTangentY, FVector& WorldTangentZ) const
@@ -174,10 +174,10 @@ LANDSCAPE_API void FLandscapeComponentDataInterface::GetWorldPositionTangents(in
 	const float ScaleFactor = (float)Component->ComponentSizeQuads / (float)(ComponentSizeVerts - 1);
 	float XOffset, YOffset;
 	GetXYOffset(LocalX, LocalY, XOffset, YOffset);
-	WorldPos = Component->ComponentToWorld.TransformPosition(FVector(LocalX * ScaleFactor + XOffset, LocalY * ScaleFactor + YOffset, LandscapeDataAccess::GetLocalHeight(Height)));
-	WorldTangentX = Component->ComponentToWorld.TransformVectorNoScale(WorldTangentX);
-	WorldTangentY = Component->ComponentToWorld.TransformVectorNoScale(WorldTangentY);
-	WorldTangentZ = Component->ComponentToWorld.TransformVectorNoScale(WorldTangentZ);
+	WorldPos = Component->GetComponentTransform().TransformPosition(FVector(LocalX * ScaleFactor + XOffset, LocalY * ScaleFactor + YOffset, LandscapeDataAccess::GetLocalHeight(Height)));
+	WorldTangentX = Component->GetComponentTransform().TransformVectorNoScale(WorldTangentX);
+	WorldTangentY = Component->GetComponentTransform().TransformVectorNoScale(WorldTangentY);
+	WorldTangentZ = Component->GetComponentTransform().TransformVectorNoScale(WorldTangentZ);
 }
 
 #endif // WITH_EDITOR
diff --git a/Engine/Source/Runtime/Landscape/Private/LandscapeGrass.cpp b/Engine/Source/Runtime/Landscape/Private/LandscapeGrass.cpp
--- a/Engine/Source/Runtime/Landscape/Private/LandscapeGrass.cpp
+++ b/Engine/Source/Runtime/Landscape/Private/LandscapeGrass.cpp
@@ -708,7 +708,7 @@ public:
 };
 
 FLandscapeComponentGrassData::FLandscapeComponentGrassData(ULandscapeComponent* Component)
-	: RotationForWPO(Component->GetLandscapeMaterial()->GetMaterial()->WorldPositionOffset.IsConnected() ? Component->ComponentToWorld.GetRotation() : FQuat(0, 0, 0, 0))
+	: RotationForWPO(Component->GetLandscapeMaterial()->GetMaterial()->WorldPositionOffset.IsConnected() ? Component->GetComponentTransform().GetRotation() : FQuat(0, 0, 0, 0))
 {
 	UMaterialInterface* Material = Component->GetLandscapeMaterial();
 	for (UMaterialInstanceConstant* MIC = Cast<UMaterialInstanceConstant>(Material); MIC; MIC = Cast<UMaterialInstanceConstant>(Material))
@@ -757,7 +757,7 @@ bool ULandscapeComponent::IsGrassMapOutdated() const
 			return true;
 		}
 
-		FQuat RotationForWPO = GetLandscapeMaterial()->GetMaterial()->WorldPositionOffset.IsConnected() ? ComponentToWorld.GetRotation() : FQuat(0, 0, 0, 0);
+		FQuat RotationForWPO = GetLandscapeMaterial()->GetMaterial()->WorldPositionOffset.IsConnected() ? GetComponentTransform().GetRotation() : FQuat(0, 0, 0, 0);
 		if (GrassData->RotationForWPO != RotationForWPO)
 		{
 			return true;
@@ -1304,7 +1304,7 @@ struct FGrassBuilderBase
 			bHaveValidData = false;
 		}
 		const FRotator DrawRot = Landscape->GetActorRotation();
-		LandscapeToWorld = Landscape->GetRootComponent()->ComponentToWorld.ToMatrixNoScale();
+		LandscapeToWorld = Landscape->GetRootComponent()->GetComponentTransform().ToMatrixNoScale();
 
 		if (bHaveValidData && SqrtSubsections != 1)
 		{
@@ -1416,7 +1416,7 @@ struct FAsyncGrassBuilder : public FGrassBuilderBase
 		, AlignToSurface(GrassVariety.AlignToSurface)
 		, PlacementJitter(GrassVariety.PlacementJitter)
 		, RandomStream(HierarchicalInstancedStaticMeshComponent->InstancingRandomSeed)
-		, XForm(LandscapeToWorld * HierarchicalInstancedStaticMeshComponent->ComponentToWorld.ToMatrixWithScale().Inverse())
+		, XForm(LandscapeToWorld * HierarchicalInstancedStaticMeshComponent->GetComponentTransform().ToMatrixWithScale().Inverse())
 		, MeshBox(GrassVariety.GrassMesh->GetBounds().GetBox())
 		, DesiredInstancesPerLeaf(HierarchicalInstancedStaticMeshComponent->DesiredInstancesPerLeaf())
 
@@ -1985,7 +1985,7 @@ void ALandscapeProxy::UpdateGrass(const TArray<FVector>& Cameras, bool bForceSyn
 					continue;
 				}
 
-				FBoxSphereBounds WorldBounds = Component->CalcBounds(Component->ComponentToWorld);
+				FBoxSphereBounds WorldBounds = Component->CalcBounds(Component->GetComponentTransform());
 				float MinDistanceToComp = Cameras.Num() ? MAX_flt : 0.0f;
 
 				for (auto& Pos : Cameras)
@@ -2053,7 +2053,7 @@ void ALandscapeProxy::UpdateGrass(const TArray<FVector>& Cameras, bool bForceSyn
 											BoxMax.Z = LocalBox.Max.Z;
 
 											FBox LocalSubBox(BoxMin, BoxMax);
-											FBox WorldSubBox = LocalSubBox.TransformBy(Component->ComponentToWorld);
+											FBox WorldSubBox = LocalSubBox.TransformBy(Component->GetComponentTransform());
 
 											MinDistanceToSubComp = Cameras.Num() ? MAX_flt : 0.0f;
 											for (auto& Pos : Cameras)
@@ -2203,7 +2203,7 @@ void ALandscapeProxy::UpdateGrass(const TArray<FVector>& Cameras, bool bForceSyn
 											QUICK_SCOPE_CYCLE_COUNTER(STAT_GrassAttachComp);
 
 											HierarchicalInstancedStaticMeshComponent->AttachToComponent(GetRootComponent(), FAttachmentTransformRules::KeepRelativeTransform);
-											FTransform DesiredTransform = GetRootComponent()->ComponentToWorld;
+											FTransform DesiredTransform = GetRootComponent()->GetComponentTransform();
 											DesiredTransform.RemoveScaling();
 											HierarchicalInstancedStaticMeshComponent->SetWorldTransform(DesiredTransform);
 
diff --git a/Engine/Source/Runtime/Landscape/Private/LandscapeLight.cpp b/Engine/Source/Runtime/Landscape/Private/LandscapeLight.cpp
--- a/Engine/Source/Runtime/Landscape/Private/LandscapeLight.cpp
+++ b/Engine/Source/Runtime/Landscape/Private/LandscapeLight.cpp
@@ -143,7 +143,7 @@ FLandscapeStaticLightingMesh::FLandscapeStaticLightingMesh(ULandscapeComponent*
 	, ExpandQuadsY(InExpandQuadsY)
 {
 	const float LODScale = (float)InComponent->ComponentSizeQuads / (((InComponent->ComponentSizeQuads + 1) >> InLOD) - 1);
-	LocalToWorld = FTransform(FQuat::Identity, FVector::ZeroVector, FVector(LODScale, LODScale, 1)) * InComponent->ComponentToWorld;
+	LocalToWorld = FTransform(FQuat::Identity, FVector::ZeroVector, FVector(LODScale, LODScale, 1)) * InComponent->GetComponentTransform();
 	ComponentSizeQuads = ((InComponent->ComponentSizeQuads + 1) >> InLOD) - 1;
 	NumVertices = ComponentSizeQuads + 2*InExpandQuadsX + 1;
 	NumQuads = NumVertices - 1;
diff --git a/Engine/Source/Runtime/Landscape/Private/LandscapeSplineRaster.cpp b/Engine/Source/Runtime/Landscape/Private/LandscapeSplineRaster.cpp
--- a/Engine/Source/Runtime/Landscape/Private/LandscapeSplineRaster.cpp
+++ b/Engine/Source/Runtime/Landscape/Private/LandscapeSplineRaster.cpp
@@ -440,7 +440,7 @@ bool ULandscapeInfo::ApplySplinesInternal(bool bOnlySelected, ALandscapeProxy* L
 
 	FScopedTransaction Transaction(LOCTEXT("LandscapeSpline_ApplySplines", "Apply Splines to Landscape"));
 
-	const FTransform SplineToLandscape = Landscape->SplineComponent->ComponentToWorld.GetRelativeTransform(Landscape->LandscapeActorToWorld());
+	const FTransform SplineToLandscape = Landscape->SplineComponent->GetComponentTransform().GetRelativeTransform(Landscape->LandscapeActorToWorld());
 
 	FLandscapeEditDataInterface LandscapeEdit(this);
 	TSet<ULandscapeComponent*> ModifiedComponents;
diff --git a/Engine/Source/Runtime/Landscape/Private/LandscapeSplines.cpp b/Engine/Source/Runtime/Landscape/Private/LandscapeSplines.cpp
--- a/Engine/Source/Runtime/Landscape/Private/LandscapeSplines.cpp
+++ b/Engine/Source/Runtime/Landscape/Private/LandscapeSplines.cpp
@@ -835,7 +835,7 @@ ULandscapeSplinesComponent* ULandscapeSplinesComponent::GetStreamingSplinesCompo
 		// this is fine, we won't have any cross-level meshes in this case anyway
 		OuterLandscape->GetLandscapeGuid().IsValid())
 	{
-		FVector LandscapeLocalLocation = ComponentToWorld.GetRelativeTransform(OuterLandscape->LandscapeActorToWorld()).TransformPosition(LocalLocation);
+		FVector LandscapeLocalLocation = GetComponentTransform().GetRelativeTransform(OuterLandscape->LandscapeActorToWorld()).TransformPosition(LocalLocation);
 		const int32 ComponentIndexX = (LandscapeLocalLocation.X >= 0.0f) ? FMath::FloorToInt(LandscapeLocalLocation.X / OuterLandscape->ComponentSizeQuads) : FMath::CeilToInt(LandscapeLocalLocation.X / OuterLandscape->ComponentSizeQuads);
 		const int32 ComponentIndexY = (LandscapeLocalLocation.Y >= 0.0f) ? FMath::FloorToInt(LandscapeLocalLocation.Y / OuterLandscape->ComponentSizeQuads) : FMath::CeilToInt(LandscapeLocalLocation.Y / OuterLandscape->ComponentSizeQuads);
 		ULandscapeComponent* LandscapeComponent = OuterLandscape->GetLandscapeInfo()->XYtoComponentMap.FindRef(FIntPoint(ComponentIndexX, ComponentIndexY));
@@ -1519,7 +1519,7 @@ void ULandscapeSplineControlPoint::UpdateSplinePoints(bool bUpdateCollision, boo
 		FRotator MeshRotation = Rotation;
 		if (MeshComponentOuterSplines != OuterSplines)
 		{
-			const FTransform RelativeTransform = OuterSplines->ComponentToWorld.GetRelativeTransform(MeshComponentOuterSplines->ComponentToWorld);
+			const FTransform RelativeTransform = OuterSplines->GetComponentTransform().GetRelativeTransform(MeshComponentOuterSplines->GetComponentTransform());
 			MeshLocation = RelativeTransform.TransformPosition(MeshLocation);
 		}
 
@@ -2381,7 +2381,7 @@ void ULandscapeSplineSegment::UpdateSplinePoints(bool bUpdateCollision)
 			auto* const MeshComponentOuterSplines = MeshComponent->GetAttachParent();
 			if (MeshComponentOuterSplines != nullptr && MeshComponentOuterSplines != OuterSplines)
 			{
-				const FTransform RelativeTransform = OuterSplines->ComponentToWorld.GetRelativeTransform(MeshComponentOuterSplines->ComponentToWorld);
+				const FTransform RelativeTransform = OuterSplines->GetComponentTransform().GetRelativeTransform(MeshComponentOuterSplines->GetComponentTransform());
 				MeshComponent->SplineParams.StartPos = RelativeTransform.TransformPosition(MeshComponent->SplineParams.StartPos);
 				MeshComponent->SplineParams.EndPos   = RelativeTransform.TransformPosition(MeshComponent->SplineParams.EndPos);
 			}
diff --git a/Engine/Source/Runtime/Launch/Private/LaunchEngineLoop.cpp b/Engine/Source/Runtime/Launch/Private/LaunchEngineLoop.cpp
--- a/Engine/Source/Runtime/Launch/Private/LaunchEngineLoop.cpp
+++ b/Engine/Source/Runtime/Launch/Private/LaunchEngineLoop.cpp
@@ -2043,6 +2043,9 @@ int32 FEngineLoop::PreInit( const TCHAR* CmdLine )
 				}
 			}
 
+			// Call init callbacks
+			FCoreDelegates::OnPostEngineInit.Broadcast();
+
 			// Load all the post-engine init modules
 			ensure(IProjectManager::Get().LoadModulesForProject(ELoadingPhase::PostEngineInit));
 			ensure(IPluginManager::Get().LoadModulesForEnabledPlugins(ELoadingPhase::PostEngineInit));
@@ -2585,7 +2588,11 @@ int32 FEngineLoop::Init()
 
 	GEngine->Init(this);
 
+	// Call init callbacks
+	PRAGMA_DISABLE_DEPRECATION_WARNINGS
 	UEngine::OnPostEngineInit.Broadcast();
+	PRAGMA_ENABLE_DEPRECATION_WARNINGS
+	FCoreDelegates::OnPostEngineInit.Broadcast();
 
 	SlowTask.EnterProgressFrame(30);
 
diff --git a/Engine/Source/Runtime/Navmesh/Private/Detour/DetourNavMesh.cpp b/Engine/Source/Runtime/Navmesh/Private/Detour/DetourNavMesh.cpp
--- a/Engine/Source/Runtime/Navmesh/Private/Detour/DetourNavMesh.cpp
+++ b/Engine/Source/Runtime/Navmesh/Private/Detour/DetourNavMesh.cpp
@@ -827,14 +827,14 @@ const dtNavMeshParams* dtNavMesh::getParams() const
 
 //////////////////////////////////////////////////////////////////////////////////////////
 int dtNavMesh::findConnectingPolys(const float* va, const float* vb,
-								   const dtMeshTile* fromTile, int fromPolyIdx,
-								   const dtMeshTile* tile, int side,
-								   dtPolyRef* con, float* conarea, int maxcon) const
+	const dtMeshTile* fromTile, int fromPolyIdx,
+	const dtMeshTile* tile, int side,
+	dtChunkArray<FConnectingPolyData>& cons) const
 {
 	if (!tile) return 0;
-	
+
 	float amin[2], amax[2], apt[3];
-	calcSlabEndPoints(va,vb, amin,amax, side);
+	calcSlabEndPoints(va, vb, amin, amax, side);
 	const float apos = getSlabCoord(va, side);
 	dtVcopy(apt, va);
 
@@ -842,9 +842,9 @@ int dtNavMesh::findConnectingPolys(const float* va, const float* vb,
 	float bmin[2], bmax[2], bpt[3];
 	unsigned short m = DT_EXT_LINK | (unsigned short)side;
 	int n = 0;
-	
+
 	dtPolyRef base = getPolyRefBase(tile);
-	
+
 	for (int i = 0; i < tile->header->polyCount; ++i)
 	{
 		dtPoly* poly = &tile->polys[i];
@@ -853,21 +853,21 @@ int dtNavMesh::findConnectingPolys(const float* va, const float* vb,
 		{
 			// Skip edges which do not point to the right side.
 			if (poly->neis[j] != m) continue;
-			
-			const float* vc = &tile->verts[poly->verts[j]*3];
-			const float* vd = &tile->verts[poly->verts[(j+1) % nv]*3];
+
+			const float* vc = &tile->verts[poly->verts[j] * 3];
+			const float* vd = &tile->verts[poly->verts[(j + 1) % nv] * 3];
 			const float bpos = getSlabCoord(vc, side);
-			
+
 			// Segments are not close enough.
-			if (dtAbs(apos-bpos) > 0.01f)
+			if (dtAbs(apos - bpos) > 0.01f)
 				continue;
-			
+
 			// Check if the segments touch.
-			calcSlabEndPoints(vc,vd, bmin,bmax, side);
-			
+			calcSlabEndPoints(vc, vd, bmin, bmax, side);
+
 			unsigned char overlapMode = 0;
-			if (!overlapSlabs(amin,amax, bmin,bmax, 0.01f, tile->header->walkableClimb, &overlapMode)) continue;
-			
+			if (!overlapSlabs(amin, amax, bmin, bmax, 0.01f, tile->header->walkableClimb, &overlapMode)) continue;
+
 			// if overlapping with only one side, verify height difference using detailed mesh
 			if (overlapMode == SLABOVERLAP_Max || overlapMode == SLABOVERLAP_Min)
 			{
@@ -884,16 +884,16 @@ int dtNavMesh::findConnectingPolys(const float* va, const float* vb,
 			}
 
 			// Add return value.
-			if (n < maxcon)
-			{
-				conarea[n*2+0] = dtMax(amin[0], bmin[0]);
-				conarea[n*2+1] = dtMin(amax[0], bmax[0]);
-				con[n] = base | (dtPolyRef)i;
-				n++;
-			}
+			FConnectingPolyData NewPolyData;
+			NewPolyData.min = dtMax(amin[0], bmin[0]);
+			NewPolyData.max = dtMin(amax[0], bmax[0]);
+			NewPolyData.ref = base | (dtPolyRef)i;
+			cons.push(NewPolyData);
+			n++;
 			break;
 		}
 	}
+
 	return n;
 }
 
@@ -946,6 +946,8 @@ void dtNavMesh::connectExtLinks(dtMeshTile* tile, dtMeshTile* target, int side,
 {
 	if (!tile) return;
 	
+	dtChunkArray<FConnectingPolyData> cons(16);
+
 	// Connect border links.
 	for (int i = 0; i < tile->header->polyCount; ++i)
 	{
@@ -968,16 +970,19 @@ void dtNavMesh::connectExtLinks(dtMeshTile* tile, dtMeshTile* target, int side,
 			// Create new links
 			const float* va = &tile->verts[poly->verts[j]*3];
 			const float* vb = &tile->verts[poly->verts[(j+1) % nv]*3];
-			dtPolyRef nei[4];
-			float neia[4*2];
-			int nnei = findConnectingPolys(va,vb, tile, i, target, dtOppositeTile(dir), nei,neia,4);
-			for (int k = 0; k < nnei; ++k)
+
+			// reset array before using
+			cons.resize(0);
+			findConnectingPolys(va,vb, tile, i, target, dtOppositeTile(dir), cons);
+
+			for (int k = 0; k < cons.size(); ++k)
 			{
+				const FConnectingPolyData& NeiData = cons[k];
 				unsigned int idx = allocLink(tile, CREATE_LINK_PREALLOCATED);
 				if (idx != DT_NULL_LINK)
 				{
 					dtLink* link = &tile->links[idx];
-					link->ref = nei[k];
+					link->ref = NeiData.ref;
 					link->edge = (unsigned char)j;
 					link->side = (unsigned char)dir;
 					
@@ -987,8 +992,8 @@ void dtNavMesh::connectExtLinks(dtMeshTile* tile, dtMeshTile* target, int side,
 					// Compress portal limits to a byte value.
 					if (dir == 0 || dir == 4)
 					{
-						float tmin = (neia[k*2+0]-va[2]) / (vb[2]-va[2]);
-						float tmax = (neia[k*2+1]-va[2]) / (vb[2]-va[2]);
+						float tmin = (NeiData.min-va[2]) / (vb[2]-va[2]);
+						float tmax = (NeiData.max-va[2]) / (vb[2]-va[2]);
 						if (tmin > tmax)
 							dtSwap(tmin,tmax);
 						link->bmin = (unsigned char)(dtClamp(tmin, 0.0f, 1.0f)*255.0f);
@@ -996,8 +1001,8 @@ void dtNavMesh::connectExtLinks(dtMeshTile* tile, dtMeshTile* target, int side,
 					}
 					else if (dir == 2 || dir == 6)
 					{
-						float tmin = (neia[k*2+0]-va[0]) / (vb[0]-va[0]);
-						float tmax = (neia[k*2+1]-va[0]) / (vb[0]-va[0]);
+						float tmin = (NeiData.min-va[0]) / (vb[0]-va[0]);
+						float tmax = (NeiData.max-va[0]) / (vb[0]-va[0]);
 						if (tmin > tmax)
 							dtSwap(tmin,tmax);
 						link->bmin = (unsigned char)(dtClamp(tmin, 0.0f, 1.0f)*255.0f);
@@ -1007,7 +1012,7 @@ void dtNavMesh::connectExtLinks(dtMeshTile* tile, dtMeshTile* target, int side,
 
 				if (updateCLinks)
 				{
-					unsigned int targetIdx = decodePolyIdPoly(nei[k]);
+					unsigned int targetIdx = decodePolyIdPoly(NeiData.ref);
 					if (tile->polyClusters && target->polyClusters &&
 						i < tile->header->offMeshBase &&
 						targetIdx < (unsigned int)target->header->offMeshBase)
diff --git a/Engine/Source/Runtime/Navmesh/Public/Detour/DetourNavMesh.h b/Engine/Source/Runtime/Navmesh/Public/Detour/DetourNavMesh.h
--- a/Engine/Source/Runtime/Navmesh/Public/Detour/DetourNavMesh.h
+++ b/Engine/Source/Runtime/Navmesh/Public/Detour/DetourNavMesh.h
@@ -783,6 +783,14 @@ public:
 	
 private:
 
+	// [UE4] result struct for findConnectingPolys
+	struct FConnectingPolyData
+	{
+		float min;
+		float max;
+		dtPolyRef ref;
+	};
+
 	/// Returns pointer to tile in the tile array.
 	dtMeshTile* getTile(int i);
 
@@ -794,12 +802,12 @@ private:
 	int getNeighbourTilesAt(const int x, const int y, const int side,
 							dtMeshTile** tiles, const int maxTiles) const;
 	
-	/// Returns all polygons in neighbour tile based on portal defined by the segment.
+	/// [UE4] Returns all polygons in neighbour tile based on portal defined by the segment.
 	int findConnectingPolys(const float* va, const float* vb,
-							const dtMeshTile* fromTile, int fromPolyIdx,
-							const dtMeshTile* tile, int side,
-							dtPolyRef* con, float* conarea, int maxcon) const;
-	
+		const dtMeshTile* fromTile, int fromPolyIdx,
+		const dtMeshTile* tile, int side,
+		dtChunkArray<FConnectingPolyData>& cons) const;
+
 	/// Builds internal polygons links for a tile.
 	void connectIntLinks(dtMeshTile* tile);
 	/// Builds internal polygons links for a tile.
diff --git a/Engine/Source/Runtime/Niagara/Classes/NiagaraDataInterfaceSpline.h b/Engine/Source/Runtime/Niagara/Classes/NiagaraDataInterfaceSpline.h
--- a/Engine/Source/Runtime/Niagara/Classes/NiagaraDataInterfaceSpline.h
+++ b/Engine/Source/Runtime/Niagara/Classes/NiagaraDataInterfaceSpline.h
@@ -60,7 +60,7 @@ private:
 	//Cached ptr to component we sample from. 
 	USplineComponent* Component;
 		
-	//Cached ComponentToWorld.
+	//Cached GetComponentTransform().
 	FMatrix Transform;
 	//InverseTranspose of above for transforming normals/tangents.
 	FMatrix TransformInverseTransposed;
diff --git a/Engine/Source/Runtime/Niagara/Classes/NiagaraDataInterfaceStaticMesh.h b/Engine/Source/Runtime/Niagara/Classes/NiagaraDataInterfaceStaticMesh.h
--- a/Engine/Source/Runtime/Niagara/Classes/NiagaraDataInterfaceStaticMesh.h
+++ b/Engine/Source/Runtime/Niagara/Classes/NiagaraDataInterfaceStaticMesh.h
@@ -176,12 +176,12 @@ private:
 	//Cached ptr to actual mesh we sample from. 
 	UStaticMesh* Mesh;
 	
-	//Cached ComponentToWorld.
+	//Cached GetComponentTransform().
 	FMatrix Transform;
 	//InverseTranspose of above for transforming normals/tangents.
 	FMatrix TransformInverseTransposed;
 
-	//Cached ComponentToWorld from previous tick.
+	//Cached GetComponentTransform() from previous tick.
 	FMatrix PrevTransform;
 	//InverseTranspose of above for transforming normals/tangents.
 	FMatrix PrevTransformInverseTransposed;
diff --git a/Engine/Source/Runtime/Niagara/Classes/NiagaraSimulation.h b/Engine/Source/Runtime/Niagara/Classes/NiagaraSimulation.h
--- a/Engine/Source/Runtime/Niagara/Classes/NiagaraSimulation.h
+++ b/Engine/Source/Runtime/Niagara/Classes/NiagaraSimulation.h
@@ -197,7 +197,7 @@ private:
 	FNiagaraDataSet Data;
 	/** Keep partial particle spawns from last frame */
 	float SpawnRemainder;
-	/** The cached ComponentToWorld transform. */
+	/** The cached GetComponentTransform() transform. */
 	FTransform CachedComponentToWorld;
 
 	NiagaraEffectRenderer *EffectRenderer;
diff --git a/Engine/Source/Runtime/Projects/Private/ModuleDescriptor.cpp b/Engine/Source/Runtime/Projects/Private/ModuleDescriptor.cpp
--- a/Engine/Source/Runtime/Projects/Private/ModuleDescriptor.cpp
+++ b/Engine/Source/Runtime/Projects/Private/ModuleDescriptor.cpp
@@ -172,6 +172,28 @@ bool FModuleDescriptor::Read(const FJsonObject& Object, FText& OutFailReason)
 		}
 	}
 
+	// Read the whitelisted targets
+	TSharedPtr<FJsonValue> WhitelistTargetsValue = Object.TryGetField(TEXT("WhitelistTargets"));
+	if (WhitelistTargetsValue.IsValid() && WhitelistTargetsValue->Type == EJson::Array)
+	{
+		const TArray< TSharedPtr< FJsonValue > >& TargetsArray = WhitelistTargetsValue->AsArray();
+		for (int Idx = 0; Idx < TargetsArray.Num(); Idx++)
+		{
+			WhitelistTargets.Add(TargetsArray[Idx]->AsString());
+		}
+	}
+
+	// Read the blacklisted targets
+	TSharedPtr<FJsonValue> BlacklistTargetsValue = Object.TryGetField(TEXT("BlacklistTargets"));
+	if (BlacklistTargetsValue.IsValid() && BlacklistTargetsValue->Type == EJson::Array)
+	{
+		const TArray< TSharedPtr< FJsonValue > >& TargetsArray = BlacklistTargetsValue->AsArray();
+		for (int Idx = 0; Idx < TargetsArray.Num(); Idx++)
+		{
+			BlacklistTargets.Add(TargetsArray[Idx]->AsString());
+		}
+	}
+
 	// Read the additional dependencies
 	TSharedPtr<FJsonValue> AdditionalDependenciesValue = Object.TryGetField(TEXT("AdditionalDependencies"));
 	if (AdditionalDependenciesValue.IsValid() && AdditionalDependenciesValue->Type == EJson::Array)
@@ -244,6 +266,24 @@ void FModuleDescriptor::Write(TJsonWriter<>& Writer) const
 		}
 		Writer.WriteArrayEnd();
 	}
+	if (WhitelistTargets.Num() > 0)
+	{
+		Writer.WriteArrayStart(TEXT("WhitelistTargets"));
+		for (int Idx = 0; Idx < WhitelistTargets.Num(); Idx++)
+		{
+			Writer.WriteValue(WhitelistTargets[Idx]);
+		}
+		Writer.WriteArrayEnd();
+	}
+	if (BlacklistTargets.Num() > 0)
+	{
+		Writer.WriteArrayStart(TEXT("BlacklistTargets"));
+		for (int Idx = 0; Idx < BlacklistTargets.Num(); Idx++)
+		{
+			Writer.WriteValue(BlacklistTargets[Idx]);
+		}
+		Writer.WriteArrayEnd();
+	}
 	if (AdditionalDependencies.Num() > 0)
 	{
 		Writer.WriteArrayStart(TEXT("AdditionalDependencies"));
@@ -286,6 +326,20 @@ bool FModuleDescriptor::IsCompiledInCurrentConfiguration() const
 		return false;
 	}
 
+	static FString UBTTarget(FPlatformMisc::GetUBTTarget());
+
+	// Check the target is whitelisted
+	if (WhitelistTargets.Num() > 0 && !WhitelistTargets.Contains(UBTTarget))
+	{
+		return false;
+	}
+
+	// Check the target is not blacklisted
+	if (BlacklistTargets.Num() > 0 && BlacklistTargets.Contains(UBTTarget))
+	{
+		return false;
+	}
+
 	// Check the module is compatible with this target. This should match ModuleDescriptor.IsCompiledInConfiguration in UBT
 	switch (Type)
 	{
diff --git a/Engine/Source/Runtime/Projects/Private/PluginDescriptor.cpp b/Engine/Source/Runtime/Projects/Private/PluginDescriptor.cpp
--- a/Engine/Source/Runtime/Projects/Private/PluginDescriptor.cpp
+++ b/Engine/Source/Runtime/Projects/Private/PluginDescriptor.cpp
@@ -18,6 +18,7 @@ FPluginDescriptor::FPluginDescriptor()
 	, bIsMod(false)
 	, bInstalled(false)
 	, bRequiresBuildPlatform(true)
+	, bIsHidden(false)
 { 
 }
 
@@ -122,6 +123,7 @@ bool FPluginDescriptor::Read(const FString& Text, bool bPluginTypeEnabledByDefau
 	{
 		bRequiresBuildPlatform = true;
 	}
+	Object.TryGetBoolField(TEXT("Hidden"), bIsHidden);
 
 	PreBuildSteps.Read(Object, TEXT("PreBuildSteps"));
 	PostBuildSteps.Read(Object, TEXT("PostBuildSteps"));
@@ -194,6 +196,11 @@ void FPluginDescriptor::Write(FString& Text, bool bPluginTypeEnabledByDefault) c
 		Writer.WriteValue(TEXT("RequiresBuildPlatform"), bRequiresBuildPlatform);
 	}
 
+	if (bIsHidden)
+	{
+		Writer.WriteValue(TEXT("Hidden"), bIsHidden);
+	}
+
 	if(!PreBuildSteps.IsEmpty())
 	{
 		PreBuildSteps.Write(Writer, TEXT("PreBuildSteps"));
diff --git a/Engine/Source/Runtime/Projects/Private/PluginManager.cpp b/Engine/Source/Runtime/Projects/Private/PluginManager.cpp
--- a/Engine/Source/Runtime/Projects/Private/PluginManager.cpp
+++ b/Engine/Source/Runtime/Projects/Private/PluginManager.cpp
@@ -120,6 +120,11 @@ bool FPlugin::IsEnabled() const
 	return bEnabled;
 }
 
+bool FPlugin::IsHidden() const
+{
+	return Descriptor.bIsHidden;
+}
+
 bool FPlugin::CanContainContent() const
 {
 	return Descriptor.bCanContainContent;
diff --git a/Engine/Source/Runtime/Projects/Private/PluginManager.h b/Engine/Source/Runtime/Projects/Private/PluginManager.h
--- a/Engine/Source/Runtime/Projects/Private/PluginManager.h
+++ b/Engine/Source/Runtime/Projects/Private/PluginManager.h
@@ -44,6 +44,7 @@ public:
 	virtual FString GetContentDir() const override;
 	virtual FString GetMountedAssetPath() const override;
 	virtual bool IsEnabled() const override;
+	virtual bool IsHidden() const override;
 	virtual bool CanContainContent() const override;
 	virtual EPluginLoadedFrom GetLoadedFrom() const override;
 	virtual const FPluginDescriptor& GetDescriptor() const override;
diff --git a/Engine/Source/Runtime/Projects/Private/ProjectManager.cpp b/Engine/Source/Runtime/Projects/Private/ProjectManager.cpp
--- a/Engine/Source/Runtime/Projects/Private/ProjectManager.cpp
+++ b/Engine/Source/Runtime/Projects/Private/ProjectManager.cpp
@@ -322,6 +322,28 @@ bool FProjectManager::SetPluginEnabled(const FString& PluginName, bool bEnabled,
 	return true;
 }
 
+bool FProjectManager::RemovePluginReference(const FString& PluginName, FText& OutFailReason)
+{
+	// Don't go any further if there's no project loaded
+	if (!CurrentProject.IsValid())
+	{
+		OutFailReason = LOCTEXT("NoProjectLoaded", "No project is currently loaded");
+		return false;
+	}
+
+	bool bPluginFound = false;
+	for (int32 PluginRefIdx = CurrentProject->Plugins.Num() - 1; PluginRefIdx >= 0 && !bPluginFound; --PluginRefIdx)
+	{
+		if (CurrentProject->Plugins[PluginRefIdx].Name == PluginName)
+		{
+			CurrentProject->Plugins.RemoveAt(PluginRefIdx);
+			bPluginFound = true;
+			break;
+		}
+	}
+	return bPluginFound;
+}
+
 void FProjectManager::GetDefaultEnabledPlugins(TArray<FString>& OutPluginNames, bool bIncludeInstalledPlugins)
 {
 	// Add all the game plugins and everything marked as enabled by default
diff --git a/Engine/Source/Runtime/Projects/Private/ProjectManager.h b/Engine/Source/Runtime/Projects/Private/ProjectManager.h
--- a/Engine/Source/Runtime/Projects/Private/ProjectManager.h
+++ b/Engine/Source/Runtime/Projects/Private/ProjectManager.h
@@ -31,6 +31,7 @@ public:
 	virtual FOnTargetPlatformsForCurrentProjectChangedEvent& OnTargetPlatformsForCurrentProjectChanged() override { return OnTargetPlatformsForCurrentProjectChangedEvent; }
 	virtual bool IsNonDefaultPluginEnabled() const override;
 	virtual bool SetPluginEnabled(const FString& PluginName, bool bEnabled, FText& OutFailReason, const FString& MarketplaceURL) override;
+	virtual bool RemovePluginReference(const FString& PluginName, FText& OutFailReason) override;
 	virtual void UpdateAdditionalPluginDirectory(const FString& Dir, const bool bAddOrRemove) override;
 	virtual bool IsCurrentProjectDirty() const override;
 	virtual bool SaveCurrentProjectToDisk(FText& OutFailReason) override;
diff --git a/Engine/Source/Runtime/Projects/Public/Interfaces/IPluginManager.h b/Engine/Source/Runtime/Projects/Public/Interfaces/IPluginManager.h
--- a/Engine/Source/Runtime/Projects/Public/Interfaces/IPluginManager.h
+++ b/Engine/Source/Runtime/Projects/Public/Interfaces/IPluginManager.h
@@ -91,6 +91,13 @@ public:
 	 */
 	virtual bool IsEnabled() const = 0;
 
+	/**
+	 * Determines if the plugin is should be displayed in-editor for the user to enable/disable freely.
+	 *
+	 * @return True if the plugin should be hidden.
+	 */
+	virtual bool IsHidden() const = 0;
+
 	/**
 	 * Determines if the plugin can contain content.
 	 *
diff --git a/Engine/Source/Runtime/Projects/Public/Interfaces/IProjectManager.h b/Engine/Source/Runtime/Projects/Public/Interfaces/IProjectManager.h
--- a/Engine/Source/Runtime/Projects/Public/Interfaces/IProjectManager.h
+++ b/Engine/Source/Runtime/Projects/Public/Interfaces/IProjectManager.h
@@ -198,6 +198,11 @@ public:
 	 */
 	virtual bool SetPluginEnabled(const FString& PluginName, bool bEnabled, FText& OutFailReason, const FString& MarketplaceURL = TEXT("")) = 0;
 
+	/**
+	 * 
+	 */
+	virtual bool RemovePluginReference(const FString& PluginName, FText& OutFailReason) = 0;
+
 	/**
 	 * Updates a directory to be scanned for plugins (added or removed)
 	 *
diff --git a/Engine/Source/Runtime/Projects/Public/ModuleDescriptor.h b/Engine/Source/Runtime/Projects/Public/ModuleDescriptor.h
--- a/Engine/Source/Runtime/Projects/Public/ModuleDescriptor.h
+++ b/Engine/Source/Runtime/Projects/Public/ModuleDescriptor.h
@@ -116,6 +116,12 @@ struct PROJECTS_API FModuleDescriptor
 	/** List of disallowed platforms */
 	TArray<FString> BlacklistPlatforms;
 
+	/** List of allowed targets */
+	TArray<FString> WhitelistTargets;
+
+	/** List of disallowed targets */
+	TArray<FString> BlacklistTargets;
+
 	/** List of additional dependencies for building this module. */
 	TArray<FString> AdditionalDependencies;
 
diff --git a/Engine/Source/Runtime/Projects/Public/PluginDescriptor.h b/Engine/Source/Runtime/Projects/Public/PluginDescriptor.h
--- a/Engine/Source/Runtime/Projects/Public/PluginDescriptor.h
+++ b/Engine/Source/Runtime/Projects/Public/PluginDescriptor.h
@@ -97,6 +97,9 @@ struct PROJECTS_API FPluginDescriptor
 	/** For plugins that are under a platform folder (eg. /PS4/), determines whether compiling the plugin requires the build platform and/or SDK to be available */
 	bool bRequiresBuildPlatform;
 
+	/** For auto-generated plugins that should not be listed in the plugin browser for users to disable freely. */
+	bool bIsHidden;
+
 	/** Pre-build steps for each host platform */
 	FCustomBuildSteps PreBuildSteps;
 
diff --git a/Engine/Source/Runtime/Renderer/Private/AtmosphereRendering.cpp b/Engine/Source/Runtime/Renderer/Private/AtmosphereRendering.cpp
--- a/Engine/Source/Runtime/Renderer/Private/AtmosphereRendering.cpp
+++ b/Engine/Source/Runtime/Renderer/Private/AtmosphereRendering.cpp
@@ -1737,7 +1737,7 @@ FAtmosphericFogSceneInfo::FAtmosphericFogSceneInfo(UAtmosphericFogComponent* InC
 	RenderFlag |= (InComponent->bDisableGroundScattering) ? EAtmosphereRenderFlag::E_DisableGroundScattering : EAtmosphereRenderFlag::E_EnableAll;
 	// Should be same as UpdateAtmosphericFogTransform
 	GroundOffset += InComponent->GetComponentLocation().Z;
-	FMatrix WorldToLight = InComponent->ComponentToWorld.ToMatrixNoScale().InverseFast();
+	FMatrix WorldToLight = InComponent->GetComponentTransform().ToMatrixNoScale().InverseFast();
 	DefaultSunDirection = FVector(WorldToLight.M[0][0],WorldToLight.M[1][0],WorldToLight.M[2][0]);
 
 #if WITH_EDITORONLY_DATA
diff --git a/Engine/Source/Runtime/Renderer/Private/PlanarReflectionRendering.cpp b/Engine/Source/Runtime/Renderer/Private/PlanarReflectionRendering.cpp
--- a/Engine/Source/Runtime/Renderer/Private/PlanarReflectionRendering.cpp
+++ b/Engine/Source/Runtime/Renderer/Private/PlanarReflectionRendering.cpp
@@ -364,7 +364,7 @@ void FScene::UpdatePlanarReflectionContents(UPlanarReflectionComponent* CaptureC
 			});
 		}
 
-		const FMatrix ComponentTransform = CaptureComponent->ComponentToWorld.ToMatrixWithScale();
+		const FMatrix ComponentTransform = CaptureComponent->GetComponentTransform().ToMatrixWithScale();
 		const FPlane MirrorPlane = FPlane(ComponentTransform.TransformPosition(FVector::ZeroVector), ComponentTransform.TransformVector(FVector(0, 0, 1)));
 
 		TArray<FSceneCaptureViewInfo> SceneCaptureViewInfo;
@@ -499,7 +499,7 @@ void FScene::UpdatePlanarReflectionTransform(UPlanarReflectionComponent* Compone
 	ENQUEUE_UNIQUE_RENDER_COMMAND_THREEPARAMETER(
 		FUpdatePlanarReflectionCommand,
 		FPlanarReflectionSceneProxy*,SceneProxy,Component->SceneProxy,
-		FMatrix,Transform,Component->ComponentToWorld.ToMatrixWithScale(),
+		FMatrix,Transform,Component->GetComponentTransform().ToMatrixWithScale(),
 		FScene*,Scene,this,
 	{
 		Scene->ReflectionSceneData.bRegisteredReflectionCapturesHasChanged = true;
diff --git a/Engine/Source/Runtime/Renderer/Private/RendererScene.cpp b/Engine/Source/Runtime/Renderer/Private/RendererScene.cpp
--- a/Engine/Source/Runtime/Renderer/Private/RendererScene.cpp
+++ b/Engine/Source/Runtime/Renderer/Private/RendererScene.cpp
@@ -1182,7 +1182,7 @@ void FScene::AddLight(ULightComponent* Light)
 		Light->SceneProxy = Proxy;
 
 		// Update the light's transform and position.
-		Proxy->SetTransform(Light->ComponentToWorld.ToMatrixNoScale(),Light->GetLightPosition());
+		Proxy->SetTransform(Light->GetComponentTransform().ToMatrixNoScale(),Light->GetLightPosition());
 
 		// Create the light scene info.
 		Proxy->LightSceneInfo = new FLightSceneInfo(Proxy, true);
@@ -1215,7 +1215,7 @@ void FScene::AddInvisibleLight(ULightComponent* Light)
 		Light->SceneProxy = Proxy;
 
 		// Update the light's transform and position.
-		Proxy->SetTransform(Light->ComponentToWorld.ToMatrixNoScale(),Light->GetLightPosition());
+		Proxy->SetTransform(Light->GetComponentTransform().ToMatrixNoScale(),Light->GetLightPosition());
 
 		// Create the light scene info.
 		Proxy->LightSceneInfo = new FLightSceneInfo(Proxy, false);
@@ -1436,7 +1436,7 @@ void FScene::UpdateReflectionCaptureTransform(UReflectionCaptureComponent* Compo
 		ENQUEUE_UNIQUE_RENDER_COMMAND_FOURPARAMETER(
 			UpdateTransformCommand,
 			FReflectionCaptureProxy*,Proxy,Component->SceneProxy,
-			FMatrix,Transform,Component->ComponentToWorld.ToMatrixWithScale(),
+			FMatrix,Transform,Component->GetComponentTransform().ToMatrixWithScale(),
 			const float*,AverageBrightness,Component->GetAverageBrightnessPtr(),
 			FScene*,Scene,this,
 		{
@@ -1696,7 +1696,7 @@ void FScene::UpdateLightTransform(ULightComponent* Light)
 	if(Light->SceneProxy)
 	{
 		FUpdateLightTransformParameters Parameters;
-		Parameters.LightToWorld = Light->ComponentToWorld.ToMatrixNoScale();
+		Parameters.LightToWorld = Light->GetComponentTransform().ToMatrixNoScale();
 		Parameters.Position = Light->GetLightPosition();
 		ENQUEUE_UNIQUE_RENDER_COMMAND_THREEPARAMETER(
 			UpdateLightTransform,
diff --git a/Engine/Source/Runtime/Slate/Private/Framework/Application/SlateApplication.cpp b/Engine/Source/Runtime/Slate/Private/Framework/Application/SlateApplication.cpp
--- a/Engine/Source/Runtime/Slate/Private/Framework/Application/SlateApplication.cpp
+++ b/Engine/Source/Runtime/Slate/Private/Framework/Application/SlateApplication.cpp
@@ -1505,7 +1505,7 @@ void FSlateApplication::FinishedInputThisFrame()
 		}
 	}
 
-	// Any widgets that may have recieved key events
+	// Any widgets that may have received key events
 	// are given a chance to process accumulated values.
 	ForEachUser([&] (FSlateUser* User) {
 		const FSlateUser::FUserFocusEntry& UserFocusEntry = User->Focus;
diff --git a/Engine/Source/Runtime/Slate/Private/Widgets/Input/SInputKeySelector.cpp b/Engine/Source/Runtime/Slate/Private/Widgets/Input/SInputKeySelector.cpp
--- a/Engine/Source/Runtime/Slate/Private/Widgets/Input/SInputKeySelector.cpp
+++ b/Engine/Source/Runtime/Slate/Private/Widgets/Input/SInputKeySelector.cpp
@@ -8,14 +8,15 @@
 void SInputKeySelector::Construct( const FArguments& InArgs )
 {
 	SelectedKey = InArgs._SelectedKey;
+	KeySelectionText = InArgs._KeySelectionText;
+	NoKeySpecifiedText = InArgs._NoKeySpecifiedText;
 	OnKeySelected = InArgs._OnKeySelected;
-	Font = InArgs._Font;
-	ColorAndOpacity = InArgs._ColorAndOpacity;
 	OnIsSelectingKeyChanged = InArgs._OnIsSelectingKeyChanged;
 	bAllowModifierKeys = InArgs._AllowModifierKeys;
+	bAllowGamepadKeys = InArgs._AllowGamepadKeys;
 	bEscapeCancelsSelection = InArgs._EscapeCancelsSelection;
-
-	DefaultFont = FCoreStyle::Get().GetWidgetStyle<FTextBlockStyle>( "NormalText" ).Font;
+	EscapeKeys = InArgs._EscapeKeys;
+	bIsFocusable = InArgs._IsFocusable;
 
 	bIsSelectingKey = false;
 
@@ -23,13 +24,13 @@ void SInputKeySelector::Construct( const FArguments& InArgs )
 	[
 		SAssignNew(Button, SButton)
 		.ButtonStyle(InArgs._ButtonStyle)
+		.IsFocusable(bIsFocusable)
 		.OnClicked(this, &SInputKeySelector::OnClicked)
 		[
-			SNew(STextBlock)
+			SAssignNew(TextBlock, STextBlock)
 			.Text(this, &SInputKeySelector::GetSelectedKeyText)
-			.Font(this, &SInputKeySelector::GetFont)
+			.TextStyle(InArgs._TextStyle)
 			.Margin(Margin)
-			.ColorAndOpacity(this, &SInputKeySelector::GetSlateColorAndOpacity)
 			.Justification(ETextJustify::Center)
 		]
 	];
@@ -51,9 +52,8 @@ FText SInputKeySelector::GetSelectedKeyText() const
 				? SelectedKey.Get().Key.GetDisplayName()
 				: SelectedKey.Get().GetInputText();
 		}
-		return FText();
 	}
-	return FText();
+	return NoKeySpecifiedText;
 }
 
 FInputChord SInputKeySelector::GetSelectedKey() const
@@ -70,16 +70,6 @@ void SInputKeySelector::SetSelectedKey( TAttribute<FInputChord> InSelectedKey )
 	}
 }
 
-FSlateFontInfo SInputKeySelector::GetFont() const
-{
-	return Font.IsSet() ? Font.Get() : DefaultFont;
-}
-
-void SInputKeySelector::SetFont( TAttribute<FSlateFontInfo> InFont )
-{
-	Font = InFont;
-}
-
 FMargin SInputKeySelector::GetMargin() const
 {
 	return Margin.Get();
@@ -90,34 +80,20 @@ void SInputKeySelector::SetMargin( TAttribute<FMargin> InMargin )
 	Margin = InMargin;
 }
 
-FSlateColor SInputKeySelector::GetSlateColorAndOpacity() const
-{
-	return FSlateColor(GetColorAndOpacity());
-}
-
-void SInputKeySelector::SetColorAndOpacity( TAttribute<FLinearColor> InColorAndOpacity )
-{
-	ColorAndOpacity = InColorAndOpacity;
-}
-
 void SInputKeySelector::SetButtonStyle(const FButtonStyle* ButtonStyle )
 {
-	Button->SetButtonStyle(ButtonStyle);
-}
-
-void SInputKeySelector::SetKeySelectionText( FText InKeySelectionText )
-{
-	KeySelectionText = InKeySelectionText;
-}
-
-void SInputKeySelector::SetAllowModifierKeys( bool bInAllowModifierKeys )
-{
-	bAllowModifierKeys = bInAllowModifierKeys;
+	if (Button.IsValid())
+	{
+		Button->SetButtonStyle(ButtonStyle);
+	}
 }
 
-bool SInputKeySelector::GetIsSelectingKey() const
+void SInputKeySelector::SetTextStyle(const FTextBlockStyle* InTextStyle)
 {
-	return bIsSelectingKey;
+	if (TextBlock.IsValid())
+	{
+		TextBlock->SetTextStyle(InTextStyle);
+	}
 }
 
 FReply SInputKeySelector::OnClicked()
@@ -152,21 +128,34 @@ void SInputKeySelector::SetIsSelectingKey( bool bInIsSelectingKey )
 	}
 }
 
+bool SInputKeySelector::IsEscapeKey(const FKey& InKey) const
+{
+	return EscapeKeys.Contains(InKey);
+}
+
 FReply SInputKeySelector::OnPreviewKeyDown( const FGeometry& MyGeometry, const FKeyEvent& InKeyEvent )
 {
-	// TODO: Add an argument to allow gamepad key selection.
-	if ( bIsSelectingKey && InKeyEvent.GetKey().IsGamepadKey() == false)
+	if ( bIsSelectingKey && (bAllowGamepadKeys || InKeyEvent.GetKey().IsGamepadKey() == false) )
 	{
 		// While selecting keys handle all key downs to prevent contained controls from
 		// interfering with key selection.
 		return FReply::Handled();
 	}
-	return SWidget::OnPreviewKeyDown( MyGeometry, InKeyEvent );
+	return SCompoundWidget::OnPreviewKeyDown( MyGeometry, InKeyEvent );
+}
+
+FReply SInputKeySelector::OnKeyDown(const FGeometry& MyGeometry, const FKeyEvent& InKeyEvent)
+{
+	if (!bIsSelectingKey && SelectedKey.IsSet() && SelectedKey.Get().Key.IsValid() && (bAllowGamepadKeys && InKeyEvent.GetKey() == EKeys::Gamepad_FaceButton_Left))
+	{
+		SelectedKey = FInputChord();
+		return FReply::Handled();
+	}
+	return SCompoundWidget::OnKeyDown(MyGeometry, InKeyEvent);
 }
 
 FReply SInputKeySelector::OnKeyUp( const FGeometry& MyGeometry, const FKeyEvent& InKeyEvent )
 {
-	// TODO: Add an argument to allow gamepad key selection.
 	FKey KeyUp = InKeyEvent.GetKey();
 	EModifierKey::Type ModifierKey = EModifierKey::FromBools(
 		InKeyEvent.IsControlDown() && KeyUp != EKeys::LeftControl && KeyUp != EKeys::RightControl,
@@ -175,11 +164,12 @@ FReply SInputKeySelector::OnKeyUp( const FGeometry& MyGeometry, const FKeyEvent&
 		InKeyEvent.IsCommandDown() && KeyUp != EKeys::LeftCommand && KeyUp != EKeys::RightCommand );
 
 	// Don't allow chords consisting of just modifier keys.
-	if ( bIsSelectingKey && KeyUp.IsGamepadKey() == false && ( KeyUp.IsModifierKey() == false || ModifierKey == EModifierKey::None ) )
+	if ( bIsSelectingKey && (bAllowGamepadKeys || KeyUp.IsGamepadKey() == false) && ( KeyUp.IsModifierKey() == false || ModifierKey == EModifierKey::None ) )
 	{
 		SetIsSelectingKey( false );
 
-		if ( KeyUp == EKeys::Escape && bEscapeCancelsSelection )
+		if ((InKeyEvent.GetKey() == EKeys::PS4_Special) || // Required?
+			(bEscapeCancelsSelection && (KeyUp == EKeys::Escape || IsEscapeKey(KeyUp))))
 		{
 			return FReply::Handled();
 		}
@@ -192,7 +182,8 @@ FReply SInputKeySelector::OnKeyUp( const FGeometry& MyGeometry, const FKeyEvent&
 			ModifierKey == EModifierKey::Command);
 		return FReply::Handled();
 	}
-	return SWidget::OnPreviewKeyDown( MyGeometry, InKeyEvent );
+
+	return SCompoundWidget::OnKeyUp( MyGeometry, InKeyEvent );
 }
 
 FReply SInputKeySelector::OnPreviewMouseButtonDown( const FGeometry& MyGeometry, const FPointerEvent& MouseEvent )
@@ -204,7 +195,27 @@ FReply SInputKeySelector::OnPreviewMouseButtonDown( const FGeometry& MyGeometry,
 		SelectKey(MouseEvent.GetEffectingButton(), false, false, false, false);
 		return FReply::Handled();
 	}
-	return SWidget::OnPreviewMouseButtonDown( MyGeometry, MouseEvent );
+	return SCompoundWidget::OnPreviewMouseButtonDown( MyGeometry, MouseEvent );
+}
+
+FReply SInputKeySelector::OnMouseButtonDown(const FGeometry& MyGeometry, const FPointerEvent& MouseEvent)
+{
+	if (!bIsSelectingKey && SelectedKey.IsSet() && SelectedKey.Get().Key.IsValid() && MouseEvent.IsMouseButtonDown(EKeys::RightMouseButton))
+	{
+		SelectedKey = FInputChord();
+		return FReply::Handled();
+	}
+	return SCompoundWidget::OnMouseButtonDown(MyGeometry, MouseEvent);
+}
+
+FNavigationReply SInputKeySelector::OnNavigation(const FGeometry& MyGeometry, const FNavigationEvent& InNavigationEvent)
+{
+	if (Button.IsValid())
+	{
+		return Button->OnNavigation(MyGeometry, InNavigationEvent);
+	}
+
+	return SCompoundWidget::OnNavigation(MyGeometry, InNavigationEvent);
 }
 
 void SInputKeySelector::OnFocusLost( const FFocusEvent& InFocusEvent )
diff --git a/Engine/Source/Runtime/Slate/Public/Framework/Commands/InputChord.h b/Engine/Source/Runtime/Slate/Public/Framework/Commands/InputChord.h
--- a/Engine/Source/Runtime/Slate/Public/Framework/Commands/InputChord.h
+++ b/Engine/Source/Runtime/Slate/Public/Framework/Commands/InputChord.h
@@ -18,23 +18,23 @@ struct SLATE_API FInputChord
 	GENERATED_USTRUCT_BODY()
 
 	/** The Key is the core of the chord. */
-	UPROPERTY(EditAnywhere, Category = Key)
+	UPROPERTY(EditAnywhere, BlueprintReadWrite, Category = Key)
 	FKey Key;
 
 	/** Whether the shift key is part of the chord.  */
-	UPROPERTY(EditAnywhere, Category = Modifier)
+	UPROPERTY(EditAnywhere, BlueprintReadWrite, Category = Modifier)
 	uint32 bShift:1;
 
 	/** Whether the control key is part of the chord.  */
-	UPROPERTY(EditAnywhere, Category = Modifier)
+	UPROPERTY(EditAnywhere, BlueprintReadWrite, Category = Modifier)
 	uint32 bCtrl:1;
 
 	/** Whether the alt key is part of the chord.  */
-	UPROPERTY(EditAnywhere, Category = Modifier)
+	UPROPERTY(EditAnywhere, BlueprintReadWrite, Category = Modifier)
 	uint32 bAlt:1;
 
 	/** Whether the command key is part of the chord.  */
-	UPROPERTY(EditAnywhere, Category = Modifier)
+	UPROPERTY(EditAnywhere, BlueprintReadWrite, Category = Modifier)
 	uint32 bCmd:1;
 
 	/** 
