Copying //UE4/Dev-Framework to //UE4/Dev-Main (Source: //UE4/Dev-Framework @ 3431384)

diff --git a/Engine/Source/Runtime/Slate/Public/Widgets/Input/SInputKeySelector.h b/Engine/Source/Runtime/Slate/Public/Widgets/Input/SInputKeySelector.h
--- a/Engine/Source/Runtime/Slate/Public/Widgets/Input/SInputKeySelector.h
+++ b/Engine/Source/Runtime/Slate/Public/Widgets/Input/SInputKeySelector.h
@@ -17,6 +17,7 @@
 #include "Framework/Commands/InputChord.h"
 
 class SButton;
+class STextBlock;
 
 /**
  * A widget for selecting keys or input chords.
@@ -30,10 +31,15 @@ public:
 	SLATE_BEGIN_ARGS( SInputKeySelector )
 		: _SelectedKey( FInputChord(EKeys::Invalid) )
 		, _ButtonStyle( &FCoreStyle::Get().GetWidgetStyle<FButtonStyle>( "Button" ) )
-		, _KeySelectionText( NSLOCTEXT("DefaultKeySelectionText", "InputKeySelector", "...") )
+		, _TextStyle( &FCoreStyle::Get().GetWidgetStyle< FTextBlockStyle >("NormalText") )
+		, _KeySelectionText( NSLOCTEXT("InputKeySelector", "DefaultKeySelectionText", "...") )
+		, _NoKeySpecifiedText(NSLOCTEXT("InputKeySelector", "DefaultEmptyText", "Empty"))
 		, _AllowModifierKeys( true )
+		, _AllowGamepadKeys( false )
 		, _EscapeCancelsSelection( true )
-	{}
+		, _IsFocusable(true)
+		{}
+
 		/** The currently selected key */
 		SLATE_ATTRIBUTE( FInputChord, SelectedKey )
 
@@ -43,27 +49,38 @@ public:
 		/** The margin around the selected key text. */
 		SLATE_ATTRIBUTE( FMargin, Margin )
 
-		/** The color of the selected key text. */
-		SLATE_ATTRIBUTE( FLinearColor, ColorAndOpacity )
-
 		/** The style of the button used to enable key selection. */
 		SLATE_STYLE_ARGUMENT( FButtonStyle, ButtonStyle )
 
+		/** The text style of the button text */
+		SLATE_STYLE_ARGUMENT( FTextBlockStyle, TextStyle)
+
 		/** The text to display while selecting a new key. */
 		SLATE_ARGUMENT( FText, KeySelectionText )
 
+		/** The text to display while no key text is available or not selecting a key. */
+		SLATE_ARGUMENT(FText, NoKeySpecifiedText)
+
 		/** When true modifier keys are captured in the selected key chord, otherwise they are ignored. */
 		SLATE_ARGUMENT( bool, AllowModifierKeys )
 
+		/** When true gamepad keys are captured in the selected key chord, otherwise they are ignored. */
+		SLATE_ARGUMENT( bool, AllowGamepadKeys )
+
 		/** When true, pressing escape will cancel the key selection, when false, pressing escape will select the escape key. */
 		SLATE_ARGUMENT( bool, EscapeCancelsSelection )
 
+		/** When EscapeCancelsSelection is true, escape on specific keys that are unbind able by the user. */
+		SLATE_ARGUMENT( TArray<FKey>, EscapeKeys )
+
 		/** Occurs whenever a new key is selected. */
 		SLATE_EVENT( FOnKeySelected, OnKeySelected )
 
 		/** Occurs whenever key selection mode starts and stops. */
 		SLATE_EVENT( FOnIsSelectingKeyChanged, OnIsSelectingKeyChanged )
 
+		/** Sometimes a button should only be mouse-clickable and never keyboard focusable. */
+		SLATE_ARGUMENT(bool, IsFocusable)
 	SLATE_END_ARGS()
 
 	void Construct( const FArguments& InArgs );
@@ -76,50 +93,52 @@ public:
 	/** Sets the currently selected key chord. */
 	void SetSelectedKey( TAttribute<FInputChord> InSelectedKey );
 
-	/** Sets the font which is used to display the currently selected key. */
-	void SetFont( TAttribute<FSlateFontInfo> InFont );
-
 	/** Sets the margin around the text used to display the currently selected key */
 	void SetMargin( TAttribute<FMargin> InMargin );
 
-	/** Sets the color of the text used to display the currently selected key. */
-	void SetColorAndOpacity( TAttribute<FLinearColor> InColorAndOpacity );
-
 	/** Sets the style of the button which is used enter key selection mode. */
 	void SetButtonStyle( const FButtonStyle* ButtonStyle );
 
+	/** Sets the style of the text on the button which is used enter key selection mode. */
+	void SetTextStyle(const FTextBlockStyle* InTextStyle);
+
 	/** Sets the text which is displayed when selecting a key. */
-	void SetKeySelectionText( FText InKeySelectionText );
+	void SetKeySelectionText( FText InKeySelectionText ) { KeySelectionText = MoveTemp(InKeySelectionText); }
+
+	/** Sets the text to display when no key text is available or not selecting a key. */
+	void SetNoKeySpecifiedText(FText InNoKeySpecifiedText) { NoKeySpecifiedText = MoveTemp(InNoKeySpecifiedText); }
 
 	/** When true modifier keys are captured in the selected key chord, otherwise they are ignored. */
-	void SetAllowModifierKeys( bool bInAllowModifierKeys );
+	void SetAllowModifierKeys(const bool bInAllowModifierKeys) { bAllowModifierKeys = bInAllowModifierKeys; }
+
+	/** When true gamepad keys are captured in the selected key chord, otherwise they are ignored. */
+	void SetAllowGamepadKeys(const bool bInAllowGamepadKeys) { bAllowGamepadKeys = bInAllowGamepadKeys; }
+
+	/** Sets the escape keys to check against. */
+	void SetEscapeKeys(TArray<FKey> InEscapeKeys) { EscapeKeys = MoveTemp(InEscapeKeys); }
 
 	/** Returns true whenever key selection mode is active, otherwise returns false. */
-	bool GetIsSelectingKey() const;
+	bool GetIsSelectingKey() const { return bIsSelectingKey; }
 
 public:
 
 	virtual FReply OnPreviewKeyDown( const FGeometry& MyGeometry, const FKeyEvent& InKeyEvent ) override;
+	virtual FReply OnKeyDown(const FGeometry& MyGeometry, const FKeyEvent& InKeyEvent) override;
 	virtual FReply OnKeyUp( const FGeometry& MyGeometry, const FKeyEvent& InKeyEvent ) override;
 	virtual FReply OnPreviewMouseButtonDown( const FGeometry& MyGeometry, const FPointerEvent& MouseEvent ) override;
-
+	virtual FReply OnMouseButtonDown(const FGeometry& MyGeometry, const FPointerEvent& MouseEvent) override;
 	virtual void OnFocusLost( const FFocusEvent& InFocusEvent ) override;
-	virtual bool SupportsKeyboardFocus() const { return true; }
+	virtual FNavigationReply OnNavigation(const FGeometry& MyGeometry, const FNavigationEvent& InNavigationEvent) override;
+	virtual bool SupportsKeyboardFocus() const override { return true; }
 
 private:
 
 	/** Gets the display text for the currently selected key. */
 	FText GetSelectedKeyText() const;
 
-	/** Gets the font used to display the currently selected key. */
-	FSlateFontInfo GetFont() const;
-
 	/**  Gets the margin around the text used to display the currently selected key. */
 	FMargin GetMargin() const;
 
-	/** Gets the color of the text used to display the currently selected key. */
-	FSlateColor GetSlateColorAndOpacity() const;
-
 	/** Handles the OnClicked event from the button which enables key selection mode. */
 	FReply OnClicked();
 
@@ -129,6 +148,9 @@ private:
 	/** Sets bIsSelectingKey and invokes the associated events. */
 	void SetIsSelectingKey(bool bInIsSelectingKey);
 
+	/** Returns true, if the key has been specified as an escape key, else false. */
+	bool IsEscapeKey(const FKey& InKey) const;
+
 private:
 
 	/** True when key selection mode is active. */
@@ -137,23 +159,26 @@ private:
 	/** The currently selected key chord. */
 	TAttribute<FInputChord> SelectedKey;
 
-	/** The font for the text used to display the currently selected key. */
-	TAttribute<FSlateFontInfo> Font;
-
 	/** The margin around the text used to display the currently selected key. */
 	TAttribute<FMargin> Margin;
 
 	/** The text to display when selecting keys. */
 	FText KeySelectionText;
 
+	/**  The text to display while no key text is available or not selecting a key. */
+	FText NoKeySpecifiedText;
+
 	/** When true modifier keys are recorded on the selected key chord, otherwise they are ignored. */
 	bool bAllowModifierKeys;
 
+	/** When true gamepad keys are recorded on the selected key chord, otherwise they are ignored. */
+	bool bAllowGamepadKeys;
+
 	/** When true, pressing escape will cancel the key selection, when false, pressing escape will select the escape key. */ 
 	bool bEscapeCancelsSelection;
 
-	/** The default font used for the text which displays the currently selected key. */
-	FSlateFontInfo DefaultFont;
+	/** When EscapeCancelsSelection is true, escape on specific keys that are unbind able by the user. */
+	TArray<FKey> EscapeKeys;
 
 	/** Delegate which is run any time a new key is selected. */
 	FOnKeySelected OnKeySelected;
@@ -161,6 +186,12 @@ private:
 	/** Delegate which is run when key selection mode starts and stops. */
 	FOnIsSelectingKeyChanged OnIsSelectingKeyChanged;
 
-	/** The button which starts key key selection mode. */
+	/** The button which starts the key selection mode. */
 	TSharedPtr<SButton> Button;
+
+	/** The text which is rendered on the button. */
+	TSharedPtr<STextBlock> TextBlock;
+
+	/** Can this button be focused? */
+	bool bIsFocusable;
 };
diff --git a/Engine/Source/Runtime/UMG/Private/Components/InputKeySelector.cpp b/Engine/Source/Runtime/UMG/Private/Components/InputKeySelector.cpp
--- a/Engine/Source/Runtime/UMG/Private/Components/InputKeySelector.cpp
+++ b/Engine/Source/Runtime/UMG/Private/Components/InputKeySelector.cpp
@@ -1,18 +1,52 @@
 // Copyright 1998-2017 Epic Games, Inc. All Rights Reserved.
 
 #include "Components/InputKeySelector.h"
+#include "Engine/Font.h"
+#include "UObject/ConstructorHelpers.h"
+#include "UObject/FrameworkObjectVersion.h"
 #include "Widgets/DeclarativeSyntaxSupport.h"
 #include "Widgets/Input/SInputKeySelector.h"
 
 UInputKeySelector::UInputKeySelector( const FObjectInitializer& ObjectInitializer )
+	: Super(ObjectInitializer)
 {
 	SInputKeySelector::FArguments InputKeySelectorDefaults;
+	WidgetStyle = *InputKeySelectorDefaults._ButtonStyle;
+	TextStyle = *InputKeySelectorDefaults._TextStyle;
+	KeySelectionText = InputKeySelectorDefaults._KeySelectionText;
+	NoKeySpecifiedText = InputKeySelectorDefaults._NoKeySpecifiedText;
 	SelectedKey = InputKeySelectorDefaults._SelectedKey.Get();
-	ButtonStyle = InputKeySelectorDefaults._ButtonStyle;
 	bAllowModifierKeys = InputKeySelectorDefaults._AllowModifierKeys;
+	bAllowGamepadKeys = InputKeySelectorDefaults._AllowGamepadKeys;
+
+	EscapeKeys.AddUnique(EKeys::Gamepad_Special_Right); // In most (if not all) cases this is going to be the menu button
+
+	if (!IsRunningDedicatedServer())
+	{
+		static ConstructorHelpers::FObjectFinder<UFont> RobotoFontObj(TEXT("/Engine/EngineFonts/Roboto"));
+		TextStyle.Font = FSlateFontInfo(RobotoFontObj.Object, 24, FName("Bold"));
+	}
+}
+
+void UInputKeySelector::Serialize(FArchive& Ar)
+{
+	Super::Serialize(Ar);
+
+	Ar.UsingCustomVersion(FFrameworkObjectVersion::GUID);
+}
+
+void UInputKeySelector::PostLoad()
+{
+	Super::PostLoad();
+
+	if (GetLinkerCustomVersion(FFrameworkObjectVersion::GUID) < FFrameworkObjectVersion::InputKeySelectorTextStyle)
+	{
+		TextStyle.Font = Font_DEPRECATED;
+		TextStyle.ColorAndOpacity = ColorAndOpacity_DEPRECATED;
+	}
 }
 
-void UInputKeySelector::SetSelectedKey( FInputChord InSelectedKey )
+void UInputKeySelector::SetSelectedKey( const FInputChord& InSelectedKey )
 {
 	if ( MyInputKeySelector.IsValid() )
 	{
@@ -27,10 +61,19 @@ void UInputKeySelector::SetKeySelectionText( FText InKeySelectionText )
 	{
 		MyInputKeySelector->SetKeySelectionText( InKeySelectionText );
 	}
-	KeySelectionText = InKeySelectionText;
+	KeySelectionText = MoveTemp(InKeySelectionText);
 }
 
-void UInputKeySelector::SetAllowModifierKeys( bool bInAllowModifierKeys )
+void UInputKeySelector::SetNoKeySpecifiedText(FText InNoKeySpecifiedText)
+{
+	if (MyInputKeySelector.IsValid())
+	{
+		MyInputKeySelector->SetNoKeySpecifiedText(InNoKeySpecifiedText);
+	}
+	NoKeySpecifiedText = MoveTemp(InNoKeySpecifiedText);
+}
+
+void UInputKeySelector::SetAllowModifierKeys( const bool bInAllowModifierKeys )
 {
 	if ( MyInputKeySelector.IsValid() )
 	{
@@ -39,6 +82,15 @@ void UInputKeySelector::SetAllowModifierKeys( bool bInAllowModifierKeys )
 	bAllowModifierKeys = bInAllowModifierKeys;
 }
 
+void UInputKeySelector::SetAllowGamepadKeys(const bool bInAllowGamepadKeys)
+{
+	if (MyInputKeySelector.IsValid())
+	{
+		MyInputKeySelector->SetAllowGamepadKeys(bInAllowGamepadKeys);
+	}
+	bAllowGamepadKeys = bInAllowGamepadKeys;
+}
+
 bool UInputKeySelector::GetIsSelectingKey() const
 {
 	return MyInputKeySelector.IsValid() ? MyInputKeySelector->GetIsSelectingKey() : false;
@@ -48,9 +100,9 @@ void UInputKeySelector::SetButtonStyle( const FButtonStyle* InButtonStyle )
 {
 	if ( MyInputKeySelector.IsValid() )
 	{
-		MyInputKeySelector->SetButtonStyle( ButtonStyle );
+		MyInputKeySelector->SetButtonStyle(InButtonStyle);
 	}
-	ButtonStyle = InButtonStyle;
+	WidgetStyle = *InButtonStyle;
 }
 
 void UInputKeySelector::SynchronizeProperties()
@@ -58,24 +110,33 @@ void UInputKeySelector::SynchronizeProperties()
 	Super::SynchronizeProperties();
 
 	MyInputKeySelector->SetSelectedKey( SelectedKey );
-	MyInputKeySelector->SetFont( Font );
 	MyInputKeySelector->SetMargin( Margin );
-	MyInputKeySelector->SetColorAndOpacity( ColorAndOpacity );
-	MyInputKeySelector->SetButtonStyle( ButtonStyle );
+	MyInputKeySelector->SetButtonStyle( &WidgetStyle );
+	MyInputKeySelector->SetTextStyle( &TextStyle );
 	MyInputKeySelector->SetKeySelectionText( KeySelectionText );
 	MyInputKeySelector->SetAllowModifierKeys( bAllowModifierKeys );
+	MyInputKeySelector->SetAllowGamepadKeys(bAllowGamepadKeys);
+	MyInputKeySelector->SetEscapeKeys(EscapeKeys);
+}
+
+void UInputKeySelector::ReleaseSlateResources(bool bReleaseChildren)
+{
+	Super::ReleaseSlateResources(bReleaseChildren);
+
+	MyInputKeySelector.Reset();
 }
 
 TSharedRef<SWidget> UInputKeySelector::RebuildWidget()
 {
 	MyInputKeySelector = SNew(SInputKeySelector)
 		.SelectedKey(SelectedKey)
-		.Font(Font)
 		.Margin(Margin)
-		.ColorAndOpacity(ColorAndOpacity)
-		.ButtonStyle(ButtonStyle)
+		.ButtonStyle(&WidgetStyle)
+		.TextStyle(&TextStyle)
 		.KeySelectionText(KeySelectionText)
 		.AllowModifierKeys(bAllowModifierKeys)
+		.AllowGamepadKeys(bAllowGamepadKeys)
+		.EscapeKeys(EscapeKeys)
 		.OnKeySelected( BIND_UOBJECT_DELEGATE( SInputKeySelector::FOnKeySelected, HandleKeySelected ) )
 		.OnIsSelectingKeyChanged( BIND_UOBJECT_DELEGATE( SInputKeySelector::FOnIsSelectingKeyChanged, HandleIsSelectingKeyChanged ) );
 	return MyInputKeySelector.ToSharedRef();
diff --git a/Engine/Source/Runtime/UMG/Private/Components/WidgetComponent.cpp b/Engine/Source/Runtime/UMG/Private/Components/WidgetComponent.cpp
--- a/Engine/Source/Runtime/UMG/Private/Components/WidgetComponent.cpp
+++ b/Engine/Source/Runtime/UMG/Private/Components/WidgetComponent.cpp
@@ -666,7 +666,7 @@ FCollisionShape UWidgetComponent::GetCollisionShape(float Inflation) const
 {
 	if ( Space != EWidgetSpace::Screen )
 	{
-		FVector BoxHalfExtent = ( FVector(0.01f, DrawSize.X * 0.5f, DrawSize.Y * 0.5f) * ComponentToWorld.GetScale3D() ) + Inflation;
+		FVector BoxHalfExtent = ( FVector(0.01f, DrawSize.X * 0.5f, DrawSize.Y * 0.5f) * GetComponentTransform().GetScale3D() ) + Inflation;
 
 		if ( Inflation < 0.0f )
 		{
@@ -1405,7 +1405,7 @@ void UWidgetComponent::GetLocalHitLocation(FVector WorldHitLocation, FVector2D&
 	ensureMsgf(GeometryMode == EWidgetGeometryMode::Plane, TEXT("Method does not support non-planar widgets."));
 
 	// Find the hit location on the component
-	FVector ComponentHitLocation = ComponentToWorld.InverseTransformPosition(WorldHitLocation);
+	FVector ComponentHitLocation = GetComponentTransform().InverseTransformPosition(WorldHitLocation);
 
 	// Convert the 3D position of component space, into the 2D equivalent
 	OutLocalWidgetHitLocation = FVector2D(-ComponentHitLocation.Y, -ComponentHitLocation.Z);
@@ -1457,8 +1457,8 @@ TTuple<FVector, FVector2D> UWidgetComponent::GetCylinderHitLocation(FVector Worl
 
 	FTransform ToWorld = GetComponentToWorld();
 
-	const FVector HitLocation_ComponentSpace = ComponentToWorld.InverseTransformPosition(WorldHitLocation);
-	const FVector HitDirection_ComponentSpace = ComponentToWorld.InverseTransformVector(WorldHitDirection);
+	const FVector HitLocation_ComponentSpace = GetComponentTransform().InverseTransformPosition(WorldHitLocation);
+	const FVector HitDirection_ComponentSpace = GetComponentTransform().InverseTransformVector(WorldHitDirection);
 
 
 	const float ArcAngleRadians = FMath::DegreesToRadians(GetCylinderArcAngle());
@@ -1527,7 +1527,7 @@ TTuple<FVector, FVector2D> UWidgetComponent::GetCylinderHitLocation(FVector Worl
 
 		const FVector2D WidgetSpaceHitCoord = FVector2D(HitAngleZeroToOne * CurrentDrawSize.X, YHitLocation);
 			
-		return MakeTuple(ComponentToWorld.TransformPosition(CylinderHitLocation_ComponentSpace), WidgetSpaceHitCoord);
+		return MakeTuple(GetComponentTransform().TransformPosition(CylinderHitLocation_ComponentSpace), WidgetSpaceHitCoord);
 	}
 	else
 	{
diff --git a/Engine/Source/Runtime/UMG/Public/Components/ExpandableArea.h b/Engine/Source/Runtime/UMG/Public/Components/ExpandableArea.h
--- a/Engine/Source/Runtime/UMG/Public/Components/ExpandableArea.h
+++ b/Engine/Source/Runtime/UMG/Public/Components/ExpandableArea.h
@@ -28,7 +28,7 @@ class UMG_API UExpandableArea : public UWidget, public INamedSlotInterface
 
 public:
 
-	UPROPERTY(EditAnywhere, BlueprintReadOnly, Category = "Style")
+	UPROPERTY(EditAnywhere, Category = "Style")
 	FExpandableAreaStyle Style;
 
 	UPROPERTY( EditAnywhere, BlueprintReadOnly, Category = "Style" )
diff --git a/Engine/Source/Runtime/UMG/Public/Components/InputKeySelector.h b/Engine/Source/Runtime/UMG/Public/Components/InputKeySelector.h
--- a/Engine/Source/Runtime/UMG/Public/Components/InputKeySelector.h
+++ b/Engine/Source/Runtime/UMG/Public/Components/InputKeySelector.h
@@ -9,6 +9,7 @@
 #include "Layout/Margin.h"
 #include "Widgets/SWidget.h"
 #include "Components/Widget.h"
+#include "Styling/SlateTypes.h"
 #include "InputKeySelector.generated.h"
 
 class SInputKeySelector;
@@ -25,31 +26,49 @@ public:
 	DECLARE_DYNAMIC_MULTICAST_DELEGATE( FOnIsSelectingKeyChanged );
 
 public:
+	/** The button style used at runtime */
+	UPROPERTY(EditAnywhere, BlueprintReadWrite, Category = Appearance, meta = (DisplayName = "Style"))
+	FButtonStyle WidgetStyle;
+
+	/** The button style used at runtime */
+	UPROPERTY(EditAnywhere, BlueprintReadWrite, Category = Appearance, meta = (DisplayName = "Text Style"))
+	FTextBlockStyle TextStyle;
+
 	/** The currently selected key chord. */
 	UPROPERTY( BlueprintReadOnly, Category = "Key Selection" )
 	FInputChord SelectedKey;
 
-	/** The font used to display the currently selected key. */
-	UPROPERTY( EditAnywhere, BlueprintReadOnly, Category = Appearance )
-	FSlateFontInfo Font;
+	UPROPERTY()
+	FSlateFontInfo Font_DEPRECATED;
 
 	/** The amount of blank space around the text used to display the currently selected key. */
 	UPROPERTY( EditAnywhere, BlueprintReadOnly, Category = Appearance )
 	FMargin Margin;
 
-	/** The color of the text used to display the currently selected key. */
-	UPROPERTY( EditAnywhere, BlueprintReadOnly, Category = Appearance )
-	FLinearColor ColorAndOpacity;
+	UPROPERTY()
+	FLinearColor ColorAndOpacity_DEPRECATED;
 
 	/** Sets the text which is displayed while selecting keys. */
 	UPROPERTY( EditAnywhere, BlueprintReadOnly, Category = Appearance )
 	FText KeySelectionText;
 
+	/** Sets the text to display when no key text is available or not selecting a key. */
+	UPROPERTY(EditAnywhere, BlueprintReadOnly, Category = Appearance)
+	FText NoKeySpecifiedText;
+
 	/** When true modifier keys such as control and alt are allowed in the */
 	/** input chord representing the selected key, if false modifier keys are ignored. */
 	UPROPERTY( EditAnywhere, BlueprintReadOnly, Category = "Key Selection" )
 	bool bAllowModifierKeys;
 
+	/** When true gamepad keys are allowed in the input chord representing the selected key, otherwise they are ignored. */
+	UPROPERTY(EditAnywhere, BlueprintReadOnly, Category = "Key Selection")
+	bool bAllowGamepadKeys;
+
+	/** When true gamepad keys are allowed in the input chord representing the selected key, otherwise they are ignored. */
+	UPROPERTY(EditAnywhere, BlueprintReadOnly, Category = "Key Selection")
+	TArray<FKey> EscapeKeys;
+
 public:
 	/** Called whenever a new key is selected by the user. */
 	UPROPERTY( BlueprintAssignable, Category = "Widget Event" )
@@ -61,16 +80,24 @@ public:
 
 	/** Sets the currently selected key. */
 	UFUNCTION( BlueprintCallable, Category = "Widget" )
-	void SetSelectedKey(FInputChord InSelectedKey);
+	void SetSelectedKey(const FInputChord& InSelectedKey);
 
 	/** Sets the text which is displayed while selecting keys. */
 	UFUNCTION( BlueprintCallable, Category = "Widget" )
 	void SetKeySelectionText( FText InKeySelectionText );
 
+	/** Sets the text to display when no key text is available or not selecting a key. */
+	UFUNCTION(BlueprintCallable, Category = "Widget")
+	void SetNoKeySpecifiedText(FText InNoKeySpecifiedText);
+
 	/** Sets whether or not modifier keys are allowed in the selected key. */
 	UFUNCTION(BlueprintCallable, Category = "Widget" )
 	void SetAllowModifierKeys(bool bInAllowModifierKeys );
 
+	/** Sets whether or not gamepad keys are allowed in the selected key. */
+	UFUNCTION(BlueprintCallable, Category = "Widget")
+	void SetAllowGamepadKeys(bool bInAllowGamepadKeys);
+
 	/** Returns true if the widget is currently selecting a key, otherwise returns false. */
 	UFUNCTION( BlueprintCallable, Category = "Widget" )
 	bool GetIsSelectingKey() const;
@@ -83,18 +110,24 @@ public:
 	//~ End UWidget Interface
 
 protected:
+	//~ Begin UObject Interface
+	virtual void PostLoad() override;
+	virtual void Serialize(FArchive& Ar) override;
+	//~ End UObject Interface
+
 	//~ Begin UWidget Interface
 	virtual TSharedRef<SWidget> RebuildWidget() override;
 	//~ End UWidget Interface
 
+	//~ Begin UVisual Interface
+	virtual void ReleaseSlateResources(bool bReleaseChildren) override;
+	//~ End UVisual Interface
+
 private:
 	virtual void HandleKeySelected(const FInputChord& InSelectedKey);
 	void HandleIsSelectingKeyChanged();
 
 private:
-	/** The style for the button used to start key selection mode. */
-	const FButtonStyle* ButtonStyle;
-
 	/** The input key selector widget managed by this object. */
 	TSharedPtr<SInputKeySelector> MyInputKeySelector;
 };
