Copying //UE4/Dev-Framework to //UE4/Dev-Main (Source: //UE4/Dev-Framework @ 3459469)

diff --git a/Engine/Source/Runtime/Slate/Private/Framework/Application/SlateApplication.cpp b/Engine/Source/Runtime/Slate/Private/Framework/Application/SlateApplication.cpp
--- a/Engine/Source/Runtime/Slate/Private/Framework/Application/SlateApplication.cpp
+++ b/Engine/Source/Runtime/Slate/Private/Framework/Application/SlateApplication.cpp
@@ -3918,7 +3918,7 @@ float FSlateApplication::GetCursorRadius() const
 	return CursorRadius;
 }
 
-FVector2D FSlateApplication::CalculatePopupWindowPosition( const FSlateRect& InAnchor, const FVector2D& InSize, const EOrientation Orientation ) const
+FVector2D FSlateApplication::CalculatePopupWindowPosition( const FSlateRect& InAnchor, const FVector2D& InSize, const FVector2D& InProposedPlacement, const EOrientation Orientation) const
 {
 	FVector2D CalculatedPopUpWindowPosition( 0, 0 );
 
@@ -3935,7 +3935,7 @@ FVector2D FSlateApplication::CalculatePopupWindowPosition( const FSlateRect& InA
 		PopUpOrientation =  EPopUpOrientation::Vertical;
 	}
 
-	if ( PlatformApplication->TryCalculatePopupWindowPosition( AnchorRect, InSize, PopUpOrientation, /*OUT*/&CalculatedPopUpWindowPosition ) )
+	if ( PlatformApplication->TryCalculatePopupWindowPosition( AnchorRect, InSize, InProposedPlacement, PopUpOrientation, /*OUT*/&CalculatedPopUpWindowPosition ) )
 	{
 		return CalculatedPopUpWindowPosition;
 	}
@@ -3955,10 +3955,15 @@ FVector2D FSlateApplication::CalculatePopupWindowPosition( const FSlateRect& InA
 			PlatformWorkArea.Left+(PlatformWorkArea.Right - PlatformWorkArea.Left), 
 			PlatformWorkArea.Top+(PlatformWorkArea.Bottom - PlatformWorkArea.Top) );
 
-		// Assume natural left-to-right, top-to-bottom flow; position popup below and to the right.
-		const FVector2D ProposedPlacement(
-			Orientation == Orient_Horizontal ? AnchorRect.Right : AnchorRect.Left,
-			Orientation == Orient_Horizontal ? AnchorRect.Top : AnchorRect.Bottom);
+		FVector2D ProposedPlacement = InProposedPlacement;
+
+		if (ProposedPlacement.IsZero())
+		{
+			// Assume natural left-to-right, top-to-bottom flow; position popup below and to the right.
+			ProposedPlacement = FVector2D(
+				Orientation == Orient_Horizontal ? AnchorRect.Right : AnchorRect.Left,
+				Orientation == Orient_Horizontal ? AnchorRect.Top : AnchorRect.Bottom);
+		}
 
 		return ComputePopupFitInRect(InAnchor, FSlateRect(ProposedPlacement, ProposedPlacement+InSize), Orientation, WorkAreaRect);
 	}
diff --git a/Engine/Source/Runtime/Slate/Private/Widgets/Input/SMenuAnchor.cpp b/Engine/Source/Runtime/Slate/Private/Widgets/Input/SMenuAnchor.cpp
--- a/Engine/Source/Runtime/Slate/Private/Widgets/Input/SMenuAnchor.cpp
+++ b/Engine/Source/Runtime/Slate/Private/Widgets/Input/SMenuAnchor.cpp
@@ -85,6 +85,7 @@ FGeometry ComputeMenuPlacement(const FGeometry& AllottedGeometry, const FVector2
 	const FVector2D NewPositionDesktopSpace = FSlateApplication::Get().CalculatePopupWindowPosition(
 		TransformRect(AllottedGeometry.GetAccumulatedLayoutTransform(), Placement.AnchorLocalSpace),
 		TransformVector(AllottedGeometry.GetAccumulatedLayoutTransform(), Placement.LocalPopupSize),
+		TransformPoint(AllottedGeometry.GetAccumulatedLayoutTransform(), Placement.LocalPopupOffset),
 		Placement.Orientation);
 
 	// transform the desktop offset into local space and use that as the layout transform for the child content.
@@ -430,7 +431,7 @@ void SMenuAnchor::SetIsOpen( bool InIsOpen, const bool bFocusMenu, const int32 F
 							MenuContentRef->SlatePrepass(FSlateApplication::Get().GetApplicationScale());
 							// @todo slate: Doesn't take into account potential window border size
 							FVector2D ExpectedSize = MenuContentRef->GetDesiredSize();
-							const FVector2D ScreenPosition = FSlateApplication::Get().CalculatePopupWindowPosition(Anchor, ExpectedSize, Orientation);
+							const FVector2D ScreenPosition = FSlateApplication::Get().CalculatePopupWindowPosition(Anchor, ExpectedSize, FVector2D::ZeroVector, Orientation);
 
 							// Release the mouse so that context can be properly restored upon closing menus.  See CL 1411833 before changing this.
 							if (bFocusMenu)
diff --git a/Engine/Source/Runtime/Slate/Private/Widgets/Input/SSubMenuHandler.cpp b/Engine/Source/Runtime/Slate/Private/Widgets/Input/SSubMenuHandler.cpp
--- a/Engine/Source/Runtime/Slate/Private/Widgets/Input/SSubMenuHandler.cpp
+++ b/Engine/Source/Runtime/Slate/Private/Widgets/Input/SSubMenuHandler.cpp
@@ -20,7 +20,7 @@ void SSubMenuHandler::Construct(const FArguments& InArgs, TWeakPtr<SMenuOwner> I
 		if (InArgs._OnGetMenuContent.IsBound() || InArgs._MenuContent.IsValid())
 		{
 			ChildSlotWidget = SAssignNew(MenuAnchor, SMenuAnchor)
-				.Placement(EMenuPlacement::MenuPlacement_MenuRight)
+				.Placement(InArgs._Placement)
 				.OnGetMenuContent(InArgs._OnGetMenuContent)
 				.MenuContent(InArgs._MenuContent)
 				[
diff --git a/Engine/Source/Runtime/Slate/Public/Framework/Application/SlateApplication.h b/Engine/Source/Runtime/Slate/Public/Framework/Application/SlateApplication.h
--- a/Engine/Source/Runtime/Slate/Public/Framework/Application/SlateApplication.h
+++ b/Engine/Source/Runtime/Slate/Public/Framework/Application/SlateApplication.h
@@ -755,14 +755,15 @@ public:
 	 * Calculates the popup window position from the passed in window position and size. 
 	 * Adjusts position for popup windows which are outside of the work area of the monitor where they reside
 	 *
-	 * @param InAnchor	The current(suggseted) window position and size of an area which may not be covered by the popup.
-	 * @param InSize		The size of the window
-	 * @param Orientation	The direction of the popup.
-	 *						If vertical it will attempt to open below the anchor but will open above if there is no room.
-	 *						If horizontal it will attempt to open below the anchor but will open above if there is no room.
+	 * @param InAnchor				The current(suggseted) window position and size of an area which may not be covered by the popup.
+	 * @param InSize				The size of the window
+	 * @param InProposedPlacement	The location on screen where the popup should go if allowed. If zero this will be determined from Orientation and Anchor
+	 * @param Orientation			The direction of the popup.
+	 *								If vertical it will attempt to open below the anchor but will open above if there is no room.
+	 *								If horizontal it will attempt to open below the anchor but will open above if there is no room.
 	 * @return The adjusted position
 	 */
-	virtual FVector2D CalculatePopupWindowPosition( const FSlateRect& InAnchor, const FVector2D& InSize, const EOrientation Orientation = Orient_Vertical ) const;
+	virtual FVector2D CalculatePopupWindowPosition( const FSlateRect& InAnchor, const FVector2D& InSize, const FVector2D& InProposedPlacement = FVector2D::ZeroVector, const EOrientation Orientation = Orient_Vertical) const;
 
 	/**
 	 * Is the window in the app's destroy queue? If so it will be destroyed next tick.
diff --git a/Engine/Source/Runtime/Slate/Public/Widgets/Input/SSubMenuHandler.h b/Engine/Source/Runtime/Slate/Public/Widgets/Input/SSubMenuHandler.h
--- a/Engine/Source/Runtime/Slate/Public/Widgets/Input/SSubMenuHandler.h
+++ b/Engine/Source/Runtime/Slate/Public/Widgets/Input/SSubMenuHandler.h
@@ -16,10 +16,12 @@ class SLATE_API SSubMenuHandler : public SCompoundWidget
 {
 public:
 	SLATE_BEGIN_ARGS( SSubMenuHandler )
+		: _Placement( EMenuPlacement::MenuPlacement_MenuRight )
 	{}
 		SLATE_DEFAULT_SLOT( FArguments, Content )
 		SLATE_ARGUMENT( TSharedPtr<SMenuAnchor>, MenuAnchor )
 		SLATE_ARGUMENT( TSharedPtr<SWidget>, MenuContent )
+		SLATE_ATTRIBUTE( EMenuPlacement, Placement )
 		SLATE_EVENT( FOnGetContent, OnGetMenuContent )
 	SLATE_END_ARGS()
 
