Slate - Properly adding and removing entries from the PointerIndexLastPositionMap so that new touch events don't inherit the movement of a previous touch event.

diff --git a/Engine/Source/Runtime/Slate/Private/Framework/Application/SlateApplication.cpp b/Engine/Source/Runtime/Slate/Private/Framework/Application/SlateApplication.cpp
--- a/Engine/Source/Runtime/Slate/Private/Framework/Application/SlateApplication.cpp
+++ b/Engine/Source/Runtime/Slate/Private/Framework/Application/SlateApplication.cpp
@@ -6106,6 +6106,9 @@ bool FSlateApplication::OnTouchStarted( const TSharedPtr< FGenericWindow >& Plat
 
 void FSlateApplication::ProcessTouchStartedEvent( const TSharedPtr< FGenericWindow >& PlatformWindow, FPointerEvent& PointerEvent )
 {
+	// Add or Update the entry if the finger has been added to the surface.
+	PointerIndexLastPositionMap.Add(FUserAndPointer(PointerEvent.GetUserIndex(), PointerEvent.GetPointerIndex()), PointerEvent.GetScreenSpacePosition());
+
 	ProcessMouseButtonDownEvent( PlatformWindow, PointerEvent );
 }
 
@@ -6175,6 +6178,9 @@ void FSlateApplication::ShouldSimulateGesture(EGestureEvent Gesture, bool bEnabl
 void FSlateApplication::ProcessTouchEndedEvent( FPointerEvent& PointerEvent )
 {
 	ProcessMouseButtonUpEvent(PointerEvent);
+
+	// Remove the entry if the finger has been removed from the surface.
+	PointerIndexLastPositionMap.Remove(FUserAndPointer(PointerEvent.GetUserIndex(), PointerEvent.GetPointerIndex()));
 }
 
 bool FSlateApplication::OnMotionDetected(const FVector& Tilt, const FVector& RotationRate, const FVector& Gravity, const FVector& Acceleration, int32 ControllerId)
