Copying //UE4/Orion-Staging to //UE4/Main (Source: //Orion/Dev-General @ 3564337)

diff --git a/Engine/Source/Runtime/Slate/Private/Widgets/SViewport.cpp b/Engine/Source/Runtime/Slate/Private/Widgets/SViewport.cpp
--- a/Engine/Source/Runtime/Slate/Private/Widgets/SViewport.cpp
+++ b/Engine/Source/Runtime/Slate/Private/Widgets/SViewport.cpp
@@ -6,7 +6,9 @@
 #include "Application/SlateApplicationBase.h"
 #include "Framework/Application/SlateApplication.h"
 
-DECLARE_CYCLE_STAT(TEXT("Game UI Tick/Paint"), STAT_ViewportUpdateTime, STATGROUP_Slate);
+DECLARE_CYCLE_STAT(TEXT("Game UI Tick"), STAT_ViewportTickTime, STATGROUP_Slate);
+DECLARE_CYCLE_STAT(TEXT("Game UI Paint"), STAT_ViewportPaintTime, STATGROUP_Slate);
+
 
 /* SViewport constructors
  *****************************************************************************/
@@ -61,7 +63,7 @@ EActiveTimerReturnType SViewport::EnsureTick(double InCurrentTime, float InDelta
 int32 SViewport::OnPaint( const FPaintArgs& Args, const FGeometry& AllottedGeometry, const FSlateRect& MyCullingRect, FSlateWindowElementList& OutDrawElements, int32 LayerId, const FWidgetStyle& InWidgetStyle, bool bParentEnabled ) const
 {
 	SCOPED_NAMED_EVENT(SViewport_OnPaint, FColor::Purple);
-	SCOPE_CYCLE_COUNTER(STAT_ViewportUpdateTime);
+	SCOPE_CYCLE_COUNTER(STAT_ViewportPaintTime);
 
 	bool bEnabled = ShouldBeEnabled( bParentEnabled );
 	bool bShowDisabledEffect = ShowDisabledEffect.Get();
@@ -168,7 +170,7 @@ int32 SViewport::OnPaint( const FPaintArgs& Args, const FGeometry& AllottedGeome
 
 void SViewport::Tick( const FGeometry& AllottedGeometry, const double InCurrentTime, const float InDeltaTime )
 {
-	SCOPE_CYCLE_COUNTER(STAT_ViewportUpdateTime);
+	SCOPE_CYCLE_COUNTER(STAT_ViewportTickTime);
 
 	if ( ViewportInterface.IsValid() )
 	{
diff --git a/Engine/Source/Runtime/Slate/Public/Framework/Application/SlateApplication.h b/Engine/Source/Runtime/Slate/Public/Framework/Application/SlateApplication.h
--- a/Engine/Source/Runtime/Slate/Public/Framework/Application/SlateApplication.h
+++ b/Engine/Source/Runtime/Slate/Public/Framework/Application/SlateApplication.h
@@ -1137,6 +1137,8 @@ public:
 
 public:
 
+	TSharedRef<FNavigationConfig> GetNavigationConfig() const { return NavigationConfigFactory(); }
+	
 	void SetNavigationConfigFactory( TFunction<TSharedRef<FNavigationConfig>()> InNavigationConfigFactory );
 
 	/** Called when the slate application is being shut down. */
diff --git a/Engine/Source/Runtime/Slate/Public/Widgets/Input/SMenuAnchor.h b/Engine/Source/Runtime/Slate/Public/Widgets/Input/SMenuAnchor.h
--- a/Engine/Source/Runtime/Slate/Public/Widgets/Input/SMenuAnchor.h
+++ b/Engine/Source/Runtime/Slate/Public/Widgets/Input/SMenuAnchor.h
@@ -41,6 +41,7 @@ public:
 		, _ShouldDeferPaintingAfterWindowContent(true)
 		, _UseApplicationMenuStack(true)
 		, _IsCollapsedByParent(false)
+		, _ApplyWidgetStyleToMenu(true)
 		{}
 		
 		SLATE_DEFAULT_SLOT( FArguments, Content )
@@ -64,6 +65,9 @@ public:
 		/** True if this menu anchor should be collapsed when its parent receives focus, false (default) otherwise */
 		SLATE_ARGUMENT(bool, IsCollapsedByParent)
 
+		/** True to apply the InWidgetStyle of the menu anchor when painting the popup, false to always paint the popup with full opacity and no tint. */
+		SLATE_ARGUMENT(bool, ApplyWidgetStyleToMenu)
+
 	SLATE_END_ARGS()
 
 	virtual ~SMenuAnchor();
@@ -131,9 +135,22 @@ protected:
 	/** Handler/callback called by menus created by this anchor, when they are dismissed */
 	void OnMenuClosed(TSharedRef<IMenu> InMenu);
 
+	/** Computes the placement geometry for menus displayed in a separately created window */
+	FGeometry ComputeNewWindowMenuPlacement(const FGeometry& AllottedGeometry, const FVector2D& PopupDesiredSize, EMenuPlacement PlacementMode) const;
+
 	static TArray<TWeakPtr<IMenu>> OpenApplicationMenus;
 
 protected:
+	struct SLATE_API FPopupPlacement
+	{
+		FPopupPlacement(const FGeometry& PlacementGeometry, const FVector2D& PopupDesiredSize, EMenuPlacement PlacementMode);
+		
+		FVector2D LocalPopupSize;
+		FVector2D LocalPopupOffset;
+		FSlateRect AnchorLocalSpace;
+		EOrientation Orientation;
+	};
+
 	/**
 	 * A pointer to the window presenting this popup.
 	 * Can be the window created to hold a menu or the window containing this anchor if the menu is drawn as a child of the this anchor.
@@ -186,6 +203,9 @@ protected:
 	/** Should the menu by created by the application stack code making it behave like and have the lifetime of a normal menu? */
 	bool bUseApplicationMenuStack;
 
+	/** Should we paint the popup using the received InWidgetStyle? */
+	bool bApplyWidgetStyleToMenu;
+
 	/**
 	 * @todo Slate : Unify geometry so that this is not necessary.
 	 * Stores the local offset of the popup as we have to compute it in OnTick;
diff --git a/Engine/Source/Runtime/Slate/Public/Widgets/Layout/SScrollBox.h b/Engine/Source/Runtime/Slate/Public/Widgets/Layout/SScrollBox.h
--- a/Engine/Source/Runtime/Slate/Public/Widgets/Layout/SScrollBox.h
+++ b/Engine/Source/Runtime/Slate/Public/Widgets/Layout/SScrollBox.h
@@ -187,6 +187,7 @@ public:
 
 	void SetScrollBarThickness(FVector2D InThickness);
 
+	void SetScrollBarRightClickDragAllowed(bool bIsAllowed);
 public:
 
 	// SWidget interface
@@ -359,6 +360,7 @@ private:
 	/** Whether the active timer to update the inertial scroll is registered */
 	bool bIsScrollingActiveTimerRegistered : 1;
 
-	/**  */
+	bool bAllowsRightClickDragScrolling : 1;
+	
 	bool bTouchPanningCapture : 1;
 };
diff --git a/Engine/Source/Runtime/SlateCore/Public/Types/SlateEnums.h b/Engine/Source/Runtime/SlateCore/Public/Types/SlateEnums.h
--- a/Engine/Source/Runtime/SlateCore/Public/Types/SlateEnums.h
+++ b/Engine/Source/Runtime/SlateCore/Public/Types/SlateEnums.h
@@ -220,6 +220,9 @@ enum EMenuPlacement
 
 	/** Place the menu's vertical center on the left side at the menu anchor's vertical center on the right side */
 	MenuPlacement_RightLeftCenter UMETA(DisplayName = "Right Center, Left Center"),
+	
+	/** Place the menu's bottom left corner directly on top of the menu anchor's bottom left corner */
+	MenuPlacement_MatchBottomLeft UMETA(DisplayName = "Match Bottom Left Corner"),
 };
 
 
diff --git a/Engine/Source/Runtime/UMG/Private/Components/ScrollBox.cpp b/Engine/Source/Runtime/UMG/Private/Components/ScrollBox.cpp
--- a/Engine/Source/Runtime/UMG/Private/Components/ScrollBox.cpp
+++ b/Engine/Source/Runtime/UMG/Private/Components/ScrollBox.cpp
@@ -28,6 +28,7 @@ UScrollBox::UScrollBox(const FObjectInitializer& ObjectInitializer)
 
 	WidgetStyle = *Defaults._Style;
 	WidgetBarStyle = *Defaults._ScrollBarStyle;
+	bAllowRightClickDragScrolling = true;
 }
 
 void UScrollBox::ReleaseSlateResources(bool bReleaseChildren)
@@ -97,6 +98,7 @@ void UScrollBox::SynchronizeProperties()
 	MyScrollBox->SetScrollBarThickness(ScrollbarThickness);
 	MyScrollBox->SetScrollBarAlwaysVisible(AlwaysShowScrollbar);
 	MyScrollBox->SetAllowOverscroll(AllowOverscroll ? EAllowOverscroll::Yes : EAllowOverscroll::No);
+	MyScrollBox->SetScrollBarRightClickDragAllowed(bAllowRightClickDragScrolling);
 }
 
 float UScrollBox::GetScrollOffset() const
diff --git a/Engine/Source/Runtime/UMG/Private/Slate/SObjectWidget.cpp b/Engine/Source/Runtime/UMG/Private/Slate/SObjectWidget.cpp
--- a/Engine/Source/Runtime/UMG/Private/Slate/SObjectWidget.cpp
+++ b/Engine/Source/Runtime/UMG/Private/Slate/SObjectWidget.cpp
@@ -449,10 +449,10 @@ FReply SObjectWidget::OnMotionDetected(const FGeometry& MyGeometry, const FMotio
 
 FNavigationReply SObjectWidget::OnNavigation(const FGeometry& MyGeometry, const FNavigationEvent& InNavigationEvent)
 {
-	// This reply will represent anything setup by the user with their Navigation Metadata, and we provide it to the user
-	// no matter what so they could make intelligent decisions like, for example, making a UUserWidget Subclass that plays a
-	// sound when the navigation is stopped, it could look at the default reply, and respond with a sound, letting the user
-	// know that way is not supported.
+	if (WidgetObject->NativeSupportsCustomNavigation())
+	{
+		return WidgetObject->NativeOnNavigation(MyGeometry, InNavigationEvent);
+	}
 	FNavigationReply Reply = SCompoundWidget::OnNavigation(MyGeometry, InNavigationEvent);
 
 	if ( CanRouteEvent() )
diff --git a/Engine/Source/Runtime/UMG/Private/UserWidget.cpp b/Engine/Source/Runtime/UMG/Private/UserWidget.cpp
--- a/Engine/Source/Runtime/UMG/Private/UserWidget.cpp
+++ b/Engine/Source/Runtime/UMG/Private/UserWidget.cpp
@@ -1645,6 +1645,11 @@ FCursorReply UUserWidget::NativeOnCursorQuery( const FGeometry& InGeometry, cons
 	return FCursorReply::Unhandled();
 }
 
+FNavigationReply UUserWidget::NativeOnNavigation(const FGeometry& InGeometry, const FNavigationEvent& InNavigationEvent)
+{
+	return FNavigationReply::Escape();
+}
+	
 void UUserWidget::NativeOnMouseCaptureLost()
 {
 	OnMouseCaptureLost();
@@ -1822,6 +1827,7 @@ UUserWidget* UUserWidget::CreateWidgetOfClass(UClass* UserWidgetClass, UGameInst
 
 	UObject* Outer = nullptr;
 	ULocalPlayer* PlayerContext = nullptr;
+	UWorld* World = InWorld;
 
 	if ( InOwningPlayer )
 	{
@@ -1844,7 +1850,7 @@ UUserWidget* UUserWidget::CreateWidgetOfClass(UClass* UserWidgetClass, UGameInst
 		}
 
 		// Assign the outer to the game instance if it exists, otherwise use the player controller's world
-		UWorld* World = InOwningPlayer->GetWorld();
+		World = InOwningPlayer->GetWorld();
 
 		Outer = World->GetGameInstance() ? StaticCast<UObject*>(World->GetGameInstance()) : StaticCast<UObject*>(World);
 		PlayerContext = CastChecked<ULocalPlayer>(InOwningPlayer->Player);
@@ -1871,7 +1877,7 @@ UUserWidget* UUserWidget::CreateWidgetOfClass(UClass* UserWidgetClass, UGameInst
 
 	if ( PlayerContext )
 	{
-		NewWidget->SetPlayerContext(FLocalPlayerContext(PlayerContext));
+		NewWidget->SetPlayerContext(FLocalPlayerContext(PlayerContext, World));
 	}
 
 	NewWidget->Initialize();
diff --git a/Engine/Source/Runtime/UMG/Public/Blueprint/UserWidget.h b/Engine/Source/Runtime/UMG/Public/Blueprint/UserWidget.h
--- a/Engine/Source/Runtime/UMG/Public/Blueprint/UserWidget.h
+++ b/Engine/Source/Runtime/UMG/Public/Blueprint/UserWidget.h
@@ -967,6 +967,7 @@ protected:
 
 	virtual bool NativeIsInteractable() const;
 	virtual bool NativeSupportsKeyboardFocus() const;
+	virtual bool NativeSupportsCustomNavigation() const { return false; }
 
 	virtual FReply NativeOnFocusReceived( const FGeometry& InGeometry, const FFocusEvent& InFocusEvent );
 	virtual void NativeOnFocusLost( const FFocusEvent& InFocusEvent );
@@ -999,6 +1000,7 @@ protected:
 	virtual FReply NativeOnTouchEnded( const FGeometry& InGeometry, const FPointerEvent& InGestureEvent );
 	virtual FReply NativeOnMotionDetected( const FGeometry& InGeometry, const FMotionEvent& InMotionEvent );
 	virtual FCursorReply NativeOnCursorQuery( const FGeometry& InGeometry, const FPointerEvent& InCursorEvent );
+	virtual FNavigationReply NativeOnNavigation(const FGeometry& InGeometry, const FNavigationEvent& InNavigationEvent);
 	virtual void NativeOnMouseCaptureLost();
 
 protected:
diff --git a/Engine/Source/Runtime/UMG/Public/Components/EditableTextBox.h b/Engine/Source/Runtime/UMG/Public/Components/EditableTextBox.h
--- a/Engine/Source/Runtime/UMG/Public/Components/EditableTextBox.h
+++ b/Engine/Source/Runtime/UMG/Public/Components/EditableTextBox.h
@@ -181,8 +181,8 @@ protected:
 	virtual TSharedRef<SWidget> RebuildWidget() override;
 	// End of UWidget
 
-	void HandleOnTextChanged(const FText& Text);
-	void HandleOnTextCommitted(const FText& Text, ETextCommit::Type CommitMethod);
+	virtual void HandleOnTextChanged(const FText& Text);
+	virtual void HandleOnTextCommitted(const FText& Text, ETextCommit::Type CommitMethod);
 
 protected:
 	TSharedPtr<SEditableTextBox> MyEditableTextBlock;
diff --git a/Engine/Source/Runtime/UMG/Public/Components/ScrollBox.h b/Engine/Source/Runtime/UMG/Public/Components/ScrollBox.h
--- a/Engine/Source/Runtime/UMG/Public/Components/ScrollBox.h
+++ b/Engine/Source/Runtime/UMG/Public/Components/ScrollBox.h
@@ -71,6 +71,10 @@ public:
 	UPROPERTY(EditAnywhere, BlueprintReadOnly, Category="Scroll")
 	float NavigationScrollPadding;
 	
+	/** Option to disable right-click-drag scrolling */
+	UPROPERTY(EditAnywhere, BlueprintReadOnly, Category = "Scroll")
+	bool bAllowRightClickDragScrolling;
+
 	UFUNCTION(BlueprintCallable, Category = "Scroll")
 	void SetOrientation(EOrientation NewOrientation);
 
diff --git a/Engine/Source/Runtime/Windows/D3D11RHI/Private/Windows/WindowsD3D11Device.cpp b/Engine/Source/Runtime/Windows/D3D11RHI/Private/Windows/WindowsD3D11Device.cpp
--- a/Engine/Source/Runtime/Windows/D3D11RHI/Private/Windows/WindowsD3D11Device.cpp
+++ b/Engine/Source/Runtime/Windows/D3D11RHI/Private/Windows/WindowsD3D11Device.cpp
@@ -1261,7 +1261,7 @@ void FD3D11DynamicRHI::InitD3DDevice()
 
 		GRHISupportsTextureStreaming = true;
 		GRHISupportsFirstInstance = true;
-
+		GRHINeedsExtraDeletionLatency = true;
 		// Set the RHI initialized flag.
 		GIsRHIInitialized = true;
 	}
