Copying //UE4/Dev-AnimPhys to //UE4/Dev-Main (Source: //UE4/Dev-AnimPhys @ 3624379)

diff --git a/Engine/Source/Runtime/Engine/Public/IAudioExtensionPlugin.h b/Engine/Source/Runtime/Engine/Public/IAudioExtensionPlugin.h
--- a/Engine/Source/Runtime/Engine/Public/IAudioExtensionPlugin.h
+++ b/Engine/Source/Runtime/Engine/Public/IAudioExtensionPlugin.h
@@ -12,6 +12,41 @@
 
 class FAudioDevice;
 
+enum class EAudioPlatform : uint8
+{
+	Windows,
+	Mac,
+	Linux,
+	IOS,
+	Android,
+	XboxOne,
+	Playstation4,
+	Switch,
+	HTML5,
+	Unknown
+};
+
+/**
+* Enumeration of audio plugin types
+*
+*/
+enum class EAudioPlugin : uint8
+{
+	SPATIALIZATION = 0,
+	REVERB = 1,
+	OCCLUSION = 2
+};
+
+class IAudioSpatialization;
+class IAudioOcclusion;
+class IAudioReverb;
+class IAudioPluginListener;
+
+typedef TSharedPtr<IAudioSpatialization, ESPMode::ThreadSafe> TAudioSpatializationPtr;
+typedef TSharedPtr<IAudioOcclusion, ESPMode::ThreadSafe> TAudioOcclusionPtr;
+typedef TSharedPtr<IAudioReverb, ESPMode::ThreadSafe> TAudioReverbPtr;
+typedef TSharedPtr<IAudioPluginListener, ESPMode::ThreadSafe> TAudioPluginListenerPtr;
+
 /**
 * FSpatializationParams
 * Struct for retrieving parameters needed for computing spatialization and occlusion plugins.
@@ -54,6 +89,35 @@ struct FSpatializationParams
 	{}
 };
 
+struct FAudioPluginInitializationParams
+{
+	//Maximum number of sources that can play simultaneously.
+	uint32 NumSources;
+
+	//Number of output channels.
+	uint32 NumOutputChannels;
+
+	//Sample rate.
+	uint32 SampleRate;
+
+	//Buffer length used for each callback.
+	uint32 BufferLength;
+
+	//Pointer to audio device owning this audio plugin.
+	//IMPORTANT: This will be deprecated once the AudioMixer
+	//           is taken out of the experimental branch.
+	FAudioDevice* AudioDevicePtr;
+
+	FAudioPluginInitializationParams()
+		: NumSources(0)
+		, NumOutputChannels(0)
+		, SampleRate(0)
+		, BufferLength(0)
+		, AudioDevicePtr(nullptr)
+	{}
+
+};
+
 struct FAudioPluginSourceInputData
 {
 	// The ID of the source voice
@@ -85,6 +149,89 @@ class ENGINE_API USpatializationPluginSourceSettingsBase : public UObject
 	GENERATED_BODY()
 };
 
+
+/************************************************************************/
+/* IAudioPluginFactory                                             */
+/* This interface is inherited by spatialization, reverb and occlusion  */
+/* plugins to describe specifics of a plugin such as platform support,  */
+/* and display names.                                                   */
+/************************************************************************/
+class IAudioPluginFactory
+{
+public:
+	/*
+	* Returns human-readable string representing the display name of this plugin.
+	* This is the name that will be used in settings and .ini files.
+	* If multiple IAudioPlugin implementations are found that return identical strings here,
+	* The first one of these loaded will be used.
+	*
+	* @return FString of the display name of this plugin.
+	*/
+	virtual FString GetDisplayName()
+	{
+		static FString DisplayName = FString(TEXT("Generic Audio Plugin"));
+		return DisplayName;
+	}
+
+	/*
+	* Returns whether this plugin supports use on the specified platform.
+	* @param Platform an enumerated platform (i.e. Windows, Playstation4, etc.)
+	* @return true if this plugin supports use on Platform, false otherwise.
+	*/
+	virtual bool SupportsPlatform(EAudioPlatform Platform) = 0;
+
+	/*
+	* Returns whether this plugin sends audio to an external renderer.
+	* if this returns true, the audio engine will not mix the result of the audio process callback
+	* from the plugin into the audio output.
+	*
+	* @return true only if the plugin will handle sending audio to the DAC itself.
+	*/
+	virtual bool IsExternalSend()
+	{
+		return false;
+	}
+};
+
+/************************************************************************/
+/* IAudioSpatializationFactory                                          */
+/* Implement this modular feature to make your Spatialialization plugin */
+/* visible to the engine.                                               */
+/************************************************************************/
+class IAudioSpatializationFactory : public IAudioPluginFactory, public IModularFeature
+{
+public:
+	/** Virtual destructor */
+	virtual ~IAudioSpatializationFactory()
+	{
+	}
+
+	// IModularFeature
+	static FName GetModularFeatureName()
+	{
+		static FName AudioExtFeatureName = FName(TEXT("AudioSpatializationPlugin"));
+		return AudioExtFeatureName;
+	}
+
+	/* Begin IAudioPluginWithMetadata implementation */
+	virtual FString GetDisplayName() override
+	{
+		static FString DisplayName = FString(TEXT("Generic Audio Spatialization Plugin"));
+		return DisplayName;
+	}
+	/* End IAudioPluginWithMetadata implementation */
+
+	/**
+	* @return a new instance of your spatialization plugin, owned by a shared pointer.
+	*/
+	virtual TAudioSpatializationPtr CreateNewSpatializationPlugin(FAudioDevice* OwningDevice) = 0;
+
+	/*
+	* @return true if this plugin uses a custom setting.
+	*/
+	virtual bool HasCustomSpatializationSetting() const { return false; }
+};
+
 /**
 * IAudioSpatialization
 *
@@ -94,7 +241,7 @@ class ENGINE_API USpatializationPluginSourceSettingsBase : public UObject
 * the effect is updated in the audio engine update loop with new position information.
 *
 */
-class IAudioSpatialization
+class IAudioSpatialization 
 {
 public:
 	/** Virtual destructor */
@@ -102,6 +249,18 @@ public:
 	{
 	}
 
+	/**
+	* Shuts down the audio plugin.
+	*
+	*/
+	virtual void Shutdown()
+	{
+	}
+
+	virtual void OnDeviceShutdown(FAudioDevice* AudioDevice)
+	{
+	}
+
 	/** DEPRECATED: sets the spatialization effect parameters. */
 	virtual void SetSpatializationParameters(uint32 SourceId, const FSpatializationParams& Params)
 	{
@@ -142,6 +301,11 @@ public:
 	{
 	}
 
+	/** Called when all sources have finished processing. */
+	virtual void OnAllSourcesProcessed()
+	{
+	}
+
 	/** Returns whether or not the spatialization effect has been initialized */
 	virtual bool IsSpatializationEffectInitialized() const
 	{
@@ -149,7 +313,7 @@ public:
 	}
 
 	/** Initializes the spatialization plugin with the given buffer length. */
-	virtual void Initialize(const uint32 SampleRate, const uint32 NumSources, const uint32 OutputBufferLength)
+	virtual void Initialize(const FAudioPluginInitializationParams InitializationParams)
 	{
 	}
 
@@ -173,6 +337,38 @@ class ENGINE_API UOcclusionPluginSourceSettingsBase : public UObject
 	GENERATED_BODY()
 };
 
+/************************************************************************/
+/* IAudioOcclusionFactory                                               */
+/*                                                                      */
+/************************************************************************/
+class IAudioOcclusionFactory : public IAudioPluginFactory, public IModularFeature
+{
+public:
+	/** Virtual destructor */
+	virtual ~IAudioOcclusionFactory()
+	{
+	}
+
+	// IModularFeature
+	static FName GetModularFeatureName()
+	{
+		static FName AudioExtFeatureName = FName(TEXT("AudioOcclusionPlugin"));
+		return AudioExtFeatureName;
+	}
+
+	/* Begin IAudioPluginWithMetadata implementation */
+	virtual FString GetDisplayName() override
+	{
+		static FString DisplayName = FString(TEXT("Generic Audio Occlusion Plugin"));
+		return DisplayName;
+	}
+	/* End IAudioPluginWithMetadata implementation */
+
+	virtual TAudioOcclusionPtr CreateNewOcclusionPlugin(FAudioDevice* OwningDevice) = 0;
+
+	virtual bool HasCustomOcclusionSetting() const { return false; }
+};
+
 class IAudioOcclusion
 {
 public:
@@ -182,7 +378,7 @@ public:
 	}
 
 	/** Initialize the occlusion plugin with the same rate and number of sources. */
-	virtual void Initialize(const int32 SampleRate, const int32 NumSources, const int32 FrameSize)
+	virtual void Initialize(const FAudioPluginInitializationParams InitializationParams)
 	{
 	}
 
@@ -209,95 +405,44 @@ class ENGINE_API UReverbPluginSourceSettingsBase : public UObject
 	GENERATED_BODY()
 };
  
-class IAudioReverb
+class IAudioReverbFactory : public IAudioPluginFactory, public IModularFeature
 {
 public:
 	/** Virtual destructor */
-	virtual ~IAudioReverb()
+	virtual ~IAudioReverbFactory()
 	{
 	}
 
-	/** Initialize the reverb plugin with the same rate and number of sources. */
-	virtual void Initialize(const int32 SampleRate, const int32 NumSources, const int32 FrameSize)
+	// IModularFeature
+	static FName GetModularFeatureName()
 	{
+		static FName AudioExtFeatureName = FName(TEXT("AudioReverbPlugin"));
+		return AudioExtFeatureName;
 	}
 
-	/** Called when a source is assigned to a voice. */
-	virtual void OnInitSource(const uint32 SourceId, const FName& AudioComponentUserId, const uint32 NumChannels, UReverbPluginSourceSettingsBase* InSettings) = 0;
-
-	/** Called when a source is done playing and is released. */
-	virtual void OnReleaseSource(const uint32 SourceId) = 0;
-
-	virtual class FSoundEffectSubmix* GetEffectSubmix(class USoundSubmix* Submix) = 0;
-
-	/** Processes audio with the given input and output data structs.*/
-	virtual void ProcessSourceAudio(const FAudioPluginSourceInputData& InputData, FAudioPluginSourceOutputData& OutputData)
+	/* Begin IAudioPluginWithMetadata implementation */
+	virtual FString GetDisplayName() override
 	{
+		static FString DisplayName = FString(TEXT("Generic Audio Reverb Plugin"));
+		return DisplayName;
 	}
-};
+	/* End IAudioPluginWithMetadata implementation */
 
-/** Override this and create it to receive listener changes. */
-class IAudioListenerObserver
-{
-public:
-	virtual ~IAudioListenerObserver()
-	{
-	}
+	virtual TAudioReverbPtr CreateNewReverbPlugin(FAudioDevice* OwningDevice) = 0;
 
-	// Called when the listener is updated on the given audio device.
-	virtual void OnListenerUpdated(FAudioDevice* AudioDevice, UWorld* ListenerWorld, const int32 ViewportIndex, const FTransform& ListenerTransform, const float InDeltaSeconds) = 0;
-	virtual void OnListenerShutdown(FAudioDevice* AudioDevice) = 0;
+	virtual bool HasCustomReverbSetting() const { return false; }
 };
 
-typedef TSharedPtr<IAudioSpatialization, ESPMode::ThreadSafe> TAudioSpatializationPtr;
-typedef TSharedPtr<IAudioOcclusion, ESPMode::ThreadSafe> TAudioOcclusionPtr;
-typedef TSharedPtr<IAudioReverb, ESPMode::ThreadSafe> TAudioReverbPtr;
-typedef TSharedPtr<IAudioListenerObserver, ESPMode::ThreadSafe> TAudioListenerObserverPtr;
-
-/**
-* The public interface of an audio plugin. Plugins that extend core features of the audio engine.
-*/
-class IAudioPlugin : public IModuleInterface, public IModularFeature
+class IAudioReverb
 {
 public:
-	// IModularFeature
-	static FName GetModularFeatureName()
-	{
-		static FName AudioExtFeatureName = FName(TEXT("AudioPlugin"));
-		return AudioExtFeatureName;
-	}
-
-	// IModuleInterface
-	virtual void StartupModule() override
-	{
-		IModularFeatures::Get().RegisterModularFeature(GetModularFeatureName(), this);
-	}
-
-	/**
-	* Singleton-like access to IAudioExtensionPlugin
-	*
-	* @return Returns IAudioExtensionPlugin singleton instance, loading the module on demand if needed
-	*/
-	static inline IAudioPlugin& Get()
-	{
-		return FModuleManager::LoadModuleChecked<IAudioPlugin>("AudioPlugin");
-	}
-
-	/**
-	* Checks to see if this module is loaded and ready.  It is only valid to call Get() if IsAvailable() returns true.
-	*
-	* @return True if the module is loaded and ready to use
-	*/
-	static inline bool IsAvailable()
+	/** Virtual destructor */
+	virtual ~IAudioReverb()
 	{
-		return FModuleManager::Get().IsModuleLoaded("AudioPlugin");
 	}
 
-	/**
-	* Initializes the audio plugin.
-	*
-	*/
-	virtual void Initialize()
+	/** Initialize the reverb plugin with the same rate and number of sources. */
+	virtual void Initialize(const FAudioPluginInitializationParams InitializationParams)
 	{
 	}
 
@@ -313,58 +458,52 @@ public:
 	{
 	}
 
-	virtual bool ImplementsSpatialization() const
-	{
-		return false;
-	}
-
-	virtual bool HasCustomSpatializationSetting() const
-	{
-		return false;
-	}
-
-	virtual bool ImplementsOcclusion() const
-	{
-		return false;
-	}
+	/** Called when a source is assigned to a voice. */
+	virtual void OnInitSource(const uint32 SourceId, const FName& AudioComponentUserId, const uint32 NumChannels, UReverbPluginSourceSettingsBase* InSettings) = 0;
 
-	virtual bool HasCustomOcclusionSetting() const
-	{
-		return false;
-	}
+	/** Called when a source is done playing and is released. */
+	virtual void OnReleaseSource(const uint32 SourceId) = 0;
 
-	virtual bool ImplementsReverb() const
-	{
-		return false;
-	}
+	virtual class FSoundEffectSubmix* GetEffectSubmix(class USoundSubmix* Submix) = 0;
 
-	virtual bool HasCustomReverbSetting() const
+	/** Processes audio with the given input and output data structs.*/
+	virtual void ProcessSourceAudio(const FAudioPluginSourceInputData& InputData, FAudioPluginSourceOutputData& OutputData)
 	{
-		return false;
 	}
+};
 
-	virtual bool SupportsMultipleAudioDevices() const
+/************************************************************************/
+/* IAudioPluginListener                                                 */
+/* Implementations of this interface can receive updates about the      */
+/* audio listener's position in the game world, as well as other data.  */
+/* to use this, register a ListenerObserver to an audio device using    */
+/* FAudioDevice::RegisterPluginListener().                              */
+/************************************************************************/
+class IAudioPluginListener
+{
+public:
+	virtual ~IAudioPluginListener()
 	{
-		return true;
 	}
 
-	virtual TAudioSpatializationPtr CreateSpatializationInterface(class FAudioDevice* AudioDevice)
+	virtual void OnDeviceShutdown(FAudioDevice* AudioDevice)
 	{
-		return nullptr;
 	}
 
-	virtual TAudioOcclusionPtr CreateOcclusionInterface(class FAudioDevice* AudioDevice)
+	//This function is called when a game world initializes a listener with an audio device this
+	//IAudioPluginListener is registered to. Please note that it is possible to miss this event
+	//if you register this IAudioPluginListener after the listener is initialized.
+	virtual void OnListenerInitialize(FAudioDevice* AudioDevice, UWorld* ListenerWorld)
 	{
-		return nullptr;
 	}
 
-	virtual TAudioReverbPtr CreateReverbInterface(class FAudioDevice* AudioDevice)
+	// Called when the listener is updated on the given audio device.
+	virtual void OnListenerUpdated(FAudioDevice* AudioDevice, const int32 ViewportIndex, const FTransform& ListenerTransform, const float InDeltaSeconds)
 	{
-		return nullptr;
 	}
 
-	virtual TAudioListenerObserverPtr CreateListenerObserverInterface(class FAudioDevice* AudioDevice)
+	//Called when the listener is shutdown.
+	virtual void OnListenerShutdown(FAudioDevice* AudioDevice)
 	{
-		return nullptr;
 	}
-};
+};
\ No newline at end of file
diff --git a/Engine/Source/Runtime/Engine/Public/PhysicsPublic.h b/Engine/Source/Runtime/Engine/Public/PhysicsPublic.h
--- a/Engine/Source/Runtime/Engine/Public/PhysicsPublic.h
+++ b/Engine/Source/Runtime/Engine/Public/PhysicsPublic.h
@@ -680,7 +680,7 @@ struct FPhysSceneShaderInfo
 #endif
 
 /** Enum to indicate types of simple shapes */
-enum EKCollisionPrimitiveType
+enum DEPRECATED(4.17, "Please use EAggCollisionShape::Type") EKCollisionPrimitiveType
 {
 	KPT_Sphere = 0,
 	KPT_Box,
diff --git a/Engine/Source/Runtime/Engine/Public/PreviewScene.h b/Engine/Source/Runtime/Engine/Public/PreviewScene.h
--- a/Engine/Source/Runtime/Engine/Public/PreviewScene.h
+++ b/Engine/Source/Runtime/Engine/Public/PreviewScene.h
@@ -88,6 +88,9 @@ public:
 	void SetSkyBrightness(float SkyBrightness);
 	void SetSkyCubemap(class UTextureCube* Cubemap);
 
+	/** Get the background color we use by default */
+	virtual FLinearColor GetBackgroundColor() const;
+
 	/** Load/Save settings to the config, specifying the key */
 	void LoadSettings(const TCHAR* Section);
 	void SaveSettings(const TCHAR* Section);
diff --git a/Engine/Source/Runtime/Engine/Public/SceneTypes.h b/Engine/Source/Runtime/Engine/Public/SceneTypes.h
--- a/Engine/Source/Runtime/Engine/Public/SceneTypes.h
+++ b/Engine/Source/Runtime/Engine/Public/SceneTypes.h
@@ -163,7 +163,7 @@ enum EMaterialProperty
 	//^^^ New material properties go above here ^^^^
 	MP_MaterialAttributes UMETA(Hidden),
 	MP_CustomOutput UMETA(Hidden),
-	MP_MAX UMETA(Hidden),
+	MP_MAX UMETA(DisplayName = "None"),
 };
 
 /** Blend modes supported for simple element rendering */
diff --git a/Engine/Source/Runtime/IOS/IOSAudio/Private/IOSAudioSource.cpp b/Engine/Source/Runtime/IOS/IOSAudio/Private/IOSAudioSource.cpp
--- a/Engine/Source/Runtime/IOS/IOSAudio/Private/IOSAudioSource.cpp
+++ b/Engine/Source/Runtime/IOS/IOSAudio/Private/IOSAudioSource.cpp
@@ -210,7 +210,7 @@ void FIOSAudioSoundSource::Update(void)
 
 	if (!AudioDevice->IsAudioDeviceMuted())
 	{
-		Volume = WaveInstance->Volume * WaveInstance->VolumeMultiplier;
+		Volume = WaveInstance->GetActualVolume();
 	}
 
 	if (SetStereoBleed())
@@ -233,9 +233,9 @@ void FIOSAudioSoundSource::Update(void)
 	if (IOSBuffer->NumChannels == 1 && WaveInstance->bUseSpatialization)
 	{
 		// Compute the directional offset
-		FVector Offset = WaveInstance->Location - IOSAudioDevice->PlayerLocation;
+		FVector Offset = GetSpatializationParams().EmitterPosition;
 		const AudioUnitParameterValue AzimuthRangeScale = 90.0f;
-		const AudioUnitParameterValue Pan = Offset.CosineAngle2D(IOSAudioDevice->PlayerRight) * AzimuthRangeScale;
+        const AudioUnitParameterValue Pan = Offset.Y * AzimuthRangeScale;
 
 		Status = AudioUnitSetParameter(IOSAudioDevice->GetMixerUnit(),
 		                               k3DMixerParam_Azimuth,
diff --git a/Engine/Source/Runtime/IOS/IOSRuntimeSettings/Classes/IOSRuntimeSettings.h b/Engine/Source/Runtime/IOS/IOSRuntimeSettings/Classes/IOSRuntimeSettings.h
--- a/Engine/Source/Runtime/IOS/IOSRuntimeSettings/Classes/IOSRuntimeSettings.h
+++ b/Engine/Source/Runtime/IOS/IOSRuntimeSettings/Classes/IOSRuntimeSettings.h
@@ -422,6 +422,18 @@ public:
 	UPROPERTY(config, EditAnywhere, Category = "Audio", meta = (ClampMin = "0", UIMin = "0", DisplayName = "Number of Source Workers"))
 	int32 AudioNumSourceWorkers;
 
+	/** Which of the currently enabled spatialization plugins to use on Windows. */
+	UPROPERTY(config, EditAnywhere, Category = "Audio")
+	FString SpatializationPlugin;
+
+	/** Which of the currently enabled reverb plugins to use on Windows. */
+	UPROPERTY(config, EditAnywhere, Category = "Audio")
+	FString ReverbPlugin;
+
+	/** Which of the currently enabled occlusion plugins to use on Windows. */
+	UPROPERTY(config, EditAnywhere, Category = "Audio")
+	FString OcclusionPlugin;
+
 #if WITH_EDITOR
 	// UObject interface
 	virtual void PostEditChangeProperty(struct FPropertyChangedEvent& PropertyChangedEvent) override;
diff --git a/Engine/Source/Runtime/Launch/Launch.Build.cs b/Engine/Source/Runtime/Launch/Launch.Build.cs
--- a/Engine/Source/Runtime/Launch/Launch.Build.cs
+++ b/Engine/Source/Runtime/Launch/Launch.Build.cs
@@ -73,10 +73,6 @@ public class Launch : ModuleRules
 				DynamicallyLoadedModuleNames.Add("D3D11RHI");
 				DynamicallyLoadedModuleNames.Add("XAudio2");
 				DynamicallyLoadedModuleNames.Add("AudioMixerXAudio2");
-				if (Target.Platform == UnrealTargetPlatform.Win64)
-				{
-					DynamicallyLoadedModuleNames.Add("AudioMixerSDL");
-				}
 			}
 			else if (Target.Platform == UnrealTargetPlatform.Mac)
 			{
@@ -232,7 +228,8 @@ public class Launch : ModuleRules
 		if (Target.Platform == UnrealTargetPlatform.HTML5 )
 		{
 			PrivateDependencyModuleNames.Add("ALAudio");
-			AddEngineThirdPartyPrivateStaticDependencies(Target, "SDL2");
+            PrivateDependencyModuleNames.Add("AudioMixerSDL");
+            AddEngineThirdPartyPrivateStaticDependencies(Target, "SDL2");
 		}
 
 		// @todo ps4 clang bug: this works around a PS4/clang compiler bug (optimizations)
diff --git a/Engine/Source/Runtime/Launch/Private/Android/AndroidEventManager.cpp b/Engine/Source/Runtime/Launch/Private/Android/AndroidEventManager.cpp
--- a/Engine/Source/Runtime/Launch/Private/Android/AndroidEventManager.cpp
+++ b/Engine/Source/Runtime/Launch/Private/Android/AndroidEventManager.cpp
@@ -433,9 +433,19 @@ void FAppEventManager::PauseAudio()
 {
 	bAudioPaused = true;
 
-	if (GEngine->GetMainAudioDevice())
+	UE_LOG(LogTemp, Log, TEXT("Android pause audio"));
+
+	FAudioDevice* AudioDevice = GEngine->GetMainAudioDevice();
+	if (AudioDevice)
 	{
-		GEngine->GetMainAudioDevice()->Suspend(false);
+		if (AudioDevice->IsAudioMixerEnabled())
+		{
+			AudioDevice->SuspendContext();
+		}
+		else
+		{
+			GEngine->GetMainAudioDevice()->Suspend(false);
+		}
 	}
 }
 
@@ -444,9 +454,19 @@ void FAppEventManager::ResumeAudio()
 {
 	bAudioPaused = false;
 
-	if (GEngine->GetMainAudioDevice())
+	UE_LOG(LogTemp, Log, TEXT("Android resume audio"));
+
+	FAudioDevice* AudioDevice = GEngine->GetMainAudioDevice();
+	if (AudioDevice)
 	{
-		GEngine->GetMainAudioDevice()->Suspend(true);
+		if (AudioDevice->IsAudioMixerEnabled())
+		{
+			AudioDevice->ResumeContext();
+		}
+		else
+		{
+			GEngine->GetMainAudioDevice()->Suspend(true);
+		}
 	}
 }
 
diff --git a/Engine/Source/Runtime/Linux/AudioMixerSDL/Private/AudioMixerPlatformSDL.cpp b/Engine/Source/Runtime/Linux/AudioMixerSDL/Private/AudioMixerPlatformSDL.cpp
--- a/Engine/Source/Runtime/Linux/AudioMixerSDL/Private/AudioMixerPlatformSDL.cpp
+++ b/Engine/Source/Runtime/Linux/AudioMixerSDL/Private/AudioMixerPlatformSDL.cpp
@@ -97,12 +97,20 @@ namespace Audio
 	{
 		// To figure out the output device info, attempt to init at 7.1, and 48k.
 		// SDL_OpenAudioDevice will attempt open the audio device with that spec but return what it actually used. We'll report that in OutInfo
-
+		FAudioPlatformSettings PlatformSettings = GetPlatformSettings();
 		SDL_AudioSpec DesiredSpec;
-		DesiredSpec.freq = 48000;
+		DesiredSpec.freq = PlatformSettings.SampleRate;
+
+		// HTML5 supports s16 format only
+#if PLATFORM_HTML5
+		DesiredSpec.format = AUDIO_S16;
+		DesiredSpec.channels = 2;
+#else
 		DesiredSpec.format = AUDIO_F32;
 		DesiredSpec.channels = 6;
-		DesiredSpec.samples = 1024;
+#endif
+		
+		DesiredSpec.samples = PlatformSettings.CallbackBufferFrameSize;
 		DesiredSpec.callback = OnBufferEnd;
 		DesiredSpec.userdata = (void*)this;
 
@@ -134,7 +142,13 @@ namespace Audio
 		OutInfo.DeviceId = DeviceName;
 		OutInfo.Name = OutInfo.DeviceId;
 		OutInfo.SampleRate = ActualSpec.freq;
+		
+		// HTML5 supports s16 format only
+#if PLATFORM_HTML5
+		OutInfo.Format = EAudioMixerStreamDataFormat::Int16;
+#else
 		OutInfo.Format = EAudioMixerStreamDataFormat::Float;
+#endif
 		OutInfo.NumChannels = ActualSpec.channels;
 
 		// Assume default channel map order, SDL doesn't support us querying it directly
@@ -176,7 +190,12 @@ namespace Audio
 			return false;
 		}
 
+		// HTML5 supports s16 format only
+#if PLATFORM_HTML5
+		AudioSpecPrefered.format = AUDIO_S16;
+#else
 		AudioSpecPrefered.format = AUDIO_F32;
+#endif
 		AudioSpecPrefered.freq = Params.SampleRate;
 		AudioSpecPrefered.channels = AudioStreamInfo.DeviceInfo.NumChannels;
 		AudioSpecPrefered.samples = OpenStreamParams.NumFrames;
@@ -203,7 +222,11 @@ namespace Audio
 		check(AudioSpecReceived.samples == OpenStreamParams.NumFrames);
 
 		// Compute the expected output byte length
+#if PLATFORM_HTML5
+		OutputBufferByteLength = OpenStreamParams.NumFrames * AudioStreamInfo.DeviceInfo.NumChannels * sizeof(int16);
+#else
 		OutputBufferByteLength = OpenStreamParams.NumFrames * AudioStreamInfo.DeviceInfo.NumChannels * sizeof(float);
+#endif
 		check(OutputBufferByteLength == AudioSpecReceived.size);
 
 		AudioStreamInfo.StreamState = EAudioOutputStreamState::Open;
@@ -283,6 +306,11 @@ namespace Audio
 
 	void FMixerPlatformSDL::HandleOnBufferEnd(uint8* InOutputBuffer, int32 InOutputBufferByteLength)
 	{
+		if (!bIsDeviceInitialized)
+		{
+			return;
+		}
+
 		OutputBuffer = InOutputBuffer;
 		check(InOutputBufferByteLength == OutputBufferByteLength);
 
@@ -357,8 +385,16 @@ namespace Audio
 
 	FAudioPlatformSettings FMixerPlatformSDL::GetPlatformSettings() const
 	{
+#if PLATFORM_LINUX
 		return FAudioPlatformSettings::GetPlatformSettings(TEXT("/Script/LinuxTargetPlatform.LinuxTargetSettings"));
+#else
+		// On HTML5 and Windows, use default parameters.
+		FAudioPlatformSettings Settings;
+		Settings.SampleRate = 48000;
+		Settings.MaxChannels = 0;
+		Settings.NumBuffers = 2;
+		Settings.CallbackBufferFrameSize = 1024;
+		return Settings;
+#endif
 	}
-
-
 }
diff --git a/Engine/Source/Runtime/LiveLinkInterface/LiveLinkInterface.Build.cs b/Engine/Source/Runtime/LiveLinkInterface/LiveLinkInterface.Build.cs
new file mode 100644
--- /dev/null
+++ b/Engine/Source/Runtime/LiveLinkInterface/LiveLinkInterface.Build.cs
@@ -0,0 +1,18 @@
+// Copyright 1998-2017 Epic Games, Inc. All Rights Reserved.
+
+using UnrealBuildTool;
+
+public class LiveLinkInterface : ModuleRules
+{
+    public LiveLinkInterface(ReadOnlyTargetRules Target) : base(Target)
+	{
+        PublicDependencyModuleNames.AddRange(
+            new string[] {
+                "Core",
+                "CoreUObject"
+            }
+        );
+
+        PrivateIncludePaths.Add("LiveLinkInterface/Private");
+	}
+}
diff --git a/Engine/Source/Runtime/LiveLinkInterface/Private/LiveLinkInterfaceModule.cpp b/Engine/Source/Runtime/LiveLinkInterface/Private/LiveLinkInterfaceModule.cpp
new file mode 100644
--- /dev/null
+++ b/Engine/Source/Runtime/LiveLinkInterface/Private/LiveLinkInterfaceModule.cpp
@@ -0,0 +1,27 @@
+// Copyright 1998-2017 Epic Games, Inc. All Rights Reserved.
+
+#include "ILiveLinkInterfaceModule.h"
+#include "ILiveLinkClient.h"
+
+FName ILiveLinkClient::ModularFeatureName(TEXT("ModularFeature_LiveLinkClient"));
+
+class FLiveLinkInterfaceModule : public ILiveLinkInterfaceModule
+{
+public:
+	
+	// IModuleInterface interface
+	virtual void StartupModule() override
+	{
+		
+	}
+
+	virtual bool SupportsDynamicReloading() override
+	{
+		return false;
+	}
+
+private:
+
+};
+
+IMPLEMENT_MODULE(FLiveLinkInterfaceModule, LiveLinkInterface);
diff --git a/Engine/Source/Runtime/LiveLinkInterface/Private/LiveLinkSourceFactory.cpp b/Engine/Source/Runtime/LiveLinkInterface/Private/LiveLinkSourceFactory.cpp
new file mode 100644
--- /dev/null
+++ b/Engine/Source/Runtime/LiveLinkInterface/Private/LiveLinkSourceFactory.cpp
@@ -0,0 +1,8 @@
+// Copyright 1998-2017 Epic Games, Inc. All Rights Reserved.
+
+#include "LiveLinkSourceFactory.h"
+
+ULiveLinkSourceFactory::ULiveLinkSourceFactory(const FObjectInitializer& ObjectInitializer)
+{
+
+}
\ No newline at end of file
diff --git a/Engine/Source/Runtime/LiveLinkInterface/Public/ILiveLinkClient.h b/Engine/Source/Runtime/LiveLinkInterface/Public/ILiveLinkClient.h
new file mode 100644
--- /dev/null
+++ b/Engine/Source/Runtime/LiveLinkInterface/Public/ILiveLinkClient.h
@@ -0,0 +1,31 @@
+// Copyright 1998-2017 Epic Games, Inc. All Rights Reserved.
+
+#pragma once
+
+#include "LiveLinkRefSkeleton.h"
+#include "Features/IModularFeature.h"
+#include "LiveLinkTypes.h"
+#include "Guid.h"
+
+class ILiveLinkSource;
+
+class LIVELINKINTERFACE_API ILiveLinkClient : public IModularFeature
+{
+public:
+	virtual ~ILiveLinkClient() {}
+
+	static FName ModularFeatureName;
+
+	// Helper functions for making time codes for new frames
+	virtual FLiveLinkTimeCode MakeTimeCode(double InTime, int32 InFrameNum) const = 0;
+	virtual FLiveLinkTimeCode MakeTimeCodeFromTimeOnly(double InTime) const = 0;
+
+	// Add a new live link source to the client
+	virtual void AddSource(TSharedPtr<ILiveLinkSource> InSource) = 0;
+
+	virtual void PushSubjectSkeleton(FName SubjectName, const FLiveLinkRefSkeleton& RefSkeleton) = 0;
+	virtual void PushSubjectData(FGuid SourceGuid, FName SubjectName, const TArray<FTransform>& Transforms, const TArray<FLiveLinkCurveElement>& CurveElements, const FLiveLinkTimeCode& TimeCode) = 0;
+	virtual void ClearSubject(FName SubjectName) = 0;
+
+	virtual const FLiveLinkSubjectFrame* GetSubjectData(FName SubjectName) = 0;
+};
\ No newline at end of file
diff --git a/Engine/Source/Runtime/LiveLinkInterface/Public/ILiveLinkInterfaceModule.h b/Engine/Source/Runtime/LiveLinkInterface/Public/ILiveLinkInterfaceModule.h
new file mode 100644
--- /dev/null
+++ b/Engine/Source/Runtime/LiveLinkInterface/Public/ILiveLinkInterfaceModule.h
@@ -0,0 +1,34 @@
+// Copyright 1998-2017 Epic Games, Inc. All Rights Reserved.
+
+#pragma once
+
+#include "ModuleInterface.h"
+#include "ModuleManager.h"
+
+
+class ILiveLinkInterfaceModule
+	: public IModuleInterface
+{
+public:
+
+	/**
+	 * Gets a reference to the live link interface module instance.
+	 *
+	 * @return A reference to the live link interface module.
+	 */
+	static ILiveLinkInterfaceModule& Get()
+	{
+#if PLATFORM_IOS
+        static ILiveLinkInterfaceModule& LiveLinkInterfaceModule = FModuleManager::LoadModuleChecked<ILiveLinkInterfaceModule>("LiveLinkInterface");
+        return LiveLinkInterfaceModule;
+#else
+        return FModuleManager::LoadModuleChecked<ILiveLinkInterfaceModule>("LiveLinkInterface");
+#endif
+	}
+
+public:
+
+	/** Virtual destructor. */
+	virtual ~ILiveLinkInterfaceModule() { }
+};
+
diff --git a/Engine/Source/Runtime/LiveLinkInterface/Public/ILiveLinkSource.h b/Engine/Source/Runtime/LiveLinkInterface/Public/ILiveLinkSource.h
new file mode 100644
--- /dev/null
+++ b/Engine/Source/Runtime/LiveLinkInterface/Public/ILiveLinkSource.h
@@ -0,0 +1,23 @@
+// Copyright 1998-2017 Epic Games, Inc. All Rights Reserved.
+
+#pragma once
+
+#include "CoreMinimal.h"
+#include "Guid.h"
+
+class ILiveLinkClient;
+
+class ILiveLinkSource
+{
+public:
+	virtual ~ILiveLinkSource() {}
+	virtual void ReceiveClient(ILiveLinkClient* InClient, FGuid InSourceGuid) = 0;
+
+	virtual bool IsSourceStillValid() = 0;
+
+	virtual bool RequestSourceShutdown() = 0;
+
+	virtual FText GetSourceType() const = 0;
+	virtual FText GetSourceMachineName() const = 0;
+	virtual FText GetSourceStatus() const = 0;
+};
\ No newline at end of file
diff --git a/Engine/Source/Runtime/LiveLinkInterface/Public/LiveLinkSourceFactory.h b/Engine/Source/Runtime/LiveLinkInterface/Public/LiveLinkSourceFactory.h
new file mode 100644
--- /dev/null
+++ b/Engine/Source/Runtime/LiveLinkInterface/Public/LiveLinkSourceFactory.h
@@ -0,0 +1,25 @@
+// Copyright 1998-2017 Epic Games, Inc. All Rights Reserved.
+
+#pragma once
+
+#include "CoreMinimal.h"
+#include "ObjectMacros.h"
+#include "Object.h"
+#include "LiveLinkSourceFactory.generated.h"
+
+class ILiveLinkSource;
+class ILiveLinkClient;
+class SWidget;
+
+UCLASS()
+class LIVELINKINTERFACE_API ULiveLinkSourceFactory : public UObject
+{
+	GENERATED_UCLASS_BODY()
+
+public:
+	virtual FText GetSourceDisplayName() const PURE_VIRTUAL( ULiveLinkSourceFactory::GetSourceDisplayName, return FText(); );
+	virtual FText GetSourceTooltip() const PURE_VIRTUAL(ULiveLinkSourceFactory::GetSourceTooltip, return FText(); );
+
+	virtual TSharedPtr<SWidget> CreateSourceCreationPanel() PURE_VIRTUAL(ULiveLinkSourceFactory::CreateSourceCreationPanel, return nullptr;);
+	virtual TSharedPtr<ILiveLinkSource> OnSourceCreationPanelClosed(bool bCreateSource) PURE_VIRTUAL(ULiveLinkSourceFactory::OnSourceCreationPanelClosed, return nullptr;);
+};
\ No newline at end of file
diff --git a/Engine/Source/Runtime/LiveLinkInterface/Public/LiveLinkTypes.h b/Engine/Source/Runtime/LiveLinkInterface/Public/LiveLinkTypes.h
new file mode 100644
--- /dev/null
+++ b/Engine/Source/Runtime/LiveLinkInterface/Public/LiveLinkTypes.h
@@ -0,0 +1,96 @@
+// Copyright 1998-2017 Epic Games, Inc. All Rights Reserved.
+
+#pragma once
+
+#include "CoreMinimal.h"
+#include "UObject/ObjectMacros.h"
+#include "LiveLinkRefSkeleton.h"
+#include "LiveLinkTypes.generated.h"
+
+USTRUCT()
+struct FLiveLinkCurveElement
+{
+public:
+	GENERATED_USTRUCT_BODY()
+
+	UPROPERTY()
+	FName CurveName;
+	
+	UPROPERTY()
+	float CurveValue;
+};
+
+struct FOptionalCurveElement
+{
+	/** Curve Value */
+	float					Value;
+	/** Whether this value is set or not */
+	bool					bValid;
+
+	FOptionalCurveElement(float InValue)
+		: Value(InValue)
+		, bValid(true)
+	{}
+
+	FOptionalCurveElement()
+		: Value(0.f)
+		, bValid(false)
+	{}
+
+	bool IsValid() const
+	{
+		return bValid;
+	}
+
+	void SetValue(float InValue)
+	{
+		Value = InValue;
+		bValid = true;
+	}
+};
+
+//Helper struct for updating curve data across multiple frames of live link data
+struct FLiveLinkCurveIntegrationData
+{
+public:
+
+	// Number of new curves that need to be added to existing frames
+	int32 NumNewCurves;
+
+	// Built curve buffer for current frame in existing curve key format
+	TArray<FOptionalCurveElement> CurveValues;
+};
+
+struct FLiveLinkCurveKey
+{
+	TArray<FName> CurveNames;
+
+	FLiveLinkCurveIntegrationData UpdateCurveKey(const TArray<FLiveLinkCurveElement>& CurveElements);
+};
+
+struct FLiveLinkTimeCode
+{
+	// Time for this frame. Used during interpolation. If this goes backwards we will dump already stored frames. 
+	double Time;
+
+	// Frame number for this data.
+	int32 FrameNum;
+
+	// Value calculated on create to represent the different between the source time and client time
+	double Offset;
+};
+
+struct FLiveLinkSubjectFrame
+{
+	// Ref Skeleton for transforms
+	FLiveLinkRefSkeleton RefSkeleton;
+
+	// Key for storing curve data (Names)
+	FLiveLinkCurveKey	 CurveKeyData;
+
+	// Transforms for this frame
+	TArray<FTransform>		Transforms;
+	
+	// Curve data for this frame
+	TArray<FOptionalCurveElement>	Curves;
+};
\ No newline at end of file
diff --git a/Engine/Source/Runtime/LiveLinkMessageBusFramework/LiveLinkMessageBusFramework.Build.cs b/Engine/Source/Runtime/LiveLinkMessageBusFramework/LiveLinkMessageBusFramework.Build.cs
new file mode 100644
--- /dev/null
+++ b/Engine/Source/Runtime/LiveLinkMessageBusFramework/LiveLinkMessageBusFramework.Build.cs
@@ -0,0 +1,22 @@
+// Copyright 1998-2017 Epic Games, Inc. All Rights Reserved.
+
+using UnrealBuildTool;
+
+public class LiveLinkMessageBusFramework : ModuleRules
+{
+    public LiveLinkMessageBusFramework(ReadOnlyTargetRules Target) : base(Target)
+    {
+		PrivateDependencyModuleNames.Add("MessagingCommon");
+        PrivateIncludePathModuleNames.Add("Messaging");
+		
+        PublicDependencyModuleNames.AddRange(
+            new string[] {
+                "Core",
+				"CoreUObject",
+                "LiveLinkInterface"
+            }
+        );
+
+        PrivateIncludePaths.Add("LiveLinkMessageBusFramework/Private");
+	}
+}
diff --git a/Engine/Source/Runtime/LiveLinkMessageBusFramework/Private/LiveLinkProvider.cpp b/Engine/Source/Runtime/LiveLinkMessageBusFramework/Private/LiveLinkProvider.cpp
new file mode 100644
--- /dev/null
+++ b/Engine/Source/Runtime/LiveLinkMessageBusFramework/Private/LiveLinkProvider.cpp
@@ -0,0 +1,249 @@
+// Copyright 1998-2017 Epic Games, Inc. All Rights Reserved.
+
+#include "LiveLinkProvider.h"
+#include "IMessageContext.h"
+#include "MessageEndpoint.h"
+#include "MessageEndpointBuilder.h"
+#include "LiveLinkMessages.h"
+#include "LiveLinkTypes.h"
+
+#include "HAL/PlatformProcess.h"
+
+// Subject that the application has told us about
+struct FTrackedSubject
+{
+	// Ref skeleton to go with transform data
+	FLiveLinkRefSkeleton RefSkeleton;
+
+	// Bone transform data
+	TArray<FTransform> Transforms;
+
+	// Curve data
+	TArray<FLiveLinkCurveElement> Curves;
+
+	// Incrementing time (application time) for interpolation purposes
+	double Time;
+
+	// Frame number of current data
+	int32 FrameNum;
+};
+
+// Address that we have had a connection request from
+struct FTrackedAddress
+{
+	FTrackedAddress(FMessageAddress InAddress)
+		: Address(InAddress)
+		, LastHeartbeatTime(FPlatformTime::Seconds())
+	{}
+
+	FMessageAddress Address;
+	double			LastHeartbeatTime;
+};
+
+struct FLiveLinkProvider : public ILiveLinkProvider
+{
+private:
+	// How long we give connections before we decide they are dead
+	static const double CONNECTION_TIMEOUT;
+
+	const FString ProviderName;
+	
+	const FString MachineName;
+
+	TSharedPtr<FMessageEndpoint, ESPMode::ThreadSafe> MessageEndpoint;
+
+	// Array of our current connections
+	TArray<FTrackedAddress> ConnectedAddresses;
+
+	// Cache of our current subject state
+	TMap<FName, FTrackedSubject> Subjects;
+	
+	//Message bus message handlers
+	void HandlePingMessage(const FLiveLinkPingMessage& Message, const TSharedRef<IMessageContext, ESPMode::ThreadSafe>& Context);
+	void HandleConnectMessage(const FLiveLinkConnectMessage& Message, const TSharedRef<IMessageContext, ESPMode::ThreadSafe>& Context);
+	void HandleHeartbeat(const FLiveLinkHeartbeatMessage& Message, const TSharedRef<IMessageContext, ESPMode::ThreadSafe>& Context);
+	// End message bus message handlers
+
+	// Validate our current connections
+	void ValidateConnections()
+	{
+		double CutOffTime = FPlatformTime::Seconds() - CONNECTION_TIMEOUT;
+
+		ConnectedAddresses.RemoveAll([=](const FTrackedAddress& Address) { return Address.LastHeartbeatTime < CutOffTime; });
+	}
+
+	// Get the cached data for the named subject
+	FTrackedSubject& GetTrackedSubject(const FName& SubjectName)
+	{
+		return Subjects.FindOrAdd(SubjectName);
+	}
+
+	// Clear a existing track subject
+	void ClearTrackedSubject(const FName& SubjectName)
+	{
+		Subjects.Remove(SubjectName);
+	}
+
+	// Send hierarchy data for named subject to supplied address
+	void SendSubjectToAddress(FName SubjectName, const FTrackedSubject& Subject, FMessageAddress Address)
+	{
+		FLiveLinkSubjectDataMessage* SubjectData = new FLiveLinkSubjectDataMessage;
+		SubjectData->RefSkeleton = Subject.RefSkeleton;
+		SubjectData->SubjectName = SubjectName;
+
+		MessageEndpoint->Send(SubjectData, Address);
+	}
+
+	// Send hierarchy data for named subject to supplied address
+	void SendSubjectToAddress(FName SubjectName, FMessageAddress Address)
+	{
+		SendSubjectToAddress(SubjectName, GetTrackedSubject(SubjectName), Address);
+	}
+
+	// Send frame data for named subject to supplied address
+	void SendSubjectFrameToAddress(FName SubjectName, const FTrackedSubject& Subject, FMessageAddress Address)
+	{
+		FLiveLinkSubjectFrameMessage* SubjectFrame = new FLiveLinkSubjectFrameMessage;
+		SubjectFrame->Transforms = Subject.Transforms;
+		SubjectFrame->SubjectName = SubjectName;
+		SubjectFrame->Curves = Subject.Curves;
+		SubjectFrame->Time = Subject.Time;
+		SubjectFrame->FrameNum = Subject.FrameNum;
+
+		MessageEndpoint->Send(SubjectFrame, Address);
+	}
+
+	// Send frame data for named subject to supplied address
+	void SendSubjectFrameToAddress(FName SubjectName, FMessageAddress Address)
+	{
+		SendSubjectFrameToAddress(SubjectName, GetTrackedSubject(SubjectName), Address);
+	}
+
+	void SendClearSubjectToAddress(FName SubjectName, FMessageAddress Address)
+	{
+		FLiveLinkClearSubject* ClearSubject = new FLiveLinkClearSubject(SubjectName);
+		MessageEndpoint->Send(ClearSubject, Address);
+	}
+
+	// Send hierarchy data for named subject to current connections
+	void SendSubjectToConnections(FName SubjectName)
+	{
+		ValidateConnections();
+
+		for (const FTrackedAddress& Address : ConnectedAddresses)
+		{
+			SendSubjectToAddress(SubjectName, Address.Address);
+		}
+	}
+
+	// Send frame data for named subject to current connections
+	void SendSubjectFrameToConnections(FName SubjectName)
+	{
+		ValidateConnections();
+
+		for (const FTrackedAddress& Address : ConnectedAddresses)
+		{
+			SendSubjectFrameToAddress(SubjectName, Address.Address);
+		}
+	}
+
+	void SendClearSubjectToConnections(FName SubjectName)
+	{
+		ValidateConnections();
+
+		for (const FTrackedAddress& Address : ConnectedAddresses)
+		{
+			SendClearSubjectToAddress(SubjectName, Address.Address);
+		}
+	}
+
+public:
+	FLiveLinkProvider(const FString& InProviderName)
+		: ProviderName(InProviderName)
+		, MachineName(FPlatformProcess::ComputerName())
+	{
+		MessageEndpoint = FMessageEndpoint::Builder(*InProviderName)
+			.ReceivingOnAnyThread()
+			.Handling<FLiveLinkPingMessage>(this, &FLiveLinkProvider::HandlePingMessage)
+			.Handling<FLiveLinkConnectMessage>(this, &FLiveLinkProvider::HandleConnectMessage)
+			.Handling<FLiveLinkHeartbeatMessage>(this, &FLiveLinkProvider::HandleHeartbeat);
+
+		if (MessageEndpoint.IsValid())
+		{
+			MessageEndpoint->Subscribe<FLiveLinkPingMessage>();
+		}
+	}
+
+	virtual void UpdateSubject(const FName& SubjectName, const TArray<FName>& BoneNames, const TArray<int32>& BoneParents)
+	{
+		FTrackedSubject& Subject = GetTrackedSubject(SubjectName);
+		Subject.RefSkeleton.SetBoneNames(BoneNames);
+		Subject.RefSkeleton.SetBoneParents(BoneParents);
+		Subject.Transforms.Empty();
+		
+		SendSubjectToConnections(SubjectName);
+	}
+
+	virtual void ClearSubject(const FName& SubjectName)
+	{
+		ClearTrackedSubject(SubjectName);
+		SendClearSubjectToConnections(SubjectName);
+	}
+
+	virtual void UpdateSubjectFrame(const FName& SubjectName, const TArray<FTransform>& BoneTransforms, const TArray<FLiveLinkCurveElement>& CurveData, double Time, int32 FrameNum)
+	{
+		FTrackedSubject& Subject = GetTrackedSubject(SubjectName);
+
+		Subject.Transforms = BoneTransforms;
+		Subject.Curves = CurveData;
+		Subject.Time = Time;
+		Subject.FrameNum = FrameNum;
+
+		SendSubjectFrameToConnections(SubjectName);
+	}
+
+	virtual ~FLiveLinkProvider()
+	{
+		FPlatformMisc::LowLevelOutputDebugString(TEXT("Destroyed"));
+	}
+};
+
+const double FLiveLinkProvider::CONNECTION_TIMEOUT = 10.f;
+
+void FLiveLinkProvider::HandlePingMessage(const FLiveLinkPingMessage& Message, const TSharedRef<IMessageContext, ESPMode::ThreadSafe>& Context)
+{
+	MessageEndpoint->Send(new FLiveLinkPongMessage(ProviderName, MachineName, Message.PollRequest), Context->GetSender());
+}
+
+void FLiveLinkProvider::HandleConnectMessage(const FLiveLinkConnectMessage& Message, const TSharedRef<IMessageContext, ESPMode::ThreadSafe>& Context)
+{
+	const FMessageAddress& ConnectionAddress = Context->GetSender();
+
+	if (!ConnectedAddresses.ContainsByPredicate([=](const FTrackedAddress& Address) { return Address.Address == ConnectionAddress; }))
+	{
+		ConnectedAddresses.Add(FTrackedAddress(ConnectionAddress));
+		for (const auto& Subject : Subjects)
+		{
+			SendSubjectToAddress(Subject.Key, Subject.Value, ConnectionAddress);
+			FPlatformProcess::Sleep(0.1); //HACK: Try to help these go in order, editor needs extra buffering support to make sure this isn't needed in future.
+			SendSubjectFrameToAddress(Subject.Key, Subject.Value, ConnectionAddress);
+		}
+	}
+}
+
+void FLiveLinkProvider::HandleHeartbeat(const FLiveLinkHeartbeatMessage& Message, const TSharedRef<IMessageContext, ESPMode::ThreadSafe>& Context)
+{
+	FTrackedAddress* TrackedAddress = ConnectedAddresses.FindByPredicate([=](const FTrackedAddress& ConAddress) { return ConAddress.Address == Context->GetSender(); });
+	if (TrackedAddress)
+	{
+		TrackedAddress->LastHeartbeatTime = FPlatformTime::Seconds();
+
+		// Respond so editor gets heartbeat too
+		MessageEndpoint->Send(new FLiveLinkHeartbeatMessage(), Context->GetSender());
+	}
+}
+
+TSharedPtr<ILiveLinkProvider> ILiveLinkProvider::CreateLiveLinkProvider(const FString& ProviderName)
+{
+	return MakeShareable(new FLiveLinkProvider(ProviderName));
+}
\ No newline at end of file
diff --git a/Engine/Source/Runtime/LiveLinkMessageBusFramework/Public/LiveLinkMessages.h b/Engine/Source/Runtime/LiveLinkMessageBusFramework/Public/LiveLinkMessages.h
new file mode 100644
--- /dev/null
+++ b/Engine/Source/Runtime/LiveLinkMessageBusFramework/Public/LiveLinkMessages.h
@@ -0,0 +1,103 @@
+// Copyright 1998-2017 Epic Games, Inc. All Rights Reserved.
+
+#pragma once
+#include "Object.h"
+#include "LiveLinkRefSkeleton.h"
+#include "LiveLinkTypes.h"
+#include "LiveLinkMessages.generated.h"
+
+USTRUCT()
+struct FLiveLinkSubjectDataMessage
+{
+	GENERATED_USTRUCT_BODY()
+
+	UPROPERTY()
+	FLiveLinkRefSkeleton RefSkeleton;
+
+	UPROPERTY()
+	FName SubjectName;
+};
+
+USTRUCT()
+struct FLiveLinkSubjectFrameMessage
+{
+	GENERATED_USTRUCT_BODY()
+
+	UPROPERTY()
+	FName SubjectName;
+
+	// Bone Transform data for this frame
+	UPROPERTY()
+	TArray<FTransform> Transforms;
+
+	// Curve data for this frame
+	UPROPERTY()
+	TArray<FLiveLinkCurveElement> Curves;
+
+	// Incrementing time for interpolation
+	UPROPERTY()
+	double Time;
+
+	// Frame number
+	UPROPERTY()
+	int32 FrameNum;
+};
+
+USTRUCT()
+struct FLiveLinkPingMessage
+{
+	GENERATED_USTRUCT_BODY()
+
+	UPROPERTY()
+	FGuid PollRequest;
+
+	// default constructor for the receiver
+	FLiveLinkPingMessage() = default;
+
+	FLiveLinkPingMessage(const FGuid& InPollRequest) : PollRequest(InPollRequest) {}
+};
+
+USTRUCT()
+struct FLiveLinkPongMessage
+{
+	GENERATED_USTRUCT_BODY()
+
+	UPROPERTY()
+	FString ProviderName;
+
+	UPROPERTY()
+	FString MachineName;
+
+	UPROPERTY()
+	FGuid PollRequest;
+
+	// default constructor for the receiver
+	FLiveLinkPongMessage() = default;
+
+	FLiveLinkPongMessage(const FString& InProviderName, const FString& InMachineName, const FGuid& InPollRequest) : ProviderName(InProviderName), MachineName(InMachineName), PollRequest(InPollRequest) {}
+};
+
+USTRUCT()
+struct FLiveLinkConnectMessage
+{
+	GENERATED_USTRUCT_BODY()
+};
+
+USTRUCT()
+struct FLiveLinkHeartbeatMessage
+{
+	GENERATED_USTRUCT_BODY()
+};
+
+USTRUCT()
+struct FLiveLinkClearSubject
+{
+	GENERATED_USTRUCT_BODY()
+
+	// Name of the subject to clear
+	UPROPERTY()
+	FName SubjectName;
+
+	FLiveLinkClearSubject() {}
+	FLiveLinkClearSubject(const FName& InSubjectName) : SubjectName(InSubjectName) {}
+};
diff --git a/Engine/Source/Runtime/LiveLinkMessageBusFramework/Public/LiveLinkProvider.h b/Engine/Source/Runtime/LiveLinkMessageBusFramework/Public/LiveLinkProvider.h
new file mode 100644
--- /dev/null
+++ b/Engine/Source/Runtime/LiveLinkMessageBusFramework/Public/LiveLinkProvider.h
@@ -0,0 +1,23 @@
+// Copyright 1998-2017 Epic Games, Inc. All Rights Reserved.
+
+#pragma once
+
+#include "SharedPointer.h"
+#include "LiveLinkTypes.h"
+
+struct ILiveLinkProvider
+{
+public:
+	LIVELINKMESSAGEBUSFRAMEWORK_API static TSharedPtr<ILiveLinkProvider> CreateLiveLinkProvider(const FString& ProviderName);
+
+	// Update hierarchy for named subject
+	virtual void UpdateSubject(const FName& SubjectName, const TArray<FName>& BoneNames, const TArray<int32>& BoneParents) = 0;
+
+	// Remove named subject
+	virtual void ClearSubject(const FName& SubjectName) = 0;
+
+	// Update subject with transform data
+	virtual void UpdateSubjectFrame(const FName& SubjectName, const TArray<FTransform>& BoneTransforms, const TArray<FLiveLinkCurveElement>& CurveData, double Time, int32 FrameNum) = 0;
+
+	virtual ~ILiveLinkProvider() {}
+};
\ No newline at end of file
diff --git a/Engine/Source/Runtime/MediaAssets/Private/Assets/MediaSoundComponent.cpp b/Engine/Source/Runtime/MediaAssets/Private/Assets/MediaSoundComponent.cpp
--- a/Engine/Source/Runtime/MediaAssets/Private/Assets/MediaSoundComponent.cpp
+++ b/Engine/Source/Runtime/MediaAssets/Private/Assets/MediaSoundComponent.cpp
@@ -120,7 +120,7 @@ void UMediaSoundComponent::Init(const int32 SampleRate)
 }
 
 
-void UMediaSoundComponent::OnGenerateAudio(TArray<float>& OutAudio)
+void UMediaSoundComponent::OnGenerateAudio(float* OutAudio, int32 NumSamples)
 {
 	TSharedPtr<FMediaPlayerFacade, ESPMode::ThreadSafe> PinnedPlayerFacade;
 	TSharedPtr<FMediaAudioSampleQueue, ESPMode::ThreadSafe> PinnedSampleQueue;
@@ -133,8 +133,8 @@ void UMediaSoundComponent::OnGenerateAudio(TArray<float>& OutAudio)
 
 	if (PinnedPlayerFacade.IsValid() && PinnedPlayerFacade->IsPlaying() && PinnedSampleQueue.IsValid())
 	{
-		const uint32 FramesRequested = OutAudio.Num() / NumChannels;
-		const uint32 FramesWritten = Resampler->Generate(OutAudio.GetData(), FramesRequested, PinnedPlayerFacade->GetRate(), PinnedPlayerFacade->GetTime(), *PinnedSampleQueue);
+		const uint32 FramesRequested = NumSamples / NumChannels;
+		const uint32 FramesWritten = Resampler->Generate(OutAudio, FramesRequested, PinnedPlayerFacade->GetRate(), PinnedPlayerFacade->GetTime(), *PinnedSampleQueue);
 	}
 	else
 	{
diff --git a/Engine/Source/Runtime/MovieSceneTracks/Private/Evaluation/MovieSceneSkeletalAnimationTemplate.cpp b/Engine/Source/Runtime/MovieSceneTracks/Private/Evaluation/MovieSceneSkeletalAnimationTemplate.cpp
--- a/Engine/Source/Runtime/MovieSceneTracks/Private/Evaluation/MovieSceneSkeletalAnimationTemplate.cpp
+++ b/Engine/Source/Runtime/MovieSceneTracks/Private/Evaluation/MovieSceneSkeletalAnimationTemplate.cpp
@@ -86,8 +86,12 @@ struct FPreAnimatedAnimationTokenProducer : IMovieScenePreAnimatedTokenProducer
 
 				// Reset the mesh component update flag and animation mode to what they were before we animated the object
 				Component->MeshComponentUpdateFlag = MeshComponentUpdateFlag;
-				Component->SetAnimationMode(AnimationMode);
-
+				if (Component->GetAnimationMode() != AnimationMode)
+				{
+					// this SetAnimationMode reinitializes even if the mode is same
+					// if we're using same anim blueprint, we don't want to keep reinitializing it. 
+					Component->SetAnimationMode(AnimationMode);
+				}
 			}
 
 			EMeshComponentUpdateFlag::Type MeshComponentUpdateFlag;
diff --git a/Engine/Source/Runtime/MovieSceneTracks/Public/Sections/MovieSceneSkeletalAnimationSection.h b/Engine/Source/Runtime/MovieSceneTracks/Public/Sections/MovieSceneSkeletalAnimationSection.h
--- a/Engine/Source/Runtime/MovieSceneTracks/Public/Sections/MovieSceneSkeletalAnimationSection.h
+++ b/Engine/Source/Runtime/MovieSceneTracks/Public/Sections/MovieSceneSkeletalAnimationSection.h
@@ -23,7 +23,7 @@ struct FMovieSceneSkeletalAnimationParams
 	float GetSequenceLength() const { return Animation != nullptr ? Animation->SequenceLength : 0.f; }
 
 	/** The animation this section plays */
-	UPROPERTY(EditAnywhere, Category="Animation")
+	UPROPERTY(EditAnywhere, Category="Animation", meta=(AllowedClasses = "AnimSequence, AnimComposite"))
 	UAnimSequenceBase* Animation;
 
 	/** The offset into the beginning of the animation clip */
diff --git a/Engine/Source/Runtime/Projects/Private/BuildVersion.cpp b/Engine/Source/Runtime/Projects/Private/BuildVersion.cpp
--- a/Engine/Source/Runtime/Projects/Private/BuildVersion.cpp
+++ b/Engine/Source/Runtime/Projects/Private/BuildVersion.cpp
@@ -27,10 +27,17 @@ FString FBuildVersion::GetDefaultFileName()
 
 FString FBuildVersion::GetFileNameForCurrentExecutable()
 {
+	FString AppExecutableName = FPlatformProcess::ExecutableName();
+#if PLATFORM_WINDOWS
+	if(AppExecutableName.EndsWith(TEXT("-Cmd")))
+	{
+		AppExecutableName = AppExecutableName.Left(AppExecutableName.Len() - 4);
+	}
+#endif
 #if IS_PROGRAM || IS_MONOLITHIC
-	return FPaths::ProjectDir() / TEXT("Binaries") / FPlatformProcess::GetBinariesSubdirectory() / FPlatformProcess::ExecutableName() + TEXT(".version");
+	return FPaths::ProjectDir() / TEXT("Binaries") / FPlatformProcess::GetBinariesSubdirectory() / AppExecutableName + TEXT(".version");
 #else
-	return FPaths::EngineDir() / TEXT("Binaries") / FPlatformProcess::GetBinariesSubdirectory() / FPlatformProcess::ExecutableName() + TEXT(".version");
+	return FPaths::EngineDir() / TEXT("Binaries") / FPlatformProcess::GetBinariesSubdirectory() / AppExecutableName + TEXT(".version");
 #endif
 }
 
diff --git a/Engine/Source/Runtime/Projects/Private/ModuleManifest.cpp b/Engine/Source/Runtime/Projects/Private/ModuleManifest.cpp
--- a/Engine/Source/Runtime/Projects/Private/ModuleManifest.cpp
+++ b/Engine/Source/Runtime/Projects/Private/ModuleManifest.cpp
@@ -14,7 +14,14 @@ FModuleManifest::FModuleManifest()
 
 FString FModuleManifest::GetFileName(const FString& DirectoryName, bool bIsGameFolder)
 {
-	FString FileName = DirectoryName / FPlatformProcess::ExecutableName();
+	FString AppExecutableName = FPlatformProcess::ExecutableName();
+#if PLATFORM_WINDOWS
+	if (AppExecutableName.EndsWith(TEXT("-Cmd")))
+	{
+		AppExecutableName = AppExecutableName.Left(AppExecutableName.Len() - 4);
+	}
+#endif
+	FString FileName = DirectoryName / AppExecutableName;
 	if(FApp::GetBuildConfiguration() == EBuildConfigurations::DebugGame && bIsGameFolder)
 	{
 		FileName += FString::Printf(TEXT("-%s-DebugGame"), FPlatformProcess::GetBinariesSubdirectory());
diff --git a/Engine/Source/Runtime/Slate/Private/Framework/MultiBox/Mac/MacMenu.cpp b/Engine/Source/Runtime/Slate/Private/Framework/MultiBox/Mac/MacMenu.cpp
--- a/Engine/Source/Runtime/Slate/Private/Framework/MultiBox/Mac/MacMenu.cpp
+++ b/Engine/Source/Runtime/Slate/Private/Framework/MultiBox/Mac/MacMenu.cpp
@@ -431,9 +431,9 @@ NSImage* FSlateMacMenu::GetMenuItemIcon(const TSharedRef<const FMenuEntryBlock>&
 {
 	NSImage* MenuImage = nil;
 	FSlateIcon Icon;
-	if (Block->IconOverride.IsSet())
+	if (Block->IconOverride.IsSet() && Block->IconOverride.Get().IsSet())
 	{
-		Icon = Block->IconOverride;
+		Icon = Block->IconOverride.Get();
 	}
 	else if (Block->GetAction().IsValid() && Block->GetAction()->GetIcon().IsSet())
 	{
diff --git a/Engine/Source/Runtime/Slate/Private/Framework/MultiBox/MultiBoxBuilder.cpp b/Engine/Source/Runtime/Slate/Private/Framework/MultiBox/MultiBoxBuilder.cpp
--- a/Engine/Source/Runtime/Slate/Private/Framework/MultiBox/MultiBoxBuilder.cpp
+++ b/Engine/Source/Runtime/Slate/Private/Framework/MultiBox/MultiBoxBuilder.cpp
@@ -219,7 +219,7 @@ void FMenuBuilder::AddMenuSeparator(FName InExtensionHook)
 	ApplyHook(InExtensionHook, EExtensionHook::After);
 }
 
-void FMenuBuilder::AddSubMenu( const TAttribute<FText>& InMenuLabel, const TAttribute<FText>& InToolTip, const FNewMenuDelegate& InSubMenu, const FUIAction& InUIAction, FName InExtensionHook, const EUserInterfaceActionType::Type InUserInterfaceActionType, const bool bInOpenSubMenuOnClick, const FSlateIcon& InIcon, const bool bInShouldCloseWindowAfterMenuSelection /*= true*/ )
+void FMenuBuilder::AddSubMenu( const TAttribute<FText>& InMenuLabel, const TAttribute<FText>& InToolTip, const FNewMenuDelegate& InSubMenu, const FUIAction& InUIAction, FName InExtensionHook, const EUserInterfaceActionType::Type InUserInterfaceActionType, const bool bInOpenSubMenuOnClick, const TAttribute<FSlateIcon>& InIcon, const bool bInShouldCloseWindowAfterMenuSelection /*= true*/ )
 {
 	ApplySectionBeginning();
 
@@ -229,7 +229,7 @@ void FMenuBuilder::AddSubMenu( const TAttribute<FText>& InMenuLabel, const TAttr
 	MultiBox->AddMultiBlock( NewMenuEntryBlock );
 }
 
-void FMenuBuilder::AddSubMenu( const TAttribute<FText>& InMenuLabel, const TAttribute<FText>& InToolTip, const FNewMenuDelegate& InSubMenu, const bool bInOpenSubMenuOnClick /*= false*/, const FSlateIcon& InIcon /*= FSlateIcon()*/, const bool bInShouldCloseWindowAfterMenuSelection /*= true*/ )
+void FMenuBuilder::AddSubMenu( const TAttribute<FText>& InMenuLabel, const TAttribute<FText>& InToolTip, const FNewMenuDelegate& InSubMenu, const bool bInOpenSubMenuOnClick /*= false*/, const TAttribute<FSlateIcon>& InIcon /*= TAttribute<FSlateIcon>()*/, const bool bInShouldCloseWindowAfterMenuSelection /*= true*/ )
 {
 	ApplySectionBeginning();
 
@@ -328,7 +328,12 @@ void FMenuBuilder::ApplySectionBeginning()
 
 
 
-void FMenuBarBuilder::AddPullDownMenu( const FText& InMenuLabel, const FText& InToolTip, const FNewMenuDelegate& InPullDownMenu, FName InExtensionHook, FName InTutorialHighlightName )
+void FMenuBarBuilder::AddPullDownMenu( const TAttribute<FText>& InMenuLabel, const TAttribute<FText>& InToolTip, const FNewMenuDelegate& InPullDownMenu, FName InExtensionHook, FName InTutorialHighlightName )
+{
+	AddPullDownMenu(InMenuLabel, InToolTip, FSlateIcon(), InPullDownMenu, InExtensionHook, InTutorialHighlightName);
+}
+
+void FMenuBarBuilder::AddPullDownMenu(const TAttribute<FText>& InMenuLabel, const TAttribute<FText>& InToolTip, const TAttribute<FSlateIcon>& InIcon, const FNewMenuDelegate& InPullDownMenu, FName InExtensionHook, FName InTutorialHighlightName)
 {
 	ApplySectionBeginning();
 
@@ -336,12 +341,13 @@ void FMenuBarBuilder::AddPullDownMenu( const FText& InMenuLabel, const FText& In
 
 	const bool bIsSubMenu = false;
 	const bool bOpenSubMenuOnClick = false;
+	const bool bShouldCloseWindowAfterMenuSelection = true;
 	// Pulldown menus always close all menus not just themselves
 	const bool bShouldCloseSelfOnly = false;
-	TSharedRef< FMenuEntryBlock > NewMenuEntryBlock( new FMenuEntryBlock( InExtensionHook, InMenuLabel, InToolTip, InPullDownMenu, ExtenderStack.Top(), bIsSubMenu, bOpenSubMenuOnClick, CommandListStack.Last(), bShouldCloseSelfOnly) );
+	TSharedRef< FMenuEntryBlock > NewMenuEntryBlock(new FMenuEntryBlock(InExtensionHook, InMenuLabel, InToolTip, InPullDownMenu, ExtenderStack.Top(), bIsSubMenu, bOpenSubMenuOnClick, CommandListStack.Last(), bShouldCloseSelfOnly, InIcon, bShouldCloseWindowAfterMenuSelection, bInvertLabelsOnHover));
 	NewMenuEntryBlock->SetTutorialHighlightName(GenerateTutorialIdentfierName(TutorialHighlightName, InTutorialHighlightName, nullptr, MultiBox->GetBlocks().Num()));
 
-	MultiBox->AddMultiBlock( NewMenuEntryBlock );
+	MultiBox->AddMultiBlock(NewMenuEntryBlock);
 
 	ApplyHook(InExtensionHook, EExtensionHook::After);
 }
diff --git a/Engine/Source/Runtime/Slate/Private/Framework/MultiBox/SMenuEntryBlock.cpp b/Engine/Source/Runtime/Slate/Private/Framework/MultiBox/SMenuEntryBlock.cpp
--- a/Engine/Source/Runtime/Slate/Private/Framework/MultiBox/SMenuEntryBlock.cpp
+++ b/Engine/Source/Runtime/Slate/Private/Framework/MultiBox/SMenuEntryBlock.cpp
@@ -13,7 +13,7 @@
 #include "Widgets/Input/SCheckBox.h"
 
 
-FMenuEntryBlock::FMenuEntryBlock( const FName& InExtensionHook, const TSharedPtr< const FUICommandInfo > InCommand, TSharedPtr< const FUICommandList > InCommandList, const TAttribute<FText>& InLabelOverride, const TAttribute<FText>& InToolTipOverride, const FSlateIcon& InIconOverride, bool bInCloseSelfOnly, bool bInShouldCloseWindowAfterMenuSelection )
+FMenuEntryBlock::FMenuEntryBlock( const FName& InExtensionHook, const TSharedPtr< const FUICommandInfo > InCommand, TSharedPtr< const FUICommandList > InCommandList, const TAttribute<FText>& InLabelOverride, const TAttribute<FText>& InToolTipOverride, const TAttribute<FSlateIcon>& InIconOverride, bool bInCloseSelfOnly, bool bInShouldCloseWindowAfterMenuSelection, bool bInInvertLabelOnHover)
 	: FMultiBlock( InCommand, InCommandList, InExtensionHook, EMultiBlockType::MenuEntry )
 	, LabelOverride( InLabelOverride )
 	, ToolTipOverride( InToolTipOverride )
@@ -23,11 +23,12 @@ FMenuEntryBlock::FMenuEntryBlock( const FName& InExtensionHook, const TSharedPtr
 	, UserInterfaceActionType( EUserInterfaceActionType::Button )
 	, bCloseSelfOnly( bInCloseSelfOnly )
 	, bShouldCloseWindowAfterMenuSelection( bInShouldCloseWindowAfterMenuSelection )
+	, bInvertLabelOnHover(bInInvertLabelOnHover)
 {
 }
 
 
-FMenuEntryBlock::FMenuEntryBlock( const FName& InExtensionHook, const TAttribute<FText>& InLabel, const TAttribute<FText>& InToolTip, const FNewMenuDelegate& InEntryBuilder, TSharedPtr<FExtender> InExtender, bool bInSubMenu, bool bInSubMenuOnClick, const FSlateIcon& InIcon, const FUIAction& InUIAction, const EUserInterfaceActionType::Type InUserInterfaceActionType, bool bInCloseSelfOnly, bool bInShouldCloseWindowAfterMenuSelection)
+FMenuEntryBlock::FMenuEntryBlock( const FName& InExtensionHook, const TAttribute<FText>& InLabel, const TAttribute<FText>& InToolTip, const FNewMenuDelegate& InEntryBuilder, TSharedPtr<FExtender> InExtender, bool bInSubMenu, bool bInSubMenuOnClick, const TAttribute<FSlateIcon>& InIcon, const FUIAction& InUIAction, const EUserInterfaceActionType::Type InUserInterfaceActionType, bool bInCloseSelfOnly, bool bInShouldCloseWindowAfterMenuSelection, bool bInInvertLabelOnHover)
 	: FMultiBlock( InUIAction, InExtensionHook, EMultiBlockType::MenuEntry )
 	, LabelOverride( InLabel )
 	, ToolTipOverride( InToolTip )
@@ -39,11 +40,12 @@ FMenuEntryBlock::FMenuEntryBlock( const FName& InExtensionHook, const TAttribute
 	, bCloseSelfOnly( bInCloseSelfOnly )
 	, Extender( InExtender )
 	, bShouldCloseWindowAfterMenuSelection( bInShouldCloseWindowAfterMenuSelection )
+	, bInvertLabelOnHover(bInInvertLabelOnHover)
 {
 }
 
 
-FMenuEntryBlock::FMenuEntryBlock( const FName& InExtensionHook, const TAttribute<FText>& InLabel, const TAttribute<FText>& InToolTip, const FSlateIcon& InIcon, const FUIAction& UIAction, const EUserInterfaceActionType::Type InUserInterfaceActionType, bool bInCloseSelfOnly, bool bInShouldCloseWindowAfterMenuSelection )
+FMenuEntryBlock::FMenuEntryBlock( const FName& InExtensionHook, const TAttribute<FText>& InLabel, const TAttribute<FText>& InToolTip, const TAttribute<FSlateIcon>& InIcon, const FUIAction& UIAction, const EUserInterfaceActionType::Type InUserInterfaceActionType, bool bInCloseSelfOnly, bool bInShouldCloseWindowAfterMenuSelection, bool bInInvertLabelOnHover)
 	: FMultiBlock( UIAction, InExtensionHook, EMultiBlockType::MenuEntry )
 	, LabelOverride( InLabel )
 	, ToolTipOverride( InToolTip )
@@ -53,11 +55,12 @@ FMenuEntryBlock::FMenuEntryBlock( const FName& InExtensionHook, const TAttribute
 	, UserInterfaceActionType( InUserInterfaceActionType )
 	, bCloseSelfOnly( bInCloseSelfOnly )
 	, bShouldCloseWindowAfterMenuSelection( bInShouldCloseWindowAfterMenuSelection )
+	, bInvertLabelOnHover(bInInvertLabelOnHover)
 {
 }
 
 
-FMenuEntryBlock::FMenuEntryBlock( const FName& InExtensionHook, const TAttribute<FText>& InLabel, const TAttribute<FText>& InToolTip, const FNewMenuDelegate& InEntryBuilder, TSharedPtr<FExtender> InExtender, bool bInSubMenu, bool bInSubMenuOnClick, TSharedPtr< const FUICommandList > InCommandList, bool bInCloseSelfOnly, const FSlateIcon& InIcon, bool bInShouldCloseWindowAfterMenuSelection )
+FMenuEntryBlock::FMenuEntryBlock( const FName& InExtensionHook, const TAttribute<FText>& InLabel, const TAttribute<FText>& InToolTip, const FNewMenuDelegate& InEntryBuilder, TSharedPtr<FExtender> InExtender, bool bInSubMenu, bool bInSubMenuOnClick, TSharedPtr< const FUICommandList > InCommandList, bool bInCloseSelfOnly, const TAttribute<FSlateIcon>& InIcon, bool bInShouldCloseWindowAfterMenuSelection, bool bInInvertLabelOnHover)
 	: FMultiBlock( nullptr, InCommandList, InExtensionHook, EMultiBlockType::MenuEntry )
 	, LabelOverride( InLabel )
 	, ToolTipOverride( InToolTip )
@@ -69,11 +72,12 @@ FMenuEntryBlock::FMenuEntryBlock( const FName& InExtensionHook, const TAttribute
 	, bCloseSelfOnly( bInCloseSelfOnly )
 	, Extender( InExtender )
 	, bShouldCloseWindowAfterMenuSelection( bInShouldCloseWindowAfterMenuSelection )
+	, bInvertLabelOnHover(bInInvertLabelOnHover)
 {
 }
 
 
-FMenuEntryBlock::FMenuEntryBlock( const FName& InExtensionHook, const TAttribute<FText>& InLabel, const TAttribute<FText>& InToolTip, const FOnGetContent& InMenuBuilder, TSharedPtr<FExtender> InExtender, bool bInSubMenu, bool bInSubMenuOnClick, TSharedPtr< const FUICommandList > InCommandList, bool bInCloseSelfOnly, const FSlateIcon& InIcon, bool bInShouldCloseWindowAfterMenuSelection )
+FMenuEntryBlock::FMenuEntryBlock( const FName& InExtensionHook, const TAttribute<FText>& InLabel, const TAttribute<FText>& InToolTip, const FOnGetContent& InMenuBuilder, TSharedPtr<FExtender> InExtender, bool bInSubMenu, bool bInSubMenuOnClick, TSharedPtr< const FUICommandList > InCommandList, bool bInCloseSelfOnly, const TAttribute<FSlateIcon>& InIcon, bool bInShouldCloseWindowAfterMenuSelection, bool bInInvertLabelOnHover)
 	: FMultiBlock( nullptr, InCommandList, InExtensionHook, EMultiBlockType::MenuEntry )
 	, LabelOverride( InLabel )
 	, ToolTipOverride( InToolTip )
@@ -85,11 +89,12 @@ FMenuEntryBlock::FMenuEntryBlock( const FName& InExtensionHook, const TAttribute
 	, bCloseSelfOnly( bInCloseSelfOnly )
 	, Extender( InExtender )
 	, bShouldCloseWindowAfterMenuSelection( bInShouldCloseWindowAfterMenuSelection )
+	, bInvertLabelOnHover(bInInvertLabelOnHover)
 {
 }
 
 
-FMenuEntryBlock::FMenuEntryBlock( const FName& InExtensionHook, const TAttribute<FText>& InLabel, const TAttribute<FText>& InToolTip, const TSharedPtr<SWidget>& InEntryWidget, TSharedPtr<FExtender> InExtender, bool bInSubMenu, bool bInSubMenuOnClick, TSharedPtr< const FUICommandList > InCommandList, bool bInCloseSelfOnly, const FSlateIcon& InIcon, bool bInShouldCloseWindowAfterMenuSelection )
+FMenuEntryBlock::FMenuEntryBlock( const FName& InExtensionHook, const TAttribute<FText>& InLabel, const TAttribute<FText>& InToolTip, const TSharedPtr<SWidget>& InEntryWidget, TSharedPtr<FExtender> InExtender, bool bInSubMenu, bool bInSubMenuOnClick, TSharedPtr< const FUICommandList > InCommandList, bool bInCloseSelfOnly, const TAttribute<FSlateIcon>& InIcon, bool bInShouldCloseWindowAfterMenuSelection, bool bInInvertLabelOnHover)
 	: FMultiBlock( nullptr, InCommandList, InExtensionHook, EMultiBlockType::MenuEntry )
 	, LabelOverride( InLabel )
 	, ToolTipOverride( InToolTip )
@@ -101,10 +106,11 @@ FMenuEntryBlock::FMenuEntryBlock( const FName& InExtensionHook, const TAttribute
 	, bCloseSelfOnly( bInCloseSelfOnly )
 	, Extender( InExtender )
 	, bShouldCloseWindowAfterMenuSelection( bInShouldCloseWindowAfterMenuSelection )
+	, bInvertLabelOnHover(bInInvertLabelOnHover)
 {
 }
 
-FMenuEntryBlock::FMenuEntryBlock( const FName& InExtensionHook, const FUIAction& UIAction, const TSharedRef< SWidget > Contents, const TAttribute<FText>& InToolTip, const EUserInterfaceActionType::Type InUserInterfaceActionType, bool bInCloseSelfOnly, bool bInShouldCloseWindowAfterMenuSelection )
+FMenuEntryBlock::FMenuEntryBlock( const FName& InExtensionHook, const FUIAction& UIAction, const TSharedRef< SWidget > Contents, const TAttribute<FText>& InToolTip, const EUserInterfaceActionType::Type InUserInterfaceActionType, bool bInCloseSelfOnly, bool bInShouldCloseWindowAfterMenuSelection, bool bInInvertLabelOnHover)
 	: FMultiBlock( UIAction, InExtensionHook, EMultiBlockType::MenuEntry )
 	, ToolTipOverride( InToolTip )
 	, EntryWidget( Contents )
@@ -113,11 +119,12 @@ FMenuEntryBlock::FMenuEntryBlock( const FName& InExtensionHook, const FUIAction&
 	, UserInterfaceActionType( InUserInterfaceActionType )
 	, bCloseSelfOnly( bInCloseSelfOnly )
 	, bShouldCloseWindowAfterMenuSelection( bInShouldCloseWindowAfterMenuSelection )
+	, bInvertLabelOnHover(bInInvertLabelOnHover)
 {
 }
 
 
-FMenuEntryBlock::FMenuEntryBlock( const FName& InExtensionHook, const TSharedRef< SWidget > Contents, const FNewMenuDelegate& InEntryBuilder, TSharedPtr<FExtender> InExtender, bool bInSubMenu, bool bInSubMenuOnClick, TSharedPtr< const FUICommandList > InCommandList, bool bInCloseSelfOnly, bool bInShouldCloseWindowAfterMenuSelection )
+FMenuEntryBlock::FMenuEntryBlock( const FName& InExtensionHook, const TSharedRef< SWidget > Contents, const FNewMenuDelegate& InEntryBuilder, TSharedPtr<FExtender> InExtender, bool bInSubMenu, bool bInSubMenuOnClick, TSharedPtr< const FUICommandList > InCommandList, bool bInCloseSelfOnly, bool bInShouldCloseWindowAfterMenuSelection, bool bInInvertLabelOnHover)
 	: FMultiBlock( nullptr, InCommandList, InExtensionHook, EMultiBlockType::MenuEntry )
 	, EntryBuilder( InEntryBuilder )
 	, EntryWidget( Contents )
@@ -127,11 +134,12 @@ FMenuEntryBlock::FMenuEntryBlock( const FName& InExtensionHook, const TSharedRef
 	, bCloseSelfOnly( bInCloseSelfOnly )
 	, Extender( InExtender )
 	, bShouldCloseWindowAfterMenuSelection( bInShouldCloseWindowAfterMenuSelection )
+	, bInvertLabelOnHover(bInInvertLabelOnHover)
 {
 }
 
 
-FMenuEntryBlock::FMenuEntryBlock( const FName& InExtensionHook, const FUIAction& UIAction, const TSharedRef< SWidget > Contents, const FNewMenuDelegate& InEntryBuilder, TSharedPtr<FExtender> InExtender, bool bInSubMenu, TSharedPtr< const FUICommandList > InCommandList, bool bInCloseSelfOnly, bool bInShouldCloseWindowAfterMenuSelection )
+FMenuEntryBlock::FMenuEntryBlock( const FName& InExtensionHook, const FUIAction& UIAction, const TSharedRef< SWidget > Contents, const FNewMenuDelegate& InEntryBuilder, TSharedPtr<FExtender> InExtender, bool bInSubMenu, TSharedPtr< const FUICommandList > InCommandList, bool bInCloseSelfOnly, bool bInShouldCloseWindowAfterMenuSelection, bool bInInvertLabelOnHover)
 	: FMultiBlock( UIAction, InExtensionHook, EMultiBlockType::MenuEntry )
 	, EntryBuilder( InEntryBuilder )
 	, EntryWidget( Contents )
@@ -141,6 +149,7 @@ FMenuEntryBlock::FMenuEntryBlock( const FName& InExtensionHook, const FUIAction&
 	, bCloseSelfOnly( bInCloseSelfOnly )
 	, Extender( InExtender )
 	, bShouldCloseWindowAfterMenuSelection( bInShouldCloseWindowAfterMenuSelection )
+	, bInvertLabelOnHover(bInInvertLabelOnHover)
 {
 }
 
@@ -154,7 +163,7 @@ void FMenuEntryBlock::CreateMenuEntry(FMenuBuilder& InMenuBuilder) const
 bool FMenuEntryBlock::HasIcon() const
 {
 	const FSlateIcon ActionIcon = GetAction().IsValid() ? GetAction()->GetIcon() : FSlateIcon();
-	const FSlateIcon& ActualIcon = !IconOverride.IsSet() ? ActionIcon : IconOverride;
+	const FSlateIcon& ActualIcon = !IconOverride.IsSet() || !IconOverride.Get().IsSet() ? ActionIcon : IconOverride.Get();
 
 	if (ActualIcon.IsSet())
 	{
@@ -173,7 +182,9 @@ bool FMenuEntryBlock::HasIcon() const
  */
 TSharedRef< class IMultiBlockBaseWidget > FMenuEntryBlock::ConstructWidget() const
 {
-	return SNew( SMenuEntryBlock );
+	return SNew( SMenuEntryBlock )
+		.Icon(IconOverride)
+		.InvertOnHover(bInvertLabelOnHover);
 }
 
 /**
@@ -187,17 +198,22 @@ void SMenuEntryBlock::Construct( const FArguments& InArgs )
 	/*TimeToSubMenuOpen = 0.0f;*/
 	//SubMenuRequestState = Idle;
 
+	Icon = InArgs._Icon;
+
 	// No images by default
 	CheckedImage = nullptr;
 	UncheckedImage = nullptr;
 
-	this->SetForegroundColor( TAttribute<FSlateColor>::Create( TAttribute<FSlateColor>::FGetter::CreateRaw( this, &SMenuEntryBlock::InvertOnHover ) ) );
+	if (InArgs._InvertOnHover)
+	{
+		this->SetForegroundColor(TAttribute<FSlateColor>::Create(TAttribute<FSlateColor>::FGetter::CreateRaw(this, &SMenuEntryBlock::InvertOnHover)));
+	}
 }
 
 
 TSharedRef< SWidget> SMenuEntryBlock::BuildMenuBarWidget( const FMenuEntryBuildParams& InBuildParams )
 {
-	const TAttribute<FText>& Label = InBuildParams.Label;
+	Label = InBuildParams.Label;
 	const TAttribute<FText>& EntryToolTip = InBuildParams.ToolTip;
 
 	check( OwnerMultiBoxWidget.IsValid() );
@@ -208,7 +224,7 @@ TSharedRef< SWidget> SMenuEntryBlock::BuildMenuBarWidget( const FMenuEntryBuildP
 	/* Style for menu bar button with sub menu opened */
 	MenuBarButtonBorderSubmenuOpen = StyleSet->GetBrush(StyleName, ".Button.SubMenuOpen");
 	/* Style for menu bar button with no sub menu opened */
-	MenuBarButtonBorderSubmenuClosed = FCoreStyle::Get().GetBrush("NoBorder");
+	MenuBarButtonBorderSubmenuClosed = StyleSet->GetBrush(StyleName, ".Button.SubMenuClosed");
 
 	TSharedPtr< SMenuAnchor > NewMenuAnchor;
 
@@ -224,40 +240,64 @@ TSharedRef< SWidget> SMenuEntryBlock::BuildMenuBarWidget( const FMenuEntryBuildP
 		[
 			SNew( SBorder )
 
-			.BorderImage( this, &SMenuEntryBlock::GetMenuBarButtonBorder )
+			.BorderImage(FCoreStyle::Get().GetBrush("NoBorder"))
 
-			.Padding(0)
+			.Padding( StyleSet->GetMargin(ISlateStyle::Join( StyleName, ".Label.Padding" ) ) )
 
 			[
-				// Create a button
-				SNew( SButton )
+				SNew( SBorder )
 
-				// Use the menu bar item style for this button
-				.ButtonStyle( StyleSet, ISlateStyle::Join( StyleName, ".Button" ) )
+				.BorderImage( this, &SMenuEntryBlock::GetMenuBarButtonBorder )
 
-				// Pull-down menu bar items always activate on mouse-down, not mouse-up
-				.ClickMethod( EButtonClickMethod::MouseDown )
+				.Padding(0)
 
-				// Pass along the block's tool-tip string
-				.ToolTipText( this, &SMenuEntryBlock::GetFilteredToolTipText, EntryToolTip)
+				[
+					// Create a button
+					SNew( SButton )
 
-				// Add horizontal padding between the edge of the button and the content.  Also add a bit of vertical
-				// padding to push the text down from the top of the menu bar a bit.
-				.ContentPadding( FMargin( 10.0f, 2.0f ) )
+					// Use the menu bar item style for this button
+					.ButtonStyle( StyleSet, ISlateStyle::Join( StyleName, ".Button" ) )
 
-				.ForegroundColor( FSlateColor::UseForeground() )
+					// Pull-down menu bar items always activate on mouse-down, not mouse-up
+					.ClickMethod( EButtonClickMethod::MouseDown )
 
-				.VAlign( VAlign_Center )
+					// Pass along the block's tool-tip string
+					.ToolTipText( this, &SMenuEntryBlock::GetFilteredToolTipText, EntryToolTip)
 
-				[
-					SNew( STextBlock )
-					.TextStyle( StyleSet, ISlateStyle::Join( StyleName, ".Label" ) )
-					.Text( Label )
-					.HighlightText( OwnerMultiBoxWidget.Pin().Get(), &SMultiBoxWidget::GetSearchText )
-				]
+					// Add horizontal padding between the edge of the button and the content.  Also add a bit of vertical
+					// padding to push the text down from the top of the menu bar a bit.
+					.ContentPadding( StyleSet->GetMargin(ISlateStyle::Join(StyleName, ".Label.ContentPadding") ) )
+
+					.ForegroundColor( FSlateColor::UseForeground() )
+
+					.VAlign( VAlign_Center )
+
+					[
+						SNew(SHorizontalBox)
+						+SHorizontalBox::Slot()
+						.AutoWidth()
+						.HAlign(HAlign_Left)
+						.VAlign(VAlign_Center)
+						.Padding(TAttribute<FMargin>::Create(TAttribute<FMargin>::FGetter::CreateSP(this, &SMenuEntryBlock::GetMenuEntryIconPadding)))
+						[
+							SNew(SImage)
+							.Image(this, &SMenuEntryBlock::GetMenuEntryBrush)
+						]
+
+						+SHorizontalBox::Slot()
+						.AutoWidth()
+						.VAlign(VAlign_Center)
+						[
+							SNew(STextBlock)
+							.TextStyle(StyleSet, ISlateStyle::Join(StyleName, ".Label"))
+							.Text(Label)
+							.HighlightText(OwnerMultiBoxWidget.Pin().Get(), &SMultiBoxWidget::GetSearchText)
+						]
+					]
 
-				// Bind the button's "on clicked" event to our object's method for this
-				.OnClicked( this, &SMenuEntryBlock::OnMenuItemButtonClicked )
+					// Bind the button's "on clicked" event to our object's method for this
+					.OnClicked( this, &SMenuEntryBlock::OnMenuItemButtonClicked )
+				]
 			]
 		]
 	;
@@ -419,7 +459,7 @@ private:
 
 TSharedRef< SWidget > SMenuEntryBlock::BuildMenuEntryWidget( const FMenuEntryBuildParams& InBuildParams )
 {
-	const TAttribute<FText>& Label = InBuildParams.Label;
+	Label = InBuildParams.Label;
 	const TAttribute<FText>& EntryToolTip = InBuildParams.ToolTip;
 	const TSharedPtr< const FMenuEntryBlock > MenuEntryBlock = InBuildParams.MenuEntryBlock;
 	const TSharedPtr< const FMultiBox > MultiBox = InBuildParams.MultiBox;
@@ -429,7 +469,7 @@ TSharedRef< SWidget > SMenuEntryBlock::BuildMenuEntryWidget( const FMenuEntryBui
 	const FSlateIcon ActionIcon = UICommand.IsValid() ? UICommand->GetIcon() : FSlateIcon();
 
 	// Allow the block to override the tool bar icon, too
-	const FSlateIcon& ActualIcon = !MenuEntryBlock->IconOverride.IsSet() ? ActionIcon : MenuEntryBlock->IconOverride;
+	const FSlateIcon& ActualIcon = !MenuEntryBlock->IconOverride.IsSet() || !MenuEntryBlock->IconOverride.Get().IsSet() ? ActionIcon : MenuEntryBlock->IconOverride.Get();
 
 	check( OwnerMultiBoxWidget.IsValid() );
 
@@ -689,7 +729,7 @@ private:
 
 TSharedRef< SWidget> SMenuEntryBlock::BuildSubMenuWidget( const FMenuEntryBuildParams& InBuildParams )
 {
-	const TAttribute<FText>& Label = InBuildParams.Label;
+	Label = InBuildParams.Label;
 	const TAttribute<FText>& EntryToolTip = InBuildParams.ToolTip;
 
 	const TSharedPtr< const FMenuEntryBlock > MenuEntryBlock = InBuildParams.MenuEntryBlock;
@@ -700,7 +740,7 @@ TSharedRef< SWidget> SMenuEntryBlock::BuildSubMenuWidget( const FMenuEntryBuildP
 	const FSlateIcon ActionIcon = UICommand.IsValid() ? UICommand->GetIcon() : FSlateIcon();
 	
 	// Allow the block to override the tool bar icon, too
-	const FSlateIcon& ActualIcon = !MenuEntryBlock->IconOverride.IsSet() ? ActionIcon : MenuEntryBlock->IconOverride;
+	const FSlateIcon& ActualIcon = !MenuEntryBlock->IconOverride.IsSet() || !MenuEntryBlock->IconOverride.Get().IsSet()  ? ActionIcon : MenuEntryBlock->IconOverride.Get();
 
 	check( OwnerMultiBoxWidget.IsValid() );
 
@@ -1456,3 +1496,17 @@ const FSlateBrush* SMenuEntryBlock::GetCheckBoxImageBrushFromStyle(const FCheckB
 	}
 }
 
+const FSlateBrush* SMenuEntryBlock::GetMenuEntryBrush() const
+{
+	return Icon.Get().GetIcon();
+}
+
+FMargin SMenuEntryBlock::GetMenuEntryIconPadding() const
+{
+	if ((!Icon.IsSet() || !Icon.Get().IsSet()) || (!Label.IsSet() || Label.Get().IsEmpty()))
+	{
+		return FMargin(0.0f);
+	}
+	
+	return FMargin(0.0f, 0.0f, 3.0f, 0.0f);
+}
\ No newline at end of file
diff --git a/Engine/Source/Runtime/Slate/Private/Framework/MultiBox/SMenuEntryBlock.h b/Engine/Source/Runtime/Slate/Private/Framework/MultiBox/SMenuEntryBlock.h
--- a/Engine/Source/Runtime/Slate/Private/Framework/MultiBox/SMenuEntryBlock.h
+++ b/Engine/Source/Runtime/Slate/Private/Framework/MultiBox/SMenuEntryBlock.h
@@ -35,8 +35,9 @@ public:
 	 * @param	InToolTipOverride	Optional tool tip override.	 If omitted, then the action's label will be used instead.
 	 * @param	InIconOverride		Optional icon to use for the tool bar image.  If omitted, then the action's icon will be used instead.
 	 * @param	bInShouldCloseWindowAfterMenuSelection	In the case of a submenu, whether it should close after an item is selected
+	 * @param	bInInvertLabelOnHover	Whether to invert the label text's color on hover
 	 */
-	FMenuEntryBlock( const FName& InExtensionHook, const TSharedPtr< const FUICommandInfo > InCommand, TSharedPtr< const FUICommandList > InCommandList, const TAttribute<FText>& InLabelOverride = TAttribute<FText>(), const TAttribute<FText>& InToolTipOverride = TAttribute<FText>(), const FSlateIcon& InIconOverride = FSlateIcon(), bool bInCloseSelfOnly = false, bool bInShouldCloseWindowAfterMenuSelection = true);
+	FMenuEntryBlock( const FName& InExtensionHook, const TSharedPtr< const FUICommandInfo > InCommand, TSharedPtr< const FUICommandList > InCommandList, const TAttribute<FText>& InLabelOverride = TAttribute<FText>(), const TAttribute<FText>& InToolTipOverride = TAttribute<FText>(), const TAttribute<FSlateIcon>& InIconOverride = TAttribute<FSlateIcon>(), bool bInCloseSelfOnly = false, bool bInShouldCloseWindowAfterMenuSelection = true, bool bInInvertLabelOnHover = true);
 
 
 	/**
@@ -51,8 +52,9 @@ public:
 	 * @param	InCommandList		The list of commands bound to delegates that should be executed for menu entries
 	 * @param	InIcon				The icon to display to the left of the label
 	 * @param	bInShouldCloseWindowAfterMenuSelection	In the case of a submenu, whether it should close after an item is selected
+	 * @param	bInInvertLabelOnHover	Whether to invert the label text's color on hover
 	 */
-	FMenuEntryBlock( const FName& InExtensionHook, const TAttribute<FText>& InLabel, const TAttribute<FText>& InToolTip, const FNewMenuDelegate& InEntryBuilder, TSharedPtr<FExtender> InExtender, bool bInSubMenu, bool bInSubMenuOnClick, TSharedPtr< const FUICommandList > InCommandList, bool bInCloseSelfOnly, const FSlateIcon& InIcon = FSlateIcon(), bool bInShouldCloseWindowAfterMenuSelection = true);
+	FMenuEntryBlock( const FName& InExtensionHook, const TAttribute<FText>& InLabel, const TAttribute<FText>& InToolTip, const FNewMenuDelegate& InEntryBuilder, TSharedPtr<FExtender> InExtender, bool bInSubMenu, bool bInSubMenuOnClick, TSharedPtr< const FUICommandList > InCommandList, bool bInCloseSelfOnly, const TAttribute<FSlateIcon>& InIcon = TAttribute<FSlateIcon>(), bool bInShouldCloseWindowAfterMenuSelection = true, bool bInInvertLabelOnHover = true);
 	
 
 	/**
@@ -67,8 +69,9 @@ public:
 	 * @param	InCommandList		The list of commands bound to delegates that should be executed for menu entries
 	 * @param	InIcon				The icon to display to the left of the label
 	 * @param	bInShouldCloseWindowAfterMenuSelection	In the case of a submenu, whether it should close after an item is selected
+	 * @param	bInInvertLabelOnHover	Whether to invert the label text's color on hover
 	 */
-	FMenuEntryBlock( const FName& InExtensionHook, const TAttribute<FText>& InLabel, const TAttribute<FText>& InToolTip, const FOnGetContent& InMenuBuilder, TSharedPtr<FExtender> InExtender, bool bInSubMenu, bool bInSubMenuOnClick, TSharedPtr< const FUICommandList > InCommandList, bool bInCloseSelfOnly, const FSlateIcon& InIcon = FSlateIcon(), bool bInShouldCloseWindowAfterMenuSelection = true);
+	FMenuEntryBlock( const FName& InExtensionHook, const TAttribute<FText>& InLabel, const TAttribute<FText>& InToolTip, const FOnGetContent& InMenuBuilder, TSharedPtr<FExtender> InExtender, bool bInSubMenu, bool bInSubMenuOnClick, TSharedPtr< const FUICommandList > InCommandList, bool bInCloseSelfOnly, const TAttribute<FSlateIcon>& InIcon = TAttribute<FSlateIcon>(), bool bInShouldCloseWindowAfterMenuSelection = true, bool bInInvertLabelOnHover = true);
 
 	/**
 	 * Constructor
@@ -82,8 +85,9 @@ public:
 	 * @param	InCommandList		The list of commands bound to delegates that should be executed for menu entries
 	 * @param	InIcon				The icon to display to the left of the label
 	 * @param	bInShouldCloseWindowAfterMenuSelection	In the case of a submenu, whether it should close after an item is selected
+	 * @param	bInInvertLabelOnHover	Whether to invert the label text's color on hover
 	 */
-	FMenuEntryBlock( const FName& InExtensionHook, const TAttribute<FText>& InLabel, const TAttribute<FText>& InToolTip, const TSharedPtr<SWidget>& InEntryWidget, TSharedPtr<FExtender> InExtender, bool bInSubMenu, bool bInSubMenuOnClick, TSharedPtr< const FUICommandList > InCommandList, bool bInCloseSelfOnly, const FSlateIcon& InIcon = FSlateIcon(), bool bInShouldCloseWindowAfterMenuSelection = true);
+	FMenuEntryBlock( const FName& InExtensionHook, const TAttribute<FText>& InLabel, const TAttribute<FText>& InToolTip, const TSharedPtr<SWidget>& InEntryWidget, TSharedPtr<FExtender> InExtender, bool bInSubMenu, bool bInSubMenuOnClick, TSharedPtr< const FUICommandList > InCommandList, bool bInCloseSelfOnly, const TAttribute<FSlateIcon>& InIcon = TAttribute<FSlateIcon>(), bool bInShouldCloseWindowAfterMenuSelection = true, bool bInInvertLabelOnHover = true);
 
 
 	/**
@@ -95,16 +99,17 @@ public:
 	 * @param	InUIAction			UI action to take when this menu item is clicked as well as to determine if the menu entry can be executed or appears "checked"
 	 * @param	InUserInterfaceActionType	Type of interface action
 	 * @param	bInShouldCloseWindowAfterMenuSelection	In the case of a submenu, whether it should close after an item is selected
+	 * @param	bInInvertLabelOnHover	Whether to invert the label text's color on hover
 	 */
-	FMenuEntryBlock( const FName& InExtensionHook, const TAttribute<FText>& InLabel, const TAttribute<FText>& InToolTip, const FNewMenuDelegate& InEntryBuilder, TSharedPtr<FExtender> InExtender, bool bInSubMenu, bool bInSubMenuOnClick, const FSlateIcon& InIcon, const FUIAction& InUIAction, const EUserInterfaceActionType::Type InUserInterfaceActionType, bool bInCloseSelfOnly, bool bInShouldCloseWindowAfterMenuSelection = true);
+	FMenuEntryBlock( const FName& InExtensionHook, const TAttribute<FText>& InLabel, const TAttribute<FText>& InToolTip, const FNewMenuDelegate& InEntryBuilder, TSharedPtr<FExtender> InExtender, bool bInSubMenu, bool bInSubMenuOnClick, const TAttribute<FSlateIcon>& InIcon, const FUIAction& InUIAction, const EUserInterfaceActionType::Type InUserInterfaceActionType, bool bInCloseSelfOnly, bool bInShouldCloseWindowAfterMenuSelection = true, bool bInInvertLabelOnHover = true);
 
-	FMenuEntryBlock( const FName& InExtensionHook, const TAttribute<FText>& InLabel, const TAttribute<FText>& InToolTip, const FSlateIcon& InIcon, const FUIAction& InUIAction, const EUserInterfaceActionType::Type InUserInterfaceActionType, bool bInCloseSelfOnly, bool bInShouldCloseWindowAfterMenuSelection = true );
+	FMenuEntryBlock( const FName& InExtensionHook, const TAttribute<FText>& InLabel, const TAttribute<FText>& InToolTip, const TAttribute<FSlateIcon>& InIcon, const FUIAction& InUIAction, const EUserInterfaceActionType::Type InUserInterfaceActionType, bool bInCloseSelfOnly, bool bInShouldCloseWindowAfterMenuSelection = true, bool bInInvertLabelOnHover = true);
 
-	FMenuEntryBlock( const FName& InExtensionHook, const FUIAction& UIAction, const TSharedRef< SWidget > Contents, const TAttribute<FText>& InToolTip, const EUserInterfaceActionType::Type InUserInterfaceActionType, bool bInCloseSelfOnly, bool bInShouldCloseWindowAfterMenuSelection = true );
+	FMenuEntryBlock( const FName& InExtensionHook, const FUIAction& UIAction, const TSharedRef< SWidget > Contents, const TAttribute<FText>& InToolTip, const EUserInterfaceActionType::Type InUserInterfaceActionType, bool bInCloseSelfOnly, bool bInShouldCloseWindowAfterMenuSelection = true, bool bInInvertLabelOnHover = true);
 
-	FMenuEntryBlock( const FName& InExtensionHook, const TSharedRef< SWidget > Contents, const FNewMenuDelegate& InEntryBuilder, TSharedPtr<FExtender> InExtender, bool bInSubMenu, bool bInSubMenuOnClick, TSharedPtr< const FUICommandList > InCommandList, bool bInCloseSelfOnly, bool bInShouldCloseWindowAfterMenuSelection = true );
+	FMenuEntryBlock( const FName& InExtensionHook, const TSharedRef< SWidget > Contents, const FNewMenuDelegate& InEntryBuilder, TSharedPtr<FExtender> InExtender, bool bInSubMenu, bool bInSubMenuOnClick, TSharedPtr< const FUICommandList > InCommandList, bool bInCloseSelfOnly, bool bInShouldCloseWindowAfterMenuSelection = true, bool bInInvertLabelOnHover = true);
 
-	FMenuEntryBlock( const FName& InExtensionHook, const FUIAction& UIAction, const TSharedRef< SWidget > Contents, const FNewMenuDelegate& InEntryBuilder, TSharedPtr<FExtender> InExtender, bool bInSubMenu, TSharedPtr< const FUICommandList > InCommandList, bool bInCloseSelfOnly, bool bInShouldCloseWindowAfterMenuSelection = true );
+	FMenuEntryBlock( const FName& InExtensionHook, const FUIAction& UIAction, const TSharedRef< SWidget > Contents, const FNewMenuDelegate& InEntryBuilder, TSharedPtr<FExtender> InExtender, bool bInSubMenu, TSharedPtr< const FUICommandList > InCommandList, bool bInCloseSelfOnly, bool bInShouldCloseWindowAfterMenuSelection = true, bool bInInvertLabelOnHover = true);
 
 	/** FMultiBlock interface */
 	virtual void CreateMenuEntry(class FMenuBuilder& MenuBuilder) const override;
@@ -130,7 +135,7 @@ private:
 	TAttribute<FText> ToolTipOverride;
 
 	/** Optional overridden icon for this tool bar button.  IF not set, then the action's icon will be used instead. */
-	FSlateIcon IconOverride;
+	TAttribute<FSlateIcon> IconOverride;
 
 	/** Optional menu entry builder associated with this entry for building sub-menus and pull down menus */
 	FNewMenuDelegate EntryBuilder;
@@ -159,6 +164,9 @@ private:
 
 	/** For submenus, whether the menu should be closed after something is selected */
 	bool bShouldCloseWindowAfterMenuSelection;
+
+	/** Whether to invert the label text's color on hover */
+	bool bInvertLabelOnHover;
 };
 
 
@@ -173,7 +181,15 @@ class SLATE_API SMenuEntryBlock
 
 public:
 
-	SLATE_BEGIN_ARGS( SMenuEntryBlock ){}
+	SLATE_BEGIN_ARGS( SMenuEntryBlock )
+		: _InvertOnHover(true)
+	{}
+
+	/** The icon that will be displayed int the menu block */
+	SLATE_ATTRIBUTE(FSlateIcon, Icon)
+
+	/** Whether to invert the label text's color on hover */
+	SLATE_ARGUMENT(bool, InvertOnHover)
 
 	SLATE_END_ARGS()
 
@@ -347,6 +363,12 @@ protected:
 private:
 	const FSlateBrush* GetCheckBoxImageBrushFromStyle(const FCheckBoxStyle* Style) const;
 
+	/** Get the menu entry brush for this menu block */
+	const FSlateBrush* GetMenuEntryBrush() const;
+
+	/** Get the padding to use for the menu block icon */
+	FMargin GetMenuEntryIconPadding() const;
+
 	/** One-off delayed active timer to update the open/closed state of the sub menu. */
 	EActiveTimerReturnType UpdateSubMenuState(double InCurrentTime, float InDeltaTime, bool bWantsOpen);
 
@@ -368,6 +390,12 @@ private:
 	/** Style for menu bar button with no sub menu opened */
 	const FSlateBrush* MenuBarButtonBorderSubmenuClosed;
 
+	/** The icon that will be displayed for this menu block */
+	TAttribute<FSlateIcon> Icon;
+
+	/** The label for the menu block item */
+	TAttribute<FText> Label;
+
 	///** The pending SubMenu request state */
 	//enum
 	//{
diff --git a/Engine/Source/Runtime/Slate/Private/Framework/MultiBox/SToolBarButtonBlock.cpp b/Engine/Source/Runtime/Slate/Private/Framework/MultiBox/SToolBarButtonBlock.cpp
--- a/Engine/Source/Runtime/Slate/Private/Framework/MultiBox/SToolBarButtonBlock.cpp
+++ b/Engine/Source/Runtime/Slate/Private/Framework/MultiBox/SToolBarButtonBlock.cpp
@@ -221,11 +221,13 @@ void SToolBarButtonBlock::BuildMultiBlockWidget(const ISlateStyle* StyleSet, con
 	if( UserInterfaceType == EUserInterfaceActionType::Button )
 	{
 		FName BlockStyle = EMultiBlockLocation::ToName(ISlateStyle::Join( StyleName, ".Button" ), BlockLocation);
+		FName ButtonPadding = ISlateStyle::Join(StyleName, ".SToolBarButtonBlock.Button.Padding");
+
 		ChildSlot
 		[
 			// Create a button
 			SNew( SButton )
-				.ContentPadding(0)
+				.ContentPadding(StyleSet->GetMargin(ButtonPadding) )
 
 				// Use the tool bar item style for this button
 				.ButtonStyle( StyleSet, BlockStyle )
diff --git a/Engine/Source/Runtime/Slate/Public/Framework/MultiBox/MultiBoxBuilder.h b/Engine/Source/Runtime/Slate/Public/Framework/MultiBox/MultiBoxBuilder.h
--- a/Engine/Source/Runtime/Slate/Public/Framework/MultiBox/MultiBoxBuilder.h
+++ b/Engine/Source/Runtime/Slate/Public/Framework/MultiBox/MultiBoxBuilder.h
@@ -257,9 +257,9 @@ public:
 	 * @param	InIcon				The icon to use
 	 * @param	bInShouldCloseWindowAfterMenuSelection	Whether the submenu should close after an item is selected
 	 */
-	void AddSubMenu( const TAttribute<FText>& InMenuLabel, const TAttribute<FText>& InToolTip, const FNewMenuDelegate& InSubMenu, const FUIAction& InUIAction, FName InExtensionHook, const EUserInterfaceActionType::Type InUserInterfaceActionType, const bool bInOpenSubMenuOnClick = false, const FSlateIcon& InIcon = FSlateIcon(), const bool bInShouldCloseWindowAfterMenuSelection = true );
+	void AddSubMenu( const TAttribute<FText>& InMenuLabel, const TAttribute<FText>& InToolTip, const FNewMenuDelegate& InSubMenu, const FUIAction& InUIAction, FName InExtensionHook, const EUserInterfaceActionType::Type InUserInterfaceActionType, const bool bInOpenSubMenuOnClick = false, const TAttribute<FSlateIcon>& InIcon = TAttribute<FSlateIcon>(), const bool bInShouldCloseWindowAfterMenuSelection = true );
 
-	void AddSubMenu( const TAttribute<FText>& InMenuLabel, const TAttribute<FText>& InToolTip, const FNewMenuDelegate& InSubMenu, const bool bInOpenSubMenuOnClick = false, const FSlateIcon& InIcon = FSlateIcon(), const bool bInShouldCloseWindowAfterMenuSelection = true );
+	void AddSubMenu( const TAttribute<FText>& InMenuLabel, const TAttribute<FText>& InToolTip, const FNewMenuDelegate& InSubMenu, const bool bInOpenSubMenuOnClick = false, const TAttribute<FSlateIcon>& InIcon = TAttribute<FSlateIcon>(), const bool bInShouldCloseWindowAfterMenuSelection = true );
 
 	void AddSubMenu( const TSharedRef< SWidget > Contents, const FNewMenuDelegate& InSubMenu, const bool bInOpenSubMenuOnClick = false, const bool bInShouldCloseWindowAfterMenuSelection = true );
 
@@ -334,8 +334,9 @@ public:
 	 *
 	 * @param	InCommandList	The action list that maps command infos to delegates that should be called for each command associated with a multiblock widget
 	 */
-	FMenuBarBuilder( TSharedPtr< const FUICommandList > InCommandList, TSharedPtr<FExtender> InExtender = TSharedPtr<FExtender>(), const ISlateStyle* InStyleSet = &FCoreStyle::Get() )
+	FMenuBarBuilder( TSharedPtr< const FUICommandList > InCommandList, TSharedPtr<FExtender> InExtender = TSharedPtr<FExtender>(), const ISlateStyle* InStyleSet = &FCoreStyle::Get(), bool bInInvertLabelsOnHover = true )
 		: FBaseMenuBuilder( EMultiBoxType::MenuBar, false, InCommandList, false, InExtender, InStyleSet )
+		, bInvertLabelsOnHover(bInInvertLabelsOnHover)
 	{
 	}
 
@@ -347,11 +348,25 @@ public:
 	 * @param	InToolTip			The tooltip that should be shown when the menu is hovered over
 	 * @param	InPullDownMenu		Pull down menu object for the menu to add.
 	 */
-	void AddPullDownMenu( const FText& InMenuLabel, const FText& InToolTip, const FNewMenuDelegate& InPullDownMenu, FName InExtensionHook = NAME_None, FName InTutorialHighlightName = NAME_None );
+	void AddPullDownMenu( const TAttribute<FText>& InMenuLabel, const TAttribute<FText>& InToolTip, const FNewMenuDelegate& InPullDownMenu, FName InExtensionHook = NAME_None, FName InTutorialHighlightName = NAME_None);
 	
+	/**
+	 * Adds a pull down menu
+	 *
+	 * @param	InMenuLabel			The text that should be shown for the menu
+	 * @param	InToolTip			The tooltip that should be shown when the menu is hovered over
+	 * @param	InPullDownMenu		Pull down menu object for the menu to add.
+	 * @param	InIcon				Icon to use for the menu bar entry
+	 */
+	void AddPullDownMenu( const TAttribute<FText>& InMenuLabel, const TAttribute<FText>& InToolTip, const TAttribute<FSlateIcon>& InIcon, const FNewMenuDelegate& InPullDownMenu, FName InExtensionHook = NAME_None, FName InTutorialHighlightName = NAME_None );
+
 protected:
 	/** FMultiBoxBuilder interface */
 	virtual void ApplyHook(FName InExtensionHook, EExtensionHook::Position HookPosition) override;
+
+protected:
+	/** Whether to invert the label text's color on hover */
+	bool bInvertLabelsOnHover;
 };
 
 
diff --git a/Engine/Source/Runtime/SlateCore/Private/Styling/CoreStyle.cpp b/Engine/Source/Runtime/SlateCore/Private/Styling/CoreStyle.cpp
--- a/Engine/Source/Runtime/SlateCore/Private/Styling/CoreStyle.cpp
+++ b/Engine/Source/Runtime/SlateCore/Private/Styling/CoreStyle.cpp
@@ -816,6 +816,7 @@ TSharedRef<ISlateStyle> FCoreStyle::Create( const FName& InStyleSetName )
 		Style->Set( "ToolBar.SToolBarButtonBlock.Padding", FMargin(4.0f,0.0f));
 		Style->Set( "ToolBar.SToolBarCheckComboButtonBlock.Padding", FMargin(4.0f,0.0f));
 		Style->Set( "ToolBar.SToolBarButtonBlock.CheckBox.Padding", FMargin(4.0f,0.0f) );
+		Style->Set( "ToolBar.SToolBarButtonBlock.Button.Padding", FMargin(0.0f));
 		Style->Set( "ToolBar.SToolBarComboButtonBlock.ComboButton.Color", DefaultForeground );
 
 		Style->Set( "ToolBar.Block.IndentedPadding", FMargin( 18.0f, 2.0f, 4.0f, 4.0f ) );
@@ -895,6 +896,7 @@ TSharedRef<ISlateStyle> FCoreStyle::Create( const FName& InStyleSetName )
 		Style->Set( "Menu.SToolBarButtonBlock.Padding", FMargin(4.0f));
 		Style->Set( "Menu.SToolBarCheckComboButtonBlock.Padding", FMargin(4.0f));
 		Style->Set( "Menu.SToolBarButtonBlock.CheckBox.Padding", FMargin(0.0f) );
+		Style->Set( "Menu.SToolBarButtonBlock.Button.Padding", FMargin(0.0f));
 		Style->Set( "Menu.SToolBarComboButtonBlock.ComboButton.Color", DefaultForeground );
 
 		Style->Set( "Menu.Block.IndentedPadding", FMargin( 18.0f, 2.0f, 4.0f, 4.0f ) );
@@ -904,6 +906,8 @@ TSharedRef<ISlateStyle> FCoreStyle::Create( const FName& InStyleSetName )
 		Style->Set( "Menu.Separator.Padding", FMargin( 0.5f ) );
 
 		Style->Set( "Menu.Label", FTextBlockStyle(NormalText) .SetFont( TTF_FONT( "Fonts/Roboto-Regular", 9 ) ) );
+		Style->Set( "Menu.Label.Padding", FMargin(0.0f, 0.0f, 0.0f, 0.0f) );
+		Style->Set( "Menu.Label.ContentPadding", FMargin(10.0f, 2.0f) );
 		Style->Set( "Menu.EditableText", FEditableTextBoxStyle(NormalEditableTextBoxStyle) .SetFont( TTF_FONT( "Fonts/Roboto-Regular", 9 ) ) );
 		Style->Set( "Menu.Keybinding", FTextBlockStyle(NormalText) .SetFont( TTF_FONT( "Fonts/Roboto-Regular", 8 ) ) );
 
@@ -984,6 +988,9 @@ TSharedRef<ISlateStyle> FCoreStyle::Create( const FName& InStyleSetName )
 
 		/* The style of a menu bar button when it has a sub menu open */
 		Style->Set( "Menu.Button.SubMenuOpen", new BORDER_BRUSH( "Common/Selection", FMargin(4.f/16.f), FLinearColor(0.10f, 0.10f, 0.10f) ) );
+
+		/* The style of a menu bar button when it has a sub menu closed */
+		Style->Set("Menu.Button.SubMenuClosed", new FSlateNoResource());
 	}
 
 	// SExpandableButton defaults...
diff --git a/Engine/Source/Runtime/Windows/AudioMixerXAudio2/Private/AudioMixerPlatformWindows.cpp b/Engine/Source/Runtime/Windows/AudioMixerXAudio2/Private/AudioMixerPlatformWindows.cpp
--- a/Engine/Source/Runtime/Windows/AudioMixerXAudio2/Private/AudioMixerPlatformWindows.cpp
+++ b/Engine/Source/Runtime/Windows/AudioMixerXAudio2/Private/AudioMixerPlatformWindows.cpp
@@ -193,7 +193,7 @@ namespace Audio
 		WindowsNotificationClient->RegisterDeviceChangedListener(this);
 	}
 
-	void FMixerPlatformXAudio2::UnRegisterDeviceChangedListener() 
+	void FMixerPlatformXAudio2::UnregisterDeviceChangedListener() 
 	{
 		WindowsNotificationClient->UnRegisterDeviceDeviceChangedListener(this);
 	}
@@ -255,7 +255,7 @@ namespace Audio
 namespace Audio
 {
 	void FMixerPlatformXAudio2::RegisterDeviceChangedListener() {}
-	void FMixerPlatformXAudio2::UnRegisterDeviceChangedListener() {}
+	void FMixerPlatformXAudio2::UnregisterDeviceChangedListener() {}
 	void FMixerPlatformXAudio2::OnDefaultCaptureDeviceChanged(const EAudioDeviceRole InAudioDeviceRole, const FString& DeviceId) {}
 	void FMixerPlatformXAudio2::OnDefaultRenderDeviceChanged(const EAudioDeviceRole InAudioDeviceRole, const FString& DeviceId) {}
 	void FMixerPlatformXAudio2::OnDeviceAdded(const FString& DeviceId) {}
diff --git a/Engine/Source/Runtime/Windows/AudioMixerXAudio2/Private/AudioMixerPlatformXAudio2.h b/Engine/Source/Runtime/Windows/AudioMixerXAudio2/Private/AudioMixerPlatformXAudio2.h
--- a/Engine/Source/Runtime/Windows/AudioMixerXAudio2/Private/AudioMixerPlatformXAudio2.h
+++ b/Engine/Source/Runtime/Windows/AudioMixerXAudio2/Private/AudioMixerPlatformXAudio2.h
@@ -71,7 +71,7 @@ namespace Audio
 
 		//~ Begin IAudioMixerDeviceChangedLister
 		virtual void RegisterDeviceChangedListener() override;
-		virtual void UnRegisterDeviceChangedListener() override;
+		virtual void UnregisterDeviceChangedListener() override;
 		virtual void OnDefaultCaptureDeviceChanged(const EAudioDeviceRole InAudioDeviceRole, const FString& DeviceId) override;
 		virtual void OnDefaultRenderDeviceChanged(const EAudioDeviceRole InAudioDeviceRole, const FString& DeviceId) override;
 		virtual void OnDeviceAdded(const FString& DeviceId) override;
diff --git a/Engine/Source/Runtime/Windows/XAudio2/Private/XAudio2Buffer.cpp b/Engine/Source/Runtime/Windows/XAudio2/Private/XAudio2Buffer.cpp
--- a/Engine/Source/Runtime/Windows/XAudio2/Private/XAudio2Buffer.cpp
+++ b/Engine/Source/Runtime/Windows/XAudio2/Private/XAudio2Buffer.cpp
@@ -580,7 +580,6 @@ FXAudio2SoundBuffer* FXAudio2SoundBuffer::Init( FAudioDevice* AudioDevice, USoun
 		break;
 
 	case DTYPE_Native:
-	case DTYPE_Xenon:
 		// Upload entire wav to XAudio2
 		if( Wave->ResourceID )
 		{
