Fix an issue making standalone tools unusable even at 1.0 dpi scale

diff --git a/Engine/Source/Runtime/Slate/Private/Widgets/Input/SMenuAnchor.cpp b/Engine/Source/Runtime/Slate/Private/Widgets/Input/SMenuAnchor.cpp
--- a/Engine/Source/Runtime/Slate/Private/Widgets/Input/SMenuAnchor.cpp
+++ b/Engine/Source/Runtime/Slate/Private/Widgets/Input/SMenuAnchor.cpp
@@ -337,7 +337,7 @@ void SMenuAnchor::SetIsOpen( bool InIsOpen, const bool bFocusMenu, const int32 F
 					const EMenuPlacement PlacementMode = Placement.Get();
 
 					const FVector2D NewPosition = MyGeometry.AbsolutePosition;
-					FVector2D NewWindowSize = MyGeometry.GetDrawSize();
+					FVector2D NewWindowSize = DesiredContentSize;
 					const FVector2D SummonLocationSize = MyGeometry.GetLocalSize();
 
 					FPopupTransitionEffect TransitionEffect( FPopupTransitionEffect::None );
@@ -365,7 +365,7 @@ void SMenuAnchor::SetIsOpen( bool InIsOpen, const bool bFocusMenu, const int32 F
 						if (MethodInUse.GetPopupMethod() == EPopupMethod::CreateNewWindow)
 						{
 							// Open the pop-up
-							TSharedPtr<IMenu> NewMenu = FSlateApplication::Get().PushMenu(AsShared(), MyWidgetPath, MenuContentRef, NewPosition, TransitionEffect, bFocusMenu, NewWindowSize, MethodInUse.GetPopupMethod(), bIsCollapsedByParent);
+							TSharedPtr<IMenu> NewMenu = FSlateApplication::Get().PushMenu(AsShared(), MyWidgetPath, MenuContentRef, NewPosition, TransitionEffect, bFocusMenu, MyGeometry.GetLocalSize(), MethodInUse.GetPopupMethod(), bIsCollapsedByParent);
 
 							PopupMenuPtr = NewMenu;
 							check(NewMenu.IsValid() && NewMenu->GetOwnedWindow().IsValid());
