Copying Orion-Staging @ 3748653 (Orion/Dev-General @ 3722124) to //UE4/Main

diff --git a/Engine/Source/Runtime/Renderer/Private/RendererScene.cpp b/Engine/Source/Runtime/Renderer/Private/RendererScene.cpp
--- a/Engine/Source/Runtime/Renderer/Private/RendererScene.cpp
+++ b/Engine/Source/Runtime/Renderer/Private/RendererScene.cpp
@@ -490,22 +490,98 @@ void FScene::CheckPrimitiveArrays()
 	check(Primitives.Num() == PrimitiveOcclusionBounds.Num());
 }
 
+
+template<typename T>
+static void SwapTArray(TArray<T>& Array, int i1, int i2)
+{
+	T tmp = Array[i1];
+	Array[i1] = Array[i2];
+	Array[i2] = tmp;
+}
+
 void FScene::AddPrimitiveSceneInfo_RenderThread(FRHICommandListImmediate& RHICmdList, FPrimitiveSceneInfo* PrimitiveSceneInfo)
 {
 	SCOPE_CYCLE_COUNTER(STAT_AddScenePrimitiveRenderThreadTime);
 	
 	CheckPrimitiveArrays();
 
-	int32 PrimitiveIndex = Primitives.Add(PrimitiveSceneInfo);
-	PrimitiveSceneInfo->PackedIndex = PrimitiveIndex;
-
-	PrimitiveSceneProxies.AddUninitialized();
+	Primitives.Add(PrimitiveSceneInfo);
+	PrimitiveSceneProxies.Add(PrimitiveSceneInfo->Proxy);
 	PrimitiveBounds.AddUninitialized();
 	PrimitiveFlagsCompact.AddUninitialized();
 	PrimitiveVisibilityIds.AddUninitialized();
 	PrimitiveOcclusionFlags.AddUninitialized();
 	PrimitiveComponentIds.AddUninitialized();
 	PrimitiveOcclusionBounds.AddUninitialized();
+	
+	const int SourceIndex = PrimitiveSceneProxies.Num() - 1;
+	PrimitiveSceneInfo->PackedIndex = SourceIndex;
+
+	{
+		bool EntryFound = false;
+		int BroadIndex = - 1;
+		SIZE_T InsertProxyHash = PrimitiveSceneInfo->Proxy->GetTypeHash();
+		//broad phase search for a matching type
+		for (BroadIndex = TypeOffsetTable.Num() - 1; BroadIndex >= 0; BroadIndex--)
+		{
+			// example how the prefix sum of the tails could look like
+			// PrimitiveSceneProxies[0,0,0,6,6,6,6,6,2,2,2,2,1,1,1,7,4,8]
+			// TypeOffsetTable[3,8,12,15,16,17,18]
+
+			if (TypeOffsetTable[BroadIndex].PrimitiveSceneProxyType == InsertProxyHash)
+			{
+				EntryFound = true;
+				break;
+			}
+		}
+
+		//new type encountered
+		if (EntryFound == false)
+		{
+			BroadIndex = TypeOffsetTable.Num();
+			if (BroadIndex)
+			{
+				FTypeOffsetTableEntry Entry = TypeOffsetTable[BroadIndex - 1];
+				//adding to the end of the list and offset of the tail (will will be incremented once during the while loop)
+				TypeOffsetTable.Push(FTypeOffsetTableEntry(InsertProxyHash, Entry.Offset));
+			}
+			else
+			{
+				//starting with an empty list and offset zero (will will be incremented once during the while loop)
+				TypeOffsetTable.Push(FTypeOffsetTableEntry(InsertProxyHash, 0));
+			}
+		}
+
+		while (BroadIndex < TypeOffsetTable.Num())
+		{
+			FTypeOffsetTableEntry& NextEntry = TypeOffsetTable[BroadIndex++];
+			int DestIndex = NextEntry.Offset++; //prepare swap and increment
+
+			// example swap chain of inserting a type of 6 at the end
+			// PrimitiveSceneProxies[0,0,0,6,6,6,6,6,2,2,2,2,1,1,1,7,4,8,6]
+			// PrimitiveSceneProxies[0,0,0,6,6,6,6,6,6,2,2,2,1,1,1,7,4,8,2]
+			// PrimitiveSceneProxies[0,0,0,6,6,6,6,6,6,2,2,2,2,1,1,7,4,8,1]
+			// PrimitiveSceneProxies[0,0,0,6,6,6,6,6,6,2,2,2,2,1,1,1,4,8,7]
+			// PrimitiveSceneProxies[0,0,0,6,6,6,6,6,6,2,2,2,2,1,1,1,7,8,4]
+			// PrimitiveSceneProxies[0,0,0,6,6,6,6,6,6,2,2,2,2,1,1,1,7,4,8]
+
+			if (DestIndex != SourceIndex)
+			{
+				checkfSlow(SourceIndex > DestIndex, TEXT("Corrupted Prefix Sum [%d, %d]"), SourceIndex, DestIndex);
+				Primitives[DestIndex]->PackedIndex = SourceIndex;
+				Primitives[SourceIndex]->PackedIndex = DestIndex;
+
+				SwapTArray(Primitives, DestIndex, SourceIndex);
+				SwapTArray(PrimitiveSceneProxies, DestIndex, SourceIndex);
+				SwapTArray(PrimitiveBounds, DestIndex, SourceIndex);
+				SwapTArray(PrimitiveFlagsCompact, DestIndex, SourceIndex);
+				SwapTArray(PrimitiveVisibilityIds, DestIndex, SourceIndex);
+				SwapTArray(PrimitiveOcclusionFlags, DestIndex, SourceIndex);
+				SwapTArray(PrimitiveComponentIds, DestIndex, SourceIndex);
+				SwapTArray(PrimitiveOcclusionBounds, DestIndex, SourceIndex);
+			}
+		}
+	}
 
 	CheckPrimitiveArrays();
 
@@ -975,22 +1051,85 @@ void FScene::RemovePrimitiveSceneInfo_RenderThread(FPrimitiveSceneInfo* Primitiv
 	CheckPrimitiveArrays();
 
 	int32 PrimitiveIndex = PrimitiveSceneInfo->PackedIndex;
-	Primitives.RemoveAtSwap(PrimitiveIndex);
-	PrimitiveSceneProxies.RemoveAtSwap(PrimitiveIndex);
-	PrimitiveBounds.RemoveAtSwap(PrimitiveIndex);
-	PrimitiveFlagsCompact.RemoveAtSwap(PrimitiveIndex);
-	PrimitiveVisibilityIds.RemoveAtSwap(PrimitiveIndex);
-	PrimitiveOcclusionFlags.RemoveAtSwap(PrimitiveIndex);
-	PrimitiveComponentIds.RemoveAtSwap(PrimitiveIndex);
-	PrimitiveOcclusionBounds.RemoveAtSwap(PrimitiveIndex);
-	if (Primitives.IsValidIndex(PrimitiveIndex))
-	{
-		FPrimitiveSceneInfo* OtherPrimitive = Primitives[PrimitiveIndex];
-		OtherPrimitive->PackedIndex = PrimitiveIndex;
-
-		// Invalidate the scene info's PackedIndex now that it is used by another primitive
-		PrimitiveSceneInfo->PackedIndex = MAX_int32;
+	{
+		int BroadIndex = -1;
+		SIZE_T InsertProxyHash = PrimitiveSceneInfo->Proxy->GetTypeHash();
+		//broad phase search for a matching type
+		for (BroadIndex = TypeOffsetTable.Num() - 1; BroadIndex >= 0; BroadIndex--)
+		{
+			// example how the prefix sum of the tails could look like
+			// PrimitiveSceneProxies[0,0,0,6,6,6,6,6,2,2,2,2,1,1,1,7,4,8]
+			// TypeOffsetTable[3,8,12,15,16,17,18]
+
+			if (TypeOffsetTable[BroadIndex].PrimitiveSceneProxyType == InsertProxyHash)
+			{
+				const int InsertionOffset = TypeOffsetTable[BroadIndex].Offset;
+				const int PrevOffset = BroadIndex > 0 ? TypeOffsetTable[BroadIndex - 1].Offset : 0;
+				checkfSlow(PrimitiveIndex >= PrevOffset && PrimitiveIndex < InsertionOffset, TEXT("PrimitiveIndex %d not in Bucket Range [%d, %d]"), PrimitiveIndex, PrevOffset, InsertionOffset);
+				break;
+			}
+		}
+
+		int SourceIndex = PrimitiveIndex;
+		const int SavedBroadIndex = BroadIndex;
+		while ( BroadIndex < TypeOffsetTable.Num() )
+		{
+			FTypeOffsetTableEntry& NextEntry = TypeOffsetTable[BroadIndex++];
+			int DestIndex = --NextEntry.Offset; //decrement and prepare swap 
+
+			// example swap chain of removing X 
+			// PrimitiveSceneProxies[0,0,0,6,X,6,6,6,2,2,2,2,1,1,1,7,4,8]
+			// PrimitiveSceneProxies[0,0,0,6,6,6,6,6,X,2,2,2,1,1,1,7,4,8]
+			// PrimitiveSceneProxies[0,0,0,6,6,6,6,6,2,2,2,X,1,1,1,7,4,8]
+			// PrimitiveSceneProxies[0,0,0,6,6,6,6,6,2,2,2,1,1,1,X,7,4,8]
+			// PrimitiveSceneProxies[0,0,0,6,6,6,6,6,2,2,2,1,1,1,7,X,4,8]
+			// PrimitiveSceneProxies[0,0,0,6,6,6,6,6,2,2,2,1,1,1,7,4,X,8]
+			// PrimitiveSceneProxies[0,0,0,6,6,6,6,6,2,2,2,1,1,1,7,4,8,X]
+
+			if (DestIndex != SourceIndex)
+			{
+				checkfSlow(DestIndex > SourceIndex, TEXT("Corrupted Prefix Sum [%d, %d]"), DestIndex, SourceIndex);
+				Primitives[DestIndex]->PackedIndex = SourceIndex;
+				Primitives[SourceIndex]->PackedIndex = DestIndex;
+
+				SwapTArray(Primitives, DestIndex, SourceIndex);
+				SwapTArray(PrimitiveSceneProxies, DestIndex, SourceIndex);
+				SwapTArray(PrimitiveBounds, DestIndex, SourceIndex);
+				SwapTArray(PrimitiveFlagsCompact, DestIndex, SourceIndex);
+				SwapTArray(PrimitiveVisibilityIds, DestIndex, SourceIndex);
+				SwapTArray(PrimitiveOcclusionFlags, DestIndex, SourceIndex);
+				SwapTArray(PrimitiveComponentIds, DestIndex, SourceIndex);
+				SwapTArray(PrimitiveOcclusionBounds, DestIndex, SourceIndex);
+				SourceIndex = DestIndex;
+			}
+		}
+
+		const int PreviousOffset = SavedBroadIndex > 0 ? TypeOffsetTable[SavedBroadIndex - 1].Offset : 0;
+		const int CurrentOffset = TypeOffsetTable[SavedBroadIndex].Offset;
+
+		checkfSlow(PreviousOffset <= CurrentOffset, TEXT("Corrupted Bucket [%d, %d]"), PreviousOffset, CurrentOffset);
+		if (CurrentOffset - PreviousOffset == 0)
+		{
+			// remove empty OffsetTable entries e.g.
+			// TypeOffsetTable[3,8,12,15,15,17,18]
+			// TypeOffsetTable[3,8,12,15,17,18]
+			TypeOffsetTable.RemoveAt(SavedBroadIndex);
+		}
+
+		checkfSlow((TypeOffsetTable.Num() == 0 && Primitives.Num() == 1) || TypeOffsetTable[TypeOffsetTable.Num() - 1].Offset == Primitives.Num() - 1, TEXT("Corrupted Tail Offset [%d, %d]"), TypeOffsetTable[TypeOffsetTable.Num() - 1].Offset, Primitives.Num() - 1);
+		checkfSlow(Primitives[Primitives.Num() - 1] == PrimitiveSceneInfo, TEXT("Removed item should be at the end"));
+
+		Primitives.Pop();
+		PrimitiveSceneProxies.Pop();
+		PrimitiveBounds.Pop();
+		PrimitiveFlagsCompact.Pop();
+		PrimitiveVisibilityIds.Pop();
+		PrimitiveOcclusionFlags.Pop();
+		PrimitiveComponentIds.Pop();
+		PrimitiveOcclusionBounds.Pop();
 	}
+
+	PrimitiveSceneInfo->PackedIndex = MAX_int32;
 	
 	CheckPrimitiveArrays();
 
diff --git a/Engine/Source/Runtime/Renderer/Private/ScenePrivate.h b/Engine/Source/Runtime/Renderer/Private/ScenePrivate.h
--- a/Engine/Source/Runtime/Renderer/Private/ScenePrivate.h
+++ b/Engine/Source/Runtime/Renderer/Private/ScenePrivate.h
@@ -1980,6 +1980,15 @@ public:
 	/** Packed array of primitive components associated with the primitive. */
 	TArray<FPrimitiveComponentId> PrimitiveComponentIds;
 
+	struct FTypeOffsetTableEntry
+	{
+		FTypeOffsetTableEntry(SIZE_T InPrimitiveSceneProxyType, uint32 InOffset) : PrimitiveSceneProxyType(InPrimitiveSceneProxyType), Offset(InOffset) {}
+		SIZE_T PrimitiveSceneProxyType;
+		uint32 Offset; //(e.g. prefix sum where the next type starts)
+	};
+	/* During insertion and deletion, used to skip large chunks of items of the same type */
+	TArray<FTypeOffsetTableEntry> TypeOffsetTable;
+
 	/** The lights in the scene. */
 	TSparseArray<FLightSceneInfoCompact> Lights;
 
diff --git a/Engine/Source/Runtime/Renderer/Private/SceneVisibility.cpp b/Engine/Source/Runtime/Renderer/Private/SceneVisibility.cpp
--- a/Engine/Source/Runtime/Renderer/Private/SceneVisibility.cpp
+++ b/Engine/Source/Runtime/Renderer/Private/SceneVisibility.cpp
@@ -136,7 +136,8 @@ static FAutoConsoleVariableRef CVarOcclusionCullParallelPrimFetch(
 	ECVF_RenderThreadSafe
 	);
 
-static int32 GILCUpdatePrimTaskEnabled = 0;
+static int32 GILCUpdatePrimTaskEnabled = 1;
+
 static FAutoConsoleVariableRef CVarILCUpdatePrimitivesTask(
 	TEXT("r.Cache.UpdatePrimsTaskEnabled"),
 	GILCUpdatePrimTaskEnabled,
@@ -2710,7 +2711,7 @@ void FSceneRenderer::PostVisibilityFrameSetup(FILCUpdatePrimTaskData& OutILCTask
 	if (ViewFamily.EngineShowFlags.HitProxies == 0 && Scene->PrecomputedLightVolumes.Num() > 0)
 	{
 		QUICK_SCOPE_CYCLE_COUNTER(STAT_PostVisibilityFrameSetup_IndirectLightingCache_Update);
-		if (GILCUpdatePrimTaskEnabled)
+		if (GILCUpdatePrimTaskEnabled && FPlatformProcess::SupportsMultithreading())
 		{
 			Scene->IndirectLightingCache.StartUpdateCachePrimitivesTask(Scene, *this, true, OutILCTaskData);
 		}
diff --git a/Engine/Source/Runtime/Renderer/Private/ShadowDepthRendering.cpp b/Engine/Source/Runtime/Renderer/Private/ShadowDepthRendering.cpp
--- a/Engine/Source/Runtime/Renderer/Private/ShadowDepthRendering.cpp
+++ b/Engine/Source/Runtime/Renderer/Private/ShadowDepthRendering.cpp
@@ -1642,7 +1642,7 @@ void FProjectedShadowInfo::CopyCachedShadowMap(FRHICommandList& RHICmdList, cons
 	DrawRenderState.ApplyToPSO(GraphicsPSOInit);
 	uint32 StencilRef = DrawRenderState.GetStencilRef();
 
-	if (CachedShadowMapData.bCachedShadowMapHasPrimitives)
+	if (CachedShadowMapData.bCachedShadowMapHasPrimitives && CachedShadowMapData.ShadowMap.IsValid())
 	{
 		SCOPED_DRAW_EVENT(RHICmdList, CopyCachedShadowMap);
 
diff --git a/Engine/Source/Runtime/Slate/Private/Framework/Application/SlateApplication.cpp b/Engine/Source/Runtime/Slate/Private/Framework/Application/SlateApplication.cpp
--- a/Engine/Source/Runtime/Slate/Private/Framework/Application/SlateApplication.cpp
+++ b/Engine/Source/Runtime/Slate/Private/Framework/Application/SlateApplication.cpp
@@ -34,6 +34,7 @@
 #include "GenericPlatform/ITextInputMethodSystem.h"
 #include "ToolboxModule.h"
 #include "Framework/Docking/TabCommands.h"
+#include "HAL/LowLevelMemTracker.h"
 #include "Math/UnitConversion.h"
 #include "HAL/LowLevelMemTracker.h"
 
@@ -921,7 +922,7 @@ void FSlateApplication::Shutdown(bool bShutdownPlatform)
 
 		if (bShutdownPlatform)
 		{
-			PlatformApplication->DestroyApplication();
+		PlatformApplication->DestroyApplication();
 		}
 
 		PlatformApplication.Reset();
@@ -1618,7 +1619,7 @@ void FSlateApplication::Tick(ESlateTickType TickType)
 
 		if (TickType == ESlateTickType::All)
 		{
-			TickPlatform(DeltaTime);
+		TickPlatform(DeltaTime);
 		}
 		TickApplication(TickType, DeltaTime);
 	}
@@ -1668,34 +1669,34 @@ void FSlateApplication::TickApplication(ESlateTickType TickType, float DeltaTime
 
 	if(TickType == ESlateTickType::All)
 	{
-		FPlatformMisc::BeginNamedEvent(FColor::Magenta, "Slate::PreTick");
-		{
-			SCOPE_CYCLE_COUNTER(STAT_SlatePreTickEvent);
-			PreTickEvent.Broadcast(DeltaTime);
-		}
-		FPlatformMisc::EndNamedEvent();
+	FPlatformMisc::BeginNamedEvent(FColor::Magenta, "Slate::PreTick");
+	{
+		SCOPE_CYCLE_COUNTER(STAT_SlatePreTickEvent);
+		PreTickEvent.Broadcast(DeltaTime);
+	}
+	FPlatformMisc::EndNamedEvent();
 
-		//FPlatformMisc::BeginNamedEvent(FColor::Magenta, "Slate::UpdateCursorLockRegion");
-		// The widget locking the cursor to its bounds may have been reshaped.
-		// Check if the widget was reshaped and update the cursor lock
-		// bounds if needed.
-		UpdateCursorLockRegion();
-		//FPlatformMisc::EndNamedEvent();
+	//FPlatformMisc::BeginNamedEvent(FColor::Magenta, "Slate::UpdateCursorLockRegion");
+	// The widget locking the cursor to its bounds may have been reshaped.
+	// Check if the widget was reshaped and update the cursor lock
+	// bounds if needed.
+	UpdateCursorLockRegion();
+	//FPlatformMisc::EndNamedEvent();
 
-		//FPlatformMisc::BeginNamedEvent(FColor::Magenta, "Slate::CaptureAndToolTipUpdate");
-		// When Slate captures the mouse, it is up to us to set the cursor 
-		// because the OS assumes that we own the mouse.
+	//FPlatformMisc::BeginNamedEvent(FColor::Magenta, "Slate::CaptureAndToolTipUpdate");
+	// When Slate captures the mouse, it is up to us to set the cursor 
+	// because the OS assumes that we own the mouse.
 		if (MouseCaptor.HasCapture() || bQueryCursorRequested)
-		{
-			QueryCursor();
-		}
+	{
+		QueryCursor();
+	}
 
-		{
+	{
 			SCOPE_CYCLE_COUNTER(STAT_SlateUpdateTooltip);
-			SLATE_CYCLE_COUNTER_SCOPE(GUpdateTooltipTime);
+		SLATE_CYCLE_COUNTER_SCOPE(GUpdateTooltipTime);
 
-			// Update tool tip, if we have one
-			const bool AllowSpawningOfToolTips = false;
+		// Update tool tip, if we have one
+		const bool AllowSpawningOfToolTips = false;
 			UpdateToolTip(AllowSpawningOfToolTips);
 		}
 	}
@@ -1738,72 +1739,72 @@ void FSlateApplication::TickApplication(ESlateTickType TickType, float DeltaTime
 
 	if (TickType == ESlateTickType::All)
 	{
-		const bool bNeedsSyntheticMouseMouse = SynthesizeMouseMovePending > 0;
+	const bool bNeedsSyntheticMouseMouse = SynthesizeMouseMovePending > 0;
 		if (bNeedsSyntheticMouseMouse && (!GIsGameThreadIdInitialized || IsInGameThread())) 
-		{
-			// Force a mouse move event to make sure all widgets know whether there is a mouse cursor hovering over them
-			SynthesizeMouseMove();
-			--SynthesizeMouseMovePending;
-		}
-
-		// Update auto-throttling based on elapsed time since user interaction
-		ThrottleApplicationBasedOnMouseMovement();
+	{
+		// Force a mouse move event to make sure all widgets know whether there is a mouse cursor hovering over them
+		SynthesizeMouseMove();
+		--SynthesizeMouseMovePending;
+	}
 
-		TSharedPtr<SWindow> ActiveModalWindow = GetActiveModalWindow();
+	// Update auto-throttling based on elapsed time since user interaction
+	ThrottleApplicationBasedOnMouseMovement();
 
-		const float SleepThreshold = SleepBufferPostInput.GetValueOnGameThread();
-		const double TimeSinceInput = LastTickTime - LastUserInteractionTime;
-		const double TimeSinceMouseMove = LastTickTime - LastMouseMoveTime;
+	TSharedPtr<SWindow> ActiveModalWindow = GetActiveModalWindow();
 
-		const bool bIsUserIdle = (TimeSinceInput > SleepThreshold) && (TimeSinceMouseMove > SleepThreshold);
-		const bool bAnyActiveTimersPending = AnyActiveTimersArePending();
-		if (bAnyActiveTimersPending)
-		{
-			// Some UI might slide under the cursor. To a widget, this is
-			// as if the cursor moved over it.
-			QueueSynthesizedMouseMove();
-		}
+	const float SleepThreshold = SleepBufferPostInput.GetValueOnGameThread();
+	const double TimeSinceInput = LastTickTime - LastUserInteractionTime;
+	const double TimeSinceMouseMove = LastTickTime - LastMouseMoveTime;
+	
+	const bool bIsUserIdle = (TimeSinceInput > SleepThreshold) && (TimeSinceMouseMove > SleepThreshold);
+	const bool bAnyActiveTimersPending = AnyActiveTimersArePending();
+	if (bAnyActiveTimersPending)
+	{
+		// Some UI might slide under the cursor. To a widget, this is
+		// as if the cursor moved over it.
+		QueueSynthesizedMouseMove();
+	}
 	// Generate any simulated gestures that we've detected.
 	ForEachUser([&] (FSlateUser* User) {
 		User->GestureDetector.GenerateGestures(*this, SimulateGestures);
 	});
 
-		// Check if any element lists used for caching need to be released
-		{
+	// Check if any element lists used for caching need to be released
+	{
 			for (int32 CacheIndex = 0; CacheIndex < ReleasedCachedElementLists.Num(); CacheIndex++)
-			{
+		{
 				if (ReleasedCachedElementLists[CacheIndex]->IsInUse() == false)
-				{
+			{
 					ensure(ReleasedCachedElementLists[CacheIndex].IsUnique());
-					ReleasedCachedElementLists[CacheIndex].Reset();
-					ReleasedCachedElementLists.RemoveAtSwap(CacheIndex, 1, false);
-					CacheIndex--;
-				}
+				ReleasedCachedElementLists[CacheIndex].Reset();
+				ReleasedCachedElementLists.RemoveAtSwap(CacheIndex, 1, false);
+				CacheIndex--;
 			}
 		}
+	}
 
-		// skip tick/draw if we are idle and there are no active timers registered that we need to drive slate for.
-		// This effectively means the slate application is totally idle and we don't need to update the UI.
-		// This relies on Widgets properly registering for Active timer when they need something to happen even
-		// when the user is not providing any input (ie, animations, viewport rendering, async polling, etc).
-		bIsSlateAsleep = true;
-		if (!AllowSlateToSleep.GetValueOnGameThread() || bAnyActiveTimersPending || !bIsUserIdle || bNeedsSyntheticMouseMouse || FApp::UseVRFocus())
-		{
-			bIsSlateAsleep = false; // if we get here, then Slate is not sleeping
-
-			// Update any notifications - this needs to be done after windows have updated themselves 
-			// (so they know their size)
-			{
-				SLATE_CYCLE_COUNTER_SCOPE(GSlateTickNotificationManager);
-				FSlateNotificationManager::Get().Tick();
-			}
+	// skip tick/draw if we are idle and there are no active timers registered that we need to drive slate for.
+	// This effectively means the slate application is totally idle and we don't need to update the UI.
+	// This relies on Widgets properly registering for Active timer when they need something to happen even
+	// when the user is not providing any input (ie, animations, viewport rendering, async polling, etc).
+	bIsSlateAsleep = true;
+	if (!AllowSlateToSleep.GetValueOnGameThread() || bAnyActiveTimersPending || !bIsUserIdle || bNeedsSyntheticMouseMouse || FApp::UseVRFocus())
+	{
+		bIsSlateAsleep = false; // if we get here, then Slate is not sleeping
 
-			// Draw all windows
-			DrawWindows();
+		// Update any notifications - this needs to be done after windows have updated themselves 
+		// (so they know their size)
+		{
+			SLATE_CYCLE_COUNTER_SCOPE(GSlateTickNotificationManager);
+			FSlateNotificationManager::Get().Tick();
 		}
 
-		PostTickEvent.Broadcast(DeltaTime);
+		// Draw all windows
+		DrawWindows();
 	}
+
+	PostTickEvent.Broadcast(DeltaTime);
+}
 }
 
 
@@ -2027,7 +2028,7 @@ EUINavigation FSlateApplication::GetNavigationDirectionFromKey(const FKeyEvent&
 EUINavigation FSlateApplication::GetNavigationDirectionFromAnalog(const FAnalogInputEvent& InAnalogEvent)
 {
 	if (const FSlateUser* User = GetUser(InAnalogEvent.GetUserIndex()))
-	{
+{
 		return User->NavigationConfig->GetNavigationDirectionFromAnalog(InAnalogEvent);
 	}
 	return EUINavigation::Invalid;
@@ -2527,26 +2528,26 @@ bool FSlateApplication::SetUserFocus(uint32 UserIndex, const TSharedPtr<SWidget>
 	{
 		if (FSlateUser* User = GetOrCreateUser(UserIndex))
 		{
-			FWidgetPath PathToWidget;
-			const bool bFound = FSlateWindowHelper::FindPathToWidget(SlateWindows, WidgetToFocus.ToSharedRef(), /*OUT*/ PathToWidget);
-			if (bFound)
-			{
+		FWidgetPath PathToWidget;
+		const bool bFound = FSlateWindowHelper::FindPathToWidget(SlateWindows, WidgetToFocus.ToSharedRef(), /*OUT*/ PathToWidget);
+		if (bFound)
+		{
 				return SetUserFocus(User, PathToWidget, ReasonFocusIsChanging);
+		}
+		else
+		{
+			const bool bFoundVirtual = FSlateWindowHelper::FindPathToWidget(SlateVirtualWindows, WidgetToFocus.ToSharedRef(), /*OUT*/ PathToWidget);
+				if (bFoundVirtual)
+			{
+					return SetUserFocus(User, PathToWidget, ReasonFocusIsChanging);
 			}
 			else
 			{
-				const bool bFoundVirtual = FSlateWindowHelper::FindPathToWidget(SlateVirtualWindows, WidgetToFocus.ToSharedRef(), /*OUT*/ PathToWidget);
-				if (bFoundVirtual)
-				{
-					return SetUserFocus(User, PathToWidget, ReasonFocusIsChanging);
-				}
-				else
-				{
-					//ensureMsgf(bFound, TEXT("Attempting to focus a widget that isn't in the tree and visible: %s. If your intent is to clear focus use ClearUserFocus()"), WidgetToFocus->ToString());
-				}
+				//ensureMsgf(bFound, TEXT("Attempting to focus a widget that isn't in the tree and visible: %s. If your intent is to clear focus use ClearUserFocus()"), WidgetToFocus->ToString());
 			}
 		}
 	}
+	}
 
 	return false;
 }
@@ -2727,7 +2728,7 @@ bool FSlateApplication::SetUserFocus(FSlateUser* User, const FWidgetPath& InFocu
 	{
 		return false;
 	}
-	
+
 	TSharedPtr<IWidgetReflector> WidgetReflector = WidgetReflectorPtr.Pin();
 	const bool bReflectorShowingFocus = WidgetReflector.IsValid() && WidgetReflector->IsShowingFocus();
 
@@ -3868,6 +3869,7 @@ void FSlateApplication::CancelDragDrop()
 
 	for( auto LastWidgetIterator = WidgetsUnderCursorLastEvent.CreateConstIterator(); LastWidgetIterator; ++LastWidgetIterator)
 	{
+		
 		FWidgetPath WidgetsToDragLeave = LastWidgetIterator.Value().ToWidgetPath(FWeakWidgetPath::EInterruptedPathHandling::Truncate);
 		if(WidgetsToDragLeave.IsValid())
 		{
@@ -4055,13 +4057,13 @@ void FSlateApplication::SetInputPreProcessor(bool bEnable, TSharedPtr<class IInp
 		RegisterInputPreProcessor(NewInputProcessor);
 	}
 	else
-	{
+{
 		if ( NewInputProcessor.IsValid() )
-		{
+	{
 			UnregisterInputPreProcessor(NewInputProcessor);
-		}
-		else
-		{
+	}
+	else
+	{
 			UnregisterAllInputPreProcessors();
 		}
 	}
@@ -4158,10 +4160,10 @@ FVector2D FSlateApplication::CalculatePopupWindowPosition( const FSlateRect& InA
 
 		if (ProposedPlacement.IsZero())
 		{
-			// Assume natural left-to-right, top-to-bottom flow; position popup below and to the right.
+		// Assume natural left-to-right, top-to-bottom flow; position popup below and to the right.
 			ProposedPlacement = FVector2D(
-				Orientation == Orient_Horizontal ? AnchorRect.Right : AnchorRect.Left,
-				Orientation == Orient_Horizontal ? AnchorRect.Top : AnchorRect.Bottom);
+			Orientation == Orient_Horizontal ? AnchorRect.Right : AnchorRect.Left,
+			Orientation == Orient_Horizontal ? AnchorRect.Top : AnchorRect.Bottom);
 		}
 
 		return ComputePopupFitInRect(InAnchor, FSlateRect(ProposedPlacement, ProposedPlacement+AdjustedSize), Orientation, WorkAreaRect) / DPIScale;
@@ -4675,11 +4677,11 @@ bool FSlateApplication::HasUserFocusedDescendants(const TSharedRef< const SWidge
 		{
 			const FWeakWidgetPath& FocusedWidgetPath = User->GetWeakFocusPath();
 			if ( FocusedWidgetPath.ContainsWidget(Widget) )
-			{
-				return true;
-			}
+		{
+			return true;
 		}
 	}
+	}
 
 	return false;
 }
@@ -4695,12 +4697,12 @@ bool FSlateApplication::HasFocusedDescendants( const TSharedRef< const SWidget >
 			{
 				const FWeakWidgetPath& FocusedWidgetPath = User->GetWeakFocusPath();
 				if ( FocusedWidgetPath.ContainsWidget(Widget) )
-				{
-					return true;
-				}
+			{
+				return true;
 			}
 		}
 	}
+	}
 	return false;
 }
 
@@ -4857,10 +4859,10 @@ bool FSlateApplication::ProcessKeyDownEvent( FKeyEvent& InKeyEvent )
 			const bool bIsWidgetReflectorPicking = WidgetReflector.IsValid() && WidgetReflector->IsInPickingMode();
 			if ( bIsWidgetReflectorPicking )
 			{
-				WidgetReflector->OnWidgetPicked();
-				Reply = FReply::Handled();
+					WidgetReflector->OnWidgetPicked();
+					Reply = FReply::Handled();
 
-				return Reply.IsEventHandled();
+					return Reply.IsEventHandled();
 			}
 		}
 #endif
@@ -5004,7 +5006,7 @@ bool FSlateApplication::ProcessAnalogInputEvent(FAnalogInputEvent& InAnalogInput
 					: FReply::Unhandled();
 			});
 
-			LOG_EVENT_CONTENT(EEventLog::AnalogInput, InAnalogInputEvent.GetKey().ToString(), Reply);
+		LOG_EVENT_CONTENT(EEventLog::AnalogInput, InAnalogInputEvent.GetKey().ToString(), Reply);
 
 			QueueSynthesizedMouseMove();
 		}
@@ -5055,13 +5057,13 @@ void FSlateApplication::SetGameIsFakingTouchEvents(const bool bIsFaking, FVector
 {
 	if ( bIsGameFakingTouch != bIsFaking )
 	{
-		if (bIsFakingTouched && !bIsFaking && bIsGameFakingTouch && !bIsFakingTouch)
-		{
-			OnTouchEnded((CursorLocation ? *CursorLocation : PlatformApplication->Cursor->GetPosition()), 0, 0);
-		}
-
-		bIsGameFakingTouch = bIsFaking;
+	if (bIsFakingTouched && !bIsFaking && bIsGameFakingTouch && !bIsFakingTouch)
+	{
+		OnTouchEnded((CursorLocation ? *CursorLocation : PlatformApplication->Cursor->GetPosition()), 0, 0);
 	}
+
+	bIsGameFakingTouch = bIsFaking;
+}
 }
 
 #endif
@@ -5096,7 +5098,7 @@ bool FSlateApplication::OnMouseDown( const TSharedPtr< FGenericWindow >& Platfor
 		0,
 		PlatformApplication->GetModifierKeys()
 		);
-	
+
 	return ProcessMouseButtonDownEvent( PlatformWindow, MouseEvent );
 }
 
@@ -5466,7 +5468,7 @@ bool FSlateApplication::RoutePointerMoveEvent(const FWidgetPath& WidgetsUnderPoi
 				LastWidgetsUnderCursor = *DetectDragForWidget;
 
 				DragDetector.ResetDetection();
-
+				
 				// Switch worlds widgets in the current path
 				FScopedSwitchWorldHack SwitchWorld(DragDetectPath);
 
@@ -5475,7 +5477,7 @@ bool FSlateApplication::RoutePointerMoveEvent(const FWidgetPath& WidgetsUnderPoi
 				{
 					return InDetectDragForMe.Widget->OnDragDetected(InDetectDragForMe.Geometry, TranslatedMouseEvent);
 				});
-
+				
 				LOG_EVENT(EEventLog::DragDetected, Reply);
 			}
 			else
@@ -5642,7 +5644,7 @@ bool FSlateApplication::RoutePointerMoveEvent(const FWidgetPath& WidgetsUnderPoi
 
 		LOG_EVENT( IsDragDropping() ? EEventLog::DragOver : EEventLog::MouseMove, Reply )
 
-		bHandled = Reply.IsEventHandled();
+			bHandled = Reply.IsEventHandled();
 	}
 
 	// Give the current drag drop operation a chance to do something
@@ -5878,7 +5880,7 @@ bool FSlateApplication::ProcessMouseWheelOrGestureEvent( FPointerEvent& InWheelE
 	{
 		bShouldProcessEvent = InWheelEvent.GetWheelDelta() != 0;
 	}
-
+	
 	if ( !bShouldProcessEvent )
 	{
 		return false;
@@ -5959,7 +5961,7 @@ bool FSlateApplication::OnMouseMove()
 			0,
 			PlatformApplication->GetModifierKeys()
 			);
-		
+
 		if (InputPreProcessors.HandleMouseMoveEvent(*this, MouseEvent))
 		{
 			return true;
@@ -6919,7 +6921,7 @@ EDropEffect::Type FSlateApplication::OnDragDrop( const TSharedPtr< FGenericWindo
 		// User dropped into a Slate window. Slate is already in drag and drop mode.
 		// It knows what to do based on a mouse up.
 		if ( ProcessMouseButtonUpEvent( MouseEvent ) )
-		{
+{
 			Result = EDropEffect::Copy;
 		}
 	}
@@ -7122,21 +7124,21 @@ bool FSlateApplication::InputPreProcessorsHelper::HandleMotionDetectedEvent(FSla
 }
 
 bool FSlateApplication::InputPreProcessorsHelper::Add(TSharedPtr<IInputProcessor> InputProcessor, const int32 Index /*= INDEX_NONE*/)
-{
+	{
 	bool bResult = false;
 	if (Index == INDEX_NONE)
-	{
+		{
 		InputPreProcessorList.AddUnique(InputProcessor);
 		bResult = true;
-	}
+		}
 	else if (!InputPreProcessorList.Find(InputProcessor))
-	{
+		{
 		InputPreProcessorList.Insert(InputProcessor, Index);
 		bResult = true;
-	}
+		}
 
 	return bResult;
-}
+	}
 
 void FSlateApplication::InputPreProcessorsHelper::Remove(TSharedPtr<IInputProcessor> InputProcessor)
 {
diff --git a/Engine/Source/Runtime/SlateCore/Private/Styling/CoreStyle.cpp b/Engine/Source/Runtime/SlateCore/Private/Styling/CoreStyle.cpp
--- a/Engine/Source/Runtime/SlateCore/Private/Styling/CoreStyle.cpp
+++ b/Engine/Source/Runtime/SlateCore/Private/Styling/CoreStyle.cpp
@@ -859,6 +859,20 @@ TSharedRef<ISlateStyle> FCoreStyle::Create( const FName& InStyleSetName )
 		/* ... and add new style */
 		Style->Set( "ToolBar.CheckBox", ToolBarCheckBoxStyle );
 
+		/* Read-only checkbox that appears next to a menu item */
+		/* Set images for various SCheckBox states associated with read-only toolbar check box items... */
+		const FCheckBoxStyle BasicToolBarCheckStyle = FCheckBoxStyle()
+			.SetUncheckedImage(IMAGE_BRUSH("Icons/Empty_14x", Icon14x14))
+			.SetUncheckedHoveredImage(IMAGE_BRUSH("Icons/Empty_14x", Icon14x14))
+			.SetUncheckedPressedImage(IMAGE_BRUSH("Common/SmallCheckBox_Hovered", Icon14x14))
+			.SetCheckedImage(IMAGE_BRUSH("Common/SmallCheck", Icon14x14))
+			.SetCheckedHoveredImage(IMAGE_BRUSH("Common/SmallCheck", Icon14x14))
+			.SetCheckedPressedImage(IMAGE_BRUSH("Common/SmallCheck", Icon14x14))
+			.SetUndeterminedImage(IMAGE_BRUSH("Icons/Empty_14x", Icon14x14))
+			.SetUndeterminedHoveredImage(FSlateNoResource())
+			.SetUndeterminedPressedImage(FSlateNoResource());
+		Style->Set("ToolBar.Check", BasicToolBarCheckStyle);
+
 		// This radio button is actually just a check box with different images
 		/* Create style for "ToolBar.RadioButton" widget ... */
 		const FCheckBoxStyle ToolbarRadioButtonCheckBoxStyle = FCheckBoxStyle()
diff --git a/Engine/Source/Runtime/SlateCore/Public/Textures/SlateShaderResource.h b/Engine/Source/Runtime/SlateCore/Public/Textures/SlateShaderResource.h
--- a/Engine/Source/Runtime/SlateCore/Public/Textures/SlateShaderResource.h
+++ b/Engine/Source/Runtime/SlateCore/Public/Textures/SlateShaderResource.h
@@ -193,13 +193,20 @@ public:
 class FSlateResourceHandle
 {
 	friend class FSlateShaderResourceManager;
+	friend class FSlateNullShaderResourceManager;
 	friend class FSlateDrawElement;
 public:
+	FSlateResourceHandle() {}
+
 	/**
 	 * @return true if the handle still points to a valid rendering resource
 	 */
 	bool IsValid() const { return Data.IsValid() && Data->Proxy; }
 private:
+	FSlateResourceHandle(const TSharedPtr<FSlateSharedHandleData>& InData)
+		: Data(InData)
+	{}
+
 	/** Internal data to pair the handle to the resource */
 	TSharedPtr<FSlateSharedHandleData> Data;
 };
diff --git a/Engine/Source/Runtime/SlateNullRenderer/Private/SlateNullRenderer.cpp b/Engine/Source/Runtime/SlateNullRenderer/Private/SlateNullRenderer.cpp
--- a/Engine/Source/Runtime/SlateNullRenderer/Private/SlateNullRenderer.cpp
+++ b/Engine/Source/Runtime/SlateNullRenderer/Private/SlateNullRenderer.cpp
@@ -3,8 +3,9 @@
 #include "SlateNullRenderer.h"
 #include "Rendering/SlateDrawBuffer.h"
 
-FSlateNullRenderer::FSlateNullRenderer(const TSharedRef<FSlateFontServices>& InSlateFontServices)
+FSlateNullRenderer::FSlateNullRenderer(const TSharedRef<FSlateFontServices>& InSlateFontServices, const TSharedRef<FSlateShaderResourceManager>& InResourceManager)
 	: FSlateRenderer(InSlateFontServices)
+	, ResourceManager(InResourceManager)
 {
 }
 
@@ -56,7 +57,7 @@ bool FSlateNullRenderer::GenerateDynamicImageResource( FName ResourceName, uint3
 
 FSlateResourceHandle FSlateNullRenderer::GetResourceHandle( const FSlateBrush& Brush )
 {
-	return FSlateResourceHandle();
+	return ResourceManager.IsValid() ? ResourceManager->GetResourceHandle(Brush) : FSlateResourceHandle();
 }
 
 void FSlateNullRenderer::RemoveDynamicBrushResource( TSharedPtr<FSlateDynamicImageBrush> BrushToRemove )
diff --git a/Engine/Source/Runtime/SlateNullRenderer/Private/SlateNullRenderer.h b/Engine/Source/Runtime/SlateNullRenderer/Private/SlateNullRenderer.h
--- a/Engine/Source/Runtime/SlateNullRenderer/Private/SlateNullRenderer.h
+++ b/Engine/Source/Runtime/SlateNullRenderer/Private/SlateNullRenderer.h
@@ -15,7 +15,7 @@ struct Rect;
 class FSlateNullRenderer : public FSlateRenderer
 {
 public:
-	explicit FSlateNullRenderer(const TSharedRef<FSlateFontServices>& InSlateFontServices);
+	FSlateNullRenderer(const TSharedRef<FSlateFontServices>& InSlateFontServices, const TSharedRef<FSlateShaderResourceManager>& InResourceManager);
 
 	/** FSlateRenderer interface */
 	virtual bool Initialize() override;
@@ -38,4 +38,7 @@ public:
 	virtual int32 RegisterCurrentScene(FSceneInterface* Scene) override;
 	virtual int32 GetCurrentSceneIndex() const override;
 	virtual void ClearScenes() override;
+
+private:
+	TSharedPtr<FSlateShaderResourceManager> ResourceManager;
 };
diff --git a/Engine/Source/Runtime/SlateNullRenderer/Private/SlateNullRendererModule.cpp b/Engine/Source/Runtime/SlateNullRenderer/Private/SlateNullRendererModule.cpp
--- a/Engine/Source/Runtime/SlateNullRenderer/Private/SlateNullRendererModule.cpp
+++ b/Engine/Source/Runtime/SlateNullRenderer/Private/SlateNullRendererModule.cpp
@@ -24,8 +24,12 @@ public:
 
 	// FSlateShaderResourceManager interface
 	virtual FSlateShaderResourceProxy* GetShaderResource( const FSlateBrush& InBrush ) override { return nullptr; }
-	virtual FSlateResourceHandle GetResourceHandle( const FSlateBrush& InBrush ) override { return FSlateResourceHandle(); }
 	virtual ISlateAtlasProvider* GetTextureAtlasProvider() { return this; }
+	virtual FSlateResourceHandle GetResourceHandle( const FSlateBrush& InBrush ) override 
+	{
+		static const FSlateResourceHandle NullHandle(MakeShareable(new FSlateSharedHandleData));
+		return NullHandle;
+	}
 };
 
 /** A null font texture resource to represent fonts */
@@ -120,7 +124,7 @@ public:
 	{
 		ConditionalCreateResources();
 
-		return MakeShareable( new FSlateNullRenderer(SlateFontServices.ToSharedRef()) );
+		return MakeShareable( new FSlateNullRenderer(SlateFontServices.ToSharedRef(), ResourceManager.ToSharedRef()) );
 	}
 
 	virtual TSharedRef<ISlateFontAtlasFactory> CreateSlateFontAtlasFactory() override
diff --git a/Engine/Source/Runtime/SlateRHIRenderer/Private/SlatePostProcessor.cpp b/Engine/Source/Runtime/SlateRHIRenderer/Private/SlatePostProcessor.cpp
--- a/Engine/Source/Runtime/SlateRHIRenderer/Private/SlatePostProcessor.cpp
+++ b/Engine/Source/Runtime/SlateRHIRenderer/Private/SlatePostProcessor.cpp
@@ -316,7 +316,7 @@ void FSlatePostProcessor::UpsampleRect(FRHICommandListImmediate& RHICmdList, IRe
 	SetRenderTarget(RHICmdList, DestTexture, FTextureRHIRef());
 	RHICmdList.ApplyCachedRenderTargets(GraphicsPSOInit);
 
-	Params.RestoreStateFunc(GraphicsPSOInit);
+	Params.RestoreStateFunc(RHICmdList, GraphicsPSOInit);
 
 	TShaderMapRef<FScreenPS> PixelShader(ShaderMap);
 
diff --git a/Engine/Source/Runtime/SlateRHIRenderer/Private/SlatePostProcessor.h b/Engine/Source/Runtime/SlateRHIRenderer/Private/SlatePostProcessor.h
--- a/Engine/Source/Runtime/SlateRHIRenderer/Private/SlatePostProcessor.h
+++ b/Engine/Source/Runtime/SlateRHIRenderer/Private/SlatePostProcessor.h
@@ -15,7 +15,7 @@ struct FPostProcessRectParams
 	FSlateRect SourceRect;
 	FSlateRect DestRect;
 	FIntPoint SourceTextureSize;
-	TFunction<void(FGraphicsPipelineStateInitializer&)> RestoreStateFunc;
+	TFunction<void(FRHICommandListImmediate&, FGraphicsPipelineStateInitializer&)> RestoreStateFunc;
 	TFunction<void()> RestoreStateFuncPostPipelineState;
 };
 
diff --git a/Engine/Source/Runtime/SlateRHIRenderer/Private/SlateRHIRenderingPolicy.cpp b/Engine/Source/Runtime/SlateRHIRenderer/Private/SlateRHIRenderingPolicy.cpp
--- a/Engine/Source/Runtime/SlateRHIRenderer/Private/SlateRHIRenderingPolicy.cpp
+++ b/Engine/Source/Runtime/SlateRHIRenderer/Private/SlateRHIRenderingPolicy.cpp
@@ -302,10 +302,8 @@ void FSlateRHIRenderingPolicy::DrawElements(
 	const float EngineGamma = GEngine ? GEngine->GetDisplayGamma() : 2.2f;
 	const float DisplayGamma = bGammaCorrect ? EngineGamma : 1.0f;
 
-#if STATS
 	int32 ScissorClips = 0;
 	int32 StencilClips = 0;
-#endif
 
 	// In order to support MaterialParameterCollections, we need to create multiple FSceneViews for 
 	// each possible Scene that we encounter. The following code creates these as separate arrays, where the first 
@@ -424,11 +422,14 @@ void FSlateRHIRenderingPolicy::DrawElements(
 		const ESlateShader::Type ShaderType = RenderBatch.ShaderType;
 		const FShaderParams& ShaderParams = RenderBatch.ShaderParams;
 
-		auto UpdateScissorRect = [&RHICmdList,
-#if STATS
+		auto UpdateScissorRect = [
 				&ScissorClips, &StencilClips,
-#endif
-				&StencilRef, &MaskingID, &BackBuffer, &RenderBatch, &ColorTarget, &DepthStencilTarget, &RenderClipStates, &LastClippingIndex, &ViewTranslation2D, &bSwitchVerticalAxis](FGraphicsPipelineStateInitializer& InGraphicsPSOInit, const FMatrix& ViewProjection, bool bForceStateChange)
+				&StencilRef, &MaskingID, 
+				&BackBuffer, &RenderBatch, 
+				&ColorTarget, &DepthStencilTarget, 
+				&RenderClipStates, &LastClippingIndex, 
+				&ViewTranslation2D, &bSwitchVerticalAxis
+				](FRHICommandListImmediate& InRHICmdList, FGraphicsPipelineStateInitializer& InGraphicsPSOInit, const FMatrix& ViewProjection, bool bForceStateChange)
 		{
 			if ( RenderBatch.ClippingIndex != LastClippingIndex || bForceStateChange )
 			{
@@ -437,9 +438,7 @@ void FSlateRHIRenderingPolicy::DrawElements(
 					const FSlateClippingState& ClipState = RenderClipStates[RenderBatch.ClippingIndex];
 					if ( ClipState.ScissorRect.IsSet() )
 					{
-#if STATS
 						ScissorClips++;
-#endif
 
 						const FSlateClippingZone& ScissorRect = ClipState.ScissorRect.GetValue();
 
@@ -451,11 +450,11 @@ void FSlateRHIRenderingPolicy::DrawElements(
 							const FIntPoint ViewSize = BackBuffer.GetSizeXY();
 							const int32 MinY = (ViewSize.Y - BottomRight.Y);
 							const int32 MaxY = (ViewSize.Y - TopLeft.Y);
-							RHICmdList.SetScissorRect(true, TopLeft.X, MinY, BottomRight.X, MaxY);
+							InRHICmdList.SetScissorRect(true, TopLeft.X, MinY, BottomRight.X, MaxY);
 						}
 						else
 						{
-							RHICmdList.SetScissorRect(true, TopLeft.X, TopLeft.Y, BottomRight.X, BottomRight.Y);
+							InRHICmdList.SetScissorRect(true, TopLeft.X, TopLeft.Y, BottomRight.X, BottomRight.Y);
 						}
 						
 						// Disable depth/stencil testing by default
@@ -464,11 +463,9 @@ void FSlateRHIRenderingPolicy::DrawElements(
 					}
 					else
 					{
-#if STATS
 						StencilClips++;
-#endif
 
-						SLATE_DRAW_EVENT(RHICmdList, StencilClipping);
+						SLATE_DRAW_EVENT(InRHICmdList, StencilClipping);
 
 						check(ClipState.StencilQuads.Num() > 0);
 
@@ -489,7 +486,7 @@ void FSlateRHIRenderingPolicy::DrawElements(
 							bClearStencil = true;
 
 							// We don't want there to be any scissor rect when we clear the stencil
-							RHICmdList.SetScissorRect(false, 0, 0, 0, 0);
+							InRHICmdList.SetScissorRect(false, 0, 0, 0, 0);
 						}
 						else
 						{
@@ -509,11 +506,11 @@ void FSlateRHIRenderingPolicy::DrawElements(
 								const FIntPoint ViewSize = BackBuffer.GetSizeXY();
 								const int32 MinY = (ViewSize.Y - BottomRight.Y);
 								const int32 MaxY = (ViewSize.Y - TopLeft.Y);
-								RHICmdList.SetScissorRect(true, TopLeft.X, MinY, BottomRight.X, MaxY);
+								InRHICmdList.SetScissorRect(true, TopLeft.X, MinY, BottomRight.X, MaxY);
 							}
 							else
 							{
-								RHICmdList.SetScissorRect(true, TopLeft.X, TopLeft.Y, BottomRight.X, BottomRight.Y);
+								InRHICmdList.SetScissorRect(true, TopLeft.X, TopLeft.Y, BottomRight.X, BottomRight.Y);
 							}
 						}
 
@@ -524,13 +521,14 @@ void FSlateRHIRenderingPolicy::DrawElements(
 							FRHIRenderTargetView ColorView(ColorTarget, ERenderTargetLoadAction::ELoad);
 							FRHIDepthRenderTargetView DepthStencilView(DepthStencilTarget, ERenderTargetLoadAction::ENoAction, ERenderTargetStoreAction::ENoAction, bClearStencil ? ERenderTargetLoadAction::EClear : ERenderTargetLoadAction::ELoad, ERenderTargetStoreAction::EStore);
 							FRHISetRenderTargetsInfo CurrentRTInfo(1, &ColorView, DepthStencilView);
-							RHICmdList.SetRenderTargetsAndClear(CurrentRTInfo);
+							InRHICmdList.SetRenderTargetsAndClear(CurrentRTInfo);
+							InRHICmdList.ApplyCachedRenderTargets(InGraphicsPSOInit);
 						}
 						
 						// Start by setting up the stenciling states so that we can write representations of the clipping zones into the stencil buffer only.
 						{
 							FGraphicsPipelineStateInitializer WriteMaskPSOInit;
-							RHICmdList.ApplyCachedRenderTargets(WriteMaskPSOInit);
+							InRHICmdList.ApplyCachedRenderTargets(WriteMaskPSOInit);
 							WriteMaskPSOInit.BlendState = TStaticBlendStateWriteMask<CW_NONE>::GetRHI();
 							WriteMaskPSOInit.RasterizerState = TStaticRasterizerState<>::GetRHI();
 							WriteMaskPSOInit.DepthStencilState =
@@ -561,17 +559,17 @@ void FSlateRHIRenderingPolicy::DrawElements(
 							WriteMaskPSOInit.BoundShaderState.PixelShaderRHI = GETSAFERHISHADER_PIXEL(*PixelShader);
 							WriteMaskPSOInit.PrimitiveType = PT_TriangleStrip;
 
-							FLocalGraphicsPipelineState MaskingGraphicsReplacePSO = RHICmdList.BuildLocalGraphicsPipelineState(WriteMaskPSOInit);
-							RHICmdList.SetLocalGraphicsPipelineState(MaskingGraphicsReplacePSO);
+							FLocalGraphicsPipelineState MaskingGraphicsReplacePSO = InRHICmdList.BuildLocalGraphicsPipelineState(WriteMaskPSOInit);
+							InRHICmdList.SetLocalGraphicsPipelineState(MaskingGraphicsReplacePSO);
 
-							VertexShader->SetViewProjection(RHICmdList, ViewProjection);
-							VertexShader->SetVerticalAxisMultiplier(RHICmdList, bSwitchVerticalAxis ? -1.0f : 1.0f);
+							VertexShader->SetViewProjection(InRHICmdList, ViewProjection);
+							VertexShader->SetVerticalAxisMultiplier(InRHICmdList, bSwitchVerticalAxis ? -1.0f : 1.0f);
 
 							// Draw the first stencil using SO_Replace, so that we stomp any pixel with a MaskingID + 1.
 							{
 								const FSlateClippingZone& MaskQuad = StencilQuads[0];
 
-								RHICmdList.SetStencilRef(MaskingID + 1);
+								InRHICmdList.SetStencilRef(MaskingID + 1);
 
 								//TODO Slate If we ever decided to add masking with a texture, we could do that here.
 								FVector2D Vertices[4];
@@ -579,7 +577,7 @@ void FSlateRHIRenderingPolicy::DrawElements(
 								Vertices[1].Set(MaskQuad.TopRight.X, MaskQuad.TopRight.Y);
 								Vertices[2].Set(MaskQuad.BottomLeft.X, MaskQuad.BottomLeft.Y);
 								Vertices[3].Set(MaskQuad.BottomRight.X, MaskQuad.BottomRight.Y);
-								DrawPrimitiveUP(RHICmdList, PT_TriangleStrip, 2, Vertices, sizeof(Vertices[0]));
+								DrawPrimitiveUP(InRHICmdList, PT_TriangleStrip, 2, Vertices, sizeof(Vertices[0]));
 							}
 
 							// Now setup the pipeline to use SO_SaturatedIncrement, since we've established the initial
@@ -604,8 +602,8 @@ void FSlateRHIRenderingPolicy::DrawElements(
 									, /*StencilReadMask*/ 0xFF
 									, /*StencilWriteMask*/ 0xFF>::GetRHI();
 
-								FLocalGraphicsPipelineState MaskingGraphicsIncrementPSO = RHICmdList.BuildLocalGraphicsPipelineState(WriteMaskPSOInit);
-								RHICmdList.SetLocalGraphicsPipelineState(MaskingGraphicsIncrementPSO);
+								FLocalGraphicsPipelineState MaskingGraphicsIncrementPSO = InRHICmdList.BuildLocalGraphicsPipelineState(WriteMaskPSOInit);
+								InRHICmdList.SetLocalGraphicsPipelineState(MaskingGraphicsIncrementPSO);
 							}
 						}
 
@@ -622,7 +620,7 @@ void FSlateRHIRenderingPolicy::DrawElements(
 							Vertices[1].Set(MaskQuad.TopRight.X, MaskQuad.TopRight.Y);
 							Vertices[2].Set(MaskQuad.BottomLeft.X, MaskQuad.BottomLeft.Y);
 							Vertices[3].Set(MaskQuad.BottomRight.X, MaskQuad.BottomRight.Y);
-							DrawPrimitiveUP(RHICmdList, PT_TriangleStrip, 2, Vertices, sizeof(Vertices[0]));
+							DrawPrimitiveUP(InRHICmdList, PT_TriangleStrip, 2, Vertices, sizeof(Vertices[0]));
 						}
 
 						// Setup the stenciling state to be read only now, disable depth writes, and restore the color buffer
@@ -655,12 +653,10 @@ void FSlateRHIRenderingPolicy::DrawElements(
 							StencilRef = MaskingID;
 						}
 					}
-
-					RHICmdList.ApplyCachedRenderTargets(InGraphicsPSOInit);
 				}
 				else
 				{
-					RHICmdList.SetScissorRect(false, 0, 0, 0, 0);
+					InRHICmdList.SetScissorRect(false, 0, 0, 0, 0);
 
 					// Disable depth/stencil testing
 					InGraphicsPSOInit.DepthStencilState = TStaticDepthStencilState<false, CF_Always>::GetRHI();
@@ -685,7 +681,7 @@ void FSlateRHIRenderingPolicy::DrawElements(
 			FMatrix DynamicOffset = FTranslationMatrix::Make(FVector(RenderBatch.DynamicOffset.X, RenderBatch.DynamicOffset.Y, 0));
 			const FMatrix ViewProjection = DynamicOffset * Options.ViewProjectionMatrix;
 
-			UpdateScissorRect(GraphicsPSOInit, ViewProjection, false);
+			UpdateScissorRect(RHICmdList, GraphicsPSOInit, ViewProjection, false);
 
 			const uint32 PrimitiveCount = RenderBatch.DrawPrimitiveType == ESlateDrawPrimitive::LineList ? RenderBatch.NumIndices / 2 : RenderBatch.NumIndices / 3; 
 
@@ -1085,8 +1081,8 @@ void FSlateRHIRenderingPolicy::DrawElements(
 				RectParams.DestRect = FSlateRect(QuadPositionData.X, QuadPositionData.Y, QuadPositionData.Z, QuadPositionData.W);
 				RectParams.SourceTextureSize = BackBuffer.GetSizeXY();
 
-				RectParams.RestoreStateFunc = [&](FGraphicsPipelineStateInitializer& InGraphicsPSOInit) {
-					UpdateScissorRect(InGraphicsPSOInit, Options.ViewProjectionMatrix, true);
+				RectParams.RestoreStateFunc = [&](FRHICommandListImmediate& InRHICmdList, FGraphicsPipelineStateInitializer& InGraphicsPSOInit) {
+					UpdateScissorRect(InRHICmdList, InGraphicsPSOInit, Options.ViewProjectionMatrix, true);
 				};
 
 				RectParams.RestoreStateFuncPostPipelineState = [&]() {
diff --git a/Engine/Source/Runtime/UMG/Private/Components/WidgetComponent.cpp b/Engine/Source/Runtime/UMG/Private/Components/WidgetComponent.cpp
--- a/Engine/Source/Runtime/UMG/Private/Components/WidgetComponent.cpp
+++ b/Engine/Source/Runtime/UMG/Private/Components/WidgetComponent.cpp
@@ -265,9 +265,15 @@ private:
 
 
 /** Represents a billboard sprite to the scene manager. */
-class FWidget3DSceneProxy : public FPrimitiveSceneProxy
+class FWidget3DSceneProxy final : public FPrimitiveSceneProxy
 {
 public:
+	SIZE_T GetTypeHash() const override
+	{
+		static size_t UniquePointer;
+		return reinterpret_cast<size_t>(&UniquePointer);
+	}
+
 	/** Initialization constructor. */
 	FWidget3DSceneProxy( UWidgetComponent* InComponent, ISlate3DRenderer& InRenderer )
 		: FPrimitiveSceneProxy( InComponent )
@@ -636,9 +642,15 @@ FPrimitiveSceneProxy* UWidgetComponent::CreateSceneProxy()
 
 #if WITH_EDITOR
 	// make something so we can see this component in the editor
-	class FWidgetBoxProxy : public FPrimitiveSceneProxy
+	class FWidgetBoxProxy final : public FPrimitiveSceneProxy
 	{
 	public:
+		SIZE_T GetTypeHash() const override
+		{
+			static size_t UniquePointer;
+			return reinterpret_cast<size_t>(&UniquePointer);
+		}
+
 		FWidgetBoxProxy(const UWidgetComponent* InComponent)
 			: FPrimitiveSceneProxy(InComponent)
 			, BoxExtents(1.f, InComponent->GetDrawSize().X / 2.0f, InComponent->GetDrawSize().Y / 2.0f)
diff --git a/Engine/Source/Runtime/UMG/Private/Slate/SRetainerWidget.cpp b/Engine/Source/Runtime/UMG/Private/Slate/SRetainerWidget.cpp
--- a/Engine/Source/Runtime/UMG/Private/Slate/SRetainerWidget.cpp
+++ b/Engine/Source/Runtime/UMG/Private/Slate/SRetainerWidget.cpp
@@ -92,7 +92,9 @@ void SRetainerWidget::Construct(const FArguments& InArgs)
 
 	SurfaceBrush.SetResourceObject(RenderTarget);
 
-	Window = SNew(SVirtualWindow);
+	Window = SNew(SVirtualWindow)
+    .Visibility(EVisibility::SelfHitTestInvisible);  // deubanks: We don't want Retainer Widgets blocking hit testing for tooltips
+
 	Window->SetShouldResolveDeferred(false);
 	
 	UpdateWidgetRenderer();
diff --git a/Engine/Source/Runtime/VectorVM/Private/VectorVM.cpp b/Engine/Source/Runtime/VectorVM/Private/VectorVM.cpp
--- a/Engine/Source/Runtime/VectorVM/Private/VectorVM.cpp
+++ b/Engine/Source/Runtime/VectorVM/Private/VectorVM.cpp
@@ -5,9 +5,14 @@
 #include "UObject/Class.h"
 #include "UObject/Package.h"
 #include "VectorVMPrivate.h"
+#include "Stats.h"
 
 IMPLEMENT_MODULE(FDefaultModuleImpl, VectorVM);
 
+
+DECLARE_STATS_GROUP(TEXT("VectorVM"), STATGROUP_VectorVM, STATCAT_Advanced);
+DECLARE_CYCLE_STAT(TEXT("Execution"), STAT_VVMExec, STATGROUP_VectorVM);
+
 DEFINE_LOG_CATEGORY_STATIC(LogVectorVM, All, All);
 
 //#define VM_FORCEINLINE
@@ -27,6 +32,72 @@ DEFINE_LOG_CATEGORY_STATIC(LogVectorVM, All, All);
 #define SRCOP_CCR (OP2_CONST | OP1_CONST | OP_REGISTER)
 #define SRCOP_CCC (OP2_CONST | OP1_CONST | OP0_CONST)
 
+//////////////////////////////////////////////////////////////////////////
+
+FVectorVMContext::FVectorVMContext()
+	: Code(nullptr)
+	, ConstantTable(nullptr)
+	, DataSetIndexTable(nullptr)
+	, DataSetOffsetTable(nullptr)
+	, NumSecondaryDataSets(0)
+	, ExternalFunctionTable(nullptr)
+	, UserPtrTable(nullptr)
+	, NumInstances(0)
+	, StartInstance(0)
+#if STATS
+	, StatScopes(nullptr)
+#endif
+{
+	uint32 TempRegisterSize = Align((VectorVM::InstancesPerChunk) * VectorVM::MaxInstanceSizeBytes, VECTOR_WIDTH_BYTES) + VECTOR_WIDTH_BYTES;
+	TempRegTable.SetNumUninitialized(TempRegisterSize * VectorVM::NumTempRegisters);
+	// Map temporary registers.
+	for (int32 i = 0; i < VectorVM::NumTempRegisters; ++i)
+	{
+		RegisterTable[i] = TempRegTable.GetData() + TempRegisterSize * i;
+	}
+}
+
+void FVectorVMContext::PrepareForExec(
+	uint8*RESTRICT*RESTRICT InputRegisters,
+	uint8*RESTRICT*RESTRICT OutputRegisters,
+	int32 NumInputRegisters,
+	int32 NumOutputRegisters,
+	const uint8* InConstantTable,
+	int32 *InDataSetIndexTable,
+	int32 *InDataSetOffsetTable,
+	int32 InNumSecondaryDatasets,
+	FVMExternalFunction* InExternalFunctionTable,
+	void** InUserPtrTable
+#if STATS
+	, const TArray<TStatId>* InStatScopes
+#endif
+)
+{
+	ConstantTable = InConstantTable;
+	DataSetIndexTable = InDataSetIndexTable;
+	DataSetOffsetTable = InDataSetOffsetTable;
+	NumSecondaryDataSets = InNumSecondaryDatasets;
+	ExternalFunctionTable = InExternalFunctionTable;
+	UserPtrTable = InUserPtrTable;
+#if STATS
+	check(InStatScopes);
+	StatScopes = InStatScopes;
+	StatCounterStack.Reserve(StatScopes->Num());
+#endif
+
+	//Map IO Registers
+	for (int32 i = 0; i < NumInputRegisters; ++i)
+	{
+		RegisterTable[VectorVM::NumTempRegisters + i] = InputRegisters[i];
+	}
+	for (int32 i = 0; i < NumOutputRegisters; ++i)
+	{
+		RegisterTable[VectorVM::NumTempRegisters + VectorVM::MaxInputRegisters + i] = OutputRegisters[i];
+	}
+}
+
+//////////////////////////////////////////////////////////////////////////
+
 uint8 VectorVM::CreateSrcOperandMask(EVectorVMOperandLocation Type0, EVectorVMOperandLocation Type1, EVectorVMOperandLocation Type2)
 {
 	return	(Type0 == EVectorVMOperandLocation::Constant ? OP0_CONST : OP_REGISTER) |
@@ -514,10 +585,10 @@ struct FVectorKernelEnterStatScope
 {
 	static VM_FORCEINLINE void Exec(FVectorVMContext& Context)
 	{
+		FConstantHandler<int32> ScopeIdx(Context);
 #if STATS
-		FConstantHandler<int32> ScopeIdx(Context);		
 		int32 CounterIdx = Context.StatCounterStack.AddDefaulted(1);
-		Context.StatCounterStack[CounterIdx].Start(Context.StatScopes[ScopeIdx.Get()]);
+		Context.StatCounterStack[CounterIdx].Start((*Context.StatScopes)[ScopeIdx.Get()]);
 #endif
 	}
 };
@@ -528,7 +599,7 @@ struct FVectorKernelExitStatScope
 	{
 #if STATS
 		Context.StatCounterStack.Last().Stop();
-		Context.StatCounterStack.Pop();
+		Context.StatCounterStack.Pop(false);
 #endif
 	}
 };
@@ -1073,7 +1144,7 @@ struct FScalarIntKernelRandom : public TUnaryScalarIntKernel<FScalarIntKernelRan
 	{
 		const float rm = RAND_MAX;
 		//EEK!. Improve this. Implement GPU style seeded rand instead of this.
-		*Dst = FMath::Rand() % (Src0 + 1);
+		*Dst = static_cast<int32>((FMath::Rand() / rm) * Src0);
 	}
 };
 
@@ -1364,39 +1435,14 @@ void VectorVM::Exec(
 #endif
 	)
 {
-	uint32 TempRegisterSize = Align((InstancesPerChunk) * MaxInstanceSizeBytes, VECTOR_WIDTH_BYTES) + VECTOR_WIDTH_BYTES;
-	//TODO: Refactor this so VMs are a persistent object with growing buffers. Once spun up, there are no allocs.
-	//Can be pooled and used for threading and branching.
-	TArray<uint8,TAlignedHeapAllocator<VECTOR_WIDTH_BYTES>> TempRegTable;
-	TempRegTable.SetNumUninitialized(TempRegisterSize * NumTempRegisters);
-	uint8* RegisterTable[MaxRegisters] = {0};
-
-
-	// Map temporary registers.
-	for (int32 i = 0; i < NumTempRegisters; ++i)
-	{
-		RegisterTable[i] = TempRegTable.GetData() + TempRegisterSize * i;
-	}
-
-	// Map input and output registers.
-	//Input and output registers are indexed absolutely directly in their kernels.
-	//TODO: No need for these to be in the same table now.
-	//TODO: Also no need for the i/o size table as the ops will deal with that now.
- 	for (int32 i = 0; i < NumInputRegisters; ++i)
- 	{
- 		RegisterTable[NumTempRegisters + i] = InputRegisters[i];
- 	}
- 	for (int32 i = 0; i < NumOutputRegisters; ++i)
- 	{
- 		RegisterTable[NumTempRegisters + MaxInputRegisters + i] = OutputRegisters[i];
- 	}
+	SCOPE_CYCLE_COUNTER(STAT_VVMExec);
 
 	// table of index counters, one for each data set
-	TArray<int32> DataSetIndexTable;
+	TArray<int32, TInlineAllocator<16>> DataSetIndexTable;
+	TArray<int32, TInlineAllocator<16>> DataSetOffsetTable;
 
 	// map secondary data sets and fill in the offset table into the register table
 	//
-	TArray<int32> DataSetOffsetTable;
 	for (int32 Idx = 0; Idx < DataSetMetaTable.Num(); Idx++)
 	{
 		uint32 DataSetOffset = DataSetMetaTable[Idx].NumVariables;
@@ -1404,7 +1450,13 @@ void VectorVM::Exec(
 		DataSetIndexTable.Add(DataSetMetaTable[Idx].DataSetAccessIndex);	// prime counter index table with the data set offset; will be incremented with every write for each instance
 	}
 
-
+	FVectorVMContext& Context = FVectorVMContext::Get();
+	Context.PrepareForExec(InputRegisters, OutputRegisters, NumInputRegisters, NumOutputRegisters, ConstantTable, DataSetIndexTable.GetData(), DataSetOffsetTable.GetData(), DataSetOffsetTable.Num(),
+		ExternalFunctionTable, UserPtrTable
+#if STATS
+		, &StatScopes
+#endif
+	);
 
 	// Process one chunk at a time.
 	int32 InstancesLeft = NumInstances;
@@ -1412,13 +1464,8 @@ void VectorVM::Exec(
 	while (InstancesLeft > 0)
 	{
 		// Setup execution context.
-		FVectorVMContext Context(Code, RegisterTable, ConstantTable, DataSetIndexTable.GetData(), DataSetOffsetTable.GetData(), 
-			ExternalFunctionTable, UserPtrTable, FMath::Min(InstancesLeft, (int32)InstancesPerChunk), InstancesPerChunk * ChunkIdx
-#if STATS
-			, StatScopes
-#endif
-		);
-		Context.NumSecondaryDataSets = DataSetOffsetTable.Num();
+		Context.PrepareForChunk(Code, FMath::Min(InstancesLeft, (int32)InstancesPerChunk), InstancesPerChunk * ChunkIdx);
+
 		EVectorVMOp Op = EVectorVMOp::done;
 
 		// Execute VM on all vectors in this chunk.
@@ -1524,7 +1571,7 @@ void VectorVM::Exec(
 
 			// Opcode not recognized / implemented.
 			default:
-				UE_LOG(LogVectorVM, Error, TEXT("Unknown op code 0x%02x"), (uint32)Op);
+				UE_LOG(LogVectorVM, Fatal, TEXT("Unknown op code 0x%02x"), (uint32)Op);
 				return;//BAIL
 			}
 		} while (Op != EVectorVMOp::done);
diff --git a/Engine/Source/Runtime/VectorVM/Public/VectorVM.h b/Engine/Source/Runtime/VectorVM/Public/VectorVM.h
--- a/Engine/Source/Runtime/VectorVM/Public/VectorVM.h
+++ b/Engine/Source/Runtime/VectorVM/Public/VectorVM.h
@@ -163,8 +163,8 @@ namespace VectorVM
 	/** Constants. */
 	enum
 	{
-		NumTempRegisters = 100,
-		MaxInputRegisters = 100,
+		NumTempRegisters = 400,
+		MaxInputRegisters = 400,
 		MaxOutputRegisters = MaxInputRegisters,
 		MaxConstants = 256,
 		FirstTempRegister = 0,
@@ -212,12 +212,10 @@ namespace VectorVM
   /**
   * Context information passed around during VM execution.
   */
-struct FVectorVMContext
+struct FVectorVMContext : TThreadSingleton<FVectorVMContext>
 {
 	/** Pointer to the next element in the byte code. */
 	uint8 const* RESTRICT Code;
-	/** Pointer to the table of vector register arrays. */
-	uint8* RESTRICT * RESTRICT RegisterTable;
 	/** Pointer to the constant table. */
 	uint8 const* RESTRICT ConstantTable;
 	/** Pointer to the data set index counter table */
@@ -234,41 +232,36 @@ struct FVectorVMContext
 	int32 StartInstance;
 
 #if STATS
-	TArray<FCycleCounter> StatCounterStack;
-	const TArray<TStatId>& StatScopes;
+	TArray<FCycleCounter, TInlineAllocator<64>> StatCounterStack;
+	const TArray<TStatId>* StatScopes;
 #endif
 
-	/** Initialization constructor. */
-	FVectorVMContext(
-		const uint8* InCode,
-		uint8** InRegisterTable,
+	TArray<uint8, TAlignedHeapAllocator<VECTOR_WIDTH_BYTES>> TempRegTable;
+	uint8 *RESTRICT RegisterTable[VectorVM::MaxRegisters];
+
+	FVectorVMContext();
+
+	void PrepareForExec(
+		uint8*RESTRICT*RESTRICT InputRegisters,
+		uint8*RESTRICT*RESTRICT OutputRegisters,
+		int32 NumInputRegisters,
+		int32 NumOutputRegisters,
 		const uint8* InConstantTable,
 		int32 *InDataSetIndexTable,
 		int32 *InDataSetOffsetTable,
+		int32 InNumSecondaryDatasets,
 		FVMExternalFunction* InExternalFunctionTable,
-		void** InUserPtrTable,
-		int32 InNumInstances,
-		int32 InStartInstance
+		void** InUserPtrTable
 #if STATS
-		, const TArray<TStatId>& InStatScopes
-#endif
-	)
-		: Code(InCode)
-		, RegisterTable(InRegisterTable)
-		, ConstantTable(InConstantTable)
-		, DataSetIndexTable(InDataSetIndexTable)
-		, DataSetOffsetTable(InDataSetOffsetTable)
-		, ExternalFunctionTable(InExternalFunctionTable)
-		, UserPtrTable(InUserPtrTable)
-		, NumInstances(InNumInstances)
-		, StartInstance(InStartInstance)
-#if STATS
-		, StatScopes(InStatScopes)
+		, const TArray<TStatId>* InStatScopes
 #endif
+	);
+
+	void PrepareForChunk(const uint8* InCode, int32 InNumInstances, int32 InStartInstance)
 	{
-#if STATS
-		StatCounterStack.Reserve(StatScopes.Num());
-#endif
+		Code = InCode;
+		NumInstances = InNumInstances;
+		StartInstance = InStartInstance;
 	}
 };
 
