Copy of CLs 3838626 and 3838748 from 4.19

Changed FMacMenu to store MultiBox and MenuEntryBox as weak pointers instead of shared pointers. This way we avoid a situation where FMacMenu would try to release them on the main Cocoa thread or where we'd try to execute a menu action for menu items that Slate considers released.

diff --git a/Engine/Source/Runtime/Slate/Private/Framework/MultiBox/Mac/MacMenu.cpp b/Engine/Source/Runtime/Slate/Private/Framework/MultiBox/Mac/MacMenu.cpp
--- a/Engine/Source/Runtime/Slate/Private/Framework/MultiBox/Mac/MacMenu.cpp
+++ b/Engine/Source/Runtime/Slate/Private/Framework/MultiBox/Mac/MacMenu.cpp
@@ -319,34 +319,38 @@ void FSlateMacMenu::UpdateCachedState()
 			TSharedRef<SWidget> Widget = SNullWidget::NullWidget;
 			if (!Menu.MultiBox.IsValid())
 			{
-				if (Menu.MenuEntryBlock->MenuBuilder.IsBound())
+				TSharedPtr<const FMenuEntryBlock> MenuEntryBlock = Menu.MenuEntryBlock.Pin();
+				if (MenuEntryBlock.IsValid())
 				{
-					Widget = Menu.MenuEntryBlock->MenuBuilder.Execute();
-				}
-				else
-				{
-					const bool bShouldCloseWindowAfterMenuSelection = true;
-					FMenuBuilder MenuBuilder(bShouldCloseWindowAfterMenuSelection, Menu.MenuEntryBlock->GetActionList(), Menu.MenuEntryBlock->Extender);
+					if (MenuEntryBlock->MenuBuilder.IsBound())
 					{
-						// Have the menu fill its contents
-						Menu.MenuEntryBlock->EntryBuilder.ExecuteIfBound(MenuBuilder);
+						Widget = MenuEntryBlock->MenuBuilder.Execute();
+					}
+					else
+					{
+						const bool bShouldCloseWindowAfterMenuSelection = true;
+						FMenuBuilder MenuBuilder(bShouldCloseWindowAfterMenuSelection, MenuEntryBlock->GetActionList(), MenuEntryBlock->Extender);
+						{
+							// Have the menu fill its contents
+							MenuEntryBlock->EntryBuilder.ExecuteIfBound(MenuBuilder);
+						}
+						Widget = MenuBuilder.MakeWidget();
 					}
-					Widget = MenuBuilder.MakeWidget();
-				}
 
-				if (Widget->GetType() == FName(TEXT("SMultiBoxWidget")))
-				{
-					Menu.MultiBox = TSharedPtr<const FMultiBox>(StaticCastSharedRef<SMultiBoxWidget>(Widget)->GetMultiBox());
-				}
-				else
-				{
-					UE_LOG(LogMac, Warning, TEXT("Unsupported type of menu widget in FSlateMacMenu::UpdateCachedState(): %s"), *Widget->GetType().ToString());
+					if (Widget->GetType() == FName(TEXT("SMultiBoxWidget")))
+					{
+						Menu.MultiBox = TSharedPtr<const FMultiBox>(StaticCastSharedRef<SMultiBoxWidget>(Widget)->GetMultiBox());
+					}
+					else
+					{
+						UE_LOG(LogMac, Warning, TEXT("Unsupported type of menu widget in FSlateMacMenu::UpdateCachedState(): %s"), *Widget->GetType().ToString());
+					}
 				}
 			}
 
 			if (Menu.MultiBox.IsValid())
 			{
-				const TArray<TSharedRef<const FMultiBlock>>& MenuBlocks = Menu.MultiBox->GetBlocks();
+				const TArray<TSharedRef<const FMultiBlock>>& MenuBlocks = Menu.MultiBox.Pin()->GetBlocks();
 				for (int32 Index = MenuState->Num(); MenuBlocks.Num() > MenuState->Num(); Index++)
 				{
 					MenuState->Add(FMacMenuItemState());
diff --git a/Engine/Source/Runtime/Slate/Private/Framework/MultiBox/Mac/MacMenu.h b/Engine/Source/Runtime/Slate/Private/Framework/MultiBox/Mac/MacMenu.h
--- a/Engine/Source/Runtime/Slate/Private/Framework/MultiBox/Mac/MacMenu.h
+++ b/Engine/Source/Runtime/Slate/Private/Framework/MultiBox/Mac/MacMenu.h
@@ -7,8 +7,8 @@
 #include "CocoaMenu.h"
 
 @interface FMacMenu : FCocoaMenu <NSMenuDelegate>
-@property (assign) TSharedPtr<const FMenuEntryBlock> MenuEntryBlock;
-@property (assign) TSharedPtr<const FMultiBox> MultiBox;
+@property (assign) TWeakPtr<const FMenuEntryBlock> MenuEntryBlock;
+@property (assign) TWeakPtr<const FMultiBox> MultiBox;
 @end
 
 class SLATE_API FSlateMacMenu
