Branch snapshot for CL 3327560

diff --git a/Engine/Source/Runtime/MovieScene/Private/Evaluation/MovieSceneEvaluationTemplate.cpp b/Engine/Source/Runtime/MovieScene/Private/Evaluation/MovieSceneEvaluationTemplate.cpp
--- a/Engine/Source/Runtime/MovieScene/Private/Evaluation/MovieSceneEvaluationTemplate.cpp
+++ b/Engine/Source/Runtime/MovieScene/Private/Evaluation/MovieSceneEvaluationTemplate.cpp
@@ -26,15 +26,18 @@ void FCachedMovieSceneEvaluationTemplate::Regenerate(const FMovieSceneTrackCompi
 {
 	if (IsOutOfDate(Params))
 	{
-		ForceRegenerate(Params);
+		RegenerateImpl(Params);
 	}
 }
 
 void FCachedMovieSceneEvaluationTemplate::ForceRegenerate(const FMovieSceneTrackCompilationParams& Params)
 {
-	// Reassign the template back to default
-	static_cast<FMovieSceneEvaluationTemplate&>(*this) = FMovieSceneEvaluationTemplate();
+	ResetGeneratedData();
+	RegenerateImpl(Params);
+}
 
+void FCachedMovieSceneEvaluationTemplate::RegenerateImpl(const FMovieSceneTrackCompilationParams& Params)
+{
 	CachedSignatures.Reset();
 	CachedCompilationParams = Params;
 
@@ -165,6 +168,18 @@ void FMovieSceneGenerationLedger::AddTrack(const FGuid& InSignature, FMovieScene
 	++TrackReferenceCounts.FindOrAdd(Identifier);
 }
 
+void FMovieSceneEvaluationTemplate::ResetGeneratedData()
+{
+	Ledger.TrackSignatureToTrackIdentifier.Reset();
+	Ledger.TrackReferenceCounts.Reset();
+
+	Tracks.Reset();
+	StaleTracks.Reset();
+	EvaluationField = FMovieSceneEvaluationField();
+	Hierarchy = FMovieSceneSequenceHierarchy();
+	bHasLegacyTrackInstances = false;
+}
+
 FMovieSceneTrackIdentifier FMovieSceneEvaluationTemplate::AddTrack(const FGuid& InSignature, FMovieSceneEvaluationTrack&& InTrack)
 {
 	FMovieSceneTrackIdentifier NewIdentifier = ++Ledger.LastTrackIdentifier;
diff --git a/Engine/Source/Runtime/MovieScene/Private/Evaluation/MovieSceneEvaluationTemplateInstance.cpp b/Engine/Source/Runtime/MovieScene/Private/Evaluation/MovieSceneEvaluationTemplateInstance.cpp
--- a/Engine/Source/Runtime/MovieScene/Private/Evaluation/MovieSceneEvaluationTemplateInstance.cpp
+++ b/Engine/Source/Runtime/MovieScene/Private/Evaluation/MovieSceneEvaluationTemplateInstance.cpp
@@ -89,6 +89,14 @@ void FMovieSceneRootEvaluationTemplateInstance::Initialize(UMovieSceneSequence&
 {
 	Reset();
 
+	if (RootSequence.Get() != &InRootSequence)
+	{
+		// Always ensure that there is no persistent data when initializing a new sequence
+		// to ensure we don't collide with the previous sequence's entity keys
+		Player.State.PersistentEntityData.Reset();
+		Player.State.PersistentSharedData.Reset();
+	}
+
 	const bool bAddEvents = TemplateStore->AreTemplatesVolatile();
 	
 	RootSequence = &InRootSequence;
@@ -365,7 +373,7 @@ void FMovieSceneRootEvaluationTemplateInstance::CallSetupTearDown(IMovieScenePla
 		else
 		{
 			PersistentDataProxy.SetSectionKey(Key);
-			if (Track)
+			if (Track && Track->HasChildTemplate(Key.SectionIdentifier))
 			{
 				Track->GetChildTemplate(Key.SectionIdentifier).OnEndEvaluation(PersistentDataProxy, Player);
 			}
diff --git a/Engine/Source/Runtime/MovieScene/Private/Evaluation/MovieSceneSequenceTemplateStore.cpp b/Engine/Source/Runtime/MovieScene/Private/Evaluation/MovieSceneSequenceTemplateStore.cpp
--- a/Engine/Source/Runtime/MovieScene/Private/Evaluation/MovieSceneSequenceTemplateStore.cpp
+++ b/Engine/Source/Runtime/MovieScene/Private/Evaluation/MovieSceneSequenceTemplateStore.cpp
@@ -6,9 +6,9 @@
 FMovieSceneEvaluationTemplate& FMovieSceneSequenceTemplateStore::GetCompiledTemplate(UMovieSceneSequence& Sequence)
 {
 #if WITH_EDITORONLY_DATA
-	if (AreTemplatesVolatile())
+	if (AreTemplatesVolatile() && Sequence.EvaluationTemplate.IsOutOfDate(Sequence.TemplateParameters))
 	{
-		Sequence.EvaluationTemplate.Regenerate(Sequence.TemplateParameters);
+		Sequence.EvaluationTemplate.ForceRegenerate(Sequence.TemplateParameters);
 	}
 #endif
 	return Sequence.EvaluationTemplate;
diff --git a/Engine/Source/Runtime/MovieScene/Public/Evaluation/MovieSceneEvaluationTemplate.h b/Engine/Source/Runtime/MovieScene/Public/Evaluation/MovieSceneEvaluationTemplate.h
--- a/Engine/Source/Runtime/MovieScene/Public/Evaluation/MovieSceneEvaluationTemplate.h
+++ b/Engine/Source/Runtime/MovieScene/Public/Evaluation/MovieSceneEvaluationTemplate.h
@@ -95,6 +95,11 @@ struct FMovieSceneEvaluationTemplate
 
 public:
 
+	/**
+	 * Reset this evaluation template, whilst ensuring that track identifiers do not get reused
+	 */
+	MOVIESCENE_API void ResetGeneratedData();
+
 	/**
 	 * Attempt to locate a track with the specified identifier
 	 */
@@ -236,6 +241,8 @@ struct FCachedMovieSceneEvaluationTemplate : public FMovieSceneEvaluationTemplat
 
 private:
 
+	void RegenerateImpl(const FMovieSceneTrackCompilationParams& NewParams);
+
 	/** Transient data */
 	TWeakObjectPtr<UMovieSceneSequence> SourceSequence;
 	FMovieSceneSequenceTemplateStore* Origin;
diff --git a/Engine/Source/Runtime/MovieSceneTracks/Private/Evaluation/MovieSceneAudioTemplate.cpp b/Engine/Source/Runtime/MovieSceneTracks/Private/Evaluation/MovieSceneAudioTemplate.cpp
--- a/Engine/Source/Runtime/MovieSceneTracks/Private/Evaluation/MovieSceneAudioTemplate.cpp
+++ b/Engine/Source/Runtime/MovieSceneTracks/Private/Evaluation/MovieSceneAudioTemplate.cpp
@@ -268,7 +268,7 @@ void FMovieSceneAudioSectionTemplateData::EnsureAudioIsPlaying(UAudioComponent&
 		AudioComponent.Stop();
 		AudioComponent.SetSound(Sound);
 		AudioComponent.bIsUISound = true;
-		AudioComponent.Play(AudioStartOffset < 0 ? 0 : AudioStartOffset);
+		AudioComponent.Play((Context.GetTime() - AudioRange.GetLowerBoundValue()) + FMath::Max(AudioStartOffset, 0.f));
 
 		if (Context.GetStatus() == EMovieScenePlayerStatus::Scrubbing)
 		{
diff --git a/Engine/Source/Runtime/Renderer/Private/IndirectLightingCache.cpp b/Engine/Source/Runtime/Renderer/Private/IndirectLightingCache.cpp
--- a/Engine/Source/Runtime/Renderer/Private/IndirectLightingCache.cpp
+++ b/Engine/Source/Runtime/Renderer/Private/IndirectLightingCache.cpp
@@ -326,7 +326,7 @@ void FIndirectLightingCache::CalculateBlockScaleAndAdd(FIntVector InTexelMin, in
 
 FIndirectLightingCacheAllocation* FIndirectLightingCache::AllocatePrimitive(const FPrimitiveSceneInfo* PrimitiveSceneInfo, bool bUnbuiltPreview)
 {
-	const bool bPointSample = PrimitiveSceneInfo->Proxy->GetIndirectLightingCacheQuality() == ILCQ_Point || bUnbuiltPreview;
+	const bool bPointSample = !bUnbuiltPreview && PrimitiveSceneInfo->Proxy->GetIndirectLightingCacheQuality() == ILCQ_Point;
 	const int32 BlockSize = bPointSample ? 1 : GLightingCacheMovableObjectAllocationSize;
 	return PrimitiveAllocations.Add(PrimitiveSceneInfo->PrimitiveComponentId, CreateAllocation(BlockSize, PrimitiveSceneInfo->Proxy->GetBounds(), bPointSample, bUnbuiltPreview));
 }
@@ -686,7 +686,7 @@ void FIndirectLightingCache::UpdateCachePrimitive(
 		{
 			FIndirectLightingCacheAllocation* OriginalAllocation = PrimitiveAllocation;
 			const bool bUnbuiltPreview = bAllowUnbuiltPreview && !bIsMovable;
-			const bool bPointSample = PrimitiveSceneProxy->GetIndirectLightingCacheQuality() == ILCQ_Point || bUnbuiltPreview || !bOpaqueRelevance;
+			const bool bPointSample = !bUnbuiltPreview && ( PrimitiveSceneProxy->GetIndirectLightingCacheQuality() == ILCQ_Point || !bOpaqueRelevance );
 			const int32 BlockSize = bPointSample ? 1 : GLightingCacheMovableObjectAllocationSize;
 
 			// Light with the cumulative bounds of the entire attachment group
diff --git a/Engine/Source/Runtime/Renderer/Private/VelocityRendering.cpp b/Engine/Source/Runtime/Renderer/Private/VelocityRendering.cpp
--- a/Engine/Source/Runtime/Renderer/Private/VelocityRendering.cpp
+++ b/Engine/Source/Runtime/Renderer/Private/VelocityRendering.cpp
@@ -39,6 +39,11 @@ static TAutoConsoleVariable<int32> CVarRHICmdVelocityPassDeferredContexts(
 	1,
 	TEXT("True to use deferred contexts to parallelize velocity pass command list execution."));
 
+RENDERER_API TAutoConsoleVariable<int32> CVarAllowMotionBlurInVR(
+	TEXT("r.AllowMotionBlurInVR"),
+	0,
+	TEXT("For projects with motion blur enabled, this allows motion blur to be enabled even while in VR."));
+
 DECLARE_FLOAT_COUNTER_STAT(TEXT("Render Velocities"), Stat_GPU_RenderVelocities, STATGROUP_GPU);
 
 bool IsParallelVelocity()
@@ -607,7 +612,7 @@ bool IsMotionBlurEnabled(const FViewInfo& View)
 		&& View.Family->bRealtimeUpdate
 		&& MotionBlurQuality > 0
 		&& !View.bIsSceneCapture
-		&& !(View.Family->Views.Num() > 1);
+		&& (CVarAllowMotionBlurInVR->GetInt() != 0 || !(View.Family->Views.Num() > 1));
 }
 
 void FDeferredShadingSceneRenderer::RenderDynamicVelocitiesMeshElementsInner(FRHICommandList& RHICmdList, const FViewInfo& View, const FDrawingPolicyRenderState& DrawRenderState, int32 FirstIndex, int32 LastIndex)
diff --git a/Engine/Source/Runtime/Slate/Private/Framework/Application/SlateApplication.cpp b/Engine/Source/Runtime/Slate/Private/Framework/Application/SlateApplication.cpp
--- a/Engine/Source/Runtime/Slate/Private/Framework/Application/SlateApplication.cpp
+++ b/Engine/Source/Runtime/Slate/Private/Framework/Application/SlateApplication.cpp
@@ -3042,7 +3042,7 @@ void FSlateApplication::ProcessReply( const FWidgetPath& CurrentEventPath, const
 	// are all operations that we shouldn't do if our application isn't Active (The OS ignores half of it, and we'd be in a half state)
 	// We do allow the release of capture and lock when deactivated, this is innocuous of some platforms but required on others when 
 	// the Application deactivated before the window. (Mac is an example of this)
-	if (bAppIsActive)
+	if (bAppIsActive || bIsVirtualInteraction)
 	{
 		TSharedPtr<SWidget> RequestedMouseCaptor = TheReply.GetMouseCaptor();
 		// Do not capture the mouse if we are also starting a drag and drop.
@@ -4856,6 +4856,12 @@ bool FSlateApplication::ProcessMouseButtonDownEvent( const TSharedPtr< FGenericW
 	}
 	PressedMouseButtons.Add( MouseEvent.GetEffectingButton() );
 
+	// Input preprocessor gets first chance at the input
+	if( InputPreProcessor.IsValid() && InputPreProcessor->HandleMouseButtonDownEvent( *this, MouseEvent ) )
+	{
+		return true;
+	}
+
 	bool bInGame = false;
 
 	// Only process mouse down messages if we are not drag/dropping
@@ -5029,6 +5035,15 @@ FReply FSlateApplication::RoutePointerUpEvent(FWidgetPath& WidgetsUnderPointer,
 	TSharedPtr<SWindow> TopLevelWindow;
 #endif
 
+	if ( DragDetector.DetectDragForWidget.IsValid() )
+	{
+		if ( PointerEvent.GetEffectingButton() == DragDetector.DetectDragButton && PointerEvent.GetPointerIndex() == DragDetector.DetectDragPointerIndex )
+		{
+			// The user has released the button (or finger) that was supposed to start the drag; stop detecting it.
+			DragDetector = FDragDetector();
+		}
+	}
+
 	if (MouseCaptor.HasCaptureForPointerIndex(PointerEvent.GetUserIndex(), PointerEvent.GetPointerIndex()))
 	{
 		//FWidgetPath MouseCaptorPath = MouseCaptor.ToWidgetPath(PointerEvent.GetPointerIndex());
@@ -5135,7 +5150,14 @@ bool FSlateApplication::RoutePointerMoveEvent(const FWidgetPath& WidgetsUnderPoi
 
 	// User asked us to detect a drag.
 	bool bDragDetected = false;
-	if( DragDetector.DetectDragForWidget.IsValid() && !bIsSynthetic )
+	bool bShouldStartDetectingDrag = true;
+
+	//@TODO VREDITOR - Remove and move to interaction component
+	if (OnDragDropCheckOverride.IsBound())
+	{
+		bShouldStartDetectingDrag = OnDragDropCheckOverride.Execute();
+	}
+	if( DragDetector.DetectDragForWidget.IsValid() && !bIsSynthetic && bShouldStartDetectingDrag)
 	{	
 		const FVector2D DragDelta = (DragDetector.DetectDragStartLocation - PointerEvent.GetScreenSpacePosition());
 		bDragDetected = ( DragDelta.SizeSquared() > FMath::Square(FSlateApplication::Get().GetDragTriggerDistance()) );
@@ -5341,6 +5363,13 @@ bool FSlateApplication::RoutePointerMoveEvent(const FWidgetPath& WidgetsUnderPoi
 	if (IsDragDropping())
 	{
 		FDragDropEvent DragDropEvent( PointerEvent, DragDropContent );
+		//@TODO VREDITOR - Remove and move to interaction component
+		if (OnDragDropCheckOverride.IsBound() && DragDropEvent.GetOperation().IsValid())
+		{
+			DragDropEvent.GetOperation()->SetDecoratorVisibility(false);
+			DragDropEvent.GetOperation()->SetCursorOverride(EMouseCursor::None);
+			DragDropContent->SetCursorOverride(EMouseCursor::None);
+		}
 		FScopedSwitchWorldHack SwitchWorld( WidgetsUnderPointer );
 		DragDropContent->OnDragged( DragDropEvent );
 
@@ -5490,13 +5519,10 @@ bool FSlateApplication::ProcessMouseButtonUpEvent( FPointerEvent& MouseEvent )
 	LastUserInteractionTimeForThrottling = LastUserInteractionTime;
 	PressedMouseButtons.Remove( MouseEvent.GetEffectingButton() );
 
-	if ( DragDetector.DetectDragForWidget.IsValid() )
+	// Input preprocessor gets first chance at the input
+	if( InputPreProcessor.IsValid() && InputPreProcessor->HandleMouseButtonUpEvent( *this, MouseEvent ) )
 	{
-		if ( MouseEvent.GetEffectingButton() == DragDetector.DetectDragButton && MouseEvent.GetPointerIndex() == DragDetector.DetectDragPointerIndex )
-		{
-			// The user has released the button (or finger) that was supposed to start the drag; stop detecting it.
-			DragDetector = FDragDetector();
-		}
+		return true;
 	}
 
 	// An empty widget path is passed in.  As an optimization, one will be generated only if a captured mouse event isn't routed
diff --git a/Engine/Source/Runtime/Slate/Private/Framework/Docking/TabManager.cpp b/Engine/Source/Runtime/Slate/Private/Framework/Docking/TabManager.cpp
--- a/Engine/Source/Runtime/Slate/Private/Framework/Docking/TabManager.cpp
+++ b/Engine/Source/Runtime/Slate/Private/Framework/Docking/TabManager.cpp
@@ -1147,11 +1147,10 @@ TSharedRef<SDockTab> FTabManager::InvokeTab_Internal( const FTabId& TabId )
 			->AddTab( TabId, ETabState::OpenedTab )
 		);
 
-		TSharedRef<SDockingArea> DockingArea = RestoreArea( NewAreaForTab, GetPrivateApi().GetParentWindow() );
-				
+		TSharedRef<SDockingArea> DockingArea = RestoreArea(NewAreaForTab, GetPrivateApi().GetParentWindow());
 		const TSharedPtr<SDockTab> NewlyOpenedTab = DockingArea->GetAllChildTabs()[0];
 		check(NewlyOpenedTab.IsValid());
-		return NewlyOpenedTab.ToSharedRef();
+		return NewlyOpenedTab.ToSharedRef();	
 	}
 }
 
@@ -1907,7 +1906,10 @@ void FGlobalTabmanager::DrawAttentionToTabManager( const TSharedRef<FTabManager>
 		// #HACK VREDITOR
 		if ( ProxyTabManager.IsValid() )
 		{
-			ProxyTabManager->DrawAttention(Tab.ToSharedRef());
+			if( ProxyTabManager->IsTabSupported( Tab->GetLayoutIdentifier() ) )
+			{
+				ProxyTabManager->DrawAttention(Tab.ToSharedRef());
+			}
 		}
 	}
 }
@@ -2115,7 +2117,7 @@ void FGlobalTabmanager::UpdateStats()
 
 void FGlobalTabmanager::OpenUnmanagedTab(FName PlaceholderId, const FSearchPreference& SearchPreference, const TSharedRef<SDockTab>& UnmanagedTab)
 {
-	if ( ProxyTabManager.IsValid() )
+	if ( ProxyTabManager.IsValid() && ProxyTabManager->IsTabSupported( UnmanagedTab->GetLayoutIdentifier() ) )
 	{
 		ProxyTabManager->OpenUnmanagedTab(PlaceholderId, SearchPreference, UnmanagedTab);
 	}
@@ -2130,29 +2132,43 @@ void FGlobalTabmanager::SetProxyTabManager(TSharedPtr<FProxyTabmanager> InProxyT
 	ProxyTabManager = InProxyTabManager;
 }
 
-void FProxyTabmanager::OpenUnmanagedTab(FName PlaceholderId, const FSearchPreference& SearchPreference, const TSharedRef<SDockTab>& UnmanagedTab)
+bool FProxyTabmanager::IsTabSupported( const FTabId TabId ) const
 {
-	// No layout info about this tab found; start 
-	TSharedRef<FArea> NewAreaForTab = FTabManager::NewPrimaryArea()
-		->Split
-		(
-			FTabManager::NewStack()
-			->AddTab(UnmanagedTab->GetLayoutIdentifier(), ETabState::OpenedTab)
-		);
+	bool bIsTabSupported = true;
+	if( OnIsTabSupported.IsBound() )
+	{
+		OnIsTabSupported.Broadcast( TabId, /* In/Out */ bIsTabSupported );
+	}
+
+	return bIsTabSupported;
+}
 
+void FProxyTabmanager::OpenUnmanagedTab(FName PlaceholderId, const FSearchPreference& SearchPreference, const TSharedRef<SDockTab>& UnmanagedTab)
+{
 	TSharedPtr<SWindow> ParentWindowPtr = ParentWindow.Pin();
-	TSharedRef<SDockingArea> DockingArea = RestoreArea(NewAreaForTab, ParentWindowPtr, false);
-	ParentWindowPtr->SetContent(DockingArea);
+	if( ensure( ParentWindowPtr.IsValid() ) )
+	{
+		// No layout info about this tab found; start 
+		TSharedRef<FArea> NewAreaForTab = FTabManager::NewPrimaryArea()
+			->Split
+			(
+				FTabManager::NewStack()
+				->AddTab( UnmanagedTab->GetLayoutIdentifier(), ETabState::OpenedTab )
+			);
+
+		TSharedRef<SDockingArea> DockingArea = RestoreArea(NewAreaForTab, ParentWindowPtr, false);
+		ParentWindowPtr->SetContent(DockingArea);
 
-	const TSharedPtr<SDockTab> NewlyOpenedTab = DockingArea->GetAllChildTabs()[0];
-	check(NewlyOpenedTab.IsValid());
+		const TSharedPtr<SDockTab> NewlyOpenedTab = DockingArea->GetAllChildTabs()[0];
+		check(NewlyOpenedTab.IsValid());
 		
-	NewlyOpenedTab->GetParent()->GetParentDockTabStack()->OpenTab(UnmanagedTab);
-	NewlyOpenedTab->RequestCloseTab();
+		NewlyOpenedTab->GetParent()->GetParentDockTabStack()->OpenTab(UnmanagedTab);
+		NewlyOpenedTab->RequestCloseTab();
 
-	MainNonCloseableTab = UnmanagedTab;
+		MainNonCloseableTab = UnmanagedTab;
 
-	OnTabOpened.Broadcast(UnmanagedTab);
+		OnTabOpened.Broadcast(UnmanagedTab);
+	}
 }
 
 void FProxyTabmanager::DrawAttention(const TSharedRef<SDockTab>& TabToHighlight)
@@ -2163,4 +2179,9 @@ void FProxyTabmanager::DrawAttention(const TSharedRef<SDockTab>& TabToHighlight)
 }
 
 
+void FProxyTabmanager::SetParentWindow(TSharedRef<SWindow> InParentWindow)
+{
+	ParentWindow = InParentWindow;
+}
+
 #undef LOCTEXT_NAMESPACE
diff --git a/Engine/Source/Runtime/Slate/Public/Framework/Docking/TabManager.h b/Engine/Source/Runtime/Slate/Public/Framework/Docking/TabManager.h
--- a/Engine/Source/Runtime/Slate/Public/Framework/Docking/TabManager.h
+++ b/Engine/Source/Runtime/Slate/Public/Framework/Docking/TabManager.h
@@ -912,6 +912,10 @@ public:
 
 	void SetProxyTabManager(TSharedPtr<FProxyTabmanager> InProxyTabManager);
 
+	DECLARE_DELEGATE(FOnOverrideDockableAreaRestore);
+	/** Used to override dockable area restoration behavior */
+	FOnOverrideDockableAreaRestore OnOverrideDockableAreaRestore_Handler;
+
 protected:
 	virtual void OnTabForegrounded( const TSharedPtr<SDockTab>& NewForegroundTab, const TSharedPtr<SDockTab>& BackgroundedTab ) override;
 	virtual void OnTabRelocated( const TSharedRef<SDockTab>& RelocatedTab, const TSharedPtr<SWindow>& NewOwnerWindow ) override;
@@ -1005,6 +1009,7 @@ private:
 
 //#HACK VREDITOR - Had to introduce the proxy tab manager to steal asset tabs.
 
+DECLARE_MULTICAST_DELEGATE_TwoParams(FIsTabSupportedEvent, FTabId /* TabId */, bool& /* bOutIsSupported */ );
 DECLARE_MULTICAST_DELEGATE_OneParam(FOnTabEvent, TSharedPtr<SDockTab>);
 
 
@@ -1021,8 +1026,12 @@ public:
 
 	virtual void OpenUnmanagedTab(FName PlaceholderId, const FSearchPreference& SearchPreference, const TSharedRef<SDockTab>& UnmanagedTab) override;
 	virtual void DrawAttention(const TSharedRef<SDockTab>& TabToHighlight) override;
+	
+	bool IsTabSupported( const FTabId TabId ) const;
+	void SetParentWindow(TSharedRef<SWindow> InParentWindow);
 
 public:
+	FIsTabSupportedEvent OnIsTabSupported;
 	FOnTabEvent OnTabOpened;
 	FOnTabEvent OnAttentionDrawnToTab;
 
diff --git a/Engine/Source/Runtime/SlateCore/Private/Rendering/DrawElements.cpp b/Engine/Source/Runtime/SlateCore/Private/Rendering/DrawElements.cpp
--- a/Engine/Source/Runtime/SlateCore/Private/Rendering/DrawElements.cpp
+++ b/Engine/Source/Runtime/SlateCore/Private/Rendering/DrawElements.cpp
@@ -65,6 +65,10 @@ void FSlateDrawElement::Init(uint32 InLayer, const FPaintGeometry& PaintGeometry
 	{
 		DataPayload.BatchFlags |= ESlateBatchDrawFlag::NoGamma;
 	}
+	if (InDrawEffects & ESlateDrawEffect::ReverseGamma)
+	{
+		DataPayload.BatchFlags |= ESlateBatchDrawFlag::ReverseGamma;
+	}
 }
 
 void FSlateDrawElement::MakeDebugQuad( FSlateWindowElementList& ElementList, uint32 InLayer, const FPaintGeometry& PaintGeometry, const FSlateRect& InClippingRect)
diff --git a/Engine/Source/Runtime/SlateCore/Public/Rendering/RenderingCommon.h b/Engine/Source/Runtime/SlateCore/Public/Rendering/RenderingCommon.h
--- a/Engine/Source/Runtime/SlateCore/Public/Rendering/RenderingCommon.h
+++ b/Engine/Source/Runtime/SlateCore/Public/Rendering/RenderingCommon.h
@@ -74,6 +74,8 @@ namespace ESlateDrawEffect
 	const Type PreMultipliedAlpha	= 1 << 3;
 	/** Advanced: No gamma correction should be done */
 	const Type NoGamma				= 1 << 4;
+	/** Advanced: Existing Gamma correction should be reversed */
+	const Type ReverseGamma			= 1 << 5;
 };
 
 
@@ -95,6 +97,8 @@ namespace ESlateBatchDrawFlag
 	const Type TileU				= 1 << 4;
 	/** The element should be tiled vertically */
 	const Type TileV				= 1 << 5;
+	/** Reverse gamma correction */
+	const Type ReverseGamma			 = 1 << 6;
 };
 
 namespace ESlateLineJoinType
diff --git a/Engine/Source/Runtime/SlateRHIRenderer/Private/SlateRHIRenderingPolicy.cpp b/Engine/Source/Runtime/SlateRHIRenderer/Private/SlateRHIRenderingPolicy.cpp
--- a/Engine/Source/Runtime/SlateRHIRenderer/Private/SlateRHIRenderingPolicy.cpp
+++ b/Engine/Source/Runtime/SlateRHIRenderer/Private/SlateRHIRenderingPolicy.cpp
@@ -239,7 +239,8 @@ void FSlateRHIRenderingPolicy::DrawElements(FRHICommandListImmediate& RHICmdList
 
 	static const FEngineShowFlags DefaultShowFlags(ESFIM_Game);
 
-	const float DisplayGamma = bGammaCorrect ? GEngine ? GEngine->GetDisplayGamma() : 2.2f : 1.0f;
+	const float EngineGamma = GEngine ? GEngine->GetDisplayGamma() : 2.2f;
+	const float DisplayGamma = bGammaCorrect ? EngineGamma : 1.0f;
 
 	// In order to support MaterialParameterCollections, we need to create multiple FSceneViews for 
 	// each possible Scene that we encounter. The following code creates these as separate arrays, where the first 
@@ -326,6 +327,8 @@ void FSlateRHIRenderingPolicy::DrawElements(FRHICommandListImmediate& RHICmdList
 		const ESlateShader::Type ShaderType = RenderBatch.ShaderType;
 		const FShaderParams& ShaderParams = RenderBatch.ShaderParams;
 
+		const float FinalGamma = DrawFlags & ESlateBatchDrawFlag::ReverseGamma ? 1.0f / EngineGamma : (DrawFlags & ESlateBatchDrawFlag::NoGamma) ? 1.0f : DisplayGamma;
+
 		auto UpdateScissorRect = [&RHICmdList, &RenderBatch]()
 		{
 			if (RenderBatch.ScissorRect.IsSet())
@@ -474,7 +477,9 @@ void FSlateRHIRenderingPolicy::DrawElements(FRHICommandListImmediate& RHICmdList
                 
 				PixelShader->SetTexture(RHICmdList, TextureRHI, SamplerState);
 				PixelShader->SetShaderParams(RHICmdList, ShaderParams.PixelParams);
-				PixelShader->SetDisplayGamma(RHICmdList, (DrawFlags & ESlateBatchDrawFlag::NoGamma) ? 1.0f : DisplayGamma);
+
+				
+				PixelShader->SetDisplayGamma(RHICmdList, FinalGamma);
 
 	
 				// for RHIs that can't handle VertexOffset, we need to offset the stream source each time
@@ -547,8 +552,8 @@ void FSlateRHIRenderingPolicy::DrawElements(FRHICommandListImmediate& RHICmdList
 						VertexShader->SetVerticalAxisMultiplier(RHICmdList, bAllowSwitchVerticalAxis && RHINeedsToSwitchVerticalAxis(GShaderPlatformForFeatureLevel[GMaxRHIFeatureLevel]) ? -1.0f : 1.0f);
 						VertexShader->SetMaterialShaderParameters(RHICmdList, ActiveSceneView, MaterialRenderProxy, Material);
 
-						PixelShader->SetParameters(RHICmdList, ActiveSceneView, MaterialRenderProxy, Material, DisplayGamma, ShaderParams.PixelParams);
-						PixelShader->SetDisplayGamma(RHICmdList, (DrawFlags & ESlateBatchDrawFlag::NoGamma) ? 1.0f : DisplayGamma);
+						PixelShader->SetParameters(RHICmdList, ActiveSceneView, MaterialRenderProxy, Material, FinalGamma, ShaderParams.PixelParams);
+						PixelShader->SetDisplayGamma(RHICmdList, FinalGamma);
 
 						FSlateShaderResource* MaskResource = MaterialShaderResource->GetTextureMaskResource();
 						if (MaskResource)
diff --git a/README.md b/README.md
--- a/README.md
+++ b/README.md
@@ -1,6 +1,8 @@
 Unreal Engine
 =============
 
+Note: This branch contains all of the VR Editor and Mesh Editor features demonstrated at the **State of Unreal** keynote at the 2017 Game Developers Conference. The VR Editor and Mesh Editor are experimental features. For this particular release, features were tested exclusively on Windows, so there may be complications compiling for Mac or Linux. For documentation on the new features and updates, please see our [VR & Mesh Editor GDC 2017 Features Preview](https://docs.unrealengine.com/latest/INT/Engine/Editor/VR/GDC2017/index.html) docs.
+
 Welcome to the Unreal Engine source code! 
 
 From this repository you can build the Unreal Editor for Windows and Mac, compile Unreal Engine games for Android, iOS, Playstation 4, Xbox One, HTML5 and Linux,
