Copying //UE4/Dev-Niagara to //UE4/Dev-Main (Source: //UE4/Dev-Niagara @ 4074996)

diff --git a/Engine/Source/Editor/PropertyEditor/Private/DetailLayoutBuilderImpl.cpp b/Engine/Source/Editor/PropertyEditor/Private/DetailLayoutBuilderImpl.cpp
--- a/Engine/Source/Editor/PropertyEditor/Private/DetailLayoutBuilderImpl.cpp
+++ b/Engine/Source/Editor/PropertyEditor/Private/DetailLayoutBuilderImpl.cpp
@@ -485,6 +485,39 @@ void FDetailLayoutBuilderImpl::AddExternalRootPropertyNode(TSharedRef<FComplexPr
 	ExternalRootPropertyNodes.Add(InExternalRootNode);
 }
 
+void FDetailLayoutBuilderImpl::RemoveExternalRootPropertyNode(TSharedRef<FComplexPropertyNode> InExternalRootNode)
+{
+	ExternalRootPropertyNodes.RemoveAll([InExternalRootNode](TSharedPtr<FComplexPropertyNode> ExternalRootPropertyNode)
+	{
+		return ExternalRootPropertyNode == InExternalRootNode;
+	});
+}
+
+FDelegateHandle FDetailLayoutBuilderImpl::AddNodeVisibilityChangedHandler(FSimpleMulticastDelegate::FDelegate InOnNodeVisibilityChanged)
+{
+	return OnNodeVisibilityChanged.Add(InOnNodeVisibilityChanged);
+}
+
+void FDetailLayoutBuilderImpl::RemoveNodeVisibilityChangedHandler(FDelegateHandle DelegateHandle)
+{
+	OnNodeVisibilityChanged.Remove(DelegateHandle);
+}
+
+void FDetailLayoutBuilderImpl::NotifyNodeVisibilityChanged()
+{
+	OnNodeVisibilityChanged.Broadcast();
+}
+
+FCustomPropertyTypeLayoutMap&  FDetailLayoutBuilderImpl::GetInstancedPropertyTypeLayoutMap()
+{
+	return InstancedPropertyTypeLayoutMap;
+}
+
+void FDetailLayoutBuilderImpl::SetInstancedPropertyTypeLayoutMap(FCustomPropertyTypeLayoutMap& InInstancedPropertyTypeLayoutMap)
+{
+	InstancedPropertyTypeLayoutMap = InInstancedPropertyTypeLayoutMap;
+}
+
 TSharedPtr<FAssetThumbnailPool> FDetailLayoutBuilderImpl::GetThumbnailPool() const
 {
 	return PropertyDetailsUtilities.Pin()->GetThumbnailPool();
diff --git a/Engine/Source/Editor/PropertyEditor/Private/DetailLayoutBuilderImpl.h b/Engine/Source/Editor/PropertyEditor/Private/DetailLayoutBuilderImpl.h
--- a/Engine/Source/Editor/PropertyEditor/Private/DetailLayoutBuilderImpl.h
+++ b/Engine/Source/Editor/PropertyEditor/Private/DetailLayoutBuilderImpl.h
@@ -168,6 +168,13 @@ public:
 	 */
 	void AddExternalRootPropertyNode(TSharedRef<FComplexPropertyNode> InExternalRootNode);
 
+	/**
+	* Removes an external property root node to the list of root nodes that the details new needs to manage
+	*
+	* @param InExternalRootNode		The node to remove
+	*/
+	void RemoveExternalRootPropertyNode(TSharedRef<FComplexPropertyNode> InExternalRootNode);
+
 	/** @return The details view that owns this layout */
 	IDetailsViewPrivate* GetDetailsView() { return DetailsView; }
 	/** @return The root node for this customization */
@@ -178,6 +185,19 @@ public:
 	/** @return True if the layout is for an external root property node and not in the main set of objects the details panel is observing */
 	bool IsLayoutForExternalRoot() const { return bLayoutForExternalRoot; }
 
+	/** Adds a handler for when a node this builder owns has had a forced visibility change. */
+	FDelegateHandle AddNodeVisibilityChangedHandler(FSimpleMulticastDelegate::FDelegate InOnNodeVisibilityChanged);
+
+	/** Removes a handler for when a node this builder owns has had a forced visibility change. */
+	void RemoveNodeVisibilityChangedHandler(FDelegateHandle DelegateHandle);
+	
+	/** Notifies this detail layout builder that a node it owns had it's visibility forcibly changed. */
+	void NotifyNodeVisibilityChanged();
+
+	FCustomPropertyTypeLayoutMap& GetInstancedPropertyTypeLayoutMap();
+
+	void SetInstancedPropertyTypeLayoutMap(FCustomPropertyTypeLayoutMap& InInstancedPropertyTypeLayoutMap);
+
 private:
 	/**
 	 * Finds a property node for the current property by searching in a fast lookup map or a path search if required
@@ -225,5 +245,9 @@ private:
 	UStruct* CurrentCustomizationClass;
 	/** True if the layout is for an external root property node and not in the main set of objects the details panel is observing */
 	bool bLayoutForExternalRoot;
+	/** A delegate which is called whenever a node owned by this layout builder has it's visibility forcibly changed. */
+	FSimpleMulticastDelegate OnNodeVisibilityChanged;
+
+	FCustomPropertyTypeLayoutMap InstancedPropertyTypeLayoutMap;
 };
 
diff --git a/Engine/Source/Editor/PropertyEditor/Private/DetailPropertyRow.cpp b/Engine/Source/Editor/PropertyEditor/Private/DetailPropertyRow.cpp
--- a/Engine/Source/Editor/PropertyEditor/Private/DetailPropertyRow.cpp
+++ b/Engine/Source/Editor/PropertyEditor/Private/DetailPropertyRow.cpp
@@ -378,9 +378,7 @@ TSharedPtr<IPropertyTypeCustomization> FDetailPropertyRow::GetPropertyCustomizat
 		static FName NAME_PropertyEditor("PropertyEditor");
 		FPropertyEditorModule& PropertyEditorModule = FModuleManager::GetModuleChecked<FPropertyEditorModule>(NAME_PropertyEditor);
 
-		IDetailsViewPrivate* DetailsView = InParentCategory->GetDetailsView();
-
-		FPropertyTypeLayoutCallback LayoutCallback = PropertyEditorModule.GetPropertyTypeCustomization(Property, *PropHandle, DetailsView ? DetailsView->GetCustomPropertyTypeLayoutMap() : FCustomPropertyTypeLayoutMap() );
+		FPropertyTypeLayoutCallback LayoutCallback = PropertyEditorModule.GetPropertyTypeCustomization(Property, *PropHandle, InParentCategory->GetCustomPropertyTypeLayoutMap() );
 		if (LayoutCallback.IsValid())
 		{
 			if (PropHandle->IsValidHandle())
diff --git a/Engine/Source/Editor/PropertyEditor/Private/DetailPropertyRow.h b/Engine/Source/Editor/PropertyEditor/Private/DetailPropertyRow.h
--- a/Engine/Source/Editor/PropertyEditor/Private/DetailPropertyRow.h
+++ b/Engine/Source/Editor/PropertyEditor/Private/DetailPropertyRow.h
@@ -75,6 +75,11 @@ public:
 	 */
 	TSharedPtr<FPropertyNode> GetPropertyNode() { return PropertyNode; }
 
+	/**
+	 * @return The external root node for this row if it has one.
+	 */
+	TSharedPtr<FComplexPropertyNode> GetExternalRootNode() const { return ExternalRootNode; }
+
 	/**
 	 * @return The property node for this row
 	 */
@@ -175,7 +180,8 @@ private:
 	/** The category this row resides in */
 	TWeakPtr<FDetailCategoryImpl> ParentCategory;
 	/** Root of the property node if this node comes from an external tree */
-	TSharedPtr<FPropertyNode> ExternalRootNode;
+	TSharedPtr<FComplexPropertyNode> ExternalRootNode;
+
 	TSharedPtr<struct FDetailLayoutData> ExternalObjectLayout;
 	/** The custom expansion ID name used to save and restore expansion state on this node */
 	FName CustomExpansionIdName;
diff --git a/Engine/Source/Editor/PropertyEditor/Private/DetailTreeNode.cpp b/Engine/Source/Editor/PropertyEditor/Private/DetailTreeNode.cpp
--- a/Engine/Source/Editor/PropertyEditor/Private/DetailTreeNode.cpp
+++ b/Engine/Source/Editor/PropertyEditor/Private/DetailTreeNode.cpp
@@ -15,11 +15,17 @@ FNodeWidgets FDetailTreeNode::CreateNodeWidgets() const
 		if (Row.HasColumns())
 		{
 			Widgets.NameWidget = Row.NameWidget.Widget;
+			Widgets.NameWidgetLayoutData = FNodeWidgetLayoutData(
+				Row.NameWidget.HorizontalAlignment, Row.NameWidget.VerticalAlignment, Row.NameWidget.MinWidth, Row.NameWidget.MaxWidth);
 			Widgets.ValueWidget = Row.ValueWidget.Widget;
+			Widgets.ValueWidgetLayoutData = FNodeWidgetLayoutData(
+				Row.ValueWidget.HorizontalAlignment, Row.ValueWidget.VerticalAlignment, Row.ValueWidget.MinWidth, Row.ValueWidget.MaxWidth);
 		}
 		else
 		{
 			Widgets.WholeRowWidget = Row.WholeRowWidget.Widget;
+			Widgets.WholeRowWidgetLayoutData = FNodeWidgetLayoutData(
+				Row.WholeRowWidget.HorizontalAlignment, Row.WholeRowWidget.VerticalAlignment, Row.WholeRowWidget.MinWidth, Row.WholeRowWidget.MaxWidth);
 		}
 	}
 
diff --git a/Engine/Source/Editor/PropertyEditor/Private/PropertyRowGenerator.cpp b/Engine/Source/Editor/PropertyEditor/Private/PropertyRowGenerator.cpp
--- a/Engine/Source/Editor/PropertyEditor/Private/PropertyRowGenerator.cpp
+++ b/Engine/Source/Editor/PropertyEditor/Private/PropertyRowGenerator.cpp
@@ -22,7 +22,7 @@ public:
 	/** IPropertyUtilities interface */
 	virtual class FNotifyHook* GetNotifyHook() const override
 	{
-		return nullptr;
+		return Generator.GetNotifyHook();
 	}
 	virtual bool AreFavoritesEnabled() const override
 	{
@@ -352,6 +352,7 @@ void FPropertyRowGenerator::UpdateDetailRows()
 	//NumVisbleTopLevelObjectNodes = 0;
 
 	FDetailFilter CurrentFilter;
+	CurrentFilter.bShowAllAdvanced = true;
 
 	for (int32 RootNodeIndex = 0; RootNodeIndex < RootPropertyNodes.Num(); ++RootNodeIndex)
 	{
@@ -458,6 +459,8 @@ void FPropertyRowGenerator::UpdateSinglePropertyMap(TSharedPtr<FComplexPropertyN
 	LayoutData.ClassToPropertyMap.Empty();
 
 	TSharedPtr<FDetailLayoutBuilderImpl> DetailLayout = MakeShareable(new FDetailLayoutBuilderImpl(InRootPropertyNode, LayoutData.ClassToPropertyMap, PropertyUtilities, nullptr, false));
+	DetailLayout->AddNodeVisibilityChangedHandler(FSimpleMulticastDelegate::FDelegate::CreateRaw(this, &FPropertyRowGenerator::LayoutNodeVisibilityChanged));
+	DetailLayout->SetInstancedPropertyTypeLayoutMap(InstancedTypeToLayoutMap);
 	LayoutData.DetailLayout = DetailLayout;
 
 	TSharedPtr<FComplexPropertyNode> RootPropertyNode = InRootPropertyNode;
@@ -538,3 +541,8 @@ TSharedPtr<IDetailTreeNode> FPropertyRowGenerator::FindTreeNodeRecursive(const T
 	return nullptr;
 }
 
+void FPropertyRowGenerator::LayoutNodeVisibilityChanged()
+{
+	UpdateDetailRows();
+}
+
diff --git a/Engine/Source/Editor/PropertyEditor/Private/PropertyRowGenerator.h b/Engine/Source/Editor/PropertyEditor/Private/PropertyRowGenerator.h
--- a/Engine/Source/Editor/PropertyEditor/Private/PropertyRowGenerator.h
+++ b/Engine/Source/Editor/PropertyEditor/Private/PropertyRowGenerator.h
@@ -100,7 +100,7 @@ public:
 	virtual TStatId GetStatId() const override;
 
 	/** IPropertyUtilities interface */
-	virtual class FNotifyHook* GetNotifyHook() const { return nullptr; }
+	virtual class FNotifyHook* GetNotifyHook() const { return Args.NotifyHook; }
 	virtual void EnqueueDeferredAction(FSimpleDelegate DeferredAction);	
 	virtual bool IsPropertyEditingEnabled() const { return true; }
 	virtual void ForceRefresh();
@@ -115,6 +115,8 @@ private:
 	void UpdateSinglePropertyMap(TSharedPtr<FComplexPropertyNode> InRootPropertyNode, FDetailLayoutData& LayoutData);
 	bool ValidatePropertyNodes(const FRootPropertyNodeList& PropertyNodeList);
 	TSharedPtr<IDetailTreeNode> FindTreeNodeRecursive(const TSharedPtr<IDetailTreeNode>& StartNode, TSharedPtr<IPropertyHandle> PropertyHandle) const;
+	void LayoutNodeVisibilityChanged();
+
 private:
 	const FPropertyRowGeneratorArgs Args;
 	/** The root property nodes of the property tree for a specific set of UObjects */
diff --git a/Engine/Source/Editor/PropertyEditor/Private/SDetailSingleItemRow.cpp b/Engine/Source/Editor/PropertyEditor/Private/SDetailSingleItemRow.cpp
--- a/Engine/Source/Editor/PropertyEditor/Private/SDetailSingleItemRow.cpp
+++ b/Engine/Source/Editor/PropertyEditor/Private/SDetailSingleItemRow.cpp
@@ -657,9 +657,9 @@ TSharedRef<SWidget> SDetailSingleItemRow::CreateKeyframeButton( FDetailLayoutCus
 	{
 		TSharedPtr<IPropertyHandle> Handle = PropertyEditorHelpers::GetPropertyHandle(InCustomization.GetPropertyNode().ToSharedRef(), nullptr, nullptr);
 
-		UClass* ObjectClass = InCustomization.GetPropertyNode()->FindObjectItemParent()->GetObjectBaseClass();
-		SetKeyVisibility = KeyframeHandler.Pin()->IsPropertyKeyable(ObjectClass, *Handle) ? EVisibility::Visible : EVisibility::Collapsed;
-		
+		FObjectPropertyNode* ObjectItemParent = InCustomization.GetPropertyNode()->FindObjectItemParent();
+		UClass* ObjectClass = ObjectItemParent != nullptr ? ObjectItemParent->GetObjectBaseClass() : nullptr;
+		SetKeyVisibility = ObjectClass != nullptr && KeyframeHandler.Pin()->IsPropertyKeyable(ObjectClass, *Handle) ? EVisibility::Visible : EVisibility::Collapsed;
 	}
 
 	return
diff --git a/Engine/Source/Editor/PropertyEditor/Private/SDetailsView.cpp b/Engine/Source/Editor/PropertyEditor/Private/SDetailsView.cpp
--- a/Engine/Source/Editor/PropertyEditor/Private/SDetailsView.cpp
+++ b/Engine/Source/Editor/PropertyEditor/Private/SDetailsView.cpp
@@ -879,9 +879,6 @@ void SDetailsView::PostSetObject()
 		FObjectPropertyNode* RootPropertyNode = ComplexRootNode->AsObjectNode();
 
 		RootPropertyNode->InitNode( InitParams );
-
-		// Restore existing expanded items
-		RestoreExpandedItems(ComplexRootNode.ToSharedRef());
 	}
 
 	UpdatePropertyMaps();
diff --git a/Engine/Source/Editor/PropertyEditor/Private/SDetailsViewBase.cpp b/Engine/Source/Editor/PropertyEditor/Private/SDetailsViewBase.cpp
--- a/Engine/Source/Editor/PropertyEditor/Private/SDetailsViewBase.cpp
+++ b/Engine/Source/Editor/PropertyEditor/Private/SDetailsViewBase.cpp
@@ -706,6 +706,7 @@ void SDetailsViewBase::Tick( const FGeometry& AllottedGeometry, const double InC
 	bool bValidateExternalNodes = true;
 
 	int32 FoundIndex = RootPropertyNodes.Find(LastRootPendingKill);
+	bool bUpdateFilteredDetails = false;
 	if (FoundIndex != INDEX_NONE)
 	{ 
 		// Reaquire the root property nodes.  It may have been changed by the deferred actions if something like a blueprint editor forcefully resets a details panel during a posteditchange
@@ -720,20 +721,19 @@ void SDetailsViewBase::Tick( const FGeometry& AllottedGeometry, const double InC
 		{
 			if(RootPropertyNode == LastRootPendingKill)
 			{
+				ExpandedDetailNodes.Empty();
 				RestoreExpandedItems(RootPropertyNode.ToSharedRef());
 			}
 
 			EPropertyDataValidationResult Result = RootPropertyNode->EnsureDataIsValid();
 			if(Result == EPropertyDataValidationResult::PropertiesChanged || Result == EPropertyDataValidationResult::EditInlineNewValueChanged)
 			{
-				RestoreExpandedItems(RootPropertyNode.ToSharedRef());
 				UpdatePropertyMaps();
-				UpdateFilteredDetails();
+				bUpdateFilteredDetails = true;
 			}
 			else if(Result == EPropertyDataValidationResult::ArraySizeChanged || Result == EPropertyDataValidationResult::ChildrenRebuilt)
 			{
-				RestoreExpandedItems(RootPropertyNode.ToSharedRef());
-				UpdateFilteredDetails();
+				bUpdateFilteredDetails = true;
 			}
 			else if(Result == EPropertyDataValidationResult::ObjectInvalid)
 			{
@@ -759,26 +759,28 @@ void SDetailsViewBase::Tick( const FGeometry& AllottedGeometry, const double InC
 					EPropertyDataValidationResult Result = PropertyNode->EnsureDataIsValid();
 					if (Result == EPropertyDataValidationResult::PropertiesChanged || Result == EPropertyDataValidationResult::EditInlineNewValueChanged)
 					{
-						RestoreExpandedItems(PropertyNode.ToSharedRef());
-
 						// Note this will invalidate all the external root nodes so there is no need to continue
 						ExternalRootPropertyNodes.Empty();
 
 						UpdatePropertyMaps();
-						UpdateFilteredDetails();
+						bUpdateFilteredDetails = true;
 
 						break;
 					}
 					else if (Result == EPropertyDataValidationResult::ArraySizeChanged || Result == EPropertyDataValidationResult::ChildrenRebuilt)
 					{
-						RestoreExpandedItems(PropertyNode.ToSharedRef());
-						UpdateFilteredDetails();
+						bUpdateFilteredDetails = true;
 					}
 				}
 			}
 		}
 	}
 
+	if (bUpdateFilteredDetails)
+	{
+		UpdateFilteredDetails();
+	}
+
 	for(FDetailLayoutData& LayoutData : DetailLayouts)
 	{
 		if(LayoutData.DetailLayout.IsValid())
@@ -924,8 +926,6 @@ void SDetailsViewBase::RestoreExpandedItems(TSharedRef<FPropertyNode> InitialSta
 {
 	TSharedPtr<FPropertyNode> StartNode = InitialStartNode;
 
-	ExpandedDetailNodes.Empty();
-
 	FString ExpandedCustomItems;
 
 	UStruct* BestBaseStruct = StartNode->FindComplexParent()->GetBaseStructure();
@@ -954,6 +954,7 @@ void SDetailsViewBase::RestoreExpandedItems(TSharedRef<FPropertyNode> InitialSta
 void SDetailsViewBase::UpdateFilteredDetails()
 {
 	RootTreeNodes.Reset();
+	ExpandedDetailNodes.Empty();
 
 	FDetailNodeList InitialRootNodeList;
 	
@@ -973,6 +974,8 @@ void SDetailsViewBase::UpdateFilteredDetails()
 			RootPropertyNode->FilterNodes(CurrentFilter.FilterStrings);
 			RootPropertyNode->ProcessSeenFlags(true);
 
+			RestoreExpandedItems(RootPropertyNode.ToSharedRef());
+
 			TSharedPtr<FDetailLayoutBuilderImpl>& DetailLayout = DetailLayouts[RootNodeIndex].DetailLayout;
 			if(DetailLayout.IsValid())
 			{
diff --git a/Engine/Source/Editor/PropertyEditor/Public/IDetailTreeNode.h b/Engine/Source/Editor/PropertyEditor/Public/IDetailTreeNode.h
--- a/Engine/Source/Editor/PropertyEditor/Public/IDetailTreeNode.h
+++ b/Engine/Source/Editor/PropertyEditor/Public/IDetailTreeNode.h
@@ -3,8 +3,10 @@
 #pragma once
 
 #include "CoreMinimal.h"
+#include "Types/SlateEnums.h"
 
 class IPropertyHandle;
+class SWidget;
 
 enum class EDetailNodeType
 {
@@ -18,17 +20,54 @@ enum class EDetailNodeType
 	Object,
 };
 
+/** Layout data for node's content widgets. */
+struct FNodeWidgetLayoutData
+{
+	FNodeWidgetLayoutData()
+	{
+	}
+
+	FNodeWidgetLayoutData(EHorizontalAlignment InHorizontalAlignment, EVerticalAlignment InVerticalAlignment, TOptional<float> InMinWidth, TOptional<float> InMaxWidth)
+		: HorizontalAlignment(InHorizontalAlignment)
+		, VerticalAlignment(InVerticalAlignment)
+		, MinWidth(InMinWidth)
+		, MaxWidth(InMaxWidth)
+	{
+	}
+
+	/** The horizontal alignment requested by the widget. */
+	EHorizontalAlignment HorizontalAlignment;
+
+	/** The vertical alignment requested by the widget. */
+	EVerticalAlignment VerticalAlignment;
+
+	/** An optional minimum width requested by the widget. */
+	TOptional<float> MinWidth;
+
+	/** An optional maximum width requested by the widget. */
+	TOptional<float> MaxWidth;
+};
+
 /** The widget contents of the node.  Any of these can be null depending on how the row was generated */
 struct FNodeWidgets
 {
 	/** Widget for the name column */
 	TSharedPtr<SWidget> NameWidget;
+
+	/** Layout data for the widget in the name column. */
+	FNodeWidgetLayoutData NameWidgetLayoutData;
+
 	/** Widget for the value column*/
 	TSharedPtr<SWidget> ValueWidget;
 
+	/** Layout data for the widget in the value column. */
+	FNodeWidgetLayoutData ValueWidgetLayoutData;
 
 	/** Widget that spans the entire row.  Mutually exclusive with name/value widget */
 	TSharedPtr<SWidget> WholeRowWidget;
+
+	/** Layout data for the whole row widget. */
+	FNodeWidgetLayoutData WholeRowWidgetLayoutData;
 };
 
 class IDetailTreeNode
@@ -63,5 +102,16 @@ public:
 	 */
 	virtual void GetChildren(TArray<TSharedRef<IDetailTreeNode>>& OutChildren) = 0;
 
+	/**
+	 * Gets an identifier name for this node.  This is not a name formatted for display purposes, but can be useful for storing
+	 * UI state like if this row is expanded.
+	 */
+	virtual FName GetNodeName() const = 0;
+
 	virtual TSharedPtr<class IDetailPropertyRow> GetRow() const = 0;
+
+	/**
+	 * Gets the filter strings for this node in the tree.
+	 */
+	virtual void GetFilterStrings(TArray<FString>& OutFilterStrings) const = 0;
 };
\ No newline at end of file
diff --git a/Engine/Source/Editor/Sequencer/Private/SequencerContextMenus.cpp b/Engine/Source/Editor/Sequencer/Private/SequencerContextMenus.cpp
--- a/Engine/Source/Editor/Sequencer/Private/SequencerContextMenus.cpp
+++ b/Engine/Source/Editor/Sequencer/Private/SequencerContextMenus.cpp
@@ -1512,6 +1512,7 @@ void FPasteContextMenu::PasteInto(int32 DestinationIndex, FName KeyAreaName)
 		Selection.GetOnKeySelectionChanged().Broadcast();
 
 		Sequencer->OnClipboardUsed(Args.Clipboard);
+		Sequencer->NotifyMovieSceneDataChanged(EMovieSceneDataChangeType::TrackValueChanged);
 	}
 }
 
diff --git a/Engine/Source/Editor/Sequencer/Private/SequencerSelection.cpp b/Engine/Source/Editor/Sequencer/Private/SequencerSelection.cpp
--- a/Engine/Source/Editor/Sequencer/Private/SequencerSelection.cpp
+++ b/Engine/Source/Editor/Sequencer/Private/SequencerSelection.cpp
@@ -100,11 +100,18 @@ TArray<UMovieSceneTrack*> FSequencerSelection::GetSelectedTracks() const
 	
 	for (const TSharedRef<FSequencerDisplayNode>& SelectedNode : SelectedOutlinerNodes)
 	{
-		if (SelectedNode->GetType() == ESequencerNode::Track)
+		TSharedPtr<FSequencerDisplayNode> CurrentNode = SelectedNode;
+		while (CurrentNode.IsValid() && CurrentNode->GetType() != ESequencerNode::Track)
 		{
-			if (UMovieSceneTrack* Track = StaticCastSharedRef<FSequencerTrackNode>(SelectedNode)->GetTrack())
+			CurrentNode = CurrentNode->GetParent();
+		}
+
+		if (CurrentNode.IsValid())
+		{
+			UMovieSceneTrack* SelectedTrack = StaticCastSharedPtr<FSequencerTrackNode>(CurrentNode)->GetTrack();
+			if (SelectedTrack != nullptr)
 			{
-				OutTracks.Add(Track);
+				OutTracks.Add(SelectedTrack);
 			}
 		}
 	}
diff --git a/Engine/Source/Editor/UnrealEd/Public/DragAndDrop/DecoratedDragDropOp.h b/Engine/Source/Editor/UnrealEd/Public/DragAndDrop/DecoratedDragDropOp.h
--- a/Engine/Source/Editor/UnrealEd/Public/DragAndDrop/DecoratedDragDropOp.h
+++ b/Engine/Source/Editor/UnrealEd/Public/DragAndDrop/DecoratedDragDropOp.h
@@ -24,9 +24,14 @@ public:
 	/** Icon to be displayed */
 	const FSlateBrush*					CurrentIconBrush;
 
+	/** The color of the icon to be displayed. */
+	FSlateColor							CurrentIconColorAndOpacity;
+
 	FDecoratedDragDropOp()
 		: CurrentIconBrush(nullptr)
+		, CurrentIconColorAndOpacity(FLinearColor::White)
 		, DefaultHoverIcon(nullptr)
+		, DefaultHoverIconColorAndOpacity(FLinearColor::White)
 	{ }
 
 	/** Overridden to provide public access */
@@ -40,6 +45,7 @@ public:
 	{
 		CurrentHoverText = DefaultHoverText;
 		CurrentIconBrush = DefaultHoverIcon;
+		CurrentIconColorAndOpacity = DefaultHoverIconColorAndOpacity;
 	}
 
 	/** The widget decorator to use */
@@ -59,6 +65,7 @@ public:
 				[
 					SNew( SImage )
 					.Image( this, &FDecoratedDragDropOp::GetIcon )
+					.ColorAndOpacity( this, &FDecoratedDragDropOp::GetIconColorAndOpacity)
 				]
 
 				+ SHorizontalBox::Slot()
@@ -81,6 +88,11 @@ public:
 		return CurrentIconBrush;
 	}
 
+	FSlateColor GetIconColorAndOpacity() const
+	{
+		return CurrentIconColorAndOpacity;
+	}
+
 	/** Set the text and icon for this tooltip */
 	void SetToolTip(const FText& Text, const FSlateBrush* Icon)
 	{
@@ -93,6 +105,13 @@ public:
 	{
 		DefaultHoverText = CurrentHoverText;
 		DefaultHoverIcon = CurrentIconBrush;
+		DefaultHoverIconColorAndOpacity = CurrentIconColorAndOpacity;
+	}
+
+	/** Gets the default hover text for this drag drop op. */
+	FText GetDefaultHoverText() const
+	{
+		return DefaultHoverText;
 	}
 
 protected:
@@ -102,6 +121,9 @@ protected:
 
 	/** Default icon to be displayed */
 	const FSlateBrush*					DefaultHoverIcon;
+
+	/** Default color and opacity for the default icon to be displayed. */
+	FSlateColor							DefaultHoverIconColorAndOpacity;
 };
 
 
diff --git a/Engine/Source/Editor/UnrealEd/Public/GraphEditor.h b/Engine/Source/Editor/UnrealEd/Public/GraphEditor.h
--- a/Engine/Source/Editor/UnrealEd/Public/GraphEditor.h
+++ b/Engine/Source/Editor/UnrealEd/Public/GraphEditor.h
@@ -441,6 +441,15 @@ public:
 		}
 	}
 
+	/** Sets the current node, pin and connection factory. */
+	virtual void SetNodeFactory(const TSharedRef<class FGraphNodeFactory>& NewNodeFactory)
+	{
+		if (Implementation.IsValid())
+		{
+			Implementation->SetNodeFactory(NewNodeFactory);
+		}
+	}
+
 protected:
 	/** Invoked when the underlying Graph is being changed. */
 	virtual void OnGraphChanged(const struct FEdGraphEditAction& InAction)
diff --git a/Engine/Source/Editor/UnrealEd/Public/Kismet2/DebuggerCommands.h b/Engine/Source/Editor/UnrealEd/Public/Kismet2/DebuggerCommands.h
--- a/Engine/Source/Editor/UnrealEd/Public/Kismet2/DebuggerCommands.h
+++ b/Engine/Source/Editor/UnrealEd/Public/Kismet2/DebuggerCommands.h
@@ -153,6 +153,7 @@ public:
 
 	static void ResumePlaySession_Clicked();
 	static void PausePlaySession_Clicked();
+	static void SingleFrameAdvance_Clicked();
 
 	static bool IsInSIE();
 	static bool IsInPIE();
diff --git a/Engine/Source/Runtime/Apple/MetalRHI/Private/MetalRenderPass.cpp b/Engine/Source/Runtime/Apple/MetalRHI/Private/MetalRenderPass.cpp
--- a/Engine/Source/Runtime/Apple/MetalRHI/Private/MetalRenderPass.cpp
+++ b/Engine/Source/Runtime/Apple/MetalRHI/Private/MetalRenderPass.cpp
@@ -383,8 +383,13 @@ void FMetalRenderPass::DrawIndexedPrimitive(FMetalBuffer const& IndexBuffer, uin
 		
 		if(ClampedNumInstances < NumInstances)
 		{
+			FString ShaderName = TEXT("Unknown");
+#if !(UE_BUILD_SHIPPING || UE_BUILD_TEST)
+			ShaderName = PipelineState->PixelShader->ShaderName;
+#endif
 			// Setting NumInstances to ClampedNumInstances would fix any visual rendering bugs resulting from this bad call but these draw calls are wrong - don't hide the issue
-			UE_LOG(LogMetal, Error, TEXT("Metal DrawIndexedPrimitive requested to draw %d Instances but vertex stream only has %d instance data available."), NumInstances, ClampedNumInstances);
+			UE_LOG(LogMetal, Error, TEXT("Metal DrawIndexedPrimitive requested to draw %d Instances but vertex stream only has %d instance data available. ShaderName: %s"), NumInstances, ClampedNumInstances,
+				*ShaderName);
 		}
 	}
 #endif
diff --git a/Engine/Source/Runtime/Engine/Classes/Engine/Engine.h b/Engine/Source/Runtime/Engine/Classes/Engine/Engine.h
--- a/Engine/Source/Runtime/Engine/Classes/Engine/Engine.h
+++ b/Engine/Source/Runtime/Engine/Classes/Engine/Engine.h
@@ -598,6 +598,8 @@ enum class EFrameHitchType: uint8;
 DECLARE_MULTICAST_DELEGATE_TwoParams(FEngineHitchDetectedDelegate, EFrameHitchType /*HitchType*/, float /*HitchDurationInSeconds*/);
 
 
+DECLARE_MULTICAST_DELEGATE(FPreRenderDelegate);
+
 /**
  * Abstract base class of all Engine classes, responsible for management of systems critical to editor or game systems.
  * Also defines default classes for certain engine systems.
@@ -1519,6 +1521,11 @@ public:
 	void RegisterEndStreamingPauseRenderingDelegate( FEndStreamingPauseDelegate* InDelegate );
 	FEndStreamingPauseDelegate* EndStreamingPauseDelegate;
 
+
+	/** Delegate called just prior to rendering. */
+	FPreRenderDelegate PreRenderDelegate;
+	FPreRenderDelegate& GetPreRenderDelegate() { return PreRenderDelegate; }
+
 	/** 
 	 * Error message event relating to server travel failures 
 	 * 
diff --git a/Engine/Source/Runtime/Engine/Classes/Engine/RendererSettings.h b/Engine/Source/Runtime/Engine/Classes/Engine/RendererSettings.h
--- a/Engine/Source/Runtime/Engine/Classes/Engine/RendererSettings.h
+++ b/Engine/Source/Runtime/Engine/Classes/Engine/RendererSettings.h
@@ -478,10 +478,7 @@ class ENGINE_API URendererSettings : public UDeveloperSettings
 		ConsoleVariable = "r.MorphTarget.Mode", DisplayName = "Use GPU for computing morph targets",
 		ToolTip = "Whether to use original CPU method (loop per morph then by vertex) or use a GPU-based method on Shader Model 5 hardware."))
 	uint32 bUseGPUMorphTargets : 1;
-
-	UPROPERTY(config, EditAnywhere, Category = "Optimizations", meta = (DisplayName = "GPU Particles Support Only Local Vector Field", Tooltip = "Limits Cascade GPU Particle simulations to applying local vector fields.  Global vector fields are not applied."))
-	bool bGPUParticlesLocalVFOnly;
-
+	
 	UPROPERTY(config, EditAnywhere, Category = Debugging, meta = (
 		ConsoleVariable = "r.GPUCrashDebugging", DisplayName = "Enable vendor specific GPU crash analysis tools",
 		ToolTip = "Enables vendor specific GPU crash analysis tools.  Currently only supports NVIDIA Aftermath on DX11.",
diff --git a/Engine/Source/Runtime/Engine/Classes/Materials/MaterialExpressionDynamicParameter.h b/Engine/Source/Runtime/Engine/Classes/Materials/MaterialExpressionDynamicParameter.h
--- a/Engine/Source/Runtime/Engine/Classes/Materials/MaterialExpressionDynamicParameter.h
+++ b/Engine/Source/Runtime/Engine/Classes/Materials/MaterialExpressionDynamicParameter.h
@@ -32,6 +32,10 @@ class UMaterialExpressionDynamicParameter : public UMaterialExpression
 	UPROPERTY(EditAnywhere, Category = MaterialExpressionDynamicParameter)
 	FLinearColor DefaultValue;
 
+	/** The index of the dynamic parameter for use in tools that allow > 1 */
+	UPROPERTY(EditAnywhere, Category = MaterialExpressionDynamicParameter, meta = (UIMin = 0, ClampMin = 0, UIMax = 3, ClampMax = 3))
+	uint32 ParameterIndex;
+
 	//~ Begin UObject Interface
 #if WITH_EDITOR
 	virtual void PostEditChangeProperty(FPropertyChangedEvent& PropertyChangedEvent) override;
diff --git a/Engine/Source/Runtime/Engine/Engine.Build.cs b/Engine/Source/Runtime/Engine/Engine.Build.cs
--- a/Engine/Source/Runtime/Engine/Engine.Build.cs
+++ b/Engine/Source/Runtime/Engine/Engine.Build.cs
@@ -417,24 +417,11 @@ public class Engine : ModuleRules
 				"libOpus"
 				);
 		}
-/*
-		ConfigHierarchy Ini = ConfigCache.ReadHierarchy(ConfigHierarchyType.Engine, System.IO.DirectoryReference.FromFile(Target.ProjectFile), Target.Platform);
-		bool bLocalVectorFieldOnly = false;
-		Ini.GetBool("/Script/Engine.RendererSettings", "bGPUParticlesLocalVFOnly", out bLocalVectorFieldOnly);
-		if (bLocalVectorFieldOnly)
-		{
-			PublicDefinitions.Add("GPUPARTICLE_LOCAL_VF_ONLY=1");
-		}
-		else
-		{
-			PublicDefinitions.Add("GPUPARTICLE_LOCAL_VF_ONLY=0");
-		}
-*/
 
 		PublicDefinitions.Add("GPUPARTICLE_LOCAL_VF_ONLY=0");
 
-		// Add a reference to the stats HTML files referenced by UEngine::DumpFPSChartToHTML. Previously staged by CopyBuildToStagingDirectory.
-    if (Target.bBuildEditor || Target.Configuration != UnrealTargetConfiguration.Shipping)
+        // Add a reference to the stats HTML files referenced by UEngine::DumpFPSChartToHTML. Previously staged by CopyBuildToStagingDirectory.
+        if (Target.bBuildEditor || Target.Configuration != UnrealTargetConfiguration.Shipping)
 		{
 			RuntimeDependencies.Add("$(EngineDir)/Content/Stats/...", StagedFileType.UFS);
 		}
diff --git a/Engine/Source/Runtime/Engine/Private/Materials/HLSLMaterialTranslator.h b/Engine/Source/Runtime/Engine/Private/Materials/HLSLMaterialTranslator.h
--- a/Engine/Source/Runtime/Engine/Private/Materials/HLSLMaterialTranslator.h
+++ b/Engine/Source/Runtime/Engine/Private/Materials/HLSLMaterialTranslator.h
@@ -197,8 +197,6 @@ protected:
 	uint32 bNeedsParticlePosition : 1;
 	/** true if the material needs particle velocity. */
 	uint32 bNeedsParticleVelocity : 1;
-	/** true of the material uses a particle dynamic parameter */
-	uint32 bNeedsParticleDynamicParameter : 1;
 	/** true if the material needs particle relative time. */
 	uint32 bNeedsParticleTime : 1;
 	/** true if the material uses particle motion blur. */
@@ -246,6 +244,7 @@ protected:
 	/** Tracks the number of texture coordinates used by the vertex shader in this material. */
 	uint32 NumUserVertexTexCoords;
 
+	uint32 NumParticleDynamicParameters;
 public: 
 
 	FHLSLMaterialTranslator(FMaterial* InMaterial,
@@ -270,7 +269,6 @@ public:
 	,	bUsesSceneDepth(false)
 	,	bNeedsParticlePosition(false)
 	,	bNeedsParticleVelocity(false)
-	,	bNeedsParticleDynamicParameter(false)
 	,	bNeedsParticleTime(false)
 	,	bUsesParticleMotionBlur(false)
 	,	bNeedsParticleRandom(false)
@@ -295,6 +293,7 @@ public:
 	,	bUsesEmissiveColor(0)
 	,	NumUserTexCoords(0)
 	,	NumUserVertexTexCoords(0)
+	,	NumParticleDynamicParameters(0)
 	{
 		FMemory::Memzero(SharedPixelProperties);
 
@@ -924,9 +923,10 @@ public:
 			OutEnvironment.SetDefine(TEXT("NEEDS_PARTICLE_VELOCITY"), 1);
 		}
 
-		if (bNeedsParticleDynamicParameter)
+		if (NumParticleDynamicParameters > 0)
 		{
 			OutEnvironment.SetDefine(TEXT("USE_DYNAMIC_PARAMETERS"), 1);
+			OutEnvironment.SetDefine(TEXT("NUM_DYNAMIC_PARAMETERS"), NumParticleDynamicParameters);
 		}
 
 		if (bNeedsParticleTime)
@@ -4755,20 +4755,21 @@ protected:
 		return TransformBase(SourceCoordBasis, DestCoordBasis, A, 1);
 	}
 
-	virtual int32 DynamicParameter(FLinearColor& DefaultValue) override
+	virtual int32 DynamicParameter(FLinearColor& DefaultValue, uint32 ParameterIndex=0) override
 	{
 		if (ShaderFrequency != SF_Vertex && ShaderFrequency != SF_Pixel && ShaderFrequency != SF_Compute)
 		{
 			return NonVertexOrPixelShaderExpressionError();
 		}
 
-		bNeedsParticleDynamicParameter = true;
+		NumParticleDynamicParameters = FMath::Max(NumParticleDynamicParameters, ParameterIndex + 1);
 
 		int32 Default = Constant4(DefaultValue.R, DefaultValue.G, DefaultValue.B, DefaultValue.A);
 		return AddInlinedCodeChunk(
 			MCT_Float4,
-			TEXT("GetDynamicParameter(Parameters.Particle, %s)"),
-			*GetParameterCode(Default)
+			TEXT("GetDynamicParameter(Parameters.Particle, %s, %u)"),
+			*GetParameterCode(Default),
+			ParameterIndex
 			);
 	}
 
diff --git a/Engine/Source/Runtime/Engine/Private/Materials/MaterialExpressions.cpp b/Engine/Source/Runtime/Engine/Private/Materials/MaterialExpressions.cpp
--- a/Engine/Source/Runtime/Engine/Private/Materials/MaterialExpressions.cpp
+++ b/Engine/Source/Runtime/Engine/Private/Materials/MaterialExpressions.cpp
@@ -7081,12 +7081,14 @@ UMaterialExpressionDynamicParameter::UMaterialExpressionDynamicParameter(const F
 
 	DefaultValue = FLinearColor::White;
 
+	
+	ParameterIndex = 0;
 }
 
 #if WITH_EDITOR
 int32 UMaterialExpressionDynamicParameter::Compile( FMaterialCompiler* Compiler, int32 OutputIndex)
 {
-	return Compiler->DynamicParameter(DefaultValue);
+	return Compiler->DynamicParameter(DefaultValue, ParameterIndex);
 }
 
 TArray<FExpressionOutput>& UMaterialExpressionDynamicParameter::GetOutputs()
@@ -7172,7 +7174,7 @@ void UMaterialExpressionDynamicParameter::UpdateDynamicParameterProperties()
 
 bool UMaterialExpressionDynamicParameter::CopyDynamicParameterProperties(const UMaterialExpressionDynamicParameter* FromParam)
 {
-	if (FromParam && (FromParam != this))
+	if (FromParam && (FromParam != this) && ParameterIndex == FromParam->ParameterIndex)
 	{
 		for (int32 NameIndex = 0; NameIndex < 4; NameIndex++)
 		{
diff --git a/Engine/Source/Runtime/Engine/Private/Particles/FXSystemPrivate.h b/Engine/Source/Runtime/Engine/Private/Particles/FXSystemPrivate.h
--- a/Engine/Source/Runtime/Engine/Private/Particles/FXSystemPrivate.h
+++ b/Engine/Source/Runtime/Engine/Private/Particles/FXSystemPrivate.h
@@ -264,6 +264,16 @@ private:
 
 private:
 
+	template<typename TVectorFieldUniformParametersType>
+	void SimulateGPUParticles_Internal(
+		FRHICommandListImmediate& RHICmdList,
+		EParticleSimulatePhase::Type Phase,
+		const FUniformBufferRHIParamRef ViewUniformBuffer,
+		const FGlobalDistanceFieldParameterData* GlobalDistanceFieldParameterData,
+		FTexture2DRHIParamRef SceneDepthTexture,
+		FTexture2DRHIParamRef GBufferATexture
+	);
+
 	/*-------------------------------------------------------------------------
 		GPU simulation state.
 	-------------------------------------------------------------------------*/
diff --git a/Engine/Source/Runtime/Engine/Private/Particles/ParticleSystemRender.cpp b/Engine/Source/Runtime/Engine/Private/Particles/ParticleSystemRender.cpp
--- a/Engine/Source/Runtime/Engine/Private/Particles/ParticleSystemRender.cpp
+++ b/Engine/Source/Runtime/Engine/Private/Particles/ParticleSystemRender.cpp
@@ -38,6 +38,7 @@
 #include "Particles/ParticleLODLevel.h"
 #include "Engine/StaticMesh.h"
 
+DECLARE_CYCLE_STAT(TEXT("ParticleSystemSceneProxy Create GT"), STAT_FParticleSystemSceneProxy_Create, STATGROUP_Particles);
 DECLARE_CYCLE_STAT(TEXT("ParticleSystemSceneProxy GetMeshElements RT"), STAT_FParticleSystemSceneProxy_GetMeshElements, STATGROUP_Particles);
 DECLARE_CYCLE_STAT(TEXT("DynamicSpriteEmitterData GetDynamicMeshElementsEmitter GetParticleOrderData RT"), STAT_FDynamicSpriteEmitterData_GetDynamicMeshElementsEmitter_GetParticleOrderData, STATGROUP_Particles);
 DECLARE_CYCLE_STAT(TEXT("DynamicSpriteEmitterData PerParticleWorkOrTasks RT"), STAT_FDynamicSpriteEmitterData_PerParticleWorkOrTasks, STATGROUP_Particles);
@@ -7221,6 +7222,7 @@ void FParticleSystemSceneProxy::GatherSimpleLights(const FSceneViewFamily& ViewF
 
 FPrimitiveSceneProxy* UParticleSystemComponent::CreateSceneProxy()
 {
+	SCOPE_CYCLE_COUNTER(STAT_FParticleSystemSceneProxy_Create);
 	FParticleSystemSceneProxy* NewProxy = NULL;
 
 	//@fixme EmitterInstances.Num() check should be here to avoid proxies for dead emitters but there are some edge cases where it happens for emitters that have just activated...
diff --git a/Engine/Source/Runtime/Engine/Private/Rendering/MultiSizeIndexContainer.cpp b/Engine/Source/Runtime/Engine/Private/Rendering/MultiSizeIndexContainer.cpp
--- a/Engine/Source/Runtime/Engine/Private/Rendering/MultiSizeIndexContainer.cpp
+++ b/Engine/Source/Runtime/Engine/Private/Rendering/MultiSizeIndexContainer.cpp
@@ -161,6 +161,15 @@ void FMultiSizeIndexContainer::Serialize(FArchive& Ar, bool bNeedsCPUAccess)
 	}
 
 	IndexBuffer->Serialize(Ar);
+
+	if (bNeedsCPUAccess && IndexBuffer)
+	{
+		UE_LOG(LogSkeletalMesh, Log, TEXT("Index buffer Needs CPU: %p Num: %d"), IndexBuffer, IndexBuffer->Num());
+	}
+	else
+	{
+		UE_LOG(LogSkeletalMesh, Log, TEXT("Index buffer No CPU: %p Num: %d"), IndexBuffer, IndexBuffer->Num());
+	}
 }
 
 #if WITH_EDITOR
diff --git a/Engine/Source/Runtime/Engine/Private/ShaderCompiler/ShaderCompiler.cpp b/Engine/Source/Runtime/Engine/Private/ShaderCompiler/ShaderCompiler.cpp
--- a/Engine/Source/Runtime/Engine/Private/ShaderCompiler/ShaderCompiler.cpp
+++ b/Engine/Source/Runtime/Engine/Private/ShaderCompiler/ShaderCompiler.cpp
@@ -2591,7 +2591,7 @@ static void GenerateUniformBufferStructMember(FString& Result, const FUniformBuf
 }
 
 /* Generates the instanced stereo hlsl code that's dependent on view uniform declarations. */
-static void GenerateInstancedStereoCode(FString& Result)
+ENGINE_API void GenerateInstancedStereoCode(FString& Result)
 {
 	// Find the InstancedView uniform buffer struct
 	const FUniformBufferStruct* InstancedView = nullptr;
diff --git a/Engine/Source/Runtime/Engine/Private/SkeletalMeshLODRenderData.cpp b/Engine/Source/Runtime/Engine/Private/SkeletalMeshLODRenderData.cpp
--- a/Engine/Source/Runtime/Engine/Private/SkeletalMeshLODRenderData.cpp
+++ b/Engine/Source/Runtime/Engine/Private/SkeletalMeshLODRenderData.cpp
@@ -605,13 +605,25 @@ void FSkeletalMeshLODRenderData::Serialize(FArchive& Ar, UObject* Owner, int32 I
 
 	// Skeletal mesh buffers are kept in CPU memory after initialization to support merging of skeletal meshes.
 	bool bKeepBuffersInCPUMemory = true;
+	bool bNeedsCPUAccess = true;
+	USkeletalMesh* SkelMeshOwner = nullptr;
+
 #if !WITH_EDITOR
 	static const auto CVar = IConsoleManager::Get().FindTConsoleVariableDataInt(TEXT("r.FreeSkeletalMeshBuffers"));
 	if (CVar)
 	{
 		bKeepBuffersInCPUMemory = !CVar->GetValueOnAnyThread();
+		bNeedsCPUAccess = bKeepBuffersInCPUMemory;
 	}
 #endif
+	if (!StripFlags.IsDataStrippedForServer())
+	{
+		SkelMeshOwner = CastChecked<USkeletalMesh>(Owner);
+
+		// set cpu skinning flag on the vertex buffer so that the resource arrays know if they need to be CPU accessible
+		bNeedsCPUAccess = bKeepBuffersInCPUMemory || SkelMeshOwner->GetResourceForRendering()->RequiresCPUSkinning(GMaxRHIFeatureLevel) ||
+			SkelMeshOwner->NeedCPUData(Idx);
+	}
 
 	if (StripFlags.IsDataStrippedForServer() || StripFlags.IsClassDataStripped(MinLodStripFlag))
 	{
@@ -628,21 +640,15 @@ void FSkeletalMeshLODRenderData::Serialize(FArchive& Ar, UObject* Owner, int32 I
 	{
 		Ar << RenderSections;
 
-		MultiSizeIndexContainer.Serialize(Ar, bKeepBuffersInCPUMemory);
+		MultiSizeIndexContainer.Serialize(Ar, bNeedsCPUAccess);
 
 		Ar << ActiveBoneIndices;
 	}
 
 	Ar << RequiredBones;
-
-
+	
 	if (!StripFlags.IsDataStrippedForServer() && !StripFlags.IsClassDataStripped(MinLodStripFlag))
 	{
-		USkeletalMesh* SkelMeshOwner = CastChecked<USkeletalMesh>(Owner);
-
-		// set cpu skinning flag on the vertex buffer so that the resource arrays know if they need to be CPU accessible
-		bool bNeedsCPUAccess = bKeepBuffersInCPUMemory || SkelMeshOwner->GetResourceForRendering()->RequiresCPUSkinning(GMaxRHIFeatureLevel);
-
 		if (Ar.IsLoading())
 		{
 			SkinWeightVertexBuffer.SetNeedsCPUAccess(bNeedsCPUAccess);
@@ -652,7 +658,7 @@ void FSkeletalMeshLODRenderData::Serialize(FArchive& Ar, UObject* Owner, int32 I
 		StaticVertexBuffers.StaticMeshVertexBuffer.Serialize(Ar, bNeedsCPUAccess);
 		Ar << SkinWeightVertexBuffer;
 
-		if (SkelMeshOwner->bHasVertexColors)
+		if (SkelMeshOwner && SkelMeshOwner->bHasVertexColors)
 		{
 			StaticVertexBuffers.ColorVertexBuffer.Serialize(Ar, bKeepBuffersInCPUMemory);
 		}
diff --git a/Engine/Source/Runtime/Engine/Private/SkeletalMeshSampling.cpp b/Engine/Source/Runtime/Engine/Private/SkeletalMeshSampling.cpp
--- a/Engine/Source/Runtime/Engine/Private/SkeletalMeshSampling.cpp
+++ b/Engine/Source/Runtime/Engine/Private/SkeletalMeshSampling.cpp
@@ -290,6 +290,10 @@ void FSkeletalMeshSamplingInfo::BuildRegions(USkeletalMesh* SkeletalMesh)
 		FSkeletalMeshSamplingRegion& Region = Regions[RegionIndex];
 		FSkeletalMeshSamplingRegionBuiltData& RegionBuiltData = BuiltData.RegionBuiltData[RegionIndex];
 		int32 LODIndex = Region.LODIndex;
+		if (LODIndex == INDEX_NONE)
+		{
+			LODIndex = SkeletalMesh->GetResourceForRendering()->LODRenderData.Num() - 1;
+		}
 		if (!SkeletalMesh->IsValidLODIndex(LODIndex))
 		{
 			continue;
diff --git a/Engine/Source/Runtime/Engine/Public/MaterialCompiler.h b/Engine/Source/Runtime/Engine/Public/MaterialCompiler.h
--- a/Engine/Source/Runtime/Engine/Public/MaterialCompiler.h
+++ b/Engine/Source/Runtime/Engine/Public/MaterialCompiler.h
@@ -247,7 +247,7 @@ public:
 	virtual int32 TransformVector(EMaterialCommonBasis SourceCoordBasis, EMaterialCommonBasis DestCoordBasis, int32 A) = 0;
 	virtual int32 TransformPosition(EMaterialCommonBasis SourceCoordBasis, EMaterialCommonBasis DestCoordBasis, int32 A) = 0;
 
-	virtual int32 DynamicParameter(FLinearColor& DefaultValue) = 0;
+	virtual int32 DynamicParameter(FLinearColor& DefaultValue, uint32 ParameterIndex = 0) = 0;
 	virtual int32 LightmapUVs() = 0;
 	virtual int32 PrecomputedAOMask()  = 0;
 
@@ -461,7 +461,7 @@ public:
 		return Compiler->TransformPosition(SourceCoordBasis, DestCoordBasis, A);
 	}
 
-	virtual int32 DynamicParameter(FLinearColor& DefaultValue) override { return Compiler->DynamicParameter(DefaultValue); }
+	virtual int32 DynamicParameter(FLinearColor& DefaultValue, uint32 ParameterIndex = 0) override { return Compiler->DynamicParameter(DefaultValue, ParameterIndex); }
 	virtual int32 LightmapUVs() override { return Compiler->LightmapUVs(); }
 	virtual int32 PrecomputedAOMask() override { return Compiler->PrecomputedAOMask(); }
 
diff --git a/Engine/Source/Runtime/Engine/Public/WeightedRandomSampler.h b/Engine/Source/Runtime/Engine/Public/WeightedRandomSampler.h
--- a/Engine/Source/Runtime/Engine/Public/WeightedRandomSampler.h
+++ b/Engine/Source/Runtime/Engine/Public/WeightedRandomSampler.h
@@ -47,6 +47,8 @@ public:
 	virtual void Initialize();
 
 	virtual void Serialize(FArchive& Ar);
+
+	int32 GetNumEntries() const { return Prob.Num(); }
 protected:
 	TArray<float> Prob;
 	TArray<int32> Alias;
diff --git a/Engine/Source/Runtime/Launch/Private/LaunchEngineLoop.cpp b/Engine/Source/Runtime/Launch/Private/LaunchEngineLoop.cpp
--- a/Engine/Source/Runtime/Launch/Private/LaunchEngineLoop.cpp
+++ b/Engine/Source/Runtime/Launch/Private/LaunchEngineLoop.cpp
@@ -205,7 +205,7 @@ class FFeedbackContext;
 #endif
 
 #ifndef RHI_COMMAND_LIST_DEBUG_TRACES
-    #define RHI_COMMAND_LIST_DEBUG_TRACES 0
+	#define RHI_COMMAND_LIST_DEBUG_TRACES 0
 #endif
 
 
@@ -1446,6 +1446,18 @@ int32 FEngineLoop::PreInit(const TCHAR* CmdLine)
 
 	if (FPlatformProcess::SupportsMultithreading())
 	{
+		int StackSize = 128;
+		bool bForceEditorStackSize = false;
+#if WITH_EDITOR
+		bForceEditorStackSize = true;
+#endif
+
+		if (bHasEditorToken || bForceEditorStackSize)
+		{
+			StackSize = 1000;
+		}
+
+
 		{
 			GThreadPool = FQueuedThreadPool::Allocate();
 			int32 NumThreadsInThreadPool = FPlatformMisc::NumberOfWorkerThreadsToSpawn();
@@ -1455,7 +1467,7 @@ int32 FEngineLoop::PreInit(const TCHAR* CmdLine)
 			{
 				NumThreadsInThreadPool = 1;
 			}
-			verify(GThreadPool->Create(NumThreadsInThreadPool, 128 * 1024, TPri_SlightlyBelowNormal));
+			verify(GThreadPool->Create(NumThreadsInThreadPool, StackSize * 1024, TPri_SlightlyBelowNormal));
 		}
 		{
 			GBackgroundPriorityThreadPool = FQueuedThreadPool::Allocate();
diff --git a/Engine/Source/Runtime/MovieScene/Private/MovieSceneSection.cpp b/Engine/Source/Runtime/MovieScene/Private/MovieSceneSection.cpp
--- a/Engine/Source/Runtime/MovieScene/Private/MovieSceneSection.cpp
+++ b/Engine/Source/Runtime/MovieScene/Private/MovieSceneSection.cpp
@@ -67,7 +67,7 @@ void UMovieSceneSection::Serialize(FArchive& Ar)
 	{
 		const FFrameRate LegacyFrameRate = GetLegacyConversionFrameRate();
 
-		if (bIsInfinite_DEPRECATED)
+		if (bIsInfinite_DEPRECATED && bSupportsInfiniteRange)
 		{
 			SectionRange = TRange<FFrameNumber>::All();
 		}
diff --git a/Engine/Source/Runtime/MovieSceneTracks/Public/Sections/MovieSceneVectorSection.h b/Engine/Source/Runtime/MovieSceneTracks/Public/Sections/MovieSceneVectorSection.h
--- a/Engine/Source/Runtime/MovieSceneTracks/Public/Sections/MovieSceneVectorSection.h
+++ b/Engine/Source/Runtime/MovieSceneTracks/Public/Sections/MovieSceneVectorSection.h
@@ -129,7 +129,7 @@ protected:
 	virtual void Serialize(FArchive& Ar) override;
 	virtual TSharedPtr<FStructOnScope> GetKeyStruct(TArrayView<const FKeyHandle> KeyHandles) override;
 
-	void RecreateChannelProxy();
+	MOVIESCENETRACKS_API void RecreateChannelProxy();
 
 private:
 
diff --git a/Engine/Source/Runtime/OpenGLDrv/Private/OpenGLTexture.cpp b/Engine/Source/Runtime/OpenGLDrv/Private/OpenGLTexture.cpp
--- a/Engine/Source/Runtime/OpenGLDrv/Private/OpenGLTexture.cpp
+++ b/Engine/Source/Runtime/OpenGLDrv/Private/OpenGLTexture.cpp
@@ -2359,18 +2359,21 @@ void FOpenGLDynamicRHI::RHIUpdateTexture3D(FTexture3DRHIParamRef TextureRHI,uint
 
 void FOpenGLDynamicRHI::InvalidateTextureResourceInCache(GLuint Resource)
 {
-	for (int32 SamplerIndex = 0; SamplerIndex < FOpenGL::GetMaxCombinedTextureImageUnits(); ++SamplerIndex)
+	if (SharedContextState.Textures || RenderingContextState.Textures)
 	{
-		if (SharedContextState.Textures[SamplerIndex].Resource == Resource)
+		for (int32 SamplerIndex = 0; SamplerIndex < FOpenGL::GetMaxCombinedTextureImageUnits(); ++SamplerIndex)
 		{
-			SharedContextState.Textures[SamplerIndex].Target = GL_NONE;
-			SharedContextState.Textures[SamplerIndex].Resource = 0;
-		}
+			if (SharedContextState.Textures && SharedContextState.Textures[SamplerIndex].Resource == Resource)
+			{
+				SharedContextState.Textures[SamplerIndex].Target = GL_NONE;
+				SharedContextState.Textures[SamplerIndex].Resource = 0;
+			}
 
-		if (RenderingContextState.Textures[SamplerIndex].Resource == Resource)
-		{
-			RenderingContextState.Textures[SamplerIndex].Target = GL_NONE;
-			RenderingContextState.Textures[SamplerIndex].Resource = 0;
+			if (RenderingContextState.Textures && RenderingContextState.Textures[SamplerIndex].Resource == Resource)
+			{
+				RenderingContextState.Textures[SamplerIndex].Target = GL_NONE;
+				RenderingContextState.Textures[SamplerIndex].Resource = 0;
+			}
 		}
 	}
 	
diff --git a/Engine/Source/Runtime/RHI/Public/DynamicRHI.h b/Engine/Source/Runtime/RHI/Public/DynamicRHI.h
--- a/Engine/Source/Runtime/RHI/Public/DynamicRHI.h
+++ b/Engine/Source/Runtime/RHI/Public/DynamicRHI.h
@@ -202,6 +202,16 @@ public:
 		return new FRHIComputeFence(Name);
 	}
 
+	virtual FGPUFenceRHIRef RHICreateGPUFence(const FName &Name)
+	{
+		return new FRHIGPUFence(Name);
+	}
+
+	virtual FStagingBufferRHIRef RHICreateStagingBuffer()
+	{
+		return new FRHIStagingBuffer();
+	}
+
 	/**
 	* Creates a bound shader state instance which encapsulates a decl, vertex shader, hull shader, domain shader and pixel shader
 	* CAUTION: Even though this is marked as threadsafe, it is only valid to call from the render thread or the RHI thread. It need not be threadsafe unless the RHI support parallel translation.
diff --git a/Engine/Source/Runtime/RHI/Public/RHICommandList.h b/Engine/Source/Runtime/RHI/Public/RHICommandList.h
--- a/Engine/Source/Runtime/RHI/Public/RHICommandList.h
+++ b/Engine/Source/Runtime/RHI/Public/RHICommandList.h
@@ -3222,7 +3222,17 @@ public:
 	{		
 		return GDynamicRHI->RHICreateComputeFence(Name);
 	}	
-	
+
+	FORCEINLINE FGPUFenceRHIRef CreateGPUFence(const FName& Name)
+	{
+		return GDynamicRHI->RHICreateGPUFence(Name);
+	}
+
+	FORCEINLINE FStagingBufferRHIRef CreateStagingBuffer()
+	{
+		return GDynamicRHI->RHICreateStagingBuffer();
+	}
+
 	FORCEINLINE FBoundShaderStateRHIRef CreateBoundShaderState(FVertexDeclarationRHIParamRef VertexDeclaration, FVertexShaderRHIParamRef VertexShader, FHullShaderRHIParamRef HullShader, FDomainShaderRHIParamRef DomainShader, FPixelShaderRHIParamRef PixelShader, FGeometryShaderRHIParamRef GeometryShader)
 	{
 		LLM_SCOPE(ELLMTag::Shaders);
@@ -4132,6 +4142,19 @@ FORCEINLINE FComputeFenceRHIRef RHICreateComputeFence(const FName& Name)
 	return FRHICommandListExecutor::GetImmediateCommandList().CreateComputeFence(Name);
 }
 
+FORCEINLINE FGPUFenceRHIRef RHICreateGPUFence(const FName& Name)
+{
+	return FRHICommandListExecutor::GetImmediateCommandList().CreateGPUFence(Name);
+}
+
+
+FORCEINLINE FStagingBufferRHIRef RHICreateStagingBuffer()
+{
+	return FRHICommandListExecutor::GetImmediateCommandList().CreateStagingBuffer();
+}
+
+
+
 FORCEINLINE FIndexBufferRHIRef RHICreateAndLockIndexBuffer(uint32 Stride, uint32 Size, uint32 InUsage, FRHIResourceCreateInfo& CreateInfo, void*& OutDataBuffer)
 {
 	return FRHICommandListExecutor::GetImmediateCommandList().CreateAndLockIndexBuffer(Stride, Size, InUsage, CreateInfo, OutDataBuffer);
diff --git a/Engine/Source/Runtime/RHI/Public/RHIGPUReadback.h b/Engine/Source/Runtime/RHI/Public/RHIGPUReadback.h
new file mode 100644
--- /dev/null
+++ b/Engine/Source/Runtime/RHI/Public/RHIGPUReadback.h
@@ -0,0 +1,121 @@
+// Copyright 1998-2018 Epic Games, Inc. All Rights Reserved.
+
+/*=============================================================================
+  RHIGPUReadback.h: classes for managing fences and staging buffers for
+  asynchronous GPU memory updates and readbacks with minimal stalls and no
+  RHI thread flushes
+=============================================================================*/
+
+#pragma once
+
+#include "CoreMinimal.h"
+#include "RHI.h"
+
+
+/* FRHIGPUMemoryUpdate: represents a memory update request scheduled with RHIScheduleGPUMemoryUpdate
+ *
+ */
+class FRHIGPUMemoryUpdate
+{
+public:
+	FRHIGPUMemoryUpdate(FRHIStagingBuffer *Staging, FVertexBufferRHIRef InBuffer, FName RequestName)
+	{
+		GPUVertexBuffer = InBuffer;
+		StagingBuffer = Staging;
+		Fence = RHICreateGPUFence(RequestName);
+		Fence->Write();
+	}
+
+	FRHIGPUMemoryUpdate(FVertexBufferRHIRef InBuffer, FName RequestName)
+	{
+		GPUVertexBuffer = InBuffer;
+		StagingBuffer = RHICreateStagingBuffer();
+		Fence = RHICreateGPUFence(RequestName);
+		Fence->Write();
+	}
+
+	bool WaitForFence(float Timeout) 
+	{ 
+		return Fence->Wait(Timeout); 
+	}
+	bool IsReady() 
+	{ 
+		return Fence->Poll(); 
+	}
+
+	void UpdateBuffer(void *InPtr, int32 NumBytes) 
+	{ 
+		check(Fence->Poll());  
+	}
+	
+	void Finish() 
+	{ 
+	}
+
+private:
+
+	FRHIStagingBuffer *StagingBuffer;
+	FVertexBufferRHIRef  GPUVertexBuffer;
+	TRefCountPtr<FRHIGPUFence> Fence;
+};
+
+
+
+/* FRHIGPUMemoryReadback: represents a memory readback request scheduled with RHIScheduleGPUMemoryUpdate
+*
+*/
+class FRHIGPUMemoryReadback
+{
+public:
+
+	FRHIGPUMemoryReadback(FStagingBufferRHIRef Staging, FVertexBufferRHIRef InBuffer, FName RequestName)
+	{
+		GPUVertexBuffer = InBuffer;
+		StagingBuffer = Staging;
+		Fence = RHICreateGPUFence(RequestName);
+		Fence->Write();
+	}
+
+	FRHIGPUMemoryReadback(FVertexBufferRHIRef InBuffer, FName RequestName)
+	{
+		GPUVertexBuffer = InBuffer;
+		StagingBuffer = RHICreateStagingBuffer();
+		Fence = RHICreateGPUFence(RequestName);
+		Fence->Write();
+	}
+
+	~FRHIGPUMemoryReadback()
+	{
+	}
+
+	bool WaitForFence(float Timeout)
+	{
+		return Fence->Wait(Timeout);
+	}
+
+	bool IsReady()
+	{
+		return Fence->Poll();
+	}
+
+	void *RetrieveData(uint32 NumBytes)
+	{
+		ensure(Fence->Poll());
+		return StagingBuffer->Lock(GPUVertexBuffer, 0, NumBytes, EResourceLockMode::RLM_ReadOnly);
+	}
+
+	void Finish()
+	{
+		StagingBuffer->Unlock();
+	}
+
+private:
+
+	FStagingBufferRHIRef StagingBuffer;
+	FVertexBufferRHIRef  GPUVertexBuffer;
+	FGPUFenceRHIRef Fence;
+};
+
+
+
+
diff --git a/Engine/Source/Runtime/RHI/Public/RHIResources.h b/Engine/Source/Runtime/RHI/Public/RHIResources.h
--- a/Engine/Source/Runtime/RHI/Public/RHIResources.h
+++ b/Engine/Source/Runtime/RHI/Public/RHIResources.h
@@ -729,6 +729,41 @@ public:
 // Misc
 //
 
+
+
+/* Generic GPU fence class used by FRHIGPUMemoryReadback and FRHIGPUMemoryUpdate
+* RHI specific fences derive from this
+*/
+class FRHIGPUFence : public FRHIResource
+{
+public:
+	FRHIGPUFence(FName InName)
+		: FenceName(InName)
+	{}
+
+	virtual ~FRHIGPUFence()
+	{}
+
+	virtual bool Write()
+	{
+		return false;
+	};
+
+	virtual bool Poll() const
+	{
+		return false;
+	};
+
+	virtual bool Wait(float TimeoutMs) const
+	{
+		return false;
+	};
+
+private:
+	FName FenceName;
+};
+
+
 class FRHIRenderQuery : public FRHIResource {};
 
 class FRHIComputeFence : public FRHIResource
@@ -896,6 +931,9 @@ typedef TRefCountPtr<FRHITextureReference> FTextureReferenceRHIRef;
 typedef FRHIRenderQuery*              FRenderQueryRHIParamRef;
 typedef TRefCountPtr<FRHIRenderQuery> FRenderQueryRHIRef;
 
+typedef FRHIGPUFence*				FGPUFenceRHIParamRef;
+typedef TRefCountPtr<FRHIGPUFence>	FGPUFenceRHIRef;
+
 typedef FRHIViewport*              FViewportRHIParamRef;
 typedef TRefCountPtr<FRHIViewport> FViewportRHIRef;
 
@@ -909,6 +947,38 @@ typedef FRHIGraphicsPipelineState*              FGraphicsPipelineStateRHIParamRe
 typedef TRefCountPtr<FRHIGraphicsPipelineState> FGraphicsPipelineStateRHIRef;
 
 
+/* Generic staging buffer class used by FRHIGPUMemoryReadback and FRHIGPUMemoryUpdate
+* RHI specific staging buffers derive from this
+*/
+class FRHIStagingBuffer : public FRHIResource
+{
+public:
+	FRHIStagingBuffer()
+		:MappedPtr(nullptr)
+		, LastLockedBuffer(nullptr)
+	{
+	}
+
+	virtual void *Lock(FVertexBufferRHIRef GPUBuffer, uint32 Offset, uint32 NumBytes, EResourceLockMode LockMode)   // copyresource, map, return ptr
+	{
+		return MappedPtr;
+	}
+
+	virtual void Unlock()	// unmap, free memory
+	{
+	}
+
+
+protected:
+	// pointer to mapped buffer; null if unmapped
+	void *MappedPtr;
+	FVertexBufferRHIParamRef LastLockedBuffer;
+};
+
+typedef FRHIStagingBuffer*				FStagingBufferRHIParamRef;
+typedef TRefCountPtr<FRHIStagingBuffer>	FStagingBufferRHIRef;
+
+
 class FRHIRenderTargetView
 {
 public:
diff --git a/Engine/Source/Runtime/RHI/Public/RHIUtilities.h b/Engine/Source/Runtime/RHI/Public/RHIUtilities.h
--- a/Engine/Source/Runtime/RHI/Public/RHIUtilities.h
+++ b/Engine/Source/Runtime/RHI/Public/RHIUtilities.h
@@ -59,8 +59,14 @@ struct FRWBuffer
 	}
 
 	// @param AdditionalUsage passed down to RHICreateVertexBuffer(), get combined with "BUF_UnorderedAccess | BUF_ShaderResource" e.g. BUF_Static
-	void Initialize(uint32 BytesPerElement, uint32 NumElements, EPixelFormat Format, uint32 AdditionalUsage = 0,  const TCHAR* InDebugName = NULL, FResourceArrayInterface *InResourceArray = nullptr)	{
-		check(GMaxRHIFeatureLevel == ERHIFeatureLevel::SM5);
+	void Initialize(uint32 BytesPerElement, uint32 NumElements, EPixelFormat Format, uint32 AdditionalUsage = 0,  const TCHAR* InDebugName = NULL, FResourceArrayInterface *InResourceArray = nullptr)	
+	{
+		check( GMaxRHIFeatureLevel == ERHIFeatureLevel::SM5 
+			|| IsVulkanPlatform(GMaxRHIShaderPlatform) 
+			|| IsMetalPlatform(GMaxRHIShaderPlatform)
+			|| (GMaxRHIFeatureLevel == ERHIFeatureLevel::ES3_1 && GSupportsResourceView)
+		);
+
 		// Provide a debug name if using Fast VRAM so the allocators diagnostics will work
 		ensure(!((AdditionalUsage & BUF_FastVRAM) && !InDebugName));
 		NumBytes = BytesPerElement * NumElements;
@@ -108,7 +114,7 @@ struct FReadBuffer
 
 	void Initialize(uint32 BytesPerElement, uint32 NumElements, EPixelFormat Format, uint32 AdditionalUsage = 0)
 	{
-		check(GMaxRHIFeatureLevel >= ERHIFeatureLevel::SM4);
+		check(GSupportsResourceView);
 		NumBytes = BytesPerElement * NumElements;
 		FRHIResourceCreateInfo CreateInfo;
 		Buffer = RHICreateVertexBuffer(NumBytes, BUF_ShaderResource | AdditionalUsage, CreateInfo);
@@ -1002,4 +1008,33 @@ struct FRHILockTracker
 	}
 };
 
-extern RHI_API FRHILockTracker GRHILockTracker;
\ No newline at end of file
+extern RHI_API FRHILockTracker GRHILockTracker;
+
+
+
+
+/* Generic implementation of functions for readback and update requests
+* Both of these will need RHI specific implementations, of course
+*/
+
+
+/* Generic implementation of functions for readback and update requests
+* Both of these will need RHI specific implementations, of course
+*/
+/*
+FRHIGPUMemoryReadback *RHIScheduleGPUMemoryReadback(FRHICommandList& CmdList, FVertexBufferRHIRef GPUBuffer, FName RequestName)
+{
+	FRHIStagingBuffer *StagingBuffer = RHICreateStagingBuffer();// new FRHIStagingBuffer();
+	FRHIGPUMemoryReadback *Readback = new FRHIGPUMemoryReadback(StagingBuffer, GPUBuffer, RequestName);
+
+	return Readback;
+}
+
+FRHIGPUMemoryUpdate *RHIScheduleGPUMemoryUpdate(FRHICommandList& CmdList, FVertexBufferRHIRef GPUBuffer, FName RequestName)
+{
+	FRHIStagingBuffer *StagingBuffer = RHICreateStagingBuffer();// new FRHIStagingBuffer();
+	FRHIGPUMemoryUpdate *Update = new FRHIGPUMemoryUpdate(StagingBuffer, GPUBuffer, RequestName);
+
+	return Update;
+}
+*/
\ No newline at end of file
diff --git a/Engine/Source/Runtime/RenderCore/Public/RendererInterface.h b/Engine/Source/Runtime/RenderCore/Public/RendererInterface.h
--- a/Engine/Source/Runtime/RenderCore/Public/RendererInterface.h
+++ b/Engine/Source/Runtime/RenderCore/Public/RendererInterface.h
@@ -542,8 +542,10 @@ class FPostOpaqueRenderParameters
 		FMatrix ViewMatrix;
 		FMatrix ProjMatrix;
 		FRHITexture2D* DepthTexture;
+		FRHITexture2D* NormalTexture;
 		FRHITexture2D* SmallDepthTexture;
 		FRHICommandListImmediate* RHICmdList;
+		FUniformBufferRHIParamRef ViewUniformBuffer;
 		void* Uid; // A unique identifier for the view.
 };
 DECLARE_DELEGATE_OneParam(FPostOpaqueRenderDelegate, class FPostOpaqueRenderParameters&);
@@ -552,7 +554,7 @@ DECLARE_DELEGATE_OneParam(FPostOpaqueRenderDelegate, class FPostOpaqueRenderPara
 class FComputeDispatcher
 {
 public:
-	virtual void Execute(FRHICommandList &RHICmdList) = 0;
+	virtual void Execute(FRHICommandList &RHICmdList, FUniformBufferRHIParamRef ViewUniformBuffer) = 0;
 };
 
 
@@ -732,10 +734,13 @@ public:
 
 	virtual void RegisterPostOpaqueRenderDelegate(const FPostOpaqueRenderDelegate& PostOpaqueRenderDelegate) = 0;
 	virtual void RegisterOverlayRenderDelegate(const FPostOpaqueRenderDelegate& OverlayRenderDelegate) = 0;
+	virtual void RenderPostOpaqueExtensions(const class FViewInfo& View, FRHICommandListImmediate& RHICmdList, class FSceneRenderTargets& SceneContext) = 0;
+	virtual void RenderOverlayExtensions(const class FViewInfo& View, FRHICommandListImmediate& RHICmdList, FSceneRenderTargets& SceneContext) = 0;
+	virtual bool HasPostOpaqueExtentions() const = 0;
 
 	virtual void RegisterPostOpaqueComputeDispatcher(FComputeDispatcher *Dispatcher) = 0;
 	virtual void UnRegisterPostOpaqueComputeDispatcher(FComputeDispatcher *Dispatcher) = 0;
-	virtual void DispatchPostOpaqueCompute(FRHICommandList &CmdList) = 0;
+	virtual void DispatchPostOpaqueCompute(FRHICommandList &CmdList, FUniformBufferRHIParamRef ViewUniformBuffer) = 0;
 
 	/** Delegate that is called upon resolving scene color. */
 	DECLARE_MULTICAST_DELEGATE_TwoParams(FOnResolvedSceneColor, FRHICommandListImmediate& /*RHICmdList*/, class FSceneRenderTargets& /*SceneContext*/);
diff --git a/Engine/Source/Runtime/Renderer/Private/BasePassRendering.inl b/Engine/Source/Runtime/Renderer/Private/BasePassRendering.inl
--- a/Engine/Source/Runtime/Renderer/Private/BasePassRendering.inl
+++ b/Engine/Source/Runtime/Renderer/Private/BasePassRendering.inl
@@ -32,32 +32,35 @@ inline void TBasePassVertexShaderPolicyParamType<VertexParametersType>::SetMesh(
 	const bool bHasSkipOutputVelocityParameter = SkipOutputVelocityParameter.IsBound();
     
     float SkipOutputVelocityValue = 1.0f;
+	FMatrix PreviousLocalToWorldMatrix;
+
+	if (bHasPreviousLocalToWorldParameter || bHasSkipOutputVelocityParameter)
+	{
+		if (Proxy)
+		{
+			bool bHasPreviousLocalToWorldMatrix = false;
+			const auto& ViewInfo = (const FViewInfo&)View;
+
+			if (FVelocityDrawingPolicy::HasVelocityOnBasePass(ViewInfo, Proxy, Proxy->GetPrimitiveSceneInfo(), Mesh,
+				bHasPreviousLocalToWorldMatrix, PreviousLocalToWorldMatrix))
+			{
+				PreviousLocalToWorldMatrix = bHasPreviousLocalToWorldMatrix ? PreviousLocalToWorldMatrix : Proxy->GetLocalToWorld();
+
+				SkipOutputVelocityValue = 0.0f;
+			}
+			else
+			{
+				PreviousLocalToWorldMatrix.SetIdentity();
+			}
+		}
+		else
+		{
+			PreviousLocalToWorldMatrix.SetIdentity();
+		}
+	}
+
     if (bHasPreviousLocalToWorldParameter)
     {
-        FMatrix PreviousLocalToWorldMatrix;
-        
-        if (Proxy)
-        {
-            bool bHasPreviousLocalToWorldMatrix = false;
-            const auto& ViewInfo = (const FViewInfo&)View;
-            
-            if (FVelocityDrawingPolicy::HasVelocityOnBasePass(ViewInfo, Proxy, Proxy->GetPrimitiveSceneInfo(), Mesh,
-                                                              bHasPreviousLocalToWorldMatrix, PreviousLocalToWorldMatrix))
-            {
-                PreviousLocalToWorldMatrix = bHasPreviousLocalToWorldMatrix ? PreviousLocalToWorldMatrix : Proxy->GetLocalToWorld();
-                
-                SkipOutputVelocityValue = 0.0f;
-            }
-            else
-            {
-                PreviousLocalToWorldMatrix.SetIdentity();
-            }
-        }
-        else
-        {
-            PreviousLocalToWorldMatrix.SetIdentity();
-        }
-        
         SetShaderValue(RHICmdList, VertexShaderRHI, PreviousLocalToWorldParameter, PreviousLocalToWorldMatrix);
     }
     
diff --git a/Engine/Source/Runtime/Renderer/Private/DeferredShadingRenderer.cpp b/Engine/Source/Runtime/Renderer/Private/DeferredShadingRenderer.cpp
--- a/Engine/Source/Runtime/Renderer/Private/DeferredShadingRenderer.cpp
+++ b/Engine/Source/Runtime/Renderer/Private/DeferredShadingRenderer.cpp
@@ -710,6 +710,7 @@ void FDeferredShadingSceneRenderer::Render(FRHICommandListImmediate& RHICmdList)
 	// Dynamic vertex and index buffers need to be committed before rendering.
 	if (!bDoInitViewAftersPrepass)
 	{
+		GEngine->GetPreRenderDelegate().Broadcast();
 		{
 			SCOPE_CYCLE_COUNTER(STAT_FDeferredShadingSceneRenderer_FGlobalDynamicVertexBuffer_Commit);
 			FGlobalDynamicVertexBuffer::Get().Commit();
@@ -745,6 +746,8 @@ void FDeferredShadingSceneRenderer::Render(FRHICommandListImmediate& RHICmdList)
 				InitViewsPossiblyAfterPrepass(RHICmdList, ILCTaskData, SortEvents, UpdateViewCustomDataEvents);
 				PostInitViewCustomData(UpdateViewCustomDataEvents);
 
+				GEngine->GetPreRenderDelegate().Broadcast();
+
 				{
 					SCOPE_CYCLE_COUNTER(STAT_FDeferredShadingSceneRenderer_FGlobalDynamicVertexBuffer_Commit);
 					FGlobalDynamicVertexBuffer::Get().Commit();
@@ -1073,6 +1076,22 @@ void FDeferredShadingSceneRenderer::Render(FRHICommandListImmediate& RHICmdList)
 		ServiceLocalQueue();
 	}
 
+
+	IRendererModule& RendererModule = GetRendererModule();
+	if (RendererModule.HasPostOpaqueExtentions())
+	{
+		SceneContext.BeginRenderingSceneColor(RHICmdList);
+		for (int32 ViewIndex = 0; ViewIndex < Views.Num(); ++ViewIndex)
+		{
+			const FViewInfo& View = Views[ViewIndex];
+			RHICmdList.SetViewport(View.ViewRect.Min.X, View.ViewRect.Min.Y, 0.0f, View.ViewRect.Max.X, View.ViewRect.Max.Y, 1.0f);
+			RendererModule.RenderPostOpaqueExtensions(View, RHICmdList, SceneContext);
+		}
+	}
+	SetRenderTarget(RHICmdList, nullptr, 0, 0, nullptr);
+	RendererModule.DispatchPostOpaqueCompute(RHICmdList, Views[0].ViewUniformBuffer);
+
+
 	TRefCountPtr<IPooledRenderTarget> VelocityRT;
 
 	if (bUseVelocityGBuffer)
@@ -1252,19 +1271,6 @@ void FDeferredShadingSceneRenderer::Render(FRHICommandListImmediate& RHICmdList)
 		ServiceLocalQueue();
 	}
 
-	FRendererModule* RendererModule = FRendererModule::GetRendererModule();
-	if (RendererModule->HasPostOpaqueExtentions())
-	{
-		SceneContext.BeginRenderingSceneColor(RHICmdList);
-		for (int32 ViewIndex = 0; ViewIndex < Views.Num(); ++ViewIndex)
-		{
-			const FViewInfo& View = Views[ViewIndex];
-			RHICmdList.SetViewport(View.ViewRect.Min.X, View.ViewRect.Min.Y, 0.0f, View.ViewRect.Max.X, View.ViewRect.Max.Y, 1.0f);
-			RendererModule->RenderPostOpaqueExtensions(View, RHICmdList, SceneContext);
-		}
-	}
-
-	RendererModule->DispatchPostOpaqueCompute(RHICmdList);
 
 	// No longer needed, release
 	LightShaftOutput.LightShaftOcclusion = NULL;
@@ -1322,7 +1328,7 @@ void FDeferredShadingSceneRenderer::Render(FRHICommandListImmediate& RHICmdList)
 	{
 		const FViewInfo& View = Views[ViewIndex];
 		RHICmdList.SetViewport(View.ViewRect.Min.X, View.ViewRect.Min.Y, 0.0f, View.ViewRect.Max.X, View.ViewRect.Max.Y, 1.0f);
-		RendererModule->RenderOverlayExtensions(View, RHICmdList, SceneContext);
+		RendererModule.RenderOverlayExtensions(View, RHICmdList, SceneContext);
 	}
 
 	if (ViewFamily.EngineShowFlags.VisualizeDistanceFieldAO || ViewFamily.EngineShowFlags.VisualizeDistanceFieldGI)
@@ -1361,7 +1367,7 @@ void FDeferredShadingSceneRenderer::Render(FRHICommandListImmediate& RHICmdList)
 	if (ViewFamily.bResolveScene)
 	{
 		SCOPED_DRAW_EVENT(RHICmdList, PostProcessing);
-   		SCOPED_GPU_STAT(RHICmdList, Postprocessing);
+		SCOPED_GPU_STAT(RHICmdList, Postprocessing);
 
 		SCOPE_CYCLE_COUNTER(STAT_FinishRenderViewTargetTime);
 
diff --git a/Engine/Source/Runtime/Renderer/Private/MobileShadingRenderer.cpp b/Engine/Source/Runtime/Renderer/Private/MobileShadingRenderer.cpp
--- a/Engine/Source/Runtime/Renderer/Private/MobileShadingRenderer.cpp
+++ b/Engine/Source/Runtime/Renderer/Private/MobileShadingRenderer.cpp
@@ -185,6 +185,8 @@ void FMobileSceneRenderer::Render(FRHICommandListImmediate& RHICmdList)
 		FRHICommandListExecutor::GetImmediateCommandList().ImmediateFlush(EImmediateFlushType::FlushRHIThreadFlushResources);
 	}
 
+	GEngine->GetPreRenderDelegate().Broadcast();
+
 	// Dynamic vertex and index buffers need to be committed before rendering.
 	FGlobalDynamicVertexBuffer::Get().Commit();
 	FGlobalDynamicIndexBuffer::Get().Commit();
diff --git a/Engine/Source/Runtime/Renderer/Private/RendererModule.h b/Engine/Source/Runtime/Renderer/Private/RendererModule.h
--- a/Engine/Source/Runtime/Renderer/Private/RendererModule.h
+++ b/Engine/Source/Runtime/Renderer/Private/RendererModule.h
@@ -78,14 +78,14 @@ public:
 
 	virtual void RegisterPostOpaqueComputeDispatcher(FComputeDispatcher *Dispatcher) override;
 	virtual void UnRegisterPostOpaqueComputeDispatcher(FComputeDispatcher *Dispatcher) override;
-	virtual void DispatchPostOpaqueCompute(FRHICommandList &CmdList) override;
+	virtual void DispatchPostOpaqueCompute(FRHICommandList &CmdList, FUniformBufferRHIParamRef ViewUniformBuffer) override;
 
 	virtual void RegisterPostOpaqueRenderDelegate(const FPostOpaqueRenderDelegate& PostOpaqueRenderDelegate) override;
 	virtual void RegisterOverlayRenderDelegate(const FPostOpaqueRenderDelegate& OverlayRenderDelegate) override;
-	void RenderPostOpaqueExtensions(const FViewInfo& View, FRHICommandListImmediate& RHICmdList, FSceneRenderTargets& SceneContext);
-	void RenderOverlayExtensions(const FViewInfo& View, FRHICommandListImmediate& RHICmdList, FSceneRenderTargets& SceneContext);
+	virtual void RenderPostOpaqueExtensions(const FViewInfo& View, FRHICommandListImmediate& RHICmdList, FSceneRenderTargets& SceneContext) override;
+	virtual void RenderOverlayExtensions(const FViewInfo& View, FRHICommandListImmediate& RHICmdList, FSceneRenderTargets& SceneContext) override;
 
-	bool HasPostOpaqueExtentions() const
+	virtual bool HasPostOpaqueExtentions() const override
 	{
 		return PostOpaqueRenderDelegate.IsBound();
 	}
@@ -99,8 +99,6 @@ public:
 
 	virtual void PostRenderAllViewports() override;
 
-	static FRendererModule* GetRendererModule();
-
 private:
 	TSet<FSceneInterface*> AllocatedScenes;
 	ICustomCulling* CustomCullingImpl;
diff --git a/Engine/Source/Runtime/Renderer/Private/SceneRendering.cpp b/Engine/Source/Runtime/Renderer/Private/SceneRendering.cpp
--- a/Engine/Source/Runtime/Renderer/Private/SceneRendering.cpp
+++ b/Engine/Source/Runtime/Renderer/Private/SceneRendering.cpp
@@ -1391,6 +1391,8 @@ void FViewInfo::InitRHIResources()
 	{
 		DynamicResources[ResourceIndex]->InitPrimitiveResource();
 	}
+
+
 }
 
 // These are not real view infos, just dumb memory blocks
@@ -2971,8 +2973,6 @@ void OnChangeCVarRequiringRecreateRenderState(IConsoleVariable* Var)
 	FGlobalComponentRecreateRenderStateContext Context;
 }
 
-FRendererModule* RendererModule;
-
 FRendererModule::FRendererModule()
 	: CustomCullingImpl(nullptr)
 {
@@ -2984,14 +2984,6 @@ FRendererModule::FRendererModule()
 
 	static auto CVarEarlyZPassMovable = IConsoleManager::Get().FindConsoleVariable(TEXT("r.EarlyZPassMovable"));
 	CVarEarlyZPassMovable->SetOnChangedCallback(FConsoleVariableDelegate::CreateStatic(&OnChangeCVarRequiringRecreateRenderState));
-
-	RendererModule = this;
-}
-
-// static
-FRendererModule* FRendererModule::GetRendererModule()
-{
-	return RendererModule;
 }
 
 void FRendererModule::CreateAndInitSingleView(FRHICommandListImmediate& RHICmdList, class FSceneViewFamily* ViewFamily, const struct FSceneViewInitOptions* ViewInitOptions)
@@ -3161,11 +3153,11 @@ void FRendererModule::UnRegisterPostOpaqueComputeDispatcher(FComputeDispatcher *
 	PostOpaqueDispatchers.Remove(Dispatcher);
 }
 
-void FRendererModule::DispatchPostOpaqueCompute(FRHICommandList &RHICmdList)
+void FRendererModule::DispatchPostOpaqueCompute(FRHICommandList &RHICmdList, FUniformBufferRHIParamRef ViewUniformBuffer)
 {
 	for (FComputeDispatcher *Dispatcher : PostOpaqueDispatchers)
 	{
-		Dispatcher->Execute(RHICmdList);
+		Dispatcher->Execute(RHICmdList, ViewUniformBuffer);
 	}
 }
 
@@ -3186,7 +3178,9 @@ void FRendererModule::RenderPostOpaqueExtensions(const FViewInfo& View, FRHIComm
 	RenderParameters.ViewMatrix = View.ViewMatrices.GetViewMatrix();
 	RenderParameters.ProjMatrix = View.ViewMatrices.GetProjectionMatrix();
 	RenderParameters.DepthTexture = SceneContext.GetSceneDepthSurface()->GetTexture2D();
+	RenderParameters.NormalTexture = SceneContext.GBufferA.IsValid() ? SceneContext.GetGBufferATexture() : nullptr;
 	RenderParameters.SmallDepthTexture = SceneContext.GetSmallDepthSurface()->GetTexture2D();
+	RenderParameters.ViewUniformBuffer = View.ViewUniformBuffer;
 
 	RenderParameters.ViewportRect = View.ViewRect;
 	RenderParameters.RHICmdList = &RHICmdList;
diff --git a/Engine/Source/Runtime/ShaderCore/Private/ShaderParameters.cpp b/Engine/Source/Runtime/ShaderCore/Private/ShaderParameters.cpp
--- a/Engine/Source/Runtime/ShaderCore/Private/ShaderParameters.cpp
+++ b/Engine/Source/Runtime/ShaderCore/Private/ShaderParameters.cpp
@@ -337,7 +337,7 @@ void CreateUniformBufferShaderDeclaration(const TCHAR* Name,const FUniformBuffer
 	}
 }
 
-void CacheUniformBufferIncludes(TMap<const TCHAR*,FCachedUniformBufferDeclaration>& Cache, EShaderPlatform Platform)
+SHADERCORE_API void CacheUniformBufferIncludes(TMap<const TCHAR*,FCachedUniformBufferDeclaration>& Cache, EShaderPlatform Platform)
 {
 	for (TMap<const TCHAR*,FCachedUniformBufferDeclaration>::TIterator It(Cache); It; ++It)
 	{
diff --git a/Engine/Source/Runtime/ShaderCore/Public/Shader.h b/Engine/Source/Runtime/ShaderCore/Public/Shader.h
--- a/Engine/Source/Runtime/ShaderCore/Public/Shader.h
+++ b/Engine/Source/Runtime/ShaderCore/Public/Shader.h
@@ -1079,15 +1079,6 @@ private:
 	// DumpShaderStats needs to access ShaderIdMap.
 	friend void SHADERCORE_API DumpShaderStats( EShaderPlatform Platform, EShaderFrequency Frequency );
 
-	/** 
-	 * Cache of referenced uniform buffer includes.  
-	 * These are derived from source files so they need to be flushed when editing and recompiling shaders on the fly. 
-	 * FShaderType::Initialize will add an entry for each referenced uniform buffer, but the declarations are added on demand as shaders are compiled.
-	 */
-	TMap<const TCHAR*, FCachedUniformBufferDeclaration> ReferencedUniformBufferStructsCache;
-
-	/** Tracks what platforms ReferencedUniformBufferStructsCache has had declarations cached for. */
-	bool bCachedUniformBufferStructDeclarations[SP_NumPlatforms];
 
 	/** 
 	 * Stores a history of serialization sizes for this shader type. 
@@ -1097,6 +1088,18 @@ private:
 
 	/** Tracks whether serialization history for all shader types has been initialized. */
 	static bool bInitializedSerializationHistory;
+
+protected:
+	/** Tracks what platforms ReferencedUniformBufferStructsCache has had declarations cached for. */
+	bool bCachedUniformBufferStructDeclarations[SP_NumPlatforms];
+
+	/**
+	* Cache of referenced uniform buffer includes.
+	* These are derived from source files so they need to be flushed when editing and recompiling shaders on the fly.
+	* FShaderType::Initialize will add an entry for each referenced uniform buffer, but the declarations are added on demand as shaders are compiled.
+	*/
+	TMap<const TCHAR*, FCachedUniformBufferDeclaration> ReferencedUniformBufferStructsCache;
+
 };
 
 /**
diff --git a/Engine/Source/Runtime/ShaderCore/Public/ShaderParameters.h b/Engine/Source/Runtime/ShaderCore/Public/ShaderParameters.h
--- a/Engine/Source/Runtime/ShaderCore/Public/ShaderParameters.h
+++ b/Engine/Source/Runtime/ShaderCore/Public/ShaderParameters.h
@@ -13,6 +13,9 @@ class FShaderParameterMap;
 class FUniformBufferStruct;
 struct FShaderCompilerEnvironment;
 
+SHADERCORE_API void CacheUniformBufferIncludes(TMap<const TCHAR*, struct FCachedUniformBufferDeclaration>& Cache, EShaderPlatform Platform);
+
+
 enum EShaderParameterFlags
 {
 	// no shader error if the parameter is not used
diff --git a/Engine/Source/Runtime/Slate/Private/Widgets/Images/SThrobber.cpp b/Engine/Source/Runtime/Slate/Private/Widgets/Images/SThrobber.cpp
--- a/Engine/Source/Runtime/Slate/Private/Widgets/Images/SThrobber.cpp
+++ b/Engine/Source/Runtime/Slate/Private/Widgets/Images/SThrobber.cpp
@@ -110,6 +110,7 @@ void SCircularThrobber::Construct(const FArguments& InArgs)
 	NumPieces = InArgs._NumPieces;
 	Period = InArgs._Period;
 	Radius = InArgs._Radius;
+	ColorAndOpacity = InArgs._ColorAndOpacity;
 
 	ConstructSequence();
 }
@@ -144,7 +145,16 @@ void SCircularThrobber::ConstructSequence()
 
 int32 SCircularThrobber::OnPaint( const FPaintArgs& Args, const FGeometry& AllottedGeometry, const FSlateRect& MyCullingRect, FSlateWindowElementList& OutDrawElements, int32 LayerId, const FWidgetStyle& InWidgetStyle, bool bParentEnabled ) const
 {
-	const FLinearColor FinalColorAndOpacity( InWidgetStyle.GetColorAndOpacityTint() * PieceImage->GetTint( InWidgetStyle ) );
+	FLinearColor FinalColorAndOpacity;
+	if (ColorAndOpacity.IsSet())
+	{
+		FinalColorAndOpacity = ColorAndOpacity.Get().GetColor(InWidgetStyle);
+	}
+	else
+	{
+		FinalColorAndOpacity = InWidgetStyle.GetColorAndOpacityTint() * PieceImage->GetTint(InWidgetStyle);
+	}
+
 	const FVector2D LocalOffset = (AllottedGeometry.GetLocalSize() - PieceImage->ImageSize) * 0.5f;
 	const float DeltaAngle = NumPieces > 0 ? 2 * PI / NumPieces : 0;
 	const float Phase = Curve.GetLerp() * 2 * PI;
diff --git a/Engine/Source/Runtime/Slate/Private/Widgets/Input/SMenuAnchor.cpp b/Engine/Source/Runtime/Slate/Private/Widgets/Input/SMenuAnchor.cpp
--- a/Engine/Source/Runtime/Slate/Private/Widgets/Input/SMenuAnchor.cpp
+++ b/Engine/Source/Runtime/Slate/Private/Widgets/Input/SMenuAnchor.cpp
@@ -21,7 +21,7 @@ static FVector2D GetMenuOffsetForPlacement(const FGeometry& AllottedGeometry, EM
 			return FVector2D(-((PopupSizeLocalSpace.X / 2) - (AllottedGeometry.GetLocalSize().X / 2)), AllottedGeometry.GetLocalSize().Y);
 			break;
 		case MenuPlacement_BelowRightAnchor:
-			return FVector2D( -( PopupSizeLocalSpace.X ) - ( AllottedGeometry.GetLocalSize().X ), AllottedGeometry.GetLocalSize().Y );
+			return FVector2D( -( PopupSizeLocalSpace.X ) + ( AllottedGeometry.GetLocalSize().X ), AllottedGeometry.GetLocalSize().Y );
 			break;
 		case MenuPlacement_ComboBox:
 			return FVector2D(0.0f, AllottedGeometry.GetLocalSize().Y);
diff --git a/Engine/Source/Runtime/Slate/Private/Widgets/Input/SSearchBox.cpp b/Engine/Source/Runtime/Slate/Private/Widgets/Input/SSearchBox.cpp
--- a/Engine/Source/Runtime/Slate/Private/Widgets/Input/SSearchBox.cpp
+++ b/Engine/Source/Runtime/Slate/Private/Widgets/Input/SSearchBox.cpp
@@ -4,6 +4,8 @@
 #include "Widgets/SBoxPanel.h"
 #include "Widgets/Images/SImage.h"
 #include "Widgets/Input/SButton.h"
+#include "Widgets/Text/STextBlock.h"
+#include "Widgets/Images/SThrobber.h"
 
 const double SSearchBox::FilterDelayAfterTyping = 0.25f;
 
@@ -11,6 +13,8 @@ void SSearchBox::Construct( const FArguments& InArgs )
 {
 	check(InArgs._Style);
 
+	SearchResultData = InArgs._SearchResultData;
+	bIsSearching = InArgs._IsSearching;
 	OnSearchDelegate = InArgs._OnSearch;
 	OnTextChangedDelegate = InArgs._OnTextChanged;
 	OnTextCommittedDelegate = InArgs._OnTextCommitted;
@@ -36,9 +40,37 @@ void SSearchBox::Construct( const FArguments& InArgs )
 	// If we want to have the buttons appear to the left of the text box we have to insert the slots instead of add them
 	int32 SlotIndex = InArgs._Style->bLeftAlignButtons ? 0 : Box->NumSlots();
 
+	// Add a throbber to show if there is a search running.
+	Box->InsertSlot(SlotIndex++)
+	.AutoWidth()
+	.Padding(0, 0, 2, 0)
+	.HAlign(HAlign_Center)
+	.VAlign(VAlign_Center)
+	[
+		SNew(SCircularThrobber)
+		.Radius(9.0f)
+		.Visibility(this, &SSearchBox::GetIsSearchingThrobberVisibility)
+		.ToolTipText(NSLOCTEXT("SearchBox", "Searching", "Searching..."))
+		.ColorAndOpacity(FSlateColor::UseForeground())
+	];
+
 	// If a search delegate was bound, add a previous and next button
 	if (OnSearchDelegate.IsBound())
 	{
+		// Search result data text
+		Box->InsertSlot(SlotIndex++)
+		.AutoWidth()
+		.Padding(0, 0, 2, 0)
+		.HAlign(HAlign_Center)
+		.VAlign(VAlign_Center)
+		[
+			SNew(STextBlock)
+			.Visibility(this, &SSearchBox::GetSearchResultDataVisibility)
+			.Text(this, &SSearchBox::GetSearchResultText)
+			.ColorAndOpacity(FSlateColor::UseSubduedForeground())
+		];
+
+		// Previous result button
 		Box->InsertSlot(SlotIndex++)
 		.AutoWidth()
 		.Padding(InArgs._Style->ImagePadding)
@@ -53,12 +85,14 @@ void SSearchBox::Construct( const FArguments& InArgs )
 			.OnClicked( this, &SSearchBox::OnClickedSearch, SSearchBox::Previous )
 			.ForegroundColor( FSlateColor::UseForeground() )
 			.IsFocusable(false)
+			.Visibility(this, &SSearchBox::GetSearchResultNavigationButtonVisibility)
 			[
 				SNew(SImage)
 				.Image( &InArgs._Style->UpArrowImage )
 				.ColorAndOpacity( FSlateColor::UseForeground() )
 			]
 		];
+		// Next result button
 		Box->InsertSlot(SlotIndex++)
 		.AutoWidth()
 		.Padding(InArgs._Style->ImagePadding)
@@ -73,6 +107,7 @@ void SSearchBox::Construct( const FArguments& InArgs )
 			.OnClicked( this, &SSearchBox::OnClickedSearch, SSearchBox::Next )
 			.ForegroundColor( FSlateColor::UseForeground() )
 			.IsFocusable(false)
+			.Visibility(this, &SSearchBox::GetSearchResultNavigationButtonVisibility)
 			[
 				SNew(SImage)
 				.Image( &InArgs._Style->DownArrowImage )
@@ -157,6 +192,34 @@ void SSearchBox::HandleTextCommitted(const FText& NewText, ETextCommit::Type Com
 	OnTextCommittedDelegate.ExecuteIfBound( NewText, CommitType );
 }
 
+FText SSearchBox::GetSearchResultText() const
+{
+	TOptional<FSearchResultData> CurrentSearchResultData = SearchResultData.Get();
+	if (CurrentSearchResultData.IsSet())
+	{
+		return FText::Format(NSLOCTEXT("SearchBox", "SearchResultFormat", "{0} / {1}"), 
+			CurrentSearchResultData.GetValue().CurrentSearchResultIndex, CurrentSearchResultData.GetValue().NumSearchResults);
+	}
+	else
+	{
+		return FText();
+	}
+}
+
+EVisibility SSearchBox::GetSearchResultNavigationButtonVisibility() const
+{
+	if (SearchResultData.IsBound())
+	{
+		return SearchResultData.Get().IsSet()
+			? EVisibility::Visible
+			: EVisibility::Collapsed;
+	}
+	else
+	{
+		return EVisibility::Visible;
+	}
+}
+
 EVisibility SSearchBox::GetXVisibility() const
 {
 	return (EditableText->GetText().IsEmpty())
@@ -164,6 +227,20 @@ EVisibility SSearchBox::GetXVisibility() const
 		: EVisibility::Visible;
 }
 
+EVisibility SSearchBox::GetSearchResultDataVisibility() const
+{
+	return SearchResultData.Get().IsSet()
+		? EVisibility::Visible
+		: EVisibility::Collapsed;
+}
+
+EVisibility SSearchBox::GetIsSearchingThrobberVisibility() const
+{
+	return (bIsSearching.Get())
+		? EVisibility::Visible
+		: EVisibility::Collapsed;
+}
+
 EVisibility SSearchBox::GetSearchGlassVisibility() const
 {
 	return (EditableText->GetText().IsEmpty())
diff --git a/Engine/Source/Runtime/Slate/Public/Widgets/Images/SThrobber.h b/Engine/Source/Runtime/Slate/Public/Widgets/Images/SThrobber.h
--- a/Engine/Source/Runtime/Slate/Public/Widgets/Images/SThrobber.h
+++ b/Engine/Source/Runtime/Slate/Public/Widgets/Images/SThrobber.h
@@ -114,6 +114,8 @@ public:
 		SLATE_ARGUMENT( float, Period )
 		/** The radius of the circle */
 		SLATE_ARGUMENT( float, Radius )
+		/** Throbber color and opacity */
+		SLATE_ATTRIBUTE(FSlateColor, ColorAndOpacity)
 
 	SLATE_END_ARGS()
 
@@ -157,4 +159,7 @@ private:
 
 	/** The radius of the circle */
 	float Radius;
+
+	/** Color and opacity of the throbber images. */
+	TAttribute<FSlateColor> ColorAndOpacity;
 };
diff --git a/Engine/Source/Runtime/Slate/Public/Widgets/Input/SSearchBox.h b/Engine/Source/Runtime/Slate/Public/Widgets/Input/SSearchBox.h
--- a/Engine/Source/Runtime/Slate/Public/Widgets/Input/SSearchBox.h
+++ b/Engine/Source/Runtime/Slate/Public/Widgets/Input/SSearchBox.h
@@ -29,12 +29,23 @@ public:
 		Next,
 	};
 
+	/** Data for external search results to be shown in the search box. */
+	struct FSearchResultData
+	{
+		/** The total number of search results. */
+		int32 NumSearchResults;
+		/** The index of the currently focused search result. */
+		int32 CurrentSearchResultIndex;
+	};
+
 	DECLARE_DELEGATE_OneParam(FOnSearch, SSearchBox::SearchDirection);
 
 	SLATE_BEGIN_ARGS(SSearchBox)
 		: _Style( &FCoreStyle::Get().GetWidgetStyle<FSearchBoxStyle>("SearchBox") )
 		, _HintText( LOCTEXT("SearchHint", "Search") )
 		, _InitialText()
+		, _SearchResultData()
+		, _IsSearching(false)
 		, _OnTextChanged()
 		, _OnTextCommitted()
 		, _OnSearch()
@@ -51,6 +62,12 @@ public:
 		/** The text displayed in the SearchBox when it's created */
 		SLATE_ATTRIBUTE( FText, InitialText )
 
+		/** Optional search result data to be shown in the search bar. */
+		SLATE_ATTRIBUTE( TOptional<FSearchResultData>, SearchResultData )
+
+		/** Whether or not an external search is currently running. */
+		SLATE_ATTRIBUTE( bool, IsSearching )
+
 		/** Invoked whenever the text changes */
 		SLATE_EVENT( FOnTextChanged, OnTextChanged )
 
@@ -89,9 +106,21 @@ private:
 	/** Fires the text changed delegate on a delay */
 	EActiveTimerReturnType TriggerOnTextChanged( double InCurrentTime, float InDeltaTime, FText NewText );
 
+	/** Gets the display text for the current search results. */
+	FText GetSearchResultText() const;
+
+	/** @return should we show the search result navigation buttons. */
+	EVisibility GetSearchResultNavigationButtonVisibility() const;
+
 	/** @return should we show the X to clear search? */
 	EVisibility GetXVisibility() const;
 
+	/** @return should we show the search results text. */
+	EVisibility GetSearchResultDataVisibility() const;
+
+	/** @return should we show the is searching throbber. */
+	EVisibility GetIsSearchingThrobberVisibility() const;
+
 	/** @return should we show the search glass icon? */
 	EVisibility GetSearchGlassVisibility() const;
 
@@ -124,6 +153,12 @@ private:
 
 	/** Fonts that specify how to render search text when inactive, and active */
 	FSlateFontInfo ActiveFont, InactiveFont;
+
+	/** Optional search result data to be shown in the search bar. */
+	TAttribute<TOptional<FSearchResultData>> SearchResultData;
+
+	/** Whether or not an external search is currently running. */
+	TAttribute<bool> bIsSearching;
 };
 
 
