Fix crash if slate shuts down while a drag drop operation is in progress

diff --git a/Engine/Source/Runtime/Slate/Private/Framework/Application/SlateApplication.cpp b/Engine/Source/Runtime/Slate/Private/Framework/Application/SlateApplication.cpp
--- a/Engine/Source/Runtime/Slate/Private/Framework/Application/SlateApplication.cpp
+++ b/Engine/Source/Runtime/Slate/Private/Framework/Application/SlateApplication.cpp
@@ -974,7 +974,7 @@ void FSlateApplication::Shutdown(bool bShutdownPlatform)
 
 		if (bShutdownPlatform)
 		{
-		PlatformApplication->DestroyApplication();
+			PlatformApplication->DestroyApplication();
 		}
 
 		PlatformApplication.Reset();
@@ -2970,6 +2970,8 @@ void FSlateApplication::OnShutdown()
 
 void FSlateApplication::CloseAllWindowsImmediately()
 {
+	DragDropContent.Reset();
+
 	// Clean up our tooltip window
 	TSharedPtr< SWindow > PinnedToolTipWindow(ToolTipWindow.Pin());
 	if (PinnedToolTipWindow.IsValid())
