Fix for menu options not updating immediatly after creating or removing collision on Mac.  We set the bChachedMacMenuStateNeedsUpdate when a menu is opening, key up and mouse up.  However these can be missed on a mouse menu selection (mouse up doesn't cover this case).  This change makes sure that the menu cache updates after a menu item execute operation has been called.

diff --git a/Engine/Source/Runtime/Slate/Private/Framework/MultiBox/Mac/MacMenu.cpp b/Engine/Source/Runtime/Slate/Private/Framework/MultiBox/Mac/MacMenu.cpp
--- a/Engine/Source/Runtime/Slate/Private/Framework/MultiBox/Mac/MacMenu.cpp
+++ b/Engine/Source/Runtime/Slate/Private/Framework/MultiBox/Mac/MacMenu.cpp
@@ -53,6 +53,7 @@ static FCriticalSection GCachedMenuStateCS;
 	if ( !CocoaMenu || ![CocoaMenu isHighlightingKeyEquivalent] )
 	{
 		FSlateMacMenu::ExecuteMenuItemAction(self.MenuEntryBlock.ToSharedRef());
+		FPlatformApplicationMisc::bChachedMacMenuStateNeedsUpdate = true;
 	}
 }
 
