speculative fix for FSlateMacMenu::UpdateMenu() crashes.  I've not been able to reproduce this.  All the crash call stacks seem to be at shutdown, maybe with the save changes dialogue that would appear after modifications then a CMD+Q.  I can get similar behaviour by adjusting when the FMacMenu adds and removes itself to the GCachedMenuState map.  This fix changes the access pattern to the map from operator[] (which uses a FindChecked underneath and is the cause of this crash/assert) to a FindRef then checking for a valid node.

Since these crashes are at shutdown I think this is ok to do and I've tested with this fix in place, randomly not adding those items to GCachedMenuState object during editor operation and it now survives this and recovers if you then reenable the adding of the items.

My best guess as to the cause is that the defered update event on shutdown is happenning way after the FMacMenu dealloc has been called and so the GCachedMenuState is empty by then.

diff --git a/Engine/Source/Runtime/Slate/Private/Framework/MultiBox/Mac/MacMenu.cpp b/Engine/Source/Runtime/Slate/Private/Framework/MultiBox/Mac/MacMenu.cpp
--- a/Engine/Source/Runtime/Slate/Private/Framework/MultiBox/Mac/MacMenu.cpp
+++ b/Engine/Source/Runtime/Slate/Private/Framework/MultiBox/Mac/MacMenu.cpp
@@ -572,99 +572,102 @@ void FSlateMacMenu::UpdateMenu(FMacMenu* Menu)
 			}
 		}
 
-		TSharedPtr<TArray<FMacMenuItemState>> MenuState = GCachedMenuState[Menu];
-		int32 ItemIndexAdjust = 0;
-		for (int32 Index = 0; Index < MenuState->Num(); Index++)
+		TSharedPtr<TArray<FMacMenuItemState>> MenuState = GCachedMenuState.FindRef(Menu);
+		if(MenuState.IsValid())
 		{
-			FMacMenuItemState& MenuItemState = (*MenuState)[Index];
-			const int32 ItemIndex = (bIsWindowMenu ? Index + ItemIndexOffset : Index) - ItemIndexAdjust;
-			NSMenuItem* MenuItem = [Menu numberOfItems] > ItemIndex ? [Menu itemAtIndex:ItemIndex] : nil;
-
-			if (MenuItemState.Type == EMultiBlockType::MenuEntry)
+			int32 ItemIndexAdjust = 0;
+			for (int32 Index = 0; Index < MenuState->Num(); Index++)
 			{
-				if (MenuItem && (![MenuItem isKindOfClass:[FMacMenuItem class]] || (MenuItemState.IsSubMenu && [MenuItem submenu] == nil) || (!MenuItemState.IsSubMenu && [MenuItem submenu] != nil)))
-				{
-					[Menu removeItem:MenuItem];
-					MenuItem = nil;
-				}
-				if (!MenuItem)
-				{
-					MenuItem = [[[FMacMenuItem alloc] initWithMenuEntryBlock:MenuItemState.Block] autorelease];
+				FMacMenuItemState& MenuItemState = (*MenuState)[Index];
+				const int32 ItemIndex = (bIsWindowMenu ? Index + ItemIndexOffset : Index) - ItemIndexAdjust;
+				NSMenuItem* MenuItem = [Menu numberOfItems] > ItemIndex ? [Menu itemAtIndex:ItemIndex] : nil;
 
-					if (MenuItemState.IsSubMenu)
+				if (MenuItemState.Type == EMultiBlockType::MenuEntry)
+				{
+					if (MenuItem && (![MenuItem isKindOfClass:[FMacMenuItem class]] || (MenuItemState.IsSubMenu && [MenuItem submenu] == nil) || (!MenuItemState.IsSubMenu && [MenuItem submenu] != nil)))
+					{
+						[Menu removeItem:MenuItem];
+						MenuItem = nil;
+					}
+					if (!MenuItem)
 					{
-						FMacMenu* SubMenu = [[[FMacMenu alloc] initWithMenuEntryBlock:MenuItemState.Block] autorelease];
-						[MenuItem setSubmenu:SubMenu];
+						MenuItem = [[[FMacMenuItem alloc] initWithMenuEntryBlock:MenuItemState.Block] autorelease];
+
+						if (MenuItemState.IsSubMenu)
+						{
+							FMacMenu* SubMenu = [[[FMacMenu alloc] initWithMenuEntryBlock:MenuItemState.Block] autorelease];
+							[MenuItem setSubmenu:SubMenu];
+						}
+
+						if ([Menu numberOfItems] > ItemIndex)
+						{
+							[Menu insertItem:MenuItem atIndex:ItemIndex];
+						}
+						else
+						{
+							[Menu addItem:MenuItem];
+						}
 					}
 
-					if ([Menu numberOfItems] > ItemIndex)
+					[MenuItem setTitle:MenuItemState.Title];
+
+					[MenuItem setKeyEquivalent:MenuItemState.KeyEquivalent];
+					[MenuItem setKeyEquivalentModifierMask:MenuItemState.KeyModifiers];
+
+					if (bIsWindowMenu)
 					{
-						[Menu insertItem:MenuItem atIndex:ItemIndex];
+						NSImage* MenuImage = MenuItemState.Icon;
+						if(MenuImage)
+						{
+							[MenuItem setImage:MenuImage];
+						}
 					}
 					else
 					{
-						[Menu addItem:MenuItem];
+						[MenuItem setImage:nil];
 					}
-				}
 
-				[MenuItem setTitle:MenuItemState.Title];
-
-				[MenuItem setKeyEquivalent:MenuItemState.KeyEquivalent];
-				[MenuItem setKeyEquivalentModifierMask:MenuItemState.KeyModifiers];
-
-				if (bIsWindowMenu)
-				{
-					NSImage* MenuImage = MenuItemState.Icon;
-					if(MenuImage)
+					[MenuItem setTarget:MenuItem];
+					if(!MenuItemState.IsSubMenu)
 					{
-						[MenuItem setImage:MenuImage];
+					   if(MenuItemState.IsEnabled)
+						{
+							[MenuItem setAction:@selector(performAction)];
+						}
+						else
+						{
+							[MenuItem setAction:nil];
+						}
+					}
+					
+					if (!MenuItemState.IsSubMenu)
+					{
+						[MenuItem setState:MenuItemState.State];
 					}
 				}
-				else
-				{
-					[MenuItem setImage:nil];
-				}
-
-                [MenuItem setTarget:MenuItem];
-                if(!MenuItemState.IsSubMenu)
-                {
-                   if(MenuItemState.IsEnabled)
-                    {
-                        [MenuItem setAction:@selector(performAction)];
-                    }
-                    else
-                    {
-                        [MenuItem setAction:nil];
-                    }
-                }
-				
-				if (!MenuItemState.IsSubMenu)
-				{
-					[MenuItem setState:MenuItemState.State];
-				}
-			}
-			else if (MenuItemState.Type == EMultiBlockType::MenuSeparator)
-			{
-				if (MenuItem && ![MenuItem isSeparatorItem])
-				{
-					[Menu removeItem:MenuItem];
-				}
-				else if (!MenuItem)
+				else if (MenuItemState.Type == EMultiBlockType::MenuSeparator)
 				{
-					if ([Menu numberOfItems] > ItemIndex)
+					if (MenuItem && ![MenuItem isSeparatorItem])
 					{
-						[Menu insertItem:[NSMenuItem separatorItem] atIndex:ItemIndex];
+						[Menu removeItem:MenuItem];
 					}
-					else
+					else if (!MenuItem)
 					{
-						[Menu addItem:[NSMenuItem separatorItem]];
+						if ([Menu numberOfItems] > ItemIndex)
+						{
+							[Menu insertItem:[NSMenuItem separatorItem] atIndex:ItemIndex];
+						}
+						else
+						{
+							[Menu addItem:[NSMenuItem separatorItem]];
+						}
 					}
 				}
-			}
-			else
-			{
-				// If it's a type we skip, update ItemIndexAdjust so we can properly calculate item's index in NSMenu
-				ItemIndexAdjust++;
+				else
+				{
+					// If it's a type we skip, update ItemIndexAdjust so we can properly calculate item's index in NSMenu
+					ItemIndexAdjust++;
+				}
 			}
 		}
 	});
