Emergency Fix for bug I introduced in macOS UnrealFrontend crashing on open.  FSlateMacMenu::PostInitStartup() gets called much much earler in the front end tool than in the editor.  Unfortunatly the style set referenced is invalid as this call happens much earlier in the startup than with the game or the editor

1) Changed GetStyleSetName() - this is not imortant for us so it won't crash there again.
2) Added a check to see if the menu hasn't been setup correctly - if not initialize - this allows localization and and correct keyboard short cuts in non editor / game builds.
3) Keep Editor and Game menu startup behaviour the same as before.

diff --git a/Engine/Source/Runtime/Slate/Private/Framework/MultiBox/Mac/MacMenu.cpp b/Engine/Source/Runtime/Slate/Private/Framework/MultiBox/Mac/MacMenu.cpp
--- a/Engine/Source/Runtime/Slate/Private/Framework/MultiBox/Mac/MacMenu.cpp
+++ b/Engine/Source/Runtime/Slate/Private/Framework/MultiBox/Mac/MacMenu.cpp
@@ -105,7 +105,7 @@ public:
 		TEXT("MacMenu"),
 		NSLOCTEXT("Contexts", "MacMenu", "MacMenu"),
 		NAME_None,
-		FCoreStyle::Get().GetStyleSetName()
+		"MacMenu"
 	)
 	{}
 	
@@ -264,6 +264,9 @@ namespace MacMenuHelper
 	{
 		return GIsEditor ? NSLOCTEXT("UnrealEditor", "ApplicationTitle", "Unreal Editor").ToString().GetNSString() : FString(FApp::GetProjectName()).GetNSString();
 	}
+	
+	bool GMacPostInitStartupRequested = false;
+	bool GMacPostInitStartUpComplete = false; 
 };
 
 // Bind all low-level Application hooks that require to access this high-level MacMenu system which includes NSApp Menu's and slate menus
@@ -291,14 +294,13 @@ void FSlateMacMenu::CleanupOnShutdown()
 
 void FSlateMacMenu::PostInitStartup()
 {
-	static bool bPostInitStartup = false;
-	check(!bPostInitStartup);
+	MacMenuHelper::GMacPostInitStartupRequested = true;
 	
 	// Setup the app menu in menu bar
 	const bool bIsBundledApp = [[[NSBundle mainBundle] bundlePath] hasSuffix:@".app"];
-	if (!bPostInitStartup && bIsBundledApp)
+	if (!MacMenuHelper::GMacPostInitStartUpComplete && bIsBundledApp && MacApplication)
 	{
-		bPostInitStartup = true;
+		MacMenuHelper::GMacPostInitStartUpComplete = true;
 		
 		// Setup our Mac Specific commands
 		FMacMenuCommands::Register();
@@ -407,12 +409,17 @@ void FSlateMacMenu::LanguageChanged()
 	
 	NSMenuItem* Services = [AppMenu itemWithTag:MacMenuHelper::CmdID_ServicesMenu];
 	[Services setTitle:NSLOCTEXT("MainMenu","ServicesMenu","Services").ToString().GetNSString()];
-	
-	//@todo: Look into updating the MainMenu tab so the top level menu item names "File, Edit, Window etc" update without the need for manual focus to reflect this change...
 }
 
 void FSlateMacMenu::UpdateApplicationMenu(bool bMacApplicationModalMode)
 {
+	// In case an obsecure app startup sequence has not managed to finish the menu startup correctly
+	// However only do this if the post init has been called
+	if(!MacMenuHelper::GMacPostInitStartUpComplete && MacMenuHelper::GMacPostInitStartupRequested)
+	{
+		PostInitStartup();
+	}
+
     NSMenu* MainMenu = [NSApp mainMenu];
     NSMenuItem* AppMenuItem = [MainMenu itemWithTitle:@"AppMenuItem"];
     NSMenu* AppMenu = [AppMenuItem submenu];
