Fix UMG anchor prompts being wrong on mac

diff --git a/Engine/Source/Editor/UMGEditor/Private/Customizations/CanvasSlotCustomization.cpp b/Engine/Source/Editor/UMGEditor/Private/Customizations/CanvasSlotCustomization.cpp
--- a/Engine/Source/Editor/UMGEditor/Private/Customizations/CanvasSlotCustomization.cpp
+++ b/Engine/Source/Editor/UMGEditor/Private/Customizations/CanvasSlotCustomization.cpp
@@ -321,6 +321,9 @@ void FCanvasSlotCustomization::CustomizeAnchors(TSharedPtr<IPropertyHandle> Prop
 
 	const float FillDividePadding = 1;
 
+	FInputChord ShiftKey(EKeys::Invalid, true, false, false, false);
+	FInputChord CtrlKey(EKeys::Invalid, false, true, false, false);
+
 	AnchorsPropertyRow.CustomWidget(/*bShowChildren*/ true)
 		.NameContent()
 		[
@@ -477,14 +480,13 @@ void FCanvasSlotCustomization::CustomizeAnchors(TSharedPtr<IPropertyHandle> Prop
 								.AutoHeight()
 								[
 									SNew(STextBlock)
-									.Text(LOCTEXT("ShiftResetsAlignment", "Hold [Shift] to update the alignment to match."))
+									.Text(FText::Format(LOCTEXT("ShiftResetsAlignment", "Hold {0} to update the alignment to match."), ShiftKey.GetInputText()))
 								]
-
 								+ SVerticalBox::Slot()
 								.AutoHeight()
 								[
 									SNew(STextBlock)
-									.Text(LOCTEXT("ControlResetsPosition", "Hold [Control] to update the position to match."))
+									.Text(FText::Format(LOCTEXT("ControlResetsPosition", "Hold {0} to update the position to match."), CtrlKey.GetInputText()))
 								]
 							]
 						]
diff --git a/Engine/Source/Runtime/Slate/Private/Framework/Commands/InputChord.cpp b/Engine/Source/Runtime/Slate/Private/Framework/Commands/InputChord.cpp
--- a/Engine/Source/Runtime/Slate/Private/Framework/Commands/InputChord.cpp
+++ b/Engine/Source/Runtime/Slate/Private/Framework/Commands/InputChord.cpp
@@ -14,36 +14,40 @@
 FText FInputChord::GetInputText( ) const
 {
 #if PLATFORM_MAC
-    const FText CommandText = LOCTEXT("KeyName_Control", "Ctrl");
-    const FText ControlText = LOCTEXT("KeyName_Command", "Cmd");
+	const FText CommandText = LOCTEXT("KeyName_Control", "Ctrl");
+	const FText ControlText = LOCTEXT("KeyName_Command", "Cmd");
 #else
-    const FText ControlText = LOCTEXT("KeyName_Control", "Ctrl");
-    const FText CommandText = LOCTEXT("KeyName_Command", "Cmd"); 
+	const FText ControlText = LOCTEXT("KeyName_Control", "Ctrl");
+	const FText CommandText = LOCTEXT("KeyName_Command", "Cmd");
 #endif
-    const FText AltText = LOCTEXT("KeyName_Alt", "Alt");
-    const FText ShiftText = LOCTEXT("KeyName_Shift", "Shift");
-    
-	const FText AppenderText = LOCTEXT("ModAppender", "+");
+	const FText AltText = LOCTEXT("KeyName_Alt", "Alt");
+	const FText ShiftText = LOCTEXT("KeyName_Shift", "Shift");
+
+
+	const FText AppenderText = Key != EKeys::Invalid ? LOCTEXT("ModAppender", "+") : FText::GetEmpty();
 
 	FFormatNamedArguments Args;
 	int32 ModCount = 0;
 
-    if (bCtrl)
-    {
-		Args.Add(FString::Printf(TEXT("Mod%d"),++ModCount), ControlText);
-    }
-    if (bCmd)
-    {
-		Args.Add(FString::Printf(TEXT("Mod%d"),++ModCount), CommandText);
-    }
-    if (bAlt)
-    {
-		Args.Add(FString::Printf(TEXT("Mod%d"),++ModCount), AltText);
-    }
-    if (bShift)
-    {
-		Args.Add(FString::Printf(TEXT("Mod%d"),++ModCount), ShiftText);
-    }
+	if (bCtrl)
+	{
+		Args.Add(FString::Printf(TEXT("Mod%d"), ++ModCount), ControlText);
+	}
+
+	if (bCmd)
+	{
+		Args.Add(FString::Printf(TEXT("Mod%d"), ++ModCount), CommandText);
+	}
+
+	if (bAlt)
+	{
+		Args.Add(FString::Printf(TEXT("Mod%d"), ++ModCount), AltText);
+	}
+
+	if (bShift)
+	{
+		Args.Add(FString::Printf(TEXT("Mod%d"), ++ModCount), ShiftText);
+	}
 
 	for (int32 i = 1; i <= 4; ++i)
 	{
