Added bAllowHighDPIInGameMode setting to project's User Interface preferences, false by default. Using this setting the game can request using high-DPI-aware window on desktop platforms.

diff --git a/Engine/Config/BaseEngine.ini b/Engine/Config/BaseEngine.ini
--- a/Engine/Config/BaseEngine.ini
+++ b/Engine/Config/BaseEngine.ini
@@ -606,6 +606,7 @@ SuspendGameIO=True
 UIScaleRule=ShortestSide
 UIScaleCurve=(EditorCurveData=(Keys=((Time=480,Value=0.444),(Time=720,Value=0.666),(Time=1080,Value=1.0),(Time=8640,Value=8.0))),ExternalCurve=None)
 bLoadWidgetsOnDedicatedServer=True
+bAllowHighDPIInGameMode=False
 
 [/Script/Engine.GameEngine]
 MaxDeltaTime=0
diff --git a/Engine/Source/Runtime/Engine/Classes/Engine/UserInterfaceSettings.h b/Engine/Source/Runtime/Engine/Classes/Engine/UserInterfaceSettings.h
--- a/Engine/Source/Runtime/Engine/Classes/Engine/UserInterfaceSettings.h
+++ b/Engine/Source/Runtime/Engine/Classes/Engine/UserInterfaceSettings.h
@@ -174,6 +174,13 @@ public:
 		YAxisName="Scale"))
 	FRuntimeFloatCurve UIScaleCurve;
 
+	/**
+	 * If true, game window on desktop platforms will be created with high-DPI awareness enabled.
+	 * Recommended to be enabled only if the game's UI allows users to modify 3D resolution scaling.
+	 */
+	UPROPERTY(config, EditAnywhere, Category="DPI Scaling", meta=( DisplayName="Allow High DPI in Game Mode" ))
+	bool bAllowHighDPIInGameMode;
+
 	/**
 	 * If false, widget references will be stripped during cook for server builds and not loaded at runtime.
 	 */
diff --git a/Engine/Source/Runtime/Slate/Private/Framework/Application/SlateApplication.cpp b/Engine/Source/Runtime/Slate/Private/Framework/Application/SlateApplication.cpp
--- a/Engine/Source/Runtime/Slate/Private/Framework/Application/SlateApplication.cpp
+++ b/Engine/Source/Runtime/Slate/Private/Framework/Application/SlateApplication.cpp
@@ -1121,7 +1121,11 @@ void FSlateApplication::InitHighDPI()
 		bool bRequestEnableHighDPI = true;
 		if (GIsEditor)
 		{
-		GConfig->GetBool(TEXT("HDPI"), TEXT("EnableHighDPIAwareness"), bRequestEnableHighDPI, GEditorSettingsIni);
+			GConfig->GetBool(TEXT("HDPI"), TEXT("EnableHighDPIAwareness"), bRequestEnableHighDPI, GEditorSettingsIni);
+		}
+		else
+		{
+			GConfig->GetBool(TEXT("/Script/Engine.UserInterfaceSettings"), TEXT("bAllowHighDPIInGameMode"), bRequestEnableHighDPI, GEngineIni);
 		}
 
 		const bool bEnableHighDPI = bRequestEnableHighDPI && !FParse::Param(FCommandLine::Get(), TEXT("nohighdpi"));
