Clipping horizontal box now closes its menu when the result of taking an action in the menu causes the contents of the multibox to change, thus changing the layout of the horizontal box and disconnecting open menus from their parent.  This causes the crash in UE-59055

diff --git a/Engine/Source/Runtime/Slate/Private/Framework/MultiBox/SClippingHorizontalBox.cpp b/Engine/Source/Runtime/Slate/Private/Framework/MultiBox/SClippingHorizontalBox.cpp
--- a/Engine/Source/Runtime/Slate/Private/Framework/MultiBox/SClippingHorizontalBox.cpp
+++ b/Engine/Source/Runtime/Slate/Private/Framework/MultiBox/SClippingHorizontalBox.cpp
@@ -5,13 +5,22 @@
 #include "Rendering/DrawElements.h"
 #include "Widgets/Images/SImage.h"
 #include "Widgets/Input/SComboButton.h"
+#include "Framework/Application/SlateApplication.h"
 
 void SClippingHorizontalBox::Tick( const FGeometry& AllottedGeometry, const double InCurrentTime, const float InDeltaTime )
 {
+	const int32 PrevNumClippedChildren = NumClippedChildren;
 	// Call ArrangeChildren so we can cache off the index of the first clipped tab
 	FArrangedChildren ArrangedChildren(EVisibility::Visible);
 	this->ArrangeChildren(AllottedGeometry, ArrangedChildren);
 	ClippedIdx = ArrangedChildren.Num() - 1;
+
+	if (NumClippedChildren != PrevNumClippedChildren && WrapButton->IsOpen())
+	{
+		// Number of children changed while the wrap button was open. An event while the wrap button menu is open has caused the contents of the multibox to change.  The menu can no longer be considered safe
+		// as it most likely lost its connection to the window
+		WrapButton->SetIsOpen(false);
+	}
 }
 
 void SClippingHorizontalBox::OnArrangeChildren( const FGeometry& AllottedGeometry, FArrangedChildren& ArrangedChildren ) const
@@ -20,6 +29,7 @@ void SClippingHorizontalBox::OnArrangeChildren( const FGeometry& AllottedGeometr
 	// been called and this method isn't going to behave properly
 	check(WrapButton.IsValid());
 
+	NumClippedChildren = 0;
 	SHorizontalBox::OnArrangeChildren(AllottedGeometry, ArrangedChildren);
 
 	// Remove children that are clipped by the allotted geometry
@@ -30,6 +40,7 @@ void SClippingHorizontalBox::OnArrangeChildren( const FGeometry& AllottedGeometr
 		const FArrangedWidget& CurWidget = ArrangedChildren[ChildIdx];
 		if (FMath::TruncToInt(CurWidget.Geometry.AbsolutePosition.X + CurWidget.Geometry.GetLocalSize().X * CurWidget.Geometry.Scale) > FMath::TruncToInt(AllottedGeometry.AbsolutePosition.X + AllottedGeometry.GetLocalSize().X * CurWidget.Geometry.Scale))
 		{
+			++NumClippedChildren;
 			ArrangedChildren.Remove(ChildIdx);
 			IndexClippedAt = ChildIdx;
 		}
@@ -37,6 +48,7 @@ void SClippingHorizontalBox::OnArrangeChildren( const FGeometry& AllottedGeometr
 
 	if (IndexClippedAt == NumChildren)
 	{
+		// Note do not increment NumClippedChildrenn here. This is to remove the wrap button
 		// None of the children are being clipped, so remove the wrap button
 		ArrangedChildren.Remove(ArrangedChildren.Num() - 1);
 	}
@@ -53,6 +65,7 @@ void SClippingHorizontalBox::OnArrangeChildren( const FGeometry& AllottedGeometr
 			const FArrangedWidget& CurWidget = ArrangedChildren[ChildIdx];
 			if (FMath::TruncToInt(CurWidget.Geometry.AbsolutePosition.X + CurWidget.Geometry.GetLocalSize().X * CurWidget.Geometry.Scale) > WrapButtonXPosition)
 			{
+				++NumClippedChildren;
 				ArrangedChildren.Remove(ChildIdx);
 			}
 		}
diff --git a/Engine/Source/Runtime/Slate/Private/Framework/MultiBox/SClippingHorizontalBox.h b/Engine/Source/Runtime/Slate/Private/Framework/MultiBox/SClippingHorizontalBox.h
--- a/Engine/Source/Runtime/Slate/Private/Framework/MultiBox/SClippingHorizontalBox.h
+++ b/Engine/Source/Runtime/Slate/Private/Framework/MultiBox/SClippingHorizontalBox.h
@@ -64,6 +64,9 @@ private:
 	/** Index of the first clipped child/block */
 	int32 ClippedIdx;
 
+	/** Number of clipped children not including the wrap button */
+	mutable int32 NumClippedChildren;
+
 	/** The style to use */
 	const ISlateStyle* StyleSet;
 	FName StyleName;
