Fix for children of SInvalidationPanel not knowing the correct clipping state in the editor

diff --git a/Engine/Source/Runtime/Slate/Private/Widgets/SInvalidationPanel.cpp b/Engine/Source/Runtime/Slate/Private/Widgets/SInvalidationPanel.cpp
--- a/Engine/Source/Runtime/Slate/Private/Widgets/SInvalidationPanel.cpp
+++ b/Engine/Source/Runtime/Slate/Private/Widgets/SInvalidationPanel.cpp
@@ -341,6 +341,9 @@ int32 SInvalidationPanel::OnPaint( const FPaintArgs& Args, const FGeometry& Allo
 				// we'll have already done one pre-pass before getting here.
 				ChildSlot.GetWidget()->SlatePrepass(AllottedGeometry.Scale);
 
+				FSlateClippingState CurrentClippingState = OutDrawElements.GetClippingManager().GetClippingStates()[OutDrawElements.GetClippingManager().GetClippingIndex()];
+				CachedWindowElements->GetClippingManager().PushClippingState(CurrentClippingState);
+
 				//TODO: When SWidget::Paint is called don't drag self if volatile, and we're doing a cache pass.
 				CachedMaxChildLayer = SCompoundWidget::OnPaint(
 					Args.EnableCaching(SharedMutableThis, RootCacheNode, true, false),
@@ -351,6 +354,8 @@ int32 SInvalidationPanel::OnPaint( const FPaintArgs& Args, const FGeometry& Allo
 					InWidgetStyle,
 					bParentEnabled);
 
+				CachedWindowElements->GetClippingManager().PopClip();
+
 				{
 					CachedResources.Reset();
 					FSlateObjectReferenceCollector Collector(CachedResources);
