Fixed: Packaged Match 3 Bomb Power bar displays as max and does not update

diff --git a/Engine/Source/Runtime/Slate/Private/Widgets/SInvalidationPanel.cpp b/Engine/Source/Runtime/Slate/Private/Widgets/SInvalidationPanel.cpp
--- a/Engine/Source/Runtime/Slate/Private/Widgets/SInvalidationPanel.cpp
+++ b/Engine/Source/Runtime/Slate/Private/Widgets/SInvalidationPanel.cpp
@@ -294,6 +294,21 @@ void SInvalidationPanel::OnGlobalInvalidate()
 	InvalidateCache();
 }
 
+template<typename ElementType>
+TFunction<void(ElementType&)> UpdateClipIndexFn(int32 LastClippingIndex, int32 ClippingStateOffset)
+{
+	return [=](ElementType& Element) {
+		if (Element.GetClippingIndex() == -1)
+		{
+			Element.SetClippingIndex(LastClippingIndex);
+		}
+		else
+		{
+			Element.SetClippingIndex(ClippingStateOffset + Element.GetClippingIndex());
+		}
+	};
+}
+
 int32 SInvalidationPanel::OnPaint( const FPaintArgs& Args, const FGeometry& AllottedGeometry, const FSlateRect& MyCullingRect, FSlateWindowElementList& OutDrawElements, int32 LayerId, const FWidgetStyle& InWidgetStyle, bool bParentEnabled ) const
 {
 #if SLATE_VERBOSE_NAMED_EVENTS
@@ -373,24 +388,16 @@ int32 SInvalidationPanel::OnPaint( const FPaintArgs& Args, const FGeometry& Allo
 				
 				const int32 ClippingStateOffset = OutDrawElements.GetClippingManager().MergeClippingStates(CachedWindowElements->GetClippingManager().GetClippingStates());
 
-				TArray<FSlateDrawElement>& CachedElements = CachedWindowElements->GetDrawElements();
-				for (int32 Index = 0; Index < CachedElements.Num(); Index++)
-				{
-					FSlateDrawElement& CachedElement = CachedElements[Index];
-					if (CachedElement.GetClippingIndex() == -1)
-					{
-						CachedElement.SetClippingIndex(LastClippingIndex);
-					}
-					else
-					{
-						CachedElement.SetClippingIndex(ClippingStateOffset + CachedElement.GetClippingIndex());
-					}
-				}
-
 				if (bCacheRenderData)
 				{
 					CachedRenderData = CachedWindowElements->CacheRenderData(this);
 				}
+				else
+				{
+					// offset clipping indices for cached elements, so they point to correct entry in merged clipping states
+					CachedWindowElements->ForEachElement(UpdateClipIndexFn<FSlateDrawBase>(LastClippingIndex, ClippingStateOffset));
+					CachedWindowElements->ForEachElement(UpdateClipIndexFn<FSlateDrawElement>(LastClippingIndex, ClippingStateOffset));
+				}
 
 				LastHitTestIndex = Args.GetLastHitTestIndex();
 
diff --git a/Engine/Source/Runtime/SlateCore/Public/Rendering/DrawElements.h b/Engine/Source/Runtime/SlateCore/Public/Rendering/DrawElements.h
--- a/Engine/Source/Runtime/SlateCore/Public/Rendering/DrawElements.h
+++ b/Engine/Source/Runtime/SlateCore/Public/Rendering/DrawElements.h
@@ -1550,6 +1550,48 @@ public:
 			CachedElementBuffers.Num();
 	}
 
+	/** Apply Function to each draw element */ 
+	void ForEachElement(const TFunction<void(FSlateDrawElement&)>& InFunction)
+	{
+		for (FSlateDrawElement& Element : DrawElements)
+		{
+			InFunction(Element);
+		}
+	}
+	/** Apply Function to each draw element */ 
+	void ForEachElement(const TFunction<void(FSlateDrawBase&)>& InFunction)
+	{
+		for (FSlateDrawBox& Element : BoxElements)
+		{
+			InFunction(Element);
+		}
+
+		for (FSlateDrawBox& Element : BorderElements)
+		{
+			InFunction(Element);
+		}
+
+		for (FSlateDrawText& Element : TextElements)
+		{
+			InFunction(Element);
+		}
+
+		for (FSlateDrawShapedText& Element : ShapedTextElements)
+		{
+			InFunction(Element);
+		}
+		
+		for (FSlateDrawLines& Element : LineElements)
+		{
+			InFunction(Element);
+		}
+		
+		for (FSlateDrawCachedBuffer& Element : CachedElementBuffers)
+		{
+			InFunction(Element);
+		}
+	}
+
 public:
 	// Element batch maps sorted by layer.
 	FElementBatchMap LayerToElementBatches;
@@ -1637,6 +1679,18 @@ public:
 		return RootDrawLayer.DrawElements;
 	}
 
+	/** Apply Function to each draw element */ 
+	FORCEINLINE void ForEachElement(const TFunction<void(FSlateDrawElement&)>& InFunction)
+	{
+		RootDrawLayer.ForEachElement(InFunction);
+	}
+
+	/** Apply Function to each draw element */ 
+	FORCEINLINE void ForEachElement(const TFunction<void(FSlateDrawBase&)>& InFunction)
+	{
+		RootDrawLayer.ForEachElement(InFunction);
+	}
+	
 	/** @return the total number of elements that have been registered to be drawn. */
 	SLATECORE_API int32 GetElementCount() const;
 
