Fixed: Match3 UMG Board Backgrounds Disappear When Gems Create a Match After Falling into Place
Always re-cache elements when global clipping state is changed

diff --git a/Engine/Source/Runtime/Slate/Private/Widgets/SInvalidationPanel.cpp b/Engine/Source/Runtime/Slate/Private/Widgets/SInvalidationPanel.cpp
--- a/Engine/Source/Runtime/Slate/Private/Widgets/SInvalidationPanel.cpp
+++ b/Engine/Source/Runtime/Slate/Private/Widgets/SInvalidationPanel.cpp
@@ -146,19 +146,19 @@ bool SInvalidationPanel::IsCachingNeeded(FSlateWindowElementList& OutDrawElement
 		return true;
 	}
 	
-	if (bCacheRelativeTransforms)
+	if (LastClippingIndex != OutDrawElements.GetClippingIndex())
 	{
-		if (LastClippingIndex != OutDrawElements.GetClippingIndex())
-		{
-			return true;
-		}
-
-		int32 ClippingStateCount = OutDrawElements.GetClippingManager().GetClippingStates().Num();
-		if (LastClippingStateOffset != ClippingStateCount)
-		{
-			return true;
-		}
+		return true;
+	}
 
+	int32 ClippingStateCount = OutDrawElements.GetClippingManager().GetClippingStates().Num();
+	if (LastClippingStateOffset != ClippingStateCount)
+	{
+		return true;
+	}
+	
+	if (bCacheRelativeTransforms)
+	{
 		bool bOverlapping;
 		FVector2D IntersectionSize = AllottedGeometry.GetLayoutBoundingRect().IntersectionWith(MyCullingRect, bOverlapping).GetSize();
 		if (!LastClippingIntersectionSize.Equals(IntersectionSize, 1.0f))
