Prevent invalidation panel crash when pushing the current clip state when there isnt one

diff --git a/Engine/Source/Runtime/Slate/Private/Widgets/SInvalidationPanel.cpp b/Engine/Source/Runtime/Slate/Private/Widgets/SInvalidationPanel.cpp
--- a/Engine/Source/Runtime/Slate/Private/Widgets/SInvalidationPanel.cpp
+++ b/Engine/Source/Runtime/Slate/Private/Widgets/SInvalidationPanel.cpp
@@ -356,8 +356,15 @@ int32 SInvalidationPanel::OnPaint( const FPaintArgs& Args, const FGeometry& Allo
 				// we'll have already done one pre-pass before getting here.
 				ChildSlot.GetWidget()->SlatePrepass(AllottedGeometry.Scale);
 
-				FSlateClippingState CurrentClippingState = OutDrawElements.GetClippingManager().GetClippingStates()[OutDrawElements.GetClippingManager().GetClippingIndex()];
-				CachedWindowElements->GetClippingManager().PushClippingState(CurrentClippingState);
+				bool bPushedNewClip = false;
+				const int32 CurrentClippingIndex = OutDrawElements.GetClippingManager().GetClippingIndex();
+
+				if(CurrentClippingIndex != INDEX_NONE)
+				{
+					bPushedNewClip = true;
+					FSlateClippingState CurrentClippingState = OutDrawElements.GetClippingManager().GetClippingStates()[CurrentClippingIndex];
+					CachedWindowElements->GetClippingManager().PushClippingState(CurrentClippingState);
+				}
 
 				//TODO: When SWidget::Paint is called don't drag self if volatile, and we're doing a cache pass.
 				CachedMaxChildLayer = SCompoundWidget::OnPaint(
@@ -369,7 +376,10 @@ int32 SInvalidationPanel::OnPaint( const FPaintArgs& Args, const FGeometry& Allo
 					InWidgetStyle,
 					bParentEnabled);
 
-				CachedWindowElements->GetClippingManager().PopClip();
+				if (bPushedNewClip)
+				{
+					CachedWindowElements->GetClippingManager().PopClip();
+				}
 
 				{
 					CachedResources.Reset();
