PR #3299: UMG Slider Additions (Contributed by [REDACTED_PERSON])

diff --git a/Engine/Source/Editor/EditorStyle/Private/SlateEditorStyle.cpp b/Engine/Source/Editor/EditorStyle/Private/SlateEditorStyle.cpp
--- a/Engine/Source/Editor/EditorStyle/Private/SlateEditorStyle.cpp
+++ b/Engine/Source/Editor/EditorStyle/Private/SlateEditorStyle.cpp
@@ -1631,8 +1631,10 @@ void FSlateEditorStyle::FStyle::SetupGeneralStyles()
 	{
 		FSliderStyle SliderStyle = FSliderStyle()
 			.SetNormalBarImage(FSlateColorBrush(FColor::White))
+			.SetHoveredBarImage(FSlateColorBrush(FColor::White))
 			.SetDisabledBarImage(FSlateColorBrush(FLinearColor::Gray))
 			.SetNormalThumbImage( BOX_BRUSH( "Common/Button", 8.0f/32.0f ) )
+			.SetHoveredThumbImage(BOX_BRUSH("Common/Button", 8.0f / 32.0f))
 			.SetDisabledThumbImage( BOX_BRUSH( "Common/Button_Disabled", 8.0f/32.0f ) )
 			.SetBarThickness(2.0f);
 		Set( "Slider", SliderStyle );
diff --git a/Engine/Source/Runtime/Slate/Private/Widgets/Input/SSlider.cpp b/Engine/Source/Runtime/Slate/Private/Widgets/Input/SSlider.cpp
--- a/Engine/Source/Runtime/Slate/Private/Widgets/Input/SSlider.cpp
+++ b/Engine/Source/Runtime/Slate/Private/Widgets/Input/SSlider.cpp
@@ -10,6 +10,8 @@ void SSlider::Construct( const SSlider::FArguments& InDeclaration )
 	Style = InDeclaration._Style;
 
 	IndentHandle = InDeclaration._IndentHandle;
+	bMouseUsesStep = InDeclaration._MouseUsesStep;
+	bRequiresControllerLock = InDeclaration._RequiresControllerLock;
 	LockedAttribute = InDeclaration._Locked;
 	Orientation = InDeclaration._Orientation;
 	StepSize = InDeclaration._StepSize;
@@ -39,7 +41,7 @@ int32 SSlider::OnPaint( const FPaintArgs& Args, const FGeometry& AllottedGeometr
 	FVector2D SliderEndPoint;
 
 	// calculate slider geometry as if it's a horizontal slider (we'll rotate it later if it's vertical)
-	const FVector2D HandleSize = Style->NormalThumbImage.ImageSize;
+	const FVector2D HandleSize = GetThumbImage()->ImageSize;
 	const FVector2D HalfHandleSize = 0.5f * HandleSize;
 	const float Indentation = IndentHandle.Get() ? HandleSize.X : 0.0f;
 
@@ -78,7 +80,7 @@ int32 SSlider::OnPaint( const FPaintArgs& Args, const FGeometry& AllottedGeometr
 		OutDrawElements,
 		LayerId,
 		SliderGeometry.ToPaintGeometry(BarTopLeft, BarSize),
-		LockedAttribute.Get() ? &Style->DisabledBarImage : &Style->NormalBarImage,
+		GetBarImage(),
 		DrawEffects,
 		SliderBarColor.Get().GetColor(InWidgetStyle) * InWidgetStyle.GetColorAndOpacityTint()
 		);
@@ -89,8 +91,8 @@ int32 SSlider::OnPaint( const FPaintArgs& Args, const FGeometry& AllottedGeometr
 	FSlateDrawElement::MakeBox( 
 		OutDrawElements,
 		LayerId,
-		SliderGeometry.ToPaintGeometry(HandleTopLeftPoint, Style->NormalThumbImage.ImageSize),
-		LockedAttribute.Get() ? &Style->DisabledThumbImage : &Style->NormalThumbImage,
+		SliderGeometry.ToPaintGeometry(HandleTopLeftPoint, GetThumbImage()->ImageSize),
+		GetThumbImage(),
 		DrawEffects,
 		SliderHandleColor.Get().GetColor(InWidgetStyle) * InWidgetStyle.GetColorAndOpacityTint()
 	);
@@ -107,7 +109,8 @@ FVector2D SSlider::ComputeDesiredSize( float ) const
 		return SSliderDesiredSize;
 	}
 
-	const float Thickness = FMath::Max(Style->BarThickness, Style->NormalThumbImage.ImageSize.Y);
+	const float Thickness = FMath::Max(Style->BarThickness, 
+		FMath::Max(Style->NormalThumbImage.ImageSize.Y, Style->HoveredThumbImage.ImageSize.Y));
 
 	if (Orientation == Orient_Vertical)
 	{
@@ -151,7 +154,8 @@ FReply SSlider::OnKeyDown(const FGeometry& MyGeometry, const FKeyEvent& InKeyEve
 		// The controller's bottom face button must be pressed once to begin manipulating the slider's value.
 		// Navigation away from the widget is prevented until the button has been pressed again or focus is lost.
 		// The value can be manipulated by using the game pad's directional arrows ( relative to slider orientation ).
-		if (KeyPressed == EKeys::Enter || KeyPressed == EKeys::SpaceBar || KeyPressed == EKeys::Virtual_Accept)
+		if (KeyPressed == EKeys::Enter || KeyPressed == EKeys::SpaceBar || KeyPressed == EKeys::Virtual_Accept
+			&& bRequiresControllerLock)
 		{
 			if (bControllerInputCaptured == false)
 			{
@@ -167,7 +171,7 @@ FReply SSlider::OnKeyDown(const FGeometry& MyGeometry, const FKeyEvent& InKeyEve
 			}
 		}
 
-		if (bControllerInputCaptured)
+		if (bControllerInputCaptured || !bRequiresControllerLock)
 		{
 			float NewValue = ValueAttribute.Get();
 			if (Orientation == EOrientation::Orient_Horizontal)
@@ -175,10 +179,12 @@ FReply SSlider::OnKeyDown(const FGeometry& MyGeometry, const FKeyEvent& InKeyEve
 				if (KeyPressed == EKeys::Left || KeyPressed == EKeys::Gamepad_DPad_Left || KeyPressed == EKeys::Gamepad_LeftStick_Left)
 				{
 					NewValue -= StepSize.Get();
+					Reply = FReply::Handled();
 				}
 				else if (KeyPressed == EKeys::Right || KeyPressed == EKeys::Gamepad_DPad_Right || KeyPressed == EKeys::Gamepad_LeftStick_Right)
 				{
 					NewValue += StepSize.Get();
+					Reply = FReply::Handled();
 				}
 			}
 			else
@@ -186,15 +192,20 @@ FReply SSlider::OnKeyDown(const FGeometry& MyGeometry, const FKeyEvent& InKeyEve
 				if (KeyPressed == EKeys::Down || KeyPressed == EKeys::Gamepad_DPad_Down || KeyPressed == EKeys::Gamepad_LeftStick_Down)
 				{
 					NewValue -= StepSize.Get();
+					Reply = FReply::Handled();
 				}
 				else if (KeyPressed == EKeys::Up || KeyPressed == EKeys::Gamepad_DPad_Up || KeyPressed == EKeys::Gamepad_LeftStick_Up)
 				{
 					NewValue += StepSize.Get();
+					Reply = FReply::Handled();
 				}
 			}
 
 			CommitValue(FMath::Clamp(NewValue, 0.0f, 1.0f));
-			Reply = FReply::Handled();
+			if (!Reply.IsEventHandled())
+			{
+				Reply = SLeafWidget::OnKeyDown(MyGeometry, InKeyEvent);
+			}
 		}
 		else
 		{
@@ -294,7 +305,7 @@ float SSlider::PositionToValue( const FGeometry& MyGeometry, const FVector2D& Ab
 	float RelativeValue;
 	float Denominator;
 	// Only need X as we rotate the thumb image when rendering vertically
-	const float Indentation = Style->NormalThumbImage.ImageSize.X * (IndentHandle.Get() ? 2.f : 1.f);
+	const float Indentation = GetThumbImage()->ImageSize.X * (IndentHandle.Get() ? 2.f : 1.f);
 	const float HalfIndentation = 0.5f * Indentation;
 
 	if (Orientation == Orient_Horizontal)
@@ -308,8 +319,55 @@ float SSlider::PositionToValue( const FGeometry& MyGeometry, const FVector2D& Ab
 		// Inverse the calculation as top is 0 and bottom is 1
 		RelativeValue = (Denominator != 0.f) ? ((MyGeometry.Size.Y - LocalPosition.Y) - HalfIndentation) / Denominator : 0.f;
 	}
-	
-	return FMath::Clamp(RelativeValue, 0.0f, 1.0f);
+
+	RelativeValue = FMath::Clamp(RelativeValue, 0.0f, 1.0f);
+	if (bMouseUsesStep)
+	{
+		float direction = ValueAttribute.Get() - RelativeValue;
+		if (direction > StepSize.Get() / 2.0f)
+		{
+			return FMath::Clamp(ValueAttribute.Get() - StepSize.Get(), 0.0f, 1.0f);
+		}
+		else if (direction < StepSize.Get() / -2.0f)
+		{
+			return FMath::Clamp(ValueAttribute.Get() + StepSize.Get(), 0.0f, 1.0f);
+		}
+		else
+		{
+			return ValueAttribute.Get();
+		}
+	}
+	return RelativeValue;
+}
+
+const FSlateBrush* SSlider::GetBarImage() const {
+	if (!IsEnabled() || LockedAttribute.Get())
+	{
+		return &Style->DisabledBarImage;
+	}
+	else if (IsHovered())
+	{
+		return &Style->HoveredBarImage;
+	}
+	else
+	{
+		return &Style->NormalBarImage;
+	}
+}
+
+const FSlateBrush* SSlider::GetThumbImage() const {
+	if (!IsEnabled() || LockedAttribute.Get())
+	{
+		return &Style->DisabledThumbImage;
+	}
+	else if (IsHovered())
+	{
+		return &Style->HoveredThumbImage;
+	}
+	else
+	{
+		return &Style->NormalThumbImage;
+	}
 }
 
 float SSlider::GetValue() const
@@ -356,3 +414,12 @@ void SSlider::SetStepSize(const TAttribute<float>& InStepSize)
 {
 	StepSize = InStepSize;
 }
+
+void SSlider::SetMouseUsesStep(bool MouseUsesStep) {
+	bMouseUsesStep = MouseUsesStep;
+}
+
+void SSlider::SetRequiresControllerLock(bool RequiresControllerLock) {
+	bRequiresControllerLock = RequiresControllerLock;
+}
+
diff --git a/Engine/Source/Runtime/Slate/Public/Widgets/Input/SSlider.h b/Engine/Source/Runtime/Slate/Public/Widgets/Input/SSlider.h
--- a/Engine/Source/Runtime/Slate/Public/Widgets/Input/SSlider.h
+++ b/Engine/Source/Runtime/Slate/Public/Widgets/Input/SSlider.h
@@ -26,6 +26,8 @@ public:
 
 	SLATE_BEGIN_ARGS(SSlider)
 		: _IndentHandle(true)
+		, _MouseUsesStep(false)
+		, _RequiresControllerLock(true)
 		, _Locked(false)
 		, _Orientation(EOrientation::Orient_Horizontal)
 		, _SliderBarColor(FLinearColor::White)
@@ -42,6 +44,12 @@ public:
 		/** Whether the slidable area should be indented to fit the handle. */
 		SLATE_ATTRIBUTE( bool, IndentHandle )
 
+		/** Sets new value if mouse position is greater/less than half the step size. */
+		SLATE_ARGUMENT( bool, MouseUsesStep )
+
+		/** Sets whether we have to lock input to change the slider value. */
+		SLATE_ARGUMENT( bool, RequiresControllerLock )
+
 		/** Whether the handle is interactive or fixed. */
 		SLATE_ATTRIBUTE( bool, Locked )
 
@@ -117,6 +125,12 @@ public:
 	/** See the StepSize attribute */
 	void SetStepSize(const TAttribute<float>& InStepSize);
 
+	/** See the MouseUsesStep attribute */
+	void SetMouseUsesStep(bool MouseUsesStep);
+
+	/** See the RequiresControllerLock attribute */
+	void SetRequiresControllerLock(bool RequiresControllerLock);
+
 public:
 
 	// SWidget overrides
@@ -154,6 +168,9 @@ protected:
 	 */
 	float PositionToValue( const FGeometry& MyGeometry, const FVector2D& AbsolutePosition );
 
+	const FSlateBrush* GetBarImage() const;
+	const FSlateBrush* GetThumbImage() const;
+
 protected:
 
 	// Holds the style passed to the widget upon construction.
@@ -184,6 +201,12 @@ protected:
 	// When true, navigation away from the widget is prevented until a new value has been accepted or canceled. 
 	bool bControllerInputCaptured;
 
+	/** Sets new value if mouse position is greater/less than half the step size. */
+	bool bMouseUsesStep;
+
+	/** Sets whether we have to lock input to change the slider value. */
+	bool bRequiresControllerLock;
+
 	/** When true, this slider will be keyboard focusable. Defaults to false. */
 	bool bIsFocusable;
 
diff --git a/Engine/Source/Runtime/SlateCore/Private/Styling/CoreStyle.cpp b/Engine/Source/Runtime/SlateCore/Private/Styling/CoreStyle.cpp
--- a/Engine/Source/Runtime/SlateCore/Private/Styling/CoreStyle.cpp
+++ b/Engine/Source/Runtime/SlateCore/Private/Styling/CoreStyle.cpp
@@ -640,8 +640,10 @@ TSharedRef<ISlateStyle> FCoreStyle::Create( const FName& InStyleSetName )
 	{
 		FSliderStyle SliderStyle = FSliderStyle()
 			.SetNormalBarImage(FSlateColorBrush(FColor::White))
+			.SetHoveredBarImage(FSlateColorBrush(FColor::White))
 			.SetDisabledBarImage(FSlateColorBrush(FLinearColor::Gray))
 			.SetNormalThumbImage( IMAGE_BRUSH( "Common/Button", FVector2D(8.0f, 14.0f) ) )
+			.SetHoveredThumbImage(IMAGE_BRUSH("Common/Button", FVector2D(8.0f, 14.0f)))
 			.SetDisabledThumbImage( IMAGE_BRUSH( "Common/Button_Disabled", FVector2D(8.0f, 14.0f) ) )
 			.SetBarThickness(2.0f);
 		Style->Set( "Slider", SliderStyle );
diff --git a/Engine/Source/Runtime/SlateCore/Private/Styling/SlateTypes.cpp b/Engine/Source/Runtime/SlateCore/Private/Styling/SlateTypes.cpp
--- a/Engine/Source/Runtime/SlateCore/Private/Styling/SlateTypes.cpp
+++ b/Engine/Source/Runtime/SlateCore/Private/Styling/SlateTypes.cpp
@@ -394,8 +394,10 @@ FSliderStyle::FSliderStyle()
 void FSliderStyle::GetResources( TArray< const FSlateBrush* >& OutBrushes ) const
 {
 	OutBrushes.Add( &NormalBarImage );
+	OutBrushes.Add( &HoveredBarImage );
 	OutBrushes.Add( &DisabledBarImage );
 	OutBrushes.Add( &NormalThumbImage );
+	OutBrushes.Add( &HoveredThumbImage );
 	OutBrushes.Add( &DisabledThumbImage );
 }
 
diff --git a/Engine/Source/Runtime/SlateCore/Public/Styling/SlateTypes.h b/Engine/Source/Runtime/SlateCore/Public/Styling/SlateTypes.h
--- a/Engine/Source/Runtime/SlateCore/Public/Styling/SlateTypes.h
+++ b/Engine/Source/Runtime/SlateCore/Public/Styling/SlateTypes.h
@@ -896,6 +896,11 @@ struct SLATECORE_API FSliderStyle : public FSlateWidgetStyle
 	FSlateBrush NormalBarImage;
 	FSliderStyle& SetNormalBarImage(const FSlateBrush& InNormalBarImage){ NormalBarImage = InNormalBarImage; return *this; }
 
+	/** Image to use when the slider bar is in its hovered state */
+	UPROPERTY(EditAnywhere, BlueprintReadWrite, Category=Appearance)
+	FSlateBrush HoveredBarImage;
+	FSliderStyle& SetHoveredBarImage(const FSlateBrush& InHoveredBarImage){ HoveredBarImage = InHoveredBarImage; return *this; }
+
 	/** Image to use when the slider bar is in its disabled state */
 	UPROPERTY(EditAnywhere, BlueprintReadWrite, Category = Appearance)
 	FSlateBrush DisabledBarImage;
@@ -906,6 +911,11 @@ struct SLATECORE_API FSliderStyle : public FSlateWidgetStyle
 	FSlateBrush NormalThumbImage;
 	FSliderStyle& SetNormalThumbImage( const FSlateBrush& InNormalThumbImage ){ NormalThumbImage = InNormalThumbImage; return *this; }
 
+	/** Image to use when the slider thumb is in its hovered state */
+	UPROPERTY(EditAnywhere, BlueprintReadWrite, Category=Appearance)
+	FSlateBrush HoveredThumbImage;
+	FSliderStyle& SetHoveredThumbImage( const FSlateBrush& InHoveredThumbImage ){ HoveredThumbImage = InHoveredThumbImage; return *this; }
+
 	/** Image to use when the slider thumb is in its disabled state */
 	UPROPERTY(EditAnywhere, BlueprintReadWrite, Category=Appearance)
 	FSlateBrush DisabledThumbImage;
diff --git a/Engine/Source/Runtime/UMG/Private/Components/Slider.cpp b/Engine/Source/Runtime/UMG/Private/Components/Slider.cpp
--- a/Engine/Source/Runtime/UMG/Private/Components/Slider.cpp
+++ b/Engine/Source/Runtime/UMG/Private/Components/Slider.cpp
@@ -19,6 +19,8 @@ USlider::USlider(const FObjectInitializer& ObjectInitializer)
 	SSlider::FArguments Defaults;
 	WidgetStyle = *Defaults._Style;
 	IsFocusable = true;
+	MouseUsesStep = false;
+	RequiresControllerLock = true;
 }
 
 TSharedRef<SWidget> USlider::RebuildWidget()
@@ -42,6 +44,8 @@ void USlider::SynchronizeProperties()
 	TAttribute<float> ValueBinding = PROPERTY_BINDING(float, Value);
 	
 	MySlider->SetOrientation(Orientation);
+	MySlider->SetMouseUsesStep(MouseUsesStep);
+	MySlider->SetRequiresControllerLock(RequiresControllerLock);
 	MySlider->SetSliderBarColor(SliderBarColor);
 	MySlider->SetSliderHandleColor(SliderHandleColor);
 	MySlider->SetValue(ValueBinding);
diff --git a/Engine/Source/Runtime/UMG/Public/Components/Slider.h b/Engine/Source/Runtime/UMG/Public/Components/Slider.h
--- a/Engine/Source/Runtime/UMG/Public/Components/Slider.h
+++ b/Engine/Source/Runtime/UMG/Public/Components/Slider.h
@@ -62,6 +62,14 @@ public:
 	UPROPERTY(EditAnywhere, BlueprintReadOnly, Category=Appearance, AdvancedDisplay)
 	bool Locked;
 
+	/** Sets new value if mouse position is greater/less than half the step size. */
+	UPROPERTY(EditAnywhere, BlueprintReadOnly, Category = Appearance, AdvancedDisplay)
+	bool MouseUsesStep;
+
+	/** Sets whether we have to lock input to change the slider value. */
+	UPROPERTY(EditAnywhere, BlueprintReadOnly, Category = Appearance, AdvancedDisplay)
+	bool RequiresControllerLock;
+
 	/** The amount to adjust the value by, when using a controller or keyboard */
 	UPROPERTY(EditAnywhere, BlueprintReadOnly, Category=Appearance, meta=( ClampMin="0", ClampMax="1", UIMin="0", UIMax="1"))
 	float StepSize;
