Fix warning/error with logical operators

diff --git a/Engine/Source/Runtime/Slate/Private/Widgets/Input/SSlider.cpp b/Engine/Source/Runtime/Slate/Private/Widgets/Input/SSlider.cpp
--- a/Engine/Source/Runtime/Slate/Private/Widgets/Input/SSlider.cpp
+++ b/Engine/Source/Runtime/Slate/Private/Widgets/Input/SSlider.cpp
@@ -154,7 +154,7 @@ FReply SSlider::OnKeyDown(const FGeometry& MyGeometry, const FKeyEvent& InKeyEve
 		// The controller's bottom face button must be pressed once to begin manipulating the slider's value.
 		// Navigation away from the widget is prevented until the button has been pressed again or focus is lost.
 		// The value can be manipulated by using the game pad's directional arrows ( relative to slider orientation ).
-		if (KeyPressed == EKeys::Enter || KeyPressed == EKeys::SpaceBar || KeyPressed == EKeys::Virtual_Accept
+		if ((KeyPressed == EKeys::Enter || KeyPressed == EKeys::SpaceBar || KeyPressed == EKeys::Virtual_Accept)
 			&& bRequiresControllerLock)
 		{
 			if (bControllerInputCaptured == false)
