PR #3458: UE-43728: Always show scrollbar when necessary (Contributed by [REDACTED_PERSON])

diff --git a/Engine/Source/Runtime/Slate/Private/Widgets/Layout/SScrollBar.cpp b/Engine/Source/Runtime/Slate/Private/Widgets/Layout/SScrollBar.cpp
--- a/Engine/Source/Runtime/Slate/Private/Widgets/Layout/SScrollBar.cpp
+++ b/Engine/Source/Runtime/Slate/Private/Widgets/Layout/SScrollBar.cpp
@@ -234,7 +234,7 @@ FLinearColor SScrollBar::GetThumbOpacity() const
 	}
 	else
 	{
-		if ( bHideWhenNotInUse && !bAlwaysShow )
+		if ( bHideWhenNotInUse && !bAlwaysShowScrollbar )
 		{
 			const double LastInteractionDelta = bIsScrolling ? 0 : ( FSlateApplication::Get().GetCurrentTime() - LastInteractionTime );
 
@@ -343,7 +343,8 @@ void SScrollBar::SetThickness(TAttribute<FVector2D> InThickness)
 
 void SScrollBar::SetScrollBarAlwaysVisible(bool InAlwaysVisible)
 {
-	bAlwaysShow = InAlwaysVisible;
+	bAlwaysShowScrollbar = InAlwaysVisible;
+
 	if ( InAlwaysVisible )
 	{
 		Visibility = EVisibility::Visible;
@@ -352,4 +353,10 @@ void SScrollBar::SetScrollBarAlwaysVisible(bool InAlwaysVisible)
 	{
 		Visibility = TAttribute<EVisibility>(SharedThis(this), &SScrollBar::ShouldBeVisible);
 	}
+	Track->SetIsAlwaysVisible(InAlwaysVisible);
+}
+
+bool SScrollBar::AlwaysShowScrollbar() const
+{
+	return bAlwaysShowScrollbar;
 }
diff --git a/Engine/Source/Runtime/Slate/Private/Widgets/Layout/SScrollBarTrack.cpp b/Engine/Source/Runtime/Slate/Private/Widgets/Layout/SScrollBarTrack.cpp
--- a/Engine/Source/Runtime/Slate/Private/Widgets/Layout/SScrollBarTrack.cpp
+++ b/Engine/Source/Runtime/Slate/Private/Widgets/Layout/SScrollBarTrack.cpp
@@ -123,7 +123,7 @@ bool SScrollBarTrack::IsNeeded() const
 {
 	// We use a small epsilon here to avoid the scroll bar showing up when all of the content is already in view, due to
 	// floating point precision when the scroll bar state is set
-	return ThumbSizeFraction < (1.0f - KINDA_SMALL_NUMBER);
+	return ThumbSizeFraction < (1.0f - KINDA_SMALL_NUMBER) || bIsAlwaysVisible;
 }
 
 float SScrollBarTrack::DistanceFromTop() const
@@ -145,3 +145,8 @@ float SScrollBarTrack::GetThumbSizeFraction() const
 {
 	return ThumbSizeFraction;
 }
+
+void SScrollBarTrack::SetIsAlwaysVisible(bool InIsAlwaysVisible)
+{
+	bIsAlwaysVisible = InIsAlwaysVisible;
+}
diff --git a/Engine/Source/Runtime/Slate/Private/Widgets/Layout/SScrollBox.cpp b/Engine/Source/Runtime/Slate/Private/Widgets/Layout/SScrollBox.cpp
--- a/Engine/Source/Runtime/Slate/Private/Widgets/Layout/SScrollBox.cpp
+++ b/Engine/Source/Runtime/Slate/Private/Widgets/Layout/SScrollBox.cpp
@@ -840,7 +840,7 @@ void SScrollBox::OnMouseLeave( const FPointerEvent& MouseEvent )
 
 FReply SScrollBox::OnMouseWheel( const FGeometry& MyGeometry, const FPointerEvent& MouseEvent )
 {
-	if (ScrollBar->IsNeeded() && ConsumeMouseWheel != EConsumeMouseWheel::Never )
+	if ((ScrollBar->IsNeeded() && ConsumeMouseWheel != EConsumeMouseWheel::Never) || ConsumeMouseWheel == EConsumeMouseWheel::Always)
 	{
 		// Make sure scroll velocity is cleared so it doesn't fight with the mouse wheel input
 		InertialScrollManager.ClearScrollVelocity();
@@ -876,7 +876,7 @@ bool SScrollBox::ScrollBy(const FGeometry& AllottedGeometry, float LocalScrollAm
 	if (LocalScrollAmount != 0 )
 	{
 		const float ScrollMin = 0.0f;
-		const float ScrollMax = ContentSize - GetScrollComponentFromVector(ScrollPanelGeometry.GetLocalSize());
+		const float ScrollMax = FMath::Max(ContentSize - GetScrollComponentFromVector(ScrollPanelGeometry.GetLocalSize()), 0.0f);
 
 		if ( AllowOverscroll == EAllowOverscroll::Yes && Overscrolling == EAllowOverscroll::Yes && Overscroll.ShouldApplyOverscroll(DesiredScrollOffset == 0, DesiredScrollOffset == ScrollMax, LocalScrollAmount) )
 		{
diff --git a/Engine/Source/Runtime/Slate/Public/Widgets/Layout/SScrollBar.h b/Engine/Source/Runtime/Slate/Public/Widgets/Layout/SScrollBar.h
--- a/Engine/Source/Runtime/Slate/Public/Widgets/Layout/SScrollBar.h
+++ b/Engine/Source/Runtime/Slate/Public/Widgets/Layout/SScrollBar.h
@@ -117,6 +117,9 @@ public:
 	/** See ScrollBarAlwaysVisible attribute */
 	void SetScrollBarAlwaysVisible(bool InAlwaysVisible);
 
+	/** Returns True when the scrollbar should always be shown, else False */
+	bool AlwaysShowScrollbar() const;
+
 	/** Allows external scrolling panels to notify the scrollbar when scrolling begins. */
 	virtual void BeginScrolling();
 
@@ -149,11 +152,11 @@ protected:
 	FOnUserScrolled OnUserScrolled;
 	float DragGrabOffset;
 	EOrientation Orientation;
+	bool bAlwaysShowScrollbar;
 	EFocusCause DragFocusCause;
 	bool bHideWhenNotInUse;
 	bool bIsScrolling;
 	double LastInteractionTime;
-	bool bAlwaysShow;
 
 	/** Image to use when the scrollbar thumb is in its normal state */
 	const FSlateBrush* NormalThumbImage;
diff --git a/Engine/Source/Runtime/Slate/Public/Widgets/Layout/SScrollBarTrack.h b/Engine/Source/Runtime/Slate/Public/Widgets/Layout/SScrollBarTrack.h
--- a/Engine/Source/Runtime/Slate/Public/Widgets/Layout/SScrollBarTrack.h
+++ b/Engine/Source/Runtime/Slate/Public/Widgets/Layout/SScrollBarTrack.h
@@ -106,6 +106,8 @@ public:
 
 	float GetThumbSizeFraction() const;
 
+	void SetIsAlwaysVisible(bool InIsAlwaysVisible);
+
 protected:
 
 	static const int32 TOP_SLOT_INDEX = 0;
@@ -117,5 +119,6 @@ protected:
 	float OffsetFraction;
 	float ThumbSizeFraction;
 	float MinThumbSize;
+	bool bIsAlwaysVisible = false;
 	EOrientation Orientation;
 };
diff --git a/Engine/Source/Runtime/UMG/Public/Components/ScrollBox.h b/Engine/Source/Runtime/UMG/Public/Components/ScrollBox.h
--- a/Engine/Source/Runtime/UMG/Public/Components/ScrollBox.h
+++ b/Engine/Source/Runtime/UMG/Public/Components/ScrollBox.h
@@ -44,7 +44,7 @@ public:
 	UPROPERTY(EditAnywhere, BlueprintReadOnly, Category="Scroll")
 	ESlateVisibility ScrollBarVisibility;
 
-	/**  Enable to always consume mouse wheel event, even when scrolling is not possible */
+	/** When mouse wheel events should be consumed. */
 	UPROPERTY(EditAnywhere, BlueprintReadOnly, Category = "Scroll")
 	EConsumeMouseWheel ConsumeMouseWheel;
 
