Added early out and log if focus is attempted to be set when the window is suppost to be be disabled due to a modal window being up.

diff --git a/Engine/Source/Runtime/Slate/Private/Framework/Application/SlateApplication.cpp b/Engine/Source/Runtime/Slate/Private/Framework/Application/SlateApplication.cpp
--- a/Engine/Source/Runtime/Slate/Private/Framework/Application/SlateApplication.cpp
+++ b/Engine/Source/Runtime/Slate/Private/Framework/Application/SlateApplication.cpp
@@ -2817,6 +2817,16 @@ bool FSlateApplication::SetUserFocus(FSlateUser* User, const FWidgetPath& InFocu
 		return false;
 	}
 
+	if (InFocusPath.IsValid())
+	{
+		TSharedRef<SWindow> Window = InFocusPath.GetWindow();
+		if (ActiveModalWindows.Num() != 0 && !(Window->IsDescendantOf(GetActiveModalWindow()) || ActiveModalWindows.Top() == Window))
+		{
+			UE_LOG(LogSlate, Warning, TEXT("Ignoring SetUserFocus because it's not an active modal Window (user %i not set to %s."), User->GetUserIndex(), *InFocusPath.GetLastWidget()->ToString());
+			return false;
+		}
+	}
+
 	TSharedPtr<IWidgetReflector> WidgetReflector = WidgetReflectorPtr.Pin();
 	const bool bReflectorShowingFocus = WidgetReflector.IsValid() && WidgetReflector->IsShowingFocus();
 
@@ -4899,9 +4909,9 @@ bool FSlateApplication::ShouldProcessUserInputMessages( const TSharedPtr< FGener
 		Window = FSlateWindowHelper::FindWindowByPlatformWindow( SlateWindows, PlatformWindow.ToSharedRef() );
 	}
 
-	if ( ActiveModalWindows.Num() == 0 || 
-		( Window.IsValid() &&
-			( Window->IsDescendantOf( GetActiveModalWindow() ) || ActiveModalWindows.Contains( Window ) ) ) )
+	if (ActiveModalWindows.Num() == 0 ||
+		(Window.IsValid() &&
+		(Window->IsDescendantOf(GetActiveModalWindow()) || ActiveModalWindows.Top() == Window)))
 	{
 		return true;
 	}
