Fixed bug with InputPreProcessorsHelper::Add not correctly adding input processors

#rnx

diff --git a/Engine/Source/Runtime/Slate/Private/Framework/Application/SlateApplication.cpp b/Engine/Source/Runtime/Slate/Private/Framework/Application/SlateApplication.cpp
--- a/Engine/Source/Runtime/Slate/Private/Framework/Application/SlateApplication.cpp
+++ b/Engine/Source/Runtime/Slate/Private/Framework/Application/SlateApplication.cpp
@@ -7341,7 +7341,7 @@ bool FSlateApplication::InputPreProcessorsHelper::Add(TSharedPtr<IInputProcessor
 		InputPreProcessorList.AddUnique(InputProcessor);
 		bResult = true;
 		}
-	else if (!InputPreProcessorList.Find(InputProcessor))
+	else if (InputPreProcessorList.Find(InputProcessor) == INDEX_NONE)
 		{
 		InputPreProcessorList.Insert(InputProcessor, Index);
 		bResult = true;
