UE-60716 Show cursor in Editor text boxes when focused on Android

#Android
#4.21

Updating the widget's cursor position from the native code.
This also overrides the fix for UE49124 Cursor in virtual keyboard and UMG don't match - initially the solution was to hide the widget's cursor.

diff --git a/Engine/Build/Android/Java/src/com/epicgames/ue4/GameActivity.java.template b/Engine/Build/Android/Java/src/com/epicgames/ue4/GameActivity.java.template
--- a/Engine/Build/Android/Java/src/com/epicgames/ue4/GameActivity.java.template
+++ b/Engine/Build/Android/Java/src/com/epicgames/ue4/GameActivity.java.template
@@ -3135,6 +3135,13 @@ public class GameActivity extends NativeActivity implements SurfaceHolder.Callba
 		    return super.dispatchKeyEvent(event);
 		}
 
+		// listen for cursor position changes &  send message to the UMG widget
+		@Override   
+		protected void onSelectionChanged(int selStart, int selEnd) 
+		{
+			nativeVirtualKeyboardSendSelection(selStart, selEnd);
+		}
+
 		//extend associated InputConnection to handle special keys 
 		@Override
 		public InputConnection onCreateInputConnection(EditorInfo outAttrs) 
@@ -3475,6 +3482,8 @@ public class GameActivity extends NativeActivity implements SurfaceHolder.Callba
 	public native void nativeVirtualKeyboardChanged(String contents);
 	public native void nativeVirtualKeyboardResult(boolean update, String contents);
 	public native void nativeVirtualKeyboardSendKey(int keyCode);
+	public native void nativeVirtualKeyboardSendTextSelection(String contents, int selStart, int selEnd);
+	public native void nativeVirtualKeyboardSendSelection(int selStart, int selEnd);
 
 	public native void nativeInitHMDs();
 
diff --git a/Engine/Source/Runtime/Launch/Private/Android/AndroidJNI.cpp b/Engine/Source/Runtime/Launch/Private/Android/AndroidJNI.cpp
--- a/Engine/Source/Runtime/Launch/Private/Android/AndroidJNI.cpp
+++ b/Engine/Source/Runtime/Launch/Private/Android/AndroidJNI.cpp
@@ -1,4 +1,4 @@
-﻿// Copyright 1998-2018 Epic Games, Inc. All Rights Reserved.
+// Copyright 1998-2018 Epic Games, Inc. All Rights Reserved.
 
 #include "Android/AndroidJNI.h"
 
@@ -897,6 +897,24 @@ JNI_METHOD void Java_com_epicgames_ue4_GameActivity_nativeVirtualKeyboardSendKey
 	FAndroidInputInterface::DeferMessage(Message);
 }
 
+JNI_METHOD void Java_com_epicgames_ue4_GameActivity_nativeVirtualKeyboardSendSelection(JNIEnv* jenv, jobject thiz, jint selStart, jint selEnd)
+{
+	// call to set the widget selection on game thread
+	if (VirtualKeyboardWidget != NULL)
+	{
+		if (FTaskGraphInterface::IsRunning())
+		{
+			FGraphEventRef SetWidgetSelection = FFunctionGraphTask::CreateAndDispatchWhenReady([&, selStart, selEnd]()
+			{
+				if (VirtualKeyboardWidget != NULL)
+				{
+					VirtualKeyboardWidget->SetSelectionFromVirtualKeyboard(selStart, selEnd);
+				}
+			}, TStatId(), NULL, ENamedThreads::GameThread);
+		}
+	}
+}
+
 void AndroidThunkCpp_LaunchURL(const FString& URL)
 {
 	if (JNIEnv* Env = FAndroidApplication::GetJavaEnv())
diff --git a/Engine/Source/Runtime/Slate/Private/Framework/Text/SlateTextLayout.cpp b/Engine/Source/Runtime/Slate/Private/Framework/Text/SlateTextLayout.cpp
--- a/Engine/Source/Runtime/Slate/Private/Framework/Text/SlateTextLayout.cpp
+++ b/Engine/Source/Runtime/Slate/Private/Framework/Text/SlateTextLayout.cpp
@@ -123,12 +123,8 @@ int32 FSlateTextLayout::OnPaint( const FPaintArgs& Args, const FGeometry& Allott
 		}
 
 		// Render any overlays for this line
-		//#jira UE - 49124 Cursor in virtual keyboard and UMG don't match
-		if (FSlateApplication::Get().AllowMoveCursor())
-		{
-			const int32 HighestOverlayLayerId = OnPaintHighlights(Args, LineView, LineView.OverlayHighlights, DefaultTextStyle, AllottedGeometry, MyCullingRect, OutDrawElements, HighestBlockLayerId, InWidgetStyle, bParentEnabled);
-			HighestLayerId = FMath::Max(HighestLayerId, HighestOverlayLayerId);
-		}
+		const int32 HighestOverlayLayerId = OnPaintHighlights(Args, LineView, LineView.OverlayHighlights, DefaultTextStyle, AllottedGeometry, MyCullingRect, OutDrawElements, HighestBlockLayerId, InWidgetStyle, bParentEnabled);
+		HighestLayerId = FMath::Max(HighestLayerId, HighestOverlayLayerId);
 	}
 
 	return HighestLayerId;
diff --git a/Engine/Source/Runtime/Slate/Private/Widgets/Input/SVirtualKeyboardEntry.cpp b/Engine/Source/Runtime/Slate/Private/Widgets/Input/SVirtualKeyboardEntry.cpp
--- a/Engine/Source/Runtime/Slate/Private/Widgets/Input/SVirtualKeyboardEntry.cpp
+++ b/Engine/Source/Runtime/Slate/Private/Widgets/Input/SVirtualKeyboardEntry.cpp
@@ -70,6 +70,11 @@ void SVirtualKeyboardEntry::SetTextFromVirtualKeyboard(const FText& InNewText, E
 	}
 }
 
+void SVirtualKeyboardEntry::SetSelectionFromVirtualKeyboard(int InSelStart, int InSelEnd)
+{
+	//Nothing to do for widgets without a cursor
+}
+
 /**
  * Restores the text to the original state
  */
diff --git a/Engine/Source/Runtime/Slate/Private/Widgets/Text/SlateEditableTextLayout.cpp b/Engine/Source/Runtime/Slate/Private/Widgets/Text/SlateEditableTextLayout.cpp
--- a/Engine/Source/Runtime/Slate/Private/Widgets/Text/SlateEditableTextLayout.cpp
+++ b/Engine/Source/Runtime/Slate/Private/Widgets/Text/SlateEditableTextLayout.cpp
@@ -140,6 +140,7 @@ FSlateEditableTextLayout::FSlateEditableTextLayout(ISlateEditableTextWidget& InO
 	bHasDragSelectedSinceFocused = false;
 	bTextChangedByVirtualKeyboard = false;
 	bTextCommittedByVirtualKeyboard = false;
+	bSelectionChangedExternally = false;
 	VirtualKeyboardTextCommitType = ETextCommit::Default;
 
 	CachedSize = FVector2D::ZeroVector;
@@ -3178,6 +3179,22 @@ void FSlateEditableTextLayout::Tick(const FGeometry& AllottedGeometry, const dou
 		Refresh();
 	}
 
+	if (bSelectionChangedExternally)
+	{
+		bSelectionChangedExternally = false;
+		if (TextInputMethodContext.IsValid())
+		{
+			if (ExternalSelectionStart <= ExternalSelectionEnd)
+			{
+				TextInputMethodContext->SetSelectionRange(ExternalSelectionStart, ExternalSelectionEnd - ExternalSelectionStart, ITextInputMethodContext::ECaretPosition::Beginning);
+			}
+			else
+			{
+				TextInputMethodContext->SetSelectionRange(ExternalSelectionEnd, ExternalSelectionStart - ExternalSelectionEnd, ITextInputMethodContext::ECaretPosition::Ending);
+			}
+		}
+	}
+
 	// Update the search before we process the next PositionToScrollIntoView
 	{
 		const FText& SearchTextToSet = BoundSearchText.Get(FText::GetEmpty());
@@ -3494,6 +3511,17 @@ void FSlateEditableTextLayout::FVirtualKeyboardEntry::SetTextFromVirtualKeyboard
 	}
 }
 
+void FSlateEditableTextLayout::FVirtualKeyboardEntry::SetSelectionFromVirtualKeyboard(int InSelStart, int InSelEnd)
+{
+	// Update the text selection and the cursor position
+	// This method is called externally (eg. on Android from the native virtual keyboard implementation) 
+	// The text may also change on the same frame, so the external selection must happen in Tick after the text update
+
+	OwnerLayout->bSelectionChangedExternally = true;
+	OwnerLayout->ExternalSelectionStart = InSelStart;
+	OwnerLayout->ExternalSelectionEnd = InSelEnd;
+}
+
 FText FSlateEditableTextLayout::FVirtualKeyboardEntry::GetText() const
 {
 	return OwnerLayout->GetText();
diff --git a/Engine/Source/Runtime/Slate/Public/Widgets/Input/IVirtualKeyboardEntry.h b/Engine/Source/Runtime/Slate/Public/Widgets/Input/IVirtualKeyboardEntry.h
--- a/Engine/Source/Runtime/Slate/Public/Widgets/Input/IVirtualKeyboardEntry.h
+++ b/Engine/Source/Runtime/Slate/Public/Widgets/Input/IVirtualKeyboardEntry.h
@@ -60,6 +60,8 @@ public:
 	* @param CommitType If we are sending a TextCommitted event, what commit type is it
 	*/
 	virtual void SetTextFromVirtualKeyboard(const FText& InNewText, ETextEntryType TextEntryType) = 0;
+
+	virtual void SetSelectionFromVirtualKeyboard(int InSelStart, int SelEnd) = 0;
 	
 	/**
 	* Returns the text.
diff --git a/Engine/Source/Runtime/Slate/Public/Widgets/Input/SVirtualKeyboardEntry.h b/Engine/Source/Runtime/Slate/Public/Widgets/Input/SVirtualKeyboardEntry.h
--- a/Engine/Source/Runtime/Slate/Public/Widgets/Input/SVirtualKeyboardEntry.h
+++ b/Engine/Source/Runtime/Slate/Public/Widgets/Input/SVirtualKeyboardEntry.h
@@ -106,6 +106,7 @@ public:
 public:
 	//~ Begin IVirtualKeyboardEntry Interface
 	virtual void SetTextFromVirtualKeyboard(const FText& InNewText, ETextEntryType TextEntryType) override;
+	virtual void SetSelectionFromVirtualKeyboard(int InSelStart, int InSelEnd) override;
 
 	virtual FText GetText() const override
 	{
diff --git a/Engine/Source/Runtime/Slate/Public/Widgets/Text/SlateEditableTextLayout.h b/Engine/Source/Runtime/Slate/Public/Widgets/Text/SlateEditableTextLayout.h
--- a/Engine/Source/Runtime/Slate/Public/Widgets/Text/SlateEditableTextLayout.h
+++ b/Engine/Source/Runtime/Slate/Public/Widgets/Text/SlateEditableTextLayout.h
@@ -390,6 +390,8 @@ private:
 		static TSharedRef<FVirtualKeyboardEntry> Create(FSlateEditableTextLayout& InOwnerLayout);
 
 		virtual void SetTextFromVirtualKeyboard(const FText& InNewText, ETextEntryType TextEntryType) override;
+		virtual void SetSelectionFromVirtualKeyboard(int InSelStart, int InSelEnd) override;
+
 		virtual FText GetText() const override;
 		virtual FText GetHintText() const override;
 		virtual EKeyboardType GetVirtualKeyboardType() const override;
@@ -617,4 +619,12 @@ private:
 
 	/** Information about any active context menu widgets */
 	FActiveTextEditContextMenu ActiveContextMenu;
+
+	/** Whether the cursor position has been changed externally*/
+	bool bSelectionChangedExternally;
+
+	/** The boundaries of the external selection */
+	int ExternalSelectionStart;
+	int ExternalSelectionEnd;
+
 };
