Setting SRichTextBlock's clipping method to correctly be OnDemand by default

diff --git a/Engine/Source/Runtime/Slate/Public/Widgets/Text/SRichTextBlock.h b/Engine/Source/Runtime/Slate/Public/Widgets/Text/SRichTextBlock.h
--- a/Engine/Source/Runtime/Slate/Public/Widgets/Text/SRichTextBlock.h
+++ b/Engine/Source/Runtime/Slate/Public/Widgets/Text/SRichTextBlock.h
@@ -51,7 +51,9 @@ public:
 		, _Decorators()
 		, _Parser()
 		, _MinDesiredWidth()
-	{}
+	{
+		_Clipping = EWidgetClipping::OnDemand;
+	}
 		/** The text displayed in this text block */
 		SLATE_ATTRIBUTE( FText, Text )
 
