SWidgetSwitcher - was not correctly calculcating desired size and including the padding the current child was asking for, which was causing layout issues any time someone added padding inside a widget switcher, the inner wize was never given enough space.

Matt.Kuhlenschmidt

diff --git a/Engine/Source/Runtime/Slate/Private/Widgets/Layout/SWidgetSwitcher.cpp b/Engine/Source/Runtime/Slate/Private/Widgets/Layout/SWidgetSwitcher.cpp
--- a/Engine/Source/Runtime/Slate/Private/Widgets/Layout/SWidgetSwitcher.cpp
+++ b/Engine/Source/Runtime/Slate/Private/Widgets/Layout/SWidgetSwitcher.cpp
@@ -3,7 +3,6 @@
 #include "Widgets/Layout/SWidgetSwitcher.h"
 #include "Layout/LayoutUtils.h"
 
-
 SWidgetSwitcher::SWidgetSwitcher()
 	: WidgetIndex()
 	, AllChildren(this)
@@ -12,7 +11,6 @@ SWidgetSwitcher::SWidgetSwitcher()
 	SetCanTick(false);
 }
 
-
 SWidgetSwitcher::FSlot& SWidgetSwitcher::AddSlot( int32 SlotIndex )
 {
 	FSlot* NewSlot = new FSlot();
@@ -35,7 +33,6 @@ SWidgetSwitcher::FSlot& SWidgetSwitcher::AddSlot( int32 SlotIndex )
 	return *NewSlot;
 }
 
-
 void SWidgetSwitcher::Construct( const FArguments& InArgs )
 {
 	for (int32 Index = 0; Index < InArgs.Slots.Num(); ++Index)
@@ -46,7 +43,6 @@ void SWidgetSwitcher::Construct( const FArguments& InArgs )
 	WidgetIndex = InArgs._WidgetIndex;
 }
 
-
 TSharedPtr<SWidget> SWidgetSwitcher::GetActiveWidget( ) const
 {
 	const FSlot* ActiveSlot = GetActiveSlot();
@@ -59,7 +55,6 @@ TSharedPtr<SWidget> SWidgetSwitcher::GetActiveWidget( ) const
 	return nullptr;
 }
 
-
 TSharedPtr<SWidget> SWidgetSwitcher::GetWidget( int32 SlotIndex ) const
 {
 	if (AllChildren.IsValidIndex(SlotIndex))
@@ -70,7 +65,6 @@ TSharedPtr<SWidget> SWidgetSwitcher::GetWidget( int32 SlotIndex ) const
 	return nullptr;
 }
 
-
 int32 SWidgetSwitcher::GetWidgetIndex( TSharedRef<SWidget> Widget ) const
 {
 	for (int32 Index = 0; Index < AllChildren.Num(); ++Index)
@@ -86,10 +80,9 @@ int32 SWidgetSwitcher::GetWidgetIndex( TSharedRef<SWidget> Widget ) const
 	return INDEX_NONE;
 }
 
-
 int32 SWidgetSwitcher::RemoveSlot( TSharedRef<SWidget> WidgetToRemove )
 {
-	for (int32 SlotIndex=0; SlotIndex < AllChildren.Num(); ++SlotIndex)
+	for (int32 SlotIndex = 0; SlotIndex < AllChildren.Num(); ++SlotIndex)
 	{
 		if (AllChildren[SlotIndex].GetWidget() == WidgetToRemove)
 		{
@@ -131,35 +124,32 @@ bool SWidgetSwitcher::ValidatePathToChild(SWidget* InChild)
 	return InChild == GetActiveWidget().Get();
 }
 
-void SWidgetSwitcher::OnArrangeChildren( const FGeometry& AllottedGeometry, FArrangedChildren& ArrangedChildren ) const
+void SWidgetSwitcher::OnArrangeChildren(const FGeometry& AllottedGeometry, FArrangedChildren& ArrangedChildren) const
 {
 	const TAttribute<FVector2D> ContentScale = FVector2D::UnitVector;
 
-	if (AllChildren.Num() > 0)
+	if (const FSlot* ActiveSlotPtr = GetActiveSlot())
 	{
-		const FSlot* ActiveSlotPtr = GetActiveSlot();	// Returns null if unsafe attribute WidgetIndex is out-of-bounds
-		if (ActiveSlotPtr != nullptr)
-		{
-			ArrangeSingleChild(GSlateFlowDirection, AllottedGeometry, ArrangedChildren, *ActiveSlotPtr, ContentScale);
-		}
+		ArrangeSingleChild(GSlateFlowDirection, AllottedGeometry, ArrangedChildren, *ActiveSlotPtr, ContentScale);
 	}
 }
 	
-FVector2D SWidgetSwitcher::ComputeDesiredSize( float ) const
+FVector2D SWidgetSwitcher::ComputeDesiredSize(float) const
 {
-	if (AllChildren.Num() > 0)
+	if (const FSlot* ActiveSlotPtr = GetActiveSlot())
 	{
-		const FSlot* ActiveSlotPtr = GetActiveSlot();	// Returns null if unsafe attribute WidgetIndex is out-of-bounds
-		if (ActiveSlotPtr != nullptr)
+		const TSharedRef<SWidget>& Widget = ActiveSlotPtr->GetWidget();
+		const EVisibility ChildVisibility = Widget->GetVisibility();
+		if (ChildVisibility != EVisibility::Collapsed)
 		{
-			return ActiveSlotPtr->GetWidget()->GetDesiredSize();
+			return Widget->GetDesiredSize() + ActiveSlotPtr->SlotPadding.Get().GetDesiredSize();
 		}
 	}
 
 	return FVector2D::ZeroVector;
 }
 
-FChildren* SWidgetSwitcher::GetChildren( )
+FChildren* SWidgetSwitcher::GetChildren()
 {
 	return &OneDynamicChild;
 }
