Ensure the mouse cursor stays visible and mouse move events are processed only when the mouse button is down when faking touch events.

diff --git a/Engine/Source/Runtime/Slate/Private/Framework/Application/SlateApplication.cpp b/Engine/Source/Runtime/Slate/Private/Framework/Application/SlateApplication.cpp
--- a/Engine/Source/Runtime/Slate/Private/Framework/Application/SlateApplication.cpp
+++ b/Engine/Source/Runtime/Slate/Private/Framework/Application/SlateApplication.cpp
@@ -3661,7 +3661,10 @@ void FSlateApplication::ProcessCursorReply(const FCursorReply& CursorReply)
 		{
 			CursorReply.GetCursorWidget()->SetVisibility(EVisibility::HitTestInvisible);
 			CursorWindowPtr = CursorReply.GetCursorWindow();
-			PlatformApplication->Cursor->SetType(EMouseCursor::Custom);
+			if (!IsFakingTouchEvents())
+			{
+				PlatformApplication->Cursor->SetType(EMouseCursor::Custom);
+			}
 		}
 		else
 		{
@@ -6311,13 +6314,15 @@ FReply FSlateApplication::RouteMouseWheelOrGestureEvent(const FWidgetPath& Widge
 
 bool FSlateApplication::OnMouseMove()
 {
-	// convert to touch event if we are faking it	
-	if (bIsFakingTouched)
-	{
-		return OnTouchMoved(PlatformApplication->Cursor->GetPosition(), 1.0f, 0, 0);
-	}
-	else if (!bIsGameFakingTouch && bIsFakingTouch)
+	if (bIsFakingTouched || bIsGameFakingTouch)
 	{
+		// convert to touch event if we are faking it
+		if (bIsFakingTouched)
+		{
+			return OnTouchMoved(PlatformApplication->Cursor->GetPosition(), 1.0f, 0, 0);
+		}
+
+		// Throw out the mouse move event if we're faking touch events but the mouse button isn't down.
 		return false;
 	}
 
@@ -6351,11 +6356,18 @@ bool FSlateApplication::OnMouseMove()
 
 bool FSlateApplication::OnRawMouseMove( const int32 X, const int32 Y )
 {
-	if (bIsFakingTouched)
+	if (bIsFakingTouched || bIsGameFakingTouch)
 	{
-		return OnTouchMoved(GetCursorPos(), 1.0f, 0, 0);
+		// convert to touch event if we are faking it
+		if (bIsFakingTouched)
+		{
+			return OnTouchMoved(GetCursorPos(), 1.0f, 0, 0);
+		}
+
+		// Throw out the mouse move event if we're faking touch events but the mouse button isn't down.
+		return false;
 	}
-	
+
 	if ( X != 0 || Y != 0 )
 	{
 		FPointerEvent MouseEvent(
