SListView::NavigationSelect now requests the nav to the item prior to updating selection
- Has no effect on prior behavior, but now other nav requests can be made in response to the selection change

diff --git a/Engine/Source/Runtime/Slate/Public/Widgets/Views/SListView.h b/Engine/Source/Runtime/Slate/Public/Widgets/Views/SListView.h
--- a/Engine/Source/Runtime/Slate/Public/Widgets/Views/SListView.h
+++ b/Engine/Source/Runtime/Slate/Public/Widgets/Views/SListView.h
@@ -1596,6 +1596,9 @@ protected:
 			// Must be set before signaling selection changes because sometimes new items will be selected that need to stomp this value
 			SelectorItem = ItemToSelect;
 
+			// Always request scroll into view, otherwise partially visible items will be selected - also do this before signaling selection for similar stop-allowing reasons
+			RequestNavigateToItem(ItemToSelect, InInputEvent.GetUserIndex());
+
 			if (CurrentSelectionMode == ESelectionMode::Multi && (InInputEvent.IsShiftDown() || InInputEvent.IsControlDown()))
 			{
 				// Range select.
@@ -1617,10 +1620,6 @@ protected:
 				// Single select.
 				this->SetSelection(ItemToSelect, ESelectInfo::OnNavigation);;
 			}
-
-			// Always request scroll into view, otherwise partially visible items will be selected.
-			TSharedPtr<ITableRow> WidgetForItem = this->WidgetGenerator.GetWidgetForItem(ItemToSelect);
-			RequestNavigateToItem(ItemToSelect, InInputEvent.GetUserIndex()); 
 		}
 	}
 
