Prevent dragging a widget while already dragging one

diff --git a/Engine/Source/Runtime/Slate/Private/Framework/Application/SlateApplication.cpp b/Engine/Source/Runtime/Slate/Private/Framework/Application/SlateApplication.cpp
--- a/Engine/Source/Runtime/Slate/Private/Framework/Application/SlateApplication.cpp
+++ b/Engine/Source/Runtime/Slate/Private/Framework/Application/SlateApplication.cpp
@@ -5733,11 +5733,12 @@ bool FSlateApplication::RoutePointerMoveEvent(const FWidgetPath& WidgetsUnderPoi
 
 	// User asked us to detect a drag.
 	bool bDragDetected = false;
-	bool bShouldStartDetectingDrag = true;
+	// Currently we can only support one dragged widget at a time.
+	bool bShouldStartDetectingDrag = !DragDropContent.IsValid();
 
 #if WITH_EDITOR
 	//@TODO VREDITOR - Remove and move to interaction component
-	if (OnDragDropCheckOverride.IsBound())
+	if (bShouldStartDetectingDrag && OnDragDropCheckOverride.IsBound())
 	{
 		bShouldStartDetectingDrag = OnDragDropCheckOverride.Execute();
 	}
