Fix vertical center alignment of window title. Currently offset by a few pixels.

diff --git a/Engine/Source/Runtime/Slate/Public/Framework/Application/SWindowTitleBar.h b/Engine/Source/Runtime/Slate/Public/Framework/Application/SWindowTitleBar.h
--- a/Engine/Source/Runtime/Slate/Public/Framework/Application/SWindowTitleBar.h
+++ b/Engine/Source/Runtime/Slate/Public/Framework/Application/SWindowTitleBar.h
@@ -320,7 +320,7 @@ protected:
 			CenterContent = SNew(SBox)
 				.HAlign(HAlign_Center)
 				.Visibility(EVisibility::SelfHitTestInvisible)
-				.Padding(FMargin(5.0f, 2.0f, 2.0f, 5.0f))
+				.Padding(FMargin(5.0f, 2.0f, 5.0f, 2.0f))
 				[
 					// NOTE: We bind the window's title text to our window's GetTitle method, so that if the
 					//       title is changed later, the text will always be visually up to date
@@ -375,7 +375,7 @@ protected:
 
 						+ SHorizontalBox::Slot()
 							.HAlign(CenterContentAlignment)
-							.VAlign(VAlign_Top)
+							.VAlign(VAlign_Center)
 							.FillWidth(1.0f)
 							[
 								CenterContent.ToSharedRef()
