Backing out 4962191 due to widespread Slate focus issues

diff --git a/Engine/Source/Runtime/Slate/Private/Widgets/Text/SlateEditableTextLayout.cpp b/Engine/Source/Runtime/Slate/Private/Widgets/Text/SlateEditableTextLayout.cpp
--- a/Engine/Source/Runtime/Slate/Private/Widgets/Text/SlateEditableTextLayout.cpp
+++ b/Engine/Source/Runtime/Slate/Private/Widgets/Text/SlateEditableTextLayout.cpp
@@ -105,9 +105,6 @@ FSlateEditableTextLayout::FSlateEditableTextLayout(ISlateEditableTextWidget& InO
 		const FText& InitialTextToSet = BoundText.Get(FText::GetEmpty());
 		SetEditableText(InitialTextToSet, true);
 
-		// treat the initial text as uncommitted
-		bHasUncommittedText = true;
-
 		// Update the cached BoundText value to prevent it triggering another SetEditableText update again next Tick
 		BoundTextLastTick = FTextSnapshot(InitialTextToSet);
 		bWasPasswordLastTick = bIsPassword;
@@ -779,7 +776,16 @@ bool FSlateEditableTextLayout::HandleFocusLost(const FFocusEvent& InFocusEvent)
 		break;
 	}
 
-	CommitTextChanges(TextAction);
+	// Always clear the local undo chain on commit
+	ClearUndoStates();
+
+	const FText EditedText = GetEditableText();
+
+	OwnerWidget->OnTextCommitted(EditedText, TextAction);
+
+	// Reload underlying value now it is committed  (commit may alter the value) 
+	// so it can be re-displayed in the edit box
+	LoadText();
 
 	UpdateCursorHighlight();
 
@@ -793,38 +799,6 @@ bool FSlateEditableTextLayout::HandleFocusLost(const FFocusEvent& InFocusEvent)
 	return true;
 }
 
-void FSlateEditableTextLayout::CommitTextChanges(ETextCommit::Type CommitReason)
-{
-	// don't try to commit changes if we haven't changed anything
-	if (bHasUncommittedText)
-	{
-		bHasUncommittedText = false;
-
-		// Always clear the local undo chain on commit
-		ClearUndoStates();
-
-		const FText EditedText = GetEditableText();
-
-		OwnerWidget->OnTextCommitted(EditedText, CommitReason);
-
-		// Reload underlying value now it is committed  (commit may alter the value) 
-		// so it can be re-displayed in the edit box
-		LoadText();
-
-		// Select all text?
-		if (OwnerWidget->ShouldSelectAllTextOnCommit())
-		{
-			SelectAllText();
-		}
-
-		// Release input focus?
-		if (CommitReason == ETextCommit::OnEnter && OwnerWidget->ShouldClearKeyboardFocusOnCommit())
-		{
-			FSlateApplication::Get().ClearKeyboardFocus(EFocusCause::Cleared);
-		}
-	}
-}
-
 FReply FSlateEditableTextLayout::HandleKeyChar(const FCharacterEvent& InCharacterEvent)
 {
 	// Check for special characters
@@ -1468,7 +1442,29 @@ bool FSlateEditableTextLayout::HandleCarriageReturn()
 	}
 	else
 	{
-		CommitTextChanges(ETextCommit::OnEnter);
+		// Always clear the local undo chain on commit.
+		ClearUndoStates();
+
+		const FText EditedText = GetEditableText();
+
+		// When enter is pressed text is committed.  Let anyone interested know about it.
+		OwnerWidget->OnTextCommitted(EditedText, ETextCommit::OnEnter);
+
+		// Reload underlying value now it is committed  (commit may alter the value) 
+		// so it can be re-displayed in the edit box if it retains focus
+		LoadText();
+
+		// Select all text?
+		if (OwnerWidget->ShouldSelectAllTextOnCommit())
+		{
+			SelectAllText();
+		}
+
+		// Release input focus?
+		if (OwnerWidget->ShouldClearKeyboardFocusOnCommit())
+		{
+			FSlateApplication::Get().ClearKeyboardFocus(EFocusCause::Cleared);
+		}
 	}
 
 	return true;
@@ -2897,7 +2893,6 @@ void FSlateEditableTextLayout::RestoreOriginalText()
 
 		// Let outsiders know that the text content has been changed
 		OwnerWidget->OnTextCommitted(OriginalText.Text, ETextCommit::OnCleared);
-		bHasUncommittedText = false;
 	}
 }
 
@@ -2944,8 +2939,6 @@ void FSlateEditableTextLayout::EndEditTransaction()
 	const bool bHasTextChanged = !EditedText.ToString().Equals(StateBeforeChangingText.GetValue().Text.ToString(), ESearchCase::CaseSensitive);
 	if (bHasTextChanged)
 	{
-		bHasUncommittedText = true;
-
 		// Save text state
 		SaveText(EditedText);
 
@@ -3166,7 +3159,6 @@ void FSlateEditableTextLayout::Tick(const FGeometry& AllottedGeometry, const dou
 		// Let outsiders know that the text content has been changed
 		OwnerWidget->OnTextCommitted(GetEditableText(), VirtualKeyboardTextCommitType);
 		bTextCommittedByVirtualKeyboard = false;
-		bHasUncommittedText = false;
 	}
 
 	if (TextInputMethodChangeNotifier.IsValid() && TextInputMethodContext.IsValid() && TextInputMethodContext->UpdateCachedGeometry(AllottedGeometry))
@@ -3495,8 +3487,6 @@ void FSlateEditableTextLayout::FVirtualKeyboardEntry::SetTextFromVirtualKeyboard
 		OwnerLayout->BoundText.Set(InNewText);
 	}
 
-	OwnerLayout->bHasUncommittedText = true;
-
 	// Update the internal editable text
 	// This method is called from the main thread (i.e. not the game thread) of the device with the virtual keyboard
 	// This causes the app to crash on those devices, so we're using polling here to ensure delegates are
diff --git a/Engine/Source/Runtime/Slate/Public/Widgets/Text/SlateEditableTextLayout.h b/Engine/Source/Runtime/Slate/Public/Widgets/Text/SlateEditableTextLayout.h
--- a/Engine/Source/Runtime/Slate/Public/Widgets/Text/SlateEditableTextLayout.h
+++ b/Engine/Source/Runtime/Slate/Public/Widgets/Text/SlateEditableTextLayout.h
@@ -365,10 +365,6 @@ public:
 
 	bool HasActiveContextMenu() const;
 
-protected:
-
-	void CommitTextChanges(ETextCommit::Type CommitReason);
-
 private:
 	/** Insert the given text at the current cursor position, correctly taking into account new line characters */
 	void InsertTextAtCursorImpl(const FString& InString);
@@ -609,9 +605,6 @@ private:
 	/** Whether the text has been committed by a virtual keyboard */
 	bool bTextCommittedByVirtualKeyboard;
 
-	/** Whether the widget contains uncommitted text */
-	bool bHasUncommittedText;
-
 	/** What text was submitted by a virtual keyboard */
 	FText VirtualKeyboardText;
 
