Slate - Adding several checks to make sure all virtual keyboards entering slate are coming back on the game thread.  Updating the XboxOnePlatformTextField to correctly (1) Not pin a non-thread safe weak ptr off the main thread (2) not call back into slate off the main thread, and isntead enqeue an async task to the game thread to set/or cancel setting the text from the virtual keyboard.

We may see other platforms fail the checks I've added.

David.Harvey, Matt.Kuhlenschmidt, Ben.Woodhouse

diff --git a/Engine/Source/Runtime/Slate/Private/Widgets/Input/SVirtualKeyboardEntry.cpp b/Engine/Source/Runtime/Slate/Private/Widgets/Input/SVirtualKeyboardEntry.cpp
--- a/Engine/Source/Runtime/Slate/Private/Widgets/Input/SVirtualKeyboardEntry.cpp
+++ b/Engine/Source/Runtime/Slate/Private/Widgets/Input/SVirtualKeyboardEntry.cpp
@@ -52,6 +52,8 @@ void SVirtualKeyboardEntry::SetText(const TAttribute< FText >& InNewText)
 
 void SVirtualKeyboardEntry::SetTextFromVirtualKeyboard(const FText& InNewText, ETextEntryType TextEntryType)
 {
+	check(IsInGameThread());
+
 	// Only set the text if the text attribute doesn't have a getter binding (otherwise it would be blown away).
 	// If it is bound, we'll assume that OnTextChanged will handle the update.
 	if (!Text.IsBound())
@@ -72,6 +74,7 @@ void SVirtualKeyboardEntry::SetTextFromVirtualKeyboard(const FText& InNewText, E
 
 void SVirtualKeyboardEntry::SetSelectionFromVirtualKeyboard(int InSelStart, int InSelEnd)
 {
+	check(IsInGameThread());
 	//Nothing to do for widgets without a cursor
 }
 
diff --git a/Engine/Source/Runtime/Slate/Private/Widgets/Text/SlateEditableTextLayout.cpp b/Engine/Source/Runtime/Slate/Private/Widgets/Text/SlateEditableTextLayout.cpp
--- a/Engine/Source/Runtime/Slate/Private/Widgets/Text/SlateEditableTextLayout.cpp
+++ b/Engine/Source/Runtime/Slate/Private/Widgets/Text/SlateEditableTextLayout.cpp
@@ -219,7 +219,6 @@ FSlateEditableTextLayout::~FSlateEditableTextLayout()
 	{
 		FSlateApplication::Get().ShowVirtualKeyboard(false, 0);
 	}
-
 }
 
 void FSlateEditableTextLayout::SetText(const TAttribute<FText>& InText)
@@ -261,6 +260,8 @@ void FSlateEditableTextLayout::SetText(const TAttribute<FText>& InText)
 
 FText FSlateEditableTextLayout::GetText() const
 {
+	check(IsInGameThread());
+
 	return BoundText.Get(FText::GetEmpty());
 }
 
@@ -329,6 +330,8 @@ void FSlateEditableTextLayout::SetCompositionBrush(const TAttribute<const FSlate
 
 FText FSlateEditableTextLayout::GetPlainText() const
 {
+	check(IsInGameThread());
+
 	const TArray< FTextLayout::FLineModel >& Lines = TextLayout->GetLineModels();
 
 	const int32 NumberOfLines = Lines.Num();
@@ -346,6 +349,8 @@ FText FSlateEditableTextLayout::GetPlainText() const
 
 bool FSlateEditableTextLayout::SetEditableText(const FText& TextToSet, const bool bForce)
 {
+	check(IsInGameThread());
+
 	bool bHasTextChanged = bForce;
 	if (!bHasTextChanged)
 	{
@@ -518,6 +523,8 @@ bool FSlateEditableTextLayout::Refresh()
 
 bool FSlateEditableTextLayout::RefreshImpl(const FText* InTextToSet, const bool bForce)
 {
+	check(IsInGameThread());
+
 	bool bHasSetText = false;
 
 	const bool bIsPassword = OwnerWidget->IsTextPassword();
@@ -3146,6 +3153,8 @@ bool FSlateEditableTextLayout::ComputeVolatility() const
 
 void FSlateEditableTextLayout::Tick(const FGeometry& AllottedGeometry, const double InCurrentTime, const float InDeltaTime)
 {
+	check(IsInGameThread());
+
 	if (bTextChangedByVirtualKeyboard)
 	{
 		SetEditableText(VirtualKeyboardText);
@@ -3343,6 +3352,8 @@ void FSlateEditableTextLayout::CacheDesiredSize(float LayoutScaleMultiplier)
 
 FVector2D FSlateEditableTextLayout::ComputeDesiredSize(float LayoutScaleMultiplier) const
 {
+	check(IsInGameThread());
+
 	const float FontMaxCharHeight = FTextEditHelper::GetFontHeight(TextStyle.Font);
 	const float CaretWidth = FTextEditHelper::CalculateCaretWidth(FontMaxCharHeight);
 
@@ -3467,7 +3478,6 @@ bool FSlateEditableTextLayout::HasActiveContextMenu() const
 	return ActiveContextMenu.IsValid();
 }
 
-
 TSharedRef<FSlateEditableTextLayout::FVirtualKeyboardEntry> FSlateEditableTextLayout::FVirtualKeyboardEntry::Create(FSlateEditableTextLayout& InOwnerLayout)
 {
 	return MakeShareable(new FVirtualKeyboardEntry(InOwnerLayout));
@@ -3480,6 +3490,8 @@ FSlateEditableTextLayout::FVirtualKeyboardEntry::FVirtualKeyboardEntry(FSlateEdi
 
 void FSlateEditableTextLayout::FVirtualKeyboardEntry::SetTextFromVirtualKeyboard(const FText& InNewText, ETextEntryType TextEntryType)
 {
+	check(IsInGameThread());
+
 	// Only set the text if the text attribute doesn't have a getter binding (otherwise it would be blown away).
 	// If it is bound, we'll assume that OnTextCommitted will handle the update.
 	if (!OwnerLayout->BoundText.IsBound())
@@ -3514,6 +3526,8 @@ void FSlateEditableTextLayout::FVirtualKeyboardEntry::SetTextFromVirtualKeyboard
 
 void FSlateEditableTextLayout::FVirtualKeyboardEntry::SetSelectionFromVirtualKeyboard(int InSelStart, int InSelEnd)
 {
+	check(IsInGameThread());
+
 	// Update the text selection and the cursor position
 	// This method is called externally (eg. on Android from the native virtual keyboard implementation) 
 	// The text may also change on the same frame, so the external selection must happen in Tick after the text update
@@ -3525,26 +3539,36 @@ void FSlateEditableTextLayout::FVirtualKeyboardEntry::SetSelectionFromVirtualKey
 
 FText FSlateEditableTextLayout::FVirtualKeyboardEntry::GetText() const
 {
+	check(IsInGameThread());
+
 	return OwnerLayout->GetText();
 }
 
 FText FSlateEditableTextLayout::FVirtualKeyboardEntry::GetHintText() const
 {
+	check(IsInGameThread());
+
 	return OwnerLayout->GetHintText();
 }
 
 EKeyboardType FSlateEditableTextLayout::FVirtualKeyboardEntry::GetVirtualKeyboardType() const
 {
+	check(IsInGameThread());
+
 	return (OwnerLayout->OwnerWidget->IsTextPassword()) ? Keyboard_Password : OwnerLayout->OwnerWidget->GetVirtualKeyboardType();
 }
 
 FVirtualKeyboardOptions FSlateEditableTextLayout::FVirtualKeyboardEntry::GetVirtualKeyboardOptions() const
 {
+	check(IsInGameThread());
+
 	return OwnerLayout->OwnerWidget->GetVirtualKeyboardOptions();
 }
 
 bool FSlateEditableTextLayout::FVirtualKeyboardEntry::IsMultilineEntry() const
 {
+	check(IsInGameThread());
+
 	return OwnerLayout->OwnerWidget->IsMultiLineTextEdit();
 }
 
diff --git a/Engine/Source/Runtime/Slate/Public/Widgets/Input/SVirtualKeyboardEntry.h b/Engine/Source/Runtime/Slate/Public/Widgets/Input/SVirtualKeyboardEntry.h
--- a/Engine/Source/Runtime/Slate/Public/Widgets/Input/SVirtualKeyboardEntry.h
+++ b/Engine/Source/Runtime/Slate/Public/Widgets/Input/SVirtualKeyboardEntry.h
@@ -110,26 +110,36 @@ public:
 
 	virtual FText GetText() const override
 	{
+		check(IsInGameThread());
+
 		return Text.Get();
 	}
 
 	virtual FText GetHintText() const override
 	{
+		check(IsInGameThread());
+
 		return HintText.Get();
 	}
 
 	virtual EKeyboardType GetVirtualKeyboardType() const override
 	{
+		check(IsInGameThread());
+
 		return KeyboardType.Get();
 	}
 
 	virtual FVirtualKeyboardOptions GetVirtualKeyboardOptions() const override
 	{
+		check(IsInGameThread());
+
 		return VirtualKeyboardOptions;
 	}
 
 	virtual bool IsMultilineEntry() const override
 	{
+		check(IsInGameThread());
+
 		return false;
 	}
 	//~ End IVirtualKeyboardEntry Interface
