RHIs need frames in order to handle all their book-keeping - even when debugging a blueprint - without this they will live-leak memory into their internal pools.

diff --git a/Engine/Source/Runtime/Slate/Private/Framework/Application/SlateApplication.cpp b/Engine/Source/Runtime/Slate/Private/Framework/Application/SlateApplication.cpp
--- a/Engine/Source/Runtime/Slate/Private/Framework/Application/SlateApplication.cpp
+++ b/Engine/Source/Runtime/Slate/Private/Framework/Application/SlateApplication.cpp
@@ -4129,9 +4129,13 @@ void FSlateApplication::EnterDebuggingMode()
 	// Tick slate from here in the event that we should not return until the modal window is closed.
 	while (!bRequestLeaveDebugMode)
 	{
+		Renderer->BeginFrame();
+		
 		// Tick and render Slate
 		Tick();
 
+		Renderer->EndFrame();
+		
 		// Synchronize the game thread and the render thread so that the render thread doesn't get too far behind.
 		Renderer->Sync();
 
@@ -4144,6 +4148,7 @@ void FSlateApplication::EnterDebuggingMode()
 #endif	//WITH_EDITORONLY_DATA
 	}
 
+	Renderer->BeginFrame();
 	bRequestLeaveDebugMode = false;
 	
 	if ( PreviousGameViewport.IsValid() )
