SListView no longer jumps to the middle of the list area when navigating to a partially obscured item on either end
- Also added NavigationScrollOffset (default 0.5f) property as an additional offset to apply when navigating to an edge item to make clear there are more items in the list
- Behavior of normal scroll requests are unaffected, as are attempts to navigate to items that do not have an existing row (both still try to center the row of the target item)

diff --git a/Engine/Source/Runtime/Slate/Public/Widgets/Views/SListView.h b/Engine/Source/Runtime/Slate/Public/Widgets/Views/SListView.h
--- a/Engine/Source/Runtime/Slate/Public/Widgets/Views/SListView.h
+++ b/Engine/Source/Runtime/Slate/Public/Widgets/Views/SListView.h
@@ -81,6 +81,7 @@ public:
 		, _AllowOverscroll(EAllowOverscroll::Yes)
 		, _ConsumeMouseWheel(EConsumeMouseWheel::WhenScrollingPossible)
 		, _WheelScrollMultiplier(GetGlobalScrollAmount())
+		, _NavigationScrollOffset(0.5f)
 		, _HandleGamepadEvents( true )
 		, _HandleDirectionalNavigation( true )
 		, _IsFocusable(true)
@@ -130,6 +131,8 @@ public:
 
 		SLATE_ARGUMENT( float, WheelScrollMultiplier );
 
+		SLATE_ARGUMENT( float, NavigationScrollOffset );
+
 		SLATE_ARGUMENT( bool, HandleGamepadEvents );
 
 		SLATE_ARGUMENT( bool, HandleDirectionalNavigation );
@@ -168,8 +171,9 @@ public:
 
 		this->AllowOverscroll = InArgs._AllowOverscroll;
 		this->ConsumeMouseWheel = InArgs._ConsumeMouseWheel;
-
 		this->WheelScrollMultiplier = InArgs._WheelScrollMultiplier;
+		this->NavigationScrollOffset = InArgs._NavigationScrollOffset;
+
 		this->bHandleGamepadEvents = InArgs._HandleGamepadEvents;
 		this->bHandleDirectionalNavigation = InArgs._HandleDirectionalNavigation;
 		this->IsFocusable = InArgs._IsFocusable;
@@ -1389,20 +1393,65 @@ protected:
 				EndInertialScrolling();
 				
 				// Only scroll the item into view if it's not already in the visible range
-				const double IndexPlusOne = IndexOfItem+1;
-				if (IndexOfItem < ScrollOffset || IndexPlusOne > (ScrollOffset + NumLiveWidgets))
+				// When navigating, we don't want to scroll partially visible existing rows all the way to the center
+				const double MinDisplayedIndex = bNavigateOnScrollIntoView ? ScrollOffset - 1 : ScrollOffset;
+				const double MaxDisplayedIndex = FMath::Max(ScrollOffset + NumLiveWidgets + (bNavigateOnScrollIntoView ? 0 : -1), MinDisplayedIndex);
+				if (IndexOfItem <= MinDisplayedIndex || IndexOfItem > MaxDisplayedIndex)
 				{
 					// Scroll the top of the listview to the item in question
 					double NewScrollOffset = IndexOfItem;
 					// Center the list view on the item in question.
 					NewScrollOffset -= (NumLiveWidgets / 2);
 					//we also don't want the widget being chopped off if it is at the end of the list
-					const double MoveBackBy = FMath::Clamp<double>(IndexPlusOne - (NewScrollOffset + NumLiveWidgets), 0, FLT_MAX);
+					const double MoveBackBy = FMath::Clamp<double>(IndexOfItem - MaxDisplayedIndex, 0, FLT_MAX);
 					//Move to the correct center spot
 					NewScrollOffset += MoveBackBy;
 
 					SetScrollOffset( NewScrollOffset );
 				}
+				else if (bNavigateOnScrollIntoView)
+				{
+					if (TSharedPtr<ITableRow> TableRow = WidgetFromItem((*ItemsSource)[IndexOfItem]))
+					{
+						// Make sure the existing row for this item is fully in view
+						const FGeometry& RowGeometry = TableRow->AsWidget()->GetCachedGeometry();
+						if (RowGeometry.GetAbsolutePositionAtCoordinates(FVector2D::ZeroVector).Y < ListViewGeometry.GetAbsolutePositionAtCoordinates(FVector2D::ZeroVector).Y)
+						{
+							// This row is clipped on the top, so simply set it as the new scroll offset target to bump it down a bit
+							SetScrollOffset(IndexOfItem - NavigationScrollOffset);
+						}
+						else
+						{
+							const FVector2D BottomRight(1.f, 1.f);
+							const float RowBottomY = RowGeometry.GetAbsolutePositionAtCoordinates(BottomRight).Y;
+							const float PanelBottomY = ListViewGeometry.GetAbsolutePositionAtCoordinates(BottomRight).Y;
+							if (RowBottomY > PanelBottomY)
+							{
+								// This row is clipped on the bottom, so we need to determine the exact item offset required to get the bottom of this entry into view
+								// To do so, we need to push the current offset down by the clipped amount translated into number of rows
+								float DistanceRemaining = RowBottomY - PanelBottomY;
+								float AdditionalOffset = 0.f;
+								for (const ItemType& ItemWithWidget : WidgetGenerator.ItemsWithGeneratedWidgets)
+								{
+									const float ExistingRowSizeY = WidgetGenerator.GetWidgetForItem(ItemWithWidget)->AsWidget()->GetCachedGeometry().GetAbsoluteSize().Y;
+									if (ExistingRowSizeY < DistanceRemaining)
+									{
+										DistanceRemaining -= ExistingRowSizeY;
+										AdditionalOffset += 1.f;
+									}
+									else
+									{
+										AdditionalOffset += DistanceRemaining / ExistingRowSizeY;
+										DistanceRemaining = 0.f;
+										break;
+									}
+								}
+
+								SetScrollOffset(ScrollOffset + AdditionalOffset + NavigationScrollOffset);
+							}
+						}
+					}
+				}
 
 				RequestLayoutRefresh();
 
@@ -1596,7 +1645,7 @@ protected:
 			// Must be set before signaling selection changes because sometimes new items will be selected that need to stomp this value
 			SelectorItem = ItemToSelect;
 
-			// Always request scroll into view, otherwise partially visible items will be selected - also do this before signaling selection for similar stop-allowing reasons
+			// Always request scroll into view, otherwise partially visible items will be selected - also do this before signaling selection for similar stomp-allowing reasons
 			RequestNavigateToItem(ItemToSelect, InInputEvent.GetUserIndex());
 
 			if (CurrentSelectionMode == ESelectionMode::Multi && (InInputEvent.IsShiftDown() || InInputEvent.IsControlDown()))
@@ -1678,6 +1727,9 @@ protected:
 	/** True when the list view supports keyboard focus */
 	TAttribute<bool> IsFocusable;
 
+	/** The additional scroll offset (in items) to show when navigating to rows at the edge of the visible area (i.e. how much of the following item(s) to show) */
+	float NavigationScrollOffset = 0.5f;
+
 	/** If true, the selection will be cleared if the user clicks in empty space (not on an item) */
 	bool bClearSelectionOnClick;
 
