make FDefaultRichTextMarkupParser::Create effectively a singleton to avoid regex pattern construction on critical path GT code

rename FDefaultRichTextMarkupParser::Create to FDefaultRichTextMarkupParser::Get and fixup calling code

diff --git a/Engine/Source/Runtime/Slate/Private/Framework/Text/RichTextLayoutMarshaller.cpp b/Engine/Source/Runtime/Slate/Private/Framework/Text/RichTextLayoutMarshaller.cpp
--- a/Engine/Source/Runtime/Slate/Private/Framework/Text/RichTextLayoutMarshaller.cpp
+++ b/Engine/Source/Runtime/Slate/Private/Framework/Text/RichTextLayoutMarshaller.cpp
@@ -86,7 +86,7 @@ void FRichTextLayoutMarshaller::GetText(FString& TargetString, const FTextLayout
 }
 
 FRichTextLayoutMarshaller::FRichTextLayoutMarshaller(TArray< TSharedRef< ITextDecorator > > InDecorators, const ISlateStyle* const InDecoratorStyleSet)
-	: Parser(FDefaultRichTextMarkupParser::Create())
+	: Parser(FDefaultRichTextMarkupParser::Get())
 	, Writer(FDefaultRichTextMarkupWriter::Create())
 	, Decorators(MoveTemp(InDecorators))
 	, InlineDecorators()
diff --git a/Engine/Source/Runtime/Slate/Private/Framework/Text/RichTextMarkupProcessing.cpp b/Engine/Source/Runtime/Slate/Private/Framework/Text/RichTextMarkupProcessing.cpp
--- a/Engine/Source/Runtime/Slate/Private/Framework/Text/RichTextMarkupProcessing.cpp
+++ b/Engine/Source/Runtime/Slate/Private/Framework/Text/RichTextMarkupProcessing.cpp
@@ -48,9 +48,10 @@ struct FEscapeSequenceRegexPatternString
 };
 
 
-TSharedRef< FDefaultRichTextMarkupParser > FDefaultRichTextMarkupParser::Create()
+TSharedRef< FDefaultRichTextMarkupParser > FDefaultRichTextMarkupParser::Get()
 {
-	return MakeShareable( new FDefaultRichTextMarkupParser() );
+	static TSharedRef< FDefaultRichTextMarkupParser > Parser = MakeShareable( new FDefaultRichTextMarkupParser() );
+	return Parser;
 }
 
 FDefaultRichTextMarkupParser::FDefaultRichTextMarkupParser()
diff --git a/Engine/Source/Runtime/Slate/Private/Widgets/Text/SRichTextBlock.cpp b/Engine/Source/Runtime/Slate/Private/Widgets/Text/SRichTextBlock.cpp
--- a/Engine/Source/Runtime/Slate/Private/Widgets/Text/SRichTextBlock.cpp
+++ b/Engine/Source/Runtime/Slate/Private/Widgets/Text/SRichTextBlock.cpp
@@ -36,7 +36,7 @@ void SRichTextBlock::Construct( const FArguments& InArgs )
 		TSharedPtr<IRichTextMarkupParser> Parser = InArgs._Parser;
 		if ( !Parser.IsValid() )
 		{
-			Parser = FDefaultRichTextMarkupParser::Create();
+			Parser = FDefaultRichTextMarkupParser::Get();
 		}
 
 		TSharedPtr<FRichTextLayoutMarshaller> Marshaller = InArgs._Marshaller;
diff --git a/Engine/Source/Runtime/Slate/Public/Framework/Text/RichTextMarkupProcessing.h b/Engine/Source/Runtime/Slate/Public/Framework/Text/RichTextMarkupProcessing.h
--- a/Engine/Source/Runtime/Slate/Public/Framework/Text/RichTextMarkupProcessing.h
+++ b/Engine/Source/Runtime/Slate/Public/Framework/Text/RichTextMarkupProcessing.h
@@ -14,7 +14,7 @@
 class SLATE_API FDefaultRichTextMarkupParser : public IRichTextMarkupParser
 {
 public:
-	static TSharedRef< FDefaultRichTextMarkupParser > Create();
+	static TSharedRef< FDefaultRichTextMarkupParser > Get();
 
 public:
 	virtual void Process(TArray<FTextLineParseResults>& Results, const FString& Input, FString& Output) override;
diff --git a/Engine/Source/Runtime/UMG/Private/Components/RichTextBlock.cpp b/Engine/Source/Runtime/UMG/Private/Components/RichTextBlock.cpp
--- a/Engine/Source/Runtime/UMG/Private/Components/RichTextBlock.cpp
+++ b/Engine/Source/Runtime/UMG/Private/Components/RichTextBlock.cpp
@@ -149,7 +149,7 @@ void URichTextBlock::CreateDecorators(TArray< TSharedRef< class ITextDecorator >
 
 TSharedPtr< IRichTextMarkupParser > URichTextBlock::CreateMarkupParser()
 {
-	return FDefaultRichTextMarkupParser::Create();
+	return FDefaultRichTextMarkupParser::Get();
 }
 
 TSharedPtr< IRichTextMarkupWriter > URichTextBlock::CreateMarkupWriter()
