Changing the rich text API back to have a Create and renaming Get to GetStaticInstance to fix backwards compat and make it a bit clearer what each is doing.

diff --git a/Engine/Plugins/Tests/RuntimeTests/Source/RuntimeTests/Private/Slate/RichTextMarkupProcessingTest.cpp b/Engine/Plugins/Tests/RuntimeTests/Source/RuntimeTests/Private/Slate/RichTextMarkupProcessingTest.cpp
--- a/Engine/Plugins/Tests/RuntimeTests/Source/RuntimeTests/Private/Slate/RichTextMarkupProcessingTest.cpp
+++ b/Engine/Plugins/Tests/RuntimeTests/Source/RuntimeTests/Private/Slate/RichTextMarkupProcessingTest.cpp
@@ -100,7 +100,7 @@ namespace
 bool FRTFProcessingTest::RunTest (const FString& Parameters)
 {
 #if UE_ENABLE_ICU
-	TSharedRef< FDefaultRichTextMarkupParser > RichTextMarkupProcessing = FDefaultRichTextMarkupParser::Get();
+	TSharedRef< FDefaultRichTextMarkupParser > RichTextMarkupProcessing = FDefaultRichTextMarkupParser::GetStaticInstance();
 
 	// Text
 	{
diff --git a/Engine/Source/Runtime/Slate/Private/Framework/Text/RichTextLayoutMarshaller.cpp b/Engine/Source/Runtime/Slate/Private/Framework/Text/RichTextLayoutMarshaller.cpp
--- a/Engine/Source/Runtime/Slate/Private/Framework/Text/RichTextLayoutMarshaller.cpp
+++ b/Engine/Source/Runtime/Slate/Private/Framework/Text/RichTextLayoutMarshaller.cpp
@@ -86,7 +86,7 @@ void FRichTextLayoutMarshaller::GetText(FString& TargetString, const FTextLayout
 }
 
 FRichTextLayoutMarshaller::FRichTextLayoutMarshaller(TArray< TSharedRef< ITextDecorator > > InDecorators, const ISlateStyle* const InDecoratorStyleSet)
-	: Parser(FDefaultRichTextMarkupParser::Get())
+	: Parser(FDefaultRichTextMarkupParser::GetStaticInstance())
 	, Writer(FDefaultRichTextMarkupWriter::Create())
 	, Decorators(MoveTemp(InDecorators))
 	, InlineDecorators()
diff --git a/Engine/Source/Runtime/Slate/Private/Framework/Text/RichTextMarkupProcessing.cpp b/Engine/Source/Runtime/Slate/Private/Framework/Text/RichTextMarkupProcessing.cpp
--- a/Engine/Source/Runtime/Slate/Private/Framework/Text/RichTextMarkupProcessing.cpp
+++ b/Engine/Source/Runtime/Slate/Private/Framework/Text/RichTextMarkupProcessing.cpp
@@ -47,8 +47,12 @@ struct FEscapeSequenceRegexPatternString
 	}
 };
 
+TSharedRef< FDefaultRichTextMarkupParser > FDefaultRichTextMarkupParser::Create()
+{
+	return MakeShareable(new FDefaultRichTextMarkupParser());
+}
 
-TSharedRef< FDefaultRichTextMarkupParser > FDefaultRichTextMarkupParser::Get()
+TSharedRef< FDefaultRichTextMarkupParser > FDefaultRichTextMarkupParser::GetStaticInstance()
 {
 	static TSharedRef< FDefaultRichTextMarkupParser > Parser = MakeShareable( new FDefaultRichTextMarkupParser() );
 	return Parser;
diff --git a/Engine/Source/Runtime/Slate/Private/Widgets/Text/SRichTextBlock.cpp b/Engine/Source/Runtime/Slate/Private/Widgets/Text/SRichTextBlock.cpp
--- a/Engine/Source/Runtime/Slate/Private/Widgets/Text/SRichTextBlock.cpp
+++ b/Engine/Source/Runtime/Slate/Private/Widgets/Text/SRichTextBlock.cpp
@@ -36,7 +36,7 @@ void SRichTextBlock::Construct( const FArguments& InArgs )
 		TSharedPtr<IRichTextMarkupParser> Parser = InArgs._Parser;
 		if ( !Parser.IsValid() )
 		{
-			Parser = FDefaultRichTextMarkupParser::Get();
+			Parser = FDefaultRichTextMarkupParser::GetStaticInstance();
 		}
 
 		TSharedPtr<FRichTextLayoutMarshaller> Marshaller = InArgs._Marshaller;
diff --git a/Engine/Source/Runtime/Slate/Public/Framework/Text/RichTextMarkupProcessing.h b/Engine/Source/Runtime/Slate/Public/Framework/Text/RichTextMarkupProcessing.h
--- a/Engine/Source/Runtime/Slate/Public/Framework/Text/RichTextMarkupProcessing.h
+++ b/Engine/Source/Runtime/Slate/Public/Framework/Text/RichTextMarkupProcessing.h
@@ -14,7 +14,8 @@
 class SLATE_API FDefaultRichTextMarkupParser : public IRichTextMarkupParser
 {
 public:
-	static TSharedRef< FDefaultRichTextMarkupParser > Get();
+	static TSharedRef< FDefaultRichTextMarkupParser > Create();
+	static TSharedRef< FDefaultRichTextMarkupParser > GetStaticInstance();
 
 public:
 	virtual void Process(TArray<FTextLineParseResults>& Results, const FString& Input, FString& Output) override;
diff --git a/Engine/Source/Runtime/UMG/Private/Components/RichTextBlock.cpp b/Engine/Source/Runtime/UMG/Private/Components/RichTextBlock.cpp
--- a/Engine/Source/Runtime/UMG/Private/Components/RichTextBlock.cpp
+++ b/Engine/Source/Runtime/UMG/Private/Components/RichTextBlock.cpp
@@ -149,7 +149,7 @@ void URichTextBlock::CreateDecorators(TArray< TSharedRef< class ITextDecorator >
 
 TSharedPtr< IRichTextMarkupParser > URichTextBlock::CreateMarkupParser()
 {
-	return FDefaultRichTextMarkupParser::Get();
+	return FDefaultRichTextMarkupParser::GetStaticInstance();
 }
 
 TSharedPtr< IRichTextMarkupWriter > URichTextBlock::CreateMarkupWriter()
