Fix for hangs on macOS when stepping through blueprints in the Editor.

diff --git a/Engine/Source/Runtime/Slate/Private/Framework/Application/SlateApplication.cpp b/Engine/Source/Runtime/Slate/Private/Framework/Application/SlateApplication.cpp
--- a/Engine/Source/Runtime/Slate/Private/Framework/Application/SlateApplication.cpp
+++ b/Engine/Source/Runtime/Slate/Private/Framework/Application/SlateApplication.cpp
@@ -4115,6 +4115,8 @@ void FSlateApplication::EnterDebuggingMode()
 		PreviousGameViewport->SetActive(false);
 		GameViewportWidget.Reset();
 	}
+	
+	Renderer->EndFrame();
 
 	Renderer->FlushCommands();
 	
@@ -4136,6 +4138,8 @@ void FSlateApplication::EnterDebuggingMode()
 
 		Renderer->EndFrame();
 		
+		Renderer->FlushCommands();
+		
 		// Synchronize the game thread and the render thread so that the render thread doesn't get too far behind.
 		Renderer->Sync();
 
