Slate: Fix to the SSlider to correctly clamp the position of the thumb to make sur it stays in the range of the slider length.

diff --git a/Engine/Source/Runtime/Slate/Private/Widgets/Input/SSlider.cpp b/Engine/Source/Runtime/Slate/Private/Widgets/Input/SSlider.cpp
--- a/Engine/Source/Runtime/Slate/Private/Widgets/Input/SSlider.cpp
+++ b/Engine/Source/Runtime/Slate/Private/Widgets/Input/SSlider.cpp
@@ -45,8 +45,9 @@ int32 SSlider::OnPaint( const FPaintArgs& Args, const FGeometry& AllottedGeometr
 	const FVector2D HalfHandleSize = 0.5f * HandleSize;
 	const float Indentation = IndentHandle.Get() ? HandleSize.X : 0.0f;
 
+	// We clamp to make sure that the slider cannot go out of the slider Length.
+	const float SliderPercent = FMath::Clamp(ValueAttribute.Get(), 0.0f, 1.0f); 
 	const float SliderLength = AllottedWidth - (Indentation + HandleSize.X);
-	const float SliderPercent = ValueAttribute.Get();
 	const float SliderHandleOffset = SliderPercent * SliderLength;
 	const float SliderY = 0.5f * AllottedHeight;
 
