Adding support for knowing more about mouse capture in the slate console debugger, as well as when mouse capture is lost.

diff --git a/Engine/Source/Runtime/Slate/Private/Framework/Application/SlateApplication.cpp b/Engine/Source/Runtime/Slate/Private/Framework/Application/SlateApplication.cpp
--- a/Engine/Source/Runtime/Slate/Private/Framework/Application/SlateApplication.cpp
+++ b/Engine/Source/Runtime/Slate/Private/Framework/Application/SlateApplication.cpp
@@ -619,6 +619,11 @@ bool FSlateApplication::MouseCaptorHelper::SetMouseCaptor(uint32 UserIndex, uint
 		if (MouseCaptorWeakPath.IsValid())
 		{
 			PointerIndexToMouseCaptorWeakPathMap.Add(FUserAndPointer(UserIndex,PointerIndex), MouseCaptorWeakPath);
+
+#if WITH_SLATE_DEBUGGING
+			FSlateDebugging::BroadcastMouseCapture(UserIndex, PointerIndex, Widget);
+#endif
+
 			return true;
 		}
 	}
@@ -711,19 +716,23 @@ FWeakWidgetPath FSlateApplication::MouseCaptorHelper::ToWeakPath(uint32 UserInde
 	return FWeakWidgetPath();
 }
 
-void FSlateApplication::MouseCaptorHelper::InformCurrentCaptorOfCaptureLoss(uint32 UserIndex,uint32 PointerIndex) const
+void FSlateApplication::MouseCaptorHelper::InformCurrentCaptorOfCaptureLoss(uint32 UserIndex, uint32 PointerIndex) const
 {
 	// if we have a path to a widget then it is the current mouse captor and needs to know it has lost capture
-	const FWeakWidgetPath* MouseCaptorWeakPath = PointerIndexToMouseCaptorWeakPathMap.Find(FUserAndPointer(UserIndex,PointerIndex));
-	if (MouseCaptorWeakPath && MouseCaptorWeakPath->IsValid() )
+	const FWeakWidgetPath* MouseCaptorWeakPath = PointerIndexToMouseCaptorWeakPathMap.Find(FUserAndPointer(UserIndex, PointerIndex));
+	if (MouseCaptorWeakPath && MouseCaptorWeakPath->IsValid())
 	{
 		TWeakPtr< SWidget > WeakWidgetPtr = MouseCaptorWeakPath->GetLastWidget();
 		TSharedPtr< SWidget > SharedWidgetPtr = WeakWidgetPtr.Pin();
-		if ( SharedWidgetPtr.IsValid() )
+		if (SharedWidgetPtr.IsValid())
 		{
 			FCaptureLostEvent CaptureLostEvent(UserIndex, PointerIndex);
 			SharedWidgetPtr->OnMouseCaptureLost(CaptureLostEvent);
 		}
+
+#if WITH_SLATE_DEBUGGING
+		FSlateDebugging::BroadcastMouseCaptureLost(UserIndex, PointerIndex, SharedWidgetPtr);
+#endif
 	}
 }
 
@@ -3279,10 +3288,6 @@ void FSlateApplication::ProcessReply( const FWidgetPath& CurrentEventPath, const
 		{
 			if ( MouseCaptor.SetMouseCaptor(UserIndex, PointerIndex, CurrentEventPath, RequestedMouseCaptor) )
 			{
-#if WITH_SLATE_DEBUGGING
-				FSlateDebugging::MouseCapture(RequestedMouseCaptor);
-#endif
-
 				if (WidgetsUnderMouse)
 				{
 					const FWeakWidgetPath& LastWidgetsUnderCursor = WidgetsUnderCursorLastEvent.FindRef(FUserAndPointer(UserIndex, PointerIndex));
@@ -6516,7 +6521,7 @@ bool FSlateApplication::AttemptNavigation(const FWidgetPath& NavigationSource, c
 	}
 
 #if WITH_SLATE_DEBUGGING
-	FSlateDebugging::AttemptNavigation(NavigationEvent, NavigationReply, NavigationSource, DestinationWidget);
+	FSlateDebugging::BroadcastAttemptNavigation(NavigationEvent, NavigationReply, NavigationSource, DestinationWidget);
 #endif
 
 	return ExecuteNavigation(NavigationSource, DestinationWidget, NavigationEvent.GetUserIndex(), bAlwaysHandleNavigationAttempt);
diff --git a/Engine/Source/Runtime/SlateCore/Private/Debugging/ConsoleSlateDebugger.cpp b/Engine/Source/Runtime/SlateCore/Private/Debugging/ConsoleSlateDebugger.cpp
--- a/Engine/Source/Runtime/SlateCore/Private/Debugging/ConsoleSlateDebugger.cpp
+++ b/Engine/Source/Runtime/SlateCore/Private/Debugging/ConsoleSlateDebugger.cpp
@@ -129,7 +129,7 @@ void FConsoleSlateDebugger::UpdateListeners()
 	FSlateDebugging::InputEvent.AddRaw(this, &FConsoleSlateDebugger::OnInputEvent);
 	FSlateDebugging::FocusEvent.AddRaw(this, &FConsoleSlateDebugger::OnFocusEvent);
 	FSlateDebugging::NavigationEvent.AddRaw(this, &FConsoleSlateDebugger::OnNavigationEvent);
-	FSlateDebugging::MouseCaptureEvent.AddRaw(this, &FConsoleSlateDebugger::OnStateChangeEvent);
+	FSlateDebugging::MouseCaptureEvent.AddRaw(this, &FConsoleSlateDebugger::OnCaptureStateChangeEvent);
 #endif
 }
 
@@ -262,16 +262,18 @@ void FConsoleSlateDebugger::OnNavigationEvent(const FSlateDebuggingNavigationEve
 	OptionallyDumpCallStack();
 }
 
-void FConsoleSlateDebugger::OnStateChangeEvent(const FSlateDebuggingMouseCaptureEventArgs& EventArgs)
+void FConsoleSlateDebugger::OnCaptureStateChangeEvent(const FSlateDebuggingMouseCaptureEventArgs& EventArgs)
 {
-	static const FText StateChangeEventFormat = LOCTEXT("StateChangeEventFormat", "{0} : {1}");
+	static const FText StateChangeEventFormat = LOCTEXT("StateChangeEventFormat", "{0}({1}:{2}) : {3}");
 
-	const FText StateText = LOCTEXT("MouseCaptured", "Mouse Captured");
-	const FText SourceWidget = FText::FromString(FReflectionMetaData::GetWidgetDebugInfo(EventArgs.CapturingWidget.Get()));
+	const FText StateText = EventArgs.Captured ? LOCTEXT("MouseCaptured", "Mouse Captured") : LOCTEXT("MouseCaptureLost", "Mouse Capture Lost");
+	const FText SourceWidget = FText::FromString(FReflectionMetaData::GetWidgetDebugInfo(EventArgs.CaptureWidget.Get()));
 
 	FText EventText = FText::Format(
 		StateChangeEventFormat,
 		StateText,
+		EventArgs.UserIndex,
+		EventArgs.PointerIndex,
 		SourceWidget
 	);
 
diff --git a/Engine/Source/Runtime/SlateCore/Private/Debugging/ConsoleSlateDebugger.h b/Engine/Source/Runtime/SlateCore/Private/Debugging/ConsoleSlateDebugger.h
--- a/Engine/Source/Runtime/SlateCore/Private/Debugging/ConsoleSlateDebugger.h
+++ b/Engine/Source/Runtime/SlateCore/Private/Debugging/ConsoleSlateDebugger.h
@@ -41,7 +41,7 @@ private:
 	void OnInputEvent(const FSlateDebuggingInputEventArgs& EventArgs);
 	void OnFocusEvent(const FSlateDebuggingFocusEventArgs& EventArgs);
 	void OnNavigationEvent(const FSlateDebuggingNavigationEventArgs& EventArgs);
-	void OnStateChangeEvent(const FSlateDebuggingMouseCaptureEventArgs& EventArgs);
+	void OnCaptureStateChangeEvent(const FSlateDebuggingMouseCaptureEventArgs& EventArgs);
 
 	void OptionallyDumpCallStack();
 
diff --git a/Engine/Source/Runtime/SlateCore/Private/Debugging/SlateDebugging.cpp b/Engine/Source/Runtime/SlateCore/Private/Debugging/SlateDebugging.cpp
--- a/Engine/Source/Runtime/SlateCore/Private/Debugging/SlateDebugging.cpp
+++ b/Engine/Source/Runtime/SlateCore/Private/Debugging/SlateDebugging.cpp
@@ -50,8 +50,14 @@ FSlateDebuggingWarningEventArgs::FSlateDebuggingWarningEventArgs(
 }
 
 FSlateDebuggingMouseCaptureEventArgs::FSlateDebuggingMouseCaptureEventArgs(
+	bool InCaptured,
+	uint32 InUserIndex,
+	uint32 InPointerIndex,
 	const TSharedPtr<SWidget>& InCapturingWidget)
-	: CapturingWidget(InCapturingWidget)
+	: Captured(InCaptured)
+	, UserIndex(InUserIndex)
+	, PointerIndex(InPointerIndex)
+	, CaptureWidget(InCapturingWidget)
 {
 }
 
@@ -128,17 +134,19 @@ void FSlateDebugging::BroadcastFocusReceived(const FFocusEvent& InFocusEvent, co
 	FocusEvent.Broadcast(FSlateDebuggingFocusEventArgs(ESlateDebuggingFocusEvent::FocusReceived, InFocusEvent, InOldFocusedWidgetPath, InOldFocusedWidget, InNewFocusedWidgetPath, InNewFocusedWidget));
 }
 
-void FSlateDebugging::AttemptNavigation(const FNavigationEvent& InNavigationEvent, const FNavigationReply& InNavigationReply, const FWidgetPath& InNavigationSource, const TSharedPtr<SWidget>& InDestinationWidget)
+void FSlateDebugging::BroadcastAttemptNavigation(const FNavigationEvent& InNavigationEvent, const FNavigationReply& InNavigationReply, const FWidgetPath& InNavigationSource, const TSharedPtr<SWidget>& InDestinationWidget)
 {
 	NavigationEvent.Broadcast(FSlateDebuggingNavigationEventArgs(InNavigationEvent, InNavigationReply, InNavigationSource, InDestinationWidget));
 }
 
-DECLARE_MULTICAST_DELEGATE_OneParam(FWidgetMouseCaptureEvent, const FSlateDebuggingMouseCaptureEventArgs& /*EventArgs*/);
-static FWidgetMouseCaptureEvent MouseCaptureEvent;
+void FSlateDebugging::BroadcastMouseCapture(uint32 UserIndex, uint32 PointerIndex, const TSharedPtr<SWidget>& InCapturingWidget)
+{
+	MouseCaptureEvent.Broadcast(FSlateDebuggingMouseCaptureEventArgs(true, UserIndex, PointerIndex, InCapturingWidget));
+}
 
-void FSlateDebugging::MouseCapture(const TSharedPtr<SWidget>& InCapturingWidget)
+void FSlateDebugging::BroadcastMouseCaptureLost(uint32 UserIndex, uint32 PointerIndex, const TSharedPtr<SWidget>& InWidgetLostCapture)
 {
-	MouseCaptureEvent.Broadcast(FSlateDebuggingMouseCaptureEventArgs(InCapturingWidget));
+	MouseCaptureEvent.Broadcast(FSlateDebuggingMouseCaptureEventArgs(false, UserIndex, PointerIndex, InWidgetLostCapture));
 }
 
 #endif
diff --git a/Engine/Source/Runtime/SlateCore/Public/Debugging/SlateDebugging.h b/Engine/Source/Runtime/SlateCore/Public/Debugging/SlateDebugging.h
--- a/Engine/Source/Runtime/SlateCore/Public/Debugging/SlateDebugging.h
+++ b/Engine/Source/Runtime/SlateCore/Public/Debugging/SlateDebugging.h
@@ -127,10 +127,16 @@ struct SLATECORE_API FSlateDebuggingMouseCaptureEventArgs
 {
 public:
 	FSlateDebuggingMouseCaptureEventArgs(
+		bool InCaptured,
+		uint32 InUserIndex,
+		uint32 InPointerIndex,
 		const TSharedPtr<SWidget>& InCapturingWidget
 	);
 
-	const TSharedPtr<SWidget>& CapturingWidget;
+	bool Captured;
+	uint32 UserIndex;
+	uint32 PointerIndex;
+	const TSharedPtr<SWidget>& CaptureWidget;
 };
 
 
@@ -195,14 +201,15 @@ public:
 	DECLARE_MULTICAST_DELEGATE_OneParam(FWidgetNavigationEvent, const FSlateDebuggingNavigationEventArgs& /*EventArgs*/);
 	static FWidgetNavigationEvent NavigationEvent;
 
-	static void AttemptNavigation(const FNavigationEvent& InNavigationEvent, const FNavigationReply& InNavigationReply, const FWidgetPath& InNavigationSource, const TSharedPtr<SWidget>& InDestinationWidget);
+	static void BroadcastAttemptNavigation(const FNavigationEvent& InNavigationEvent, const FNavigationReply& InNavigationReply, const FWidgetPath& InNavigationSource, const TSharedPtr<SWidget>& InDestinationWidget);
 
 public:
 	/**  */
 	DECLARE_MULTICAST_DELEGATE_OneParam(FWidgetMouseCaptureEvent, const FSlateDebuggingMouseCaptureEventArgs& /*EventArgs*/);
 	static FWidgetMouseCaptureEvent MouseCaptureEvent;
 
-	static void MouseCapture(const TSharedPtr<SWidget>& InCapturingWidget);
+	static void BroadcastMouseCapture(uint32 UserIndex, uint32 PointerIndex, const TSharedPtr<SWidget>& InCapturingWidget);
+	static void BroadcastMouseCaptureLost(uint32 UserIndex, uint32 PointerIndex, const TSharedPtr<SWidget>& InWidgetLostCapture);
 
 public:
 	/**  */
