SlateMaterialBrushes shouldn't be destroying their resources directly we already have a system to free resources based on rendering handles, and GC.  Removing this code so that two instances of SlateMaterialBrushes pointed at the same Material don't cause one to stop drawing if the other goes away.

diff --git a/Engine/Source/Runtime/Slate/Public/SlateMaterialBrush.h b/Engine/Source/Runtime/Slate/Public/SlateMaterialBrush.h
--- a/Engine/Source/Runtime/Slate/Public/SlateMaterialBrush.h
+++ b/Engine/Source/Runtime/Slate/Public/SlateMaterialBrush.h
@@ -26,16 +26,4 @@ struct FSlateMaterialBrush : public FSlateBrush
 	{
 		ResourceName = FName( *InMaterial.GetFullName() );
 	}
-
-	/** Virtual destructor. */
-	virtual ~FSlateMaterialBrush()
-	{
-		if (FSlateApplication::IsInitialized())
-		{
-			if (FSlateRenderer* Renderer = FSlateApplication::Get().GetRenderer())
-			{
-				Renderer->ReleaseDynamicResource(*this);
-			}
-		}
-	}
 }; 
