PR #5754: Commit text changes if focus is lost from user interaction v2 (Contributed by [REDACTED_PERSON])

diff --git a/Engine/Source/Runtime/Slate/Private/Widgets/Text/SInlineEditableTextBlock.cpp b/Engine/Source/Runtime/Slate/Private/Widgets/Text/SInlineEditableTextBlock.cpp
--- a/Engine/Source/Runtime/Slate/Private/Widgets/Text/SInlineEditableTextBlock.cpp
+++ b/Engine/Source/Runtime/Slate/Private/Widgets/Text/SInlineEditableTextBlock.cpp
@@ -283,7 +283,7 @@ void SInlineEditableTextBlock::OnTextBoxCommitted(const FText& InText, ETextComm
 					return;
 				}
 			}
-			else if(InCommitType == ETextCommit::OnUserMovedFocus)
+			else if(InCommitType == ETextCommit::OnUserMovedFocus || InCommitType == ETextCommit::Default)
 			{
 				FText OutErrorMessage;
 				if(!OnVerifyTextChanged.Execute(InText, OutErrorMessage))
