PR #5731: Execute OnDrop when Esc key is pressed to end operation correctly (Contributed by [REDACTED_PERSON])

diff --git a/Engine/Source/Runtime/Slate/Private/Framework/Application/SlateApplication.cpp b/Engine/Source/Runtime/Slate/Private/Framework/Application/SlateApplication.cpp
--- a/Engine/Source/Runtime/Slate/Private/Framework/Application/SlateApplication.cpp
+++ b/Engine/Source/Runtime/Slate/Private/Framework/Application/SlateApplication.cpp
@@ -4087,13 +4087,15 @@ void FSlateApplication::CancelDragDrop()
 		return;
 	}
 
+	FPointerEvent EmptyPointerEvent;
+
 	for( auto LastWidgetIterator = WidgetsUnderCursorLastEvent.CreateConstIterator(); LastWidgetIterator; ++LastWidgetIterator)
 	{
 		
 		FWidgetPath WidgetsToDragLeave = LastWidgetIterator.Value().ToWidgetPath(FWeakWidgetPath::EInterruptedPathHandling::Truncate);
 		if(WidgetsToDragLeave.IsValid())
 		{
-			const FDragDropEvent DragDropEvent(FPointerEvent(), DragDropContent);
+			const FDragDropEvent DragDropEvent(EmptyPointerEvent, DragDropContent);
 			for(int32 WidgetIndex = WidgetsToDragLeave.Widgets.Num() - 1; WidgetIndex >= 0; --WidgetIndex)
 			{
 				WidgetsToDragLeave.Widgets[WidgetIndex].Widget->OnDragLeave(DragDropEvent);
@@ -4101,6 +4103,9 @@ void FSlateApplication::CancelDragDrop()
 		}
 	}
 
+	// cancel dragdrop operation correctly firing off callbacks
+	DragDropContent->OnDrop(false, EmptyPointerEvent);
+
 	WidgetsUnderCursorLastEvent.Empty();
 	ResetDragDropState();
 }
