PR #5585: Fix SInputKeySelector OnKeySelected event not fired when clearing selected key.

diff --git a/Engine/Source/Runtime/Slate/Private/Widgets/Input/SInputKeySelector.cpp b/Engine/Source/Runtime/Slate/Private/Widgets/Input/SInputKeySelector.cpp
--- a/Engine/Source/Runtime/Slate/Private/Widgets/Input/SInputKeySelector.cpp
+++ b/Engine/Source/Runtime/Slate/Private/Widgets/Input/SInputKeySelector.cpp
@@ -155,7 +155,7 @@ FReply SInputKeySelector::OnKeyDown(const FGeometry& MyGeometry, const FKeyEvent
 	{
 		if (SelectedKey.IsSet() && SelectedKey.Get().Key.IsValid() && (bAllowGamepadKeys && InKeyEvent.GetKey() == EKeys::Gamepad_FaceButton_Left))
 		{
-			SelectedKey = FInputChord();
+			SetSelectedKey(FInputChord());
 			return FReply::Handled();
 		}
 		else if (Button.IsValid())
@@ -218,7 +218,7 @@ FReply SInputKeySelector::OnMouseButtonDown(const FGeometry& MyGeometry, const F
 {
 	if (!bIsSelectingKey && SelectedKey.IsSet() && SelectedKey.Get().Key.IsValid() && MouseEvent.IsMouseButtonDown(EKeys::RightMouseButton))
 	{
-		SelectedKey = FInputChord();
+		SetSelectedKey(FInputChord());
 		return FReply::Handled();
 	}
 	return SCompoundWidget::OnMouseButtonDown(MyGeometry, MouseEvent);
@@ -248,4 +248,4 @@ void SInputKeySelector::SetTextBlockVisibility(EVisibility InVisibility)
 	{
 		TextBlock->SetVisibility(InVisibility);
 	}
-}
\ No newline at end of file
+}
