Fixing ScrollToItem error introduced with the recent update to edge navigation in SListView
- RequestScrollItemIntoView will now once again actually do what it says

diff --git a/Engine/Source/Runtime/Slate/Public/Widgets/Views/SListView.h b/Engine/Source/Runtime/Slate/Public/Widgets/Views/SListView.h
--- a/Engine/Source/Runtime/Slate/Public/Widgets/Views/SListView.h
+++ b/Engine/Source/Runtime/Slate/Public/Widgets/Views/SListView.h
@@ -1403,7 +1403,7 @@ protected:
 					// Center the list view on the item in question.
 					NewScrollOffset -= (NumLiveWidgets / 2);
 					//we also don't want the widget being chopped off if it is at the end of the list
-					const double MoveBackBy = FMath::Clamp<double>(IndexOfItem - MaxDisplayedIndex, 0, FLT_MAX);
+					const double MoveBackBy = FMath::Clamp<double>(IndexOfItem - (NewScrollOffset + NumLiveWidgets), 0, FLT_MAX);
 					//Move to the correct center spot
 					NewScrollOffset += MoveBackBy;
 
