Ensure rather than crash if someone tries to register a bogus UI command.

diff --git a/Engine/Source/Runtime/Slate/Private/Framework/Commands/UICommandList.cpp b/Engine/Source/Runtime/Slate/Private/Framework/Commands/UICommandList.cpp
--- a/Engine/Source/Runtime/Slate/Private/Framework/Commands/UICommandList.cpp
+++ b/Engine/Source/Runtime/Slate/Private/Framework/Commands/UICommandList.cpp
@@ -53,7 +53,10 @@ void FUICommandList::MapAction( const TSharedPtr< const FUICommandInfo > InUICom
 
 void FUICommandList::MapAction( const TSharedPtr< const FUICommandInfo > InUICommandInfo, const FUIAction& InUIAction )
 {
-	check( InUICommandInfo.IsValid() );
+	if (!ensure(InUICommandInfo.IsValid()))
+	{
+		return;
+	}
 
 	// Check against already mapped actions
 	checkfSlow(!UICommandBindingMap.Contains(InUICommandInfo), TEXT("Command list already contains a command named '%s'"), *InUICommandInfo->GetCommandName().ToString());
