This is a significant overhaul to the Curve Editor used by Sequencer which adds a plugin-based architecture and extensibility. New tools and toolbar buttons can be added to all usages of the curve editor via user plugins, and the different views for data can be created modularly so new implementations of the editor can register their own way of drawing their data and the tools should just work. Additionally, you can now write your own filters to operate on curve editor data for custom implementations of smoothing, key generation, etc.

The curve editor supports three view types by default - an absolute view (default, matches old behavior), a stacked view and a normalized view. Stacked views draw each curve separately (so non-overlapping) and normalized against their own min/max values. The normalized view draws all curves overlapping with each one normalized against its own min/max values.

A tree view has been added to help effectively manage large numbers of curves. Selecting curves in the treeview controls which curves are visible in the view area. The treeview also supports pinning curves. These pinned curves will always be visible regardless of your selection in the tree view.

A transform tool and a retiming tool have been implemented (via a plugin) which is enabled by default. The transform tool allows you to do a marquee selection of keys and then translate and scale the positions of these keys. The retiming tool allows you to create a 1 dimensional lattice to adjust the timing of your keys with a linear falloff between each lattice point. These tools work across multiple views at the same time which is especially useful if you are representing one dimensional data (such as event keys) in a view, as it allows you to adjust this data at the same time as your animation curves. A smoothing filter has been implemented (via a plugin) to allow running highpass and lowpass filters on your keys.

Opening the curve editor in Sequencer/UMG now creates a separate dockable tab which can be resized and docked as desired. A time slider has been added to the Curve Editor which is synchronized to the playback time in Sequencer. This allows you to scrub time in the curve editor without having to find the Sequencer window and adjust time there while looking at your keys and previewing your animation in the viewport at the same time.

Rudimentary support has been added for saving and later restoring a set of curves in your current session. This allows you to do a rudimentary copy/paste of entire curves but can also be useful for saving a curve, making adjustments to it and then deciding you want to go back - simply reapply the saved curve! Each curve added supports an intention name (such as "Location.X" or "FieldOfView"), and these intention names will be used when trying to apply curves. This allows you to reliably take all of the curves of a transform on one object and apply them to another object (and ensure that Location.X gets applied to the new Location.X, etc.) this can be helpful if you have a mixed set of curves buffered (such as a location and a field of view). In the event that no curves match by intention you can store and apply a single curve at a time from any intention to any other intention.

The Curve Asset editors (float, vector and color curve assets) have been changed to use the new editor. They support the same treeviews, filtering and tools that the Sequencer editor does. In addition, the Color Curve asset editor adds an additional view which provides a 1 dimensional gradient editor as an easier way to visualize and edit colors instead of the channels individually.

diff --git a/Engine/Source/Runtime/Slate/Public/Widgets/Input/SNumericList.h b/Engine/Source/Runtime/Slate/Public/Widgets/Input/SNumericList.h
new file mode 100644
--- /dev/null
+++ b/Engine/Source/Runtime/Slate/Public/Widgets/Input/SNumericList.h
@@ -0,0 +1,189 @@
+// Copyright 1998-2019 Epic Games, Inc. All Rights Reserved.
+
+#pragma once
+
+#include "CoreMinimal.h"
+#include "Misc/Attribute.h"
+#include "Layout/Visibility.h"
+#include "Input/Events.h"
+#include "Widgets/SWidget.h"
+#include "Textures/SlateIcon.h"
+#include "Layout/Margin.h"
+#include "Widgets/DeclarativeSyntaxSupport.h"
+#include "Widgets/SCompoundWidget.h"
+#include "Widgets/SBoxPanel.h"
+#include "Framework/Commands/UIAction.h"
+#include "Framework/MultiBox/MultiBoxBuilder.h"
+#include "Widgets/Text/STextBlock.h"
+#include "Framework/Application/SlateApplication.h"
+#include "Widgets/Input/SEditableTextBox.h"
+#include "Widgets/Input/SComboButton.h"
+
+/** A widget which creates a list of pre-determined numbers and a spot to enter a custom number. Similar to SNumericList but without the dropdown menu. */
+template<typename NumericType>
+class SNumericList : public SCompoundWidget
+{
+public:
+	DECLARE_DELEGATE_OneParam( FOnValueChanged, NumericType );
+
+public:
+	/** Represents a named numeric value for display in the drop down menu. */
+	class FNamedValue
+	{
+	public:
+		/** Creates a new FNamedValue 
+		  * @InValue The numeric value to be assigned
+		  * @InName The display text for the value in the UI 
+		  * @InDescription The description of the value used in tooltips or wherever a longer description is needed. */
+		FNamedValue( NumericType InValue, FText InName, FText InDescription )
+		{
+			Value = InValue;
+			Name = InName;
+			Description = InDescription;
+		}
+
+		NumericType GetValue() const
+		{
+			return Value;
+		}
+
+		FText GetName() const
+		{
+			return Name;
+		}
+
+		FText GetDescription() const
+		{
+			return Description;
+		}
+
+	private:
+		NumericType Value;
+		FText Name;
+		FText Description;
+	};
+
+public:
+	SLATE_BEGIN_ARGS( SNumericList<NumericType> )
+		: _DropDownValues()
+		, _MinDesiredValueWidth( 40 )
+		, _bShowNamedValue(true)
+		{}
+
+		/** The values which are used to populate the drop down menu. */
+		SLATE_ARGUMENT( TArray<FNamedValue>, DropDownValues )
+		/** Controls the minimum width for the text box portion of the control. */
+		SLATE_ATTRIBUTE( float, MinDesiredValueWidth )
+		/** The value displayed by the control. */
+		SLATE_ATTRIBUTE( NumericType, Value )
+		/** If enabled, attempt to show the original name specified in the Drop Down values for a given entry instead of just a straight FText::AsNumber(NumericType). */
+		SLATE_ATTRIBUTE(bool, bShowNamedValue)
+		/** An optional header to prepend to the generated list. */
+		SLATE_ATTRIBUTE( FText, HeaderText)
+		/** The callback for when the value changes. */
+		SLATE_EVENT( FOnValueChanged, OnValueChanged )
+
+	SLATE_END_ARGS()
+
+	void Construct( const FArguments& InArgs )
+	{
+		DropDownValues = InArgs._DropDownValues;
+		Value = InArgs._Value;
+		OnValueChanged = InArgs._OnValueChanged;
+
+		FMenuBuilder MenuBuilder(true, nullptr);
+
+		// Allow an optional header text because it looks nicer and helps give context to the widget.
+		if (InArgs._HeaderText.IsSet())
+		{
+			MenuBuilder.BeginSection("Header", InArgs._HeaderText.Get());
+		}
+
+		for (const FNamedValue& NamedValue : DropDownValues)
+		{
+			AddMenuEntry(MenuBuilder, NamedValue);
+		}
+
+		// Add the custom entry box at the end. 
+		MenuBuilder.AddWidget(
+			SNew(SBox)
+			.HAlign(HAlign_Right)
+			.Padding(FMargin(8, 0, 0, 0))
+			[
+				SNew(SEditableTextBox)
+				.RevertTextOnEscape(true)
+				.SelectAllTextWhenFocused(true)
+				.Text(this, &SNumericList<NumericType>::GetValueText)
+				.OnTextCommitted(this, &SNumericList<NumericType>::ValueTextCommitted)
+				.MinDesiredWidth(InArgs._MinDesiredValueWidth)
+			],
+			NSLOCTEXT("SNumericList", "CustomNumberDisplayLabel", "Custom")
+		);
+
+		ChildSlot
+		[
+			MenuBuilder.MakeWidget()
+		];
+	}
+
+private:
+	void AddMenuEntry(FMenuBuilder& MenuBuilder, const FNamedValue& Info)
+	{
+		MenuBuilder.AddMenuEntry(
+			Info.GetName(),
+			Info.GetDescription(),
+			FSlateIcon(),
+			FUIAction(
+				FExecuteAction::CreateSP(this, &SNumericList<NumericType>::SetValue, Info.GetValue()),
+				FCanExecuteAction(),
+				FIsActionChecked::CreateSP(this, &SNumericList<NumericType>::IsSameValue, Info.GetValue())
+			),
+			NAME_None,
+			EUserInterfaceActionType::RadioButton
+		);
+	}
+
+	bool IsSameValue(NumericType InValue) const
+	{
+		return Value.Get() == InValue;
+	}
+
+	/** Get the value text for displaying in the custom entry box. We use a textbox here so we can show actual names from the original list if possible. */
+	FText GetValueText() const
+	{
+		if (bShowNamedValue.Get())
+		{
+			for ( FNamedValue DropDownValue : DropDownValues )
+			{
+				if (FMath::IsNearlyEqual(DropDownValue.GetValue(), Value.Get()))
+				{
+					return DropDownValue.GetName();
+				}
+			}
+		}
+
+		return FText::FromString(LexToSanitizedString(Value.Get()));
+	}
+
+	void ValueTextCommitted(const FText& InNewText, ETextCommit::Type InTextCommit)
+	{
+		if (InNewText.IsNumeric())
+		{
+			NumericType NewValue;
+			TTypeFromString<NumericType>::FromString(NewValue, *InNewText.ToString());
+			OnValueChanged.ExecuteIfBound(NewValue);
+		}
+	}
+
+	void SetValue( NumericType InValue )
+	{
+		FSlateApplication::Get().ClearKeyboardFocus( EFocusCause::Cleared );
+		OnValueChanged.ExecuteIfBound(InValue);
+	}
+
+private:
+	TArray<FNamedValue> DropDownValues;
+	TAttribute<bool> bShowNamedValue;
+	TAttribute<NumericType> Value;
+	FOnValueChanged OnValueChanged;
+};
