Slate - The SInvalidation panel now logs excessive invalidation, (controlled by cvars, both the rate and if it's enabled).
Slate - The STextBlock can get into invalidation loops due to the nature of floating point precision.  We now compare apples to apples, with the before and after desired size of the TextLayoutCache, and if it changed after paint, it means the new size (discovered in OnPaint), has affected the text's wrapping, and thus we need another frame to lay things out correctly, so we invalidate again.  We now only do this if the text wraps, and we now add in a epsilon.

diff --git a/Engine/Source/Runtime/Slate/Private/Widgets/SInvalidationPanel.cpp b/Engine/Source/Runtime/Slate/Private/Widgets/SInvalidationPanel.cpp
--- a/Engine/Source/Runtime/Slate/Private/Widgets/SInvalidationPanel.cpp
+++ b/Engine/Source/Runtime/Slate/Private/Widgets/SInvalidationPanel.cpp
@@ -73,7 +73,23 @@ static int32 InvalidationLayerPadding = 10;
 static FAutoConsoleVariableRef CVarInvalidationLayerPadding(
 	TEXT("Slate.InvalidationLayerPadding"),
 	InvalidationLayerPadding,
-	TEXT("The invalidation panel adds this much padding to the layer id to avoid re-invalidating if it flucuates in a volatiel child widget."));
+	TEXT("The invalidation panel adds this much padding to the layer id to avoid re-invalidating if it fluctuates in a volatile child widget."));
+
+#if SLATE_VERBOSE_NAMED_EVENTS
+
+static int32 ExcessiveInvalidationFrameStreak = 10;
+static FAutoConsoleVariableRef CVarExcessiveInvalidationFrameStreak(
+	TEXT("Slate.ExcessiveInvalidationFrameStreak"),
+	ExcessiveInvalidationFrameStreak,
+	TEXT("."));
+
+static int32 LogExcessiveInvalidation = 1;
+static FAutoConsoleVariableRef CVarLogExcessiveInvalidation(
+	TEXT("Slate.LogExcessiveInvalidation"),
+	LogExcessiveInvalidation,
+	TEXT("."));
+
+#endif
 
 static bool ShouldCacheRenderData()
 {
@@ -103,6 +119,7 @@ void SInvalidationPanel::Construct( const FArguments& InArgs )
 	LastUsedCachedNodeIndex = 0;
 	LastHitTestIndex = 0;
 	MaximumLayerIdCachedAt = 0;
+	NumberOfFramesInARowWeInvalidated = 0;
 	LastClippingIntersectionSize = FVector2D::ZeroVector;
 
 	bCacheRelativeTransforms = InArgs._CacheRelativeTransforms;
@@ -293,6 +310,16 @@ void SInvalidationPanel::InvalidateWidget(SWidget* InvalidateWidget)
 {
 	bNeedsCaching = true;
 
+#if SLATE_VERBOSE_NAMED_EVENTS
+	if (LogExcessiveInvalidation && NumberOfFramesInARowWeInvalidated > 0 && (NumberOfFramesInARowWeInvalidated % ExcessiveInvalidationFrameStreak) == 0)
+	{
+		const FString InvalidationCause = FReflectionMetaData::GetWidgetDebugInfo(InvalidateWidget);
+		const FString InvalidationRoot = FReflectionMetaData::GetWidgetDebugInfo(this);
+
+		UE_LOG(LogSlate, Warning, TEXT("SInvalidationPanel(%s): Excessive Invalidation Detected: %s"), *InvalidationRoot, *InvalidationCause);
+	}
+#endif
+
 #if !(UE_BUILD_SHIPPING || UE_BUILD_TEST)
 	if ( InvalidateWidget != nullptr && IsInvalidationDebuggingEnabled() )
 	{
@@ -346,6 +373,7 @@ int32 SInvalidationPanel::OnPaint( const FPaintArgs& Args, const FGeometry& Allo
 				// Always set the caching flag to false first, during the paint / tick pass we may change something
 				// to volatile and need to re-cache.
 				bNeedsCaching = false;
+				NumberOfFramesInARowWeInvalidated++;
 
 				bNeedsCachePrepass = true;
 
@@ -443,6 +471,8 @@ int32 SInvalidationPanel::OnPaint( const FPaintArgs& Args, const FGeometry& Allo
 			}
 		else
 		{
+			NumberOfFramesInARowWeInvalidated = 0;
+
 			// The full clipping state is stored on the render batches when we cache the render data,
 			// so there's no need to merge clipping states on frames where the data has been cached successfully.
 			if (CacheRenderData == 0)
diff --git a/Engine/Source/Runtime/Slate/Private/Widgets/Text/STextBlock.cpp b/Engine/Source/Runtime/Slate/Private/Widgets/Text/STextBlock.cpp
--- a/Engine/Source/Runtime/Slate/Private/Widgets/Text/STextBlock.cpp
+++ b/Engine/Source/Runtime/Slate/Private/Widgets/Text/STextBlock.cpp
@@ -180,7 +180,6 @@ void STextBlock::SetHighlightText(TAttribute<FText> InText)
 int32 STextBlock::OnPaint( const FPaintArgs& Args, const FGeometry& AllottedGeometry, const FSlateRect& MyCullingRect, FSlateWindowElementList& OutDrawElements, int32 LayerId, const FWidgetStyle& InWidgetStyle, bool bParentEnabled ) const
 {
 	SCOPE_CYCLE_COUNTER(Stat_SlateTextBlockOnPaint);
-	//SCOPED_NAMED_EVENT_TEXT("STextBlock", FColor::Orange);
 
 	if (bSimpleTextMode)
 	{
@@ -223,13 +222,19 @@ int32 STextBlock::OnPaint( const FPaintArgs& Args, const FGeometry& AllottedGeom
 	}
 	else
 	{
+		const FVector2D LastDesiredSize = TextLayoutCache->GetDesiredSize();
+
 		// OnPaint will also update the text layout cache if required
 		LayerId = TextLayoutCache->OnPaint(Args, AllottedGeometry, MyCullingRect, OutDrawElements, LayerId, InWidgetStyle, ShouldBeEnabled(bParentEnabled));
 
+		const FVector2D NewDesiredSize = TextLayoutCache->GetDesiredSize();
+
 		// HACK: Due to the nature of wrapping and layout, we may have been arranged in a different box than what we were cached with.  Which
 		// might update wrapping, so make sure we always set the desired size to the current size of the text layout, which may have changed
 		// during paint.
-		if (TextLayoutCache->GetDesiredSize().Y > GetDesiredSize().Y)
+		bool bCanWrap = WrapTextAt.Get() > 0 || AutoWrapText.Get();
+
+		if (bCanWrap && !NewDesiredSize.Equals(LastDesiredSize))
 		{
 			const_cast<STextBlock*>(this)->Invalidate(EInvalidateWidget::Layout);
 		}
diff --git a/Engine/Source/Runtime/Slate/Public/Widgets/SInvalidationPanel.h b/Engine/Source/Runtime/Slate/Public/Widgets/SInvalidationPanel.h
--- a/Engine/Source/Runtime/Slate/Public/Widgets/SInvalidationPanel.h
+++ b/Engine/Source/Runtime/Slate/Public/Widgets/SInvalidationPanel.h
@@ -129,6 +129,7 @@ private:
 	mutable bool bNeedsCachePrepass;
 	mutable bool bPerformingCachePrepass;
 	mutable bool bIsInvalidating;
+	mutable int32 NumberOfFramesInARowWeInvalidated;
 	bool bCanCache;
 
 	bool bCacheRelativeTransforms;
