Fixed a OOB when the Linearized Item list returned an index that didn't exist in the DenseItemInfos due to a combination of tree selection and rebuilding that is used in some places (such as the Skeletal Mesh Editor).

diff --git a/Engine/Source/Runtime/Slate/Public/Widgets/Views/STreeView.h b/Engine/Source/Runtime/Slate/Public/Widgets/Views/STreeView.h
--- a/Engine/Source/Runtime/Slate/Public/Widgets/Views/STreeView.h
+++ b/Engine/Source/Runtime/Slate/Public/Widgets/Views/STreeView.h
@@ -456,15 +456,18 @@ private:
 				continue;
 			}
 
-			const FItemInfo& ItemInfo = DenseItemInfos[ItemIndex];
-			int32 ParentIndex = ItemInfo.ParentIndex;
-			while (ParentIndex != INDEX_NONE)
+			if(DenseItemInfos.IsValidIndex(ItemIndex))
 			{
-				const ItemType& ParentItem = this->LinearizedItems[ParentIndex];
-				this->Private_SetItemHighlighted(ParentItem, true);
+				const FItemInfo& ItemInfo = DenseItemInfos[ItemIndex];
+				int32 ParentIndex = ItemInfo.ParentIndex;
+				while (ParentIndex != INDEX_NONE)
+				{
+					const ItemType& ParentItem = this->LinearizedItems[ParentIndex];
+					this->Private_SetItemHighlighted(ParentItem, true);
 
-				const FItemInfo& ParentItemInfo = DenseItemInfos[ParentIndex];
-				ParentIndex = ParentItemInfo.ParentIndex;
+					const FItemInfo& ParentItemInfo = DenseItemInfos[ParentIndex];
+					ParentIndex = ParentItemInfo.ParentIndex;
+				}
 			}
 		}
 	}
