Slate - Fix numeric text ignoring "ApplyOutlineToDropShadows" setting

diff --git a/Engine/Source/Runtime/Slate/Private/Widgets/Text/STextBlock.cpp b/Engine/Source/Runtime/Slate/Private/Widgets/Text/STextBlock.cpp
--- a/Engine/Source/Runtime/Slate/Private/Widgets/Text/STextBlock.cpp
+++ b/Engine/Source/Runtime/Slate/Private/Widgets/Text/STextBlock.cpp
@@ -191,10 +191,16 @@ int32 STextBlock::OnPaint( const FPaintArgs& Args, const FGeometry& AllottedGeom
 		const bool bShouldBeEnabled = ShouldBeEnabled(bParentEnabled);
 
 		const FText& LocalText = GetText();
-		const FSlateFontInfo LocalFont = GetFont();
+		FSlateFontInfo LocalFont = GetFont();
 
 		if (ShouldDropShadow)
 		{
+			const int32 OutlineSize = LocalFont.OutlineSettings.OutlineSize;
+			if (!LocalFont.OutlineSettings.bApplyOutlineToDropShadows)
+			{
+				LocalFont.OutlineSettings.OutlineSize = 0;
+			}
+
 			FSlateDrawElement::MakeText(
 				OutDrawElements,
 				LayerId,
@@ -205,6 +211,9 @@ int32 STextBlock::OnPaint( const FPaintArgs& Args, const FGeometry& AllottedGeom
 				InWidgetStyle.GetColorAndOpacityTint() * LocalShadowColorAndOpacity
 			);
 
+			// Restore outline size for main text
+			LocalFont.OutlineSettings.OutlineSize = OutlineSize;
+
 			// actual text should appear above the shadow
 			++LayerId;
 		}
