Slate - When GSlateLayoutCaching is enabled, TextBlocks must always call GetDesiredSize().  If they don't - the text wont ever perform an initial pass to layout anything, since it's done during prepass.  Additionally - Whenever a new prepass is needed under GSlateLayoutCahcing, this needs to also imply bNeedsDesiredSize, because some widgets, like Textblocks calculate a different desired size based on that incoming layout scale.

Matt.Kuhlenschmidt

diff --git a/Engine/Source/Runtime/Slate/Private/Widgets/Text/SRichTextBlock.cpp b/Engine/Source/Runtime/Slate/Private/Widgets/Text/SRichTextBlock.cpp
--- a/Engine/Source/Runtime/Slate/Private/Widgets/Text/SRichTextBlock.cpp
+++ b/Engine/Source/Runtime/Slate/Private/Widgets/Text/SRichTextBlock.cpp
@@ -59,13 +59,27 @@ void SRichTextBlock::Construct( const FArguments& InArgs )
 
 int32 SRichTextBlock::OnPaint( const FPaintArgs& Args, const FGeometry& AllottedGeometry, const FSlateRect& MyCullingRect, FSlateWindowElementList& OutDrawElements, int32 LayerId, const FWidgetStyle& InWidgetStyle, bool bParentEnabled ) const
 {
+	const FVector2D LastDesiredSize = TextLayoutCache->GetDesiredSize();
+
+	// If we're performing layout caching, it's possible nobody ever called GetDesiredSize(),
+	// which for textblocks is required to be called, since CDS is where it actually generates
+	// a lot for the text layout.
+	if (GSlateLayoutCaching)
+	{
+		GetDesiredSize();
+	}
+
 	// OnPaint will also update the text layout cache if required
 	LayerId = TextLayoutCache->OnPaint(Args, AllottedGeometry, MyCullingRect, OutDrawElements, LayerId, InWidgetStyle, ShouldBeEnabled(bParentEnabled));
 
+	const FVector2D NewDesiredSize = TextLayoutCache->GetDesiredSize();
+
 	// HACK: Due to the nature of wrapping and layout, we may have been arranged in a different box than what we were cached with.  Which
 	// might update wrapping, so make sure we always set the desired size to the current size of the text layout, which may have changed
 	// during paint.
-	if (TextLayoutCache->GetDesiredSize().Y > GetDesiredSize().Y)
+	bool bCanWrap = WrapTextAt.Get() > 0 || AutoWrapText.Get();
+
+	if (bCanWrap && !NewDesiredSize.Equals(LastDesiredSize))
 	{
 		const_cast<SRichTextBlock*>(this)->Invalidate(EInvalidateWidget::Layout);
 	}
diff --git a/Engine/Source/Runtime/Slate/Private/Widgets/Text/STextBlock.cpp b/Engine/Source/Runtime/Slate/Private/Widgets/Text/STextBlock.cpp
--- a/Engine/Source/Runtime/Slate/Private/Widgets/Text/STextBlock.cpp
+++ b/Engine/Source/Runtime/Slate/Private/Widgets/Text/STextBlock.cpp
@@ -233,6 +233,14 @@ int32 STextBlock::OnPaint( const FPaintArgs& Args, const FGeometry& AllottedGeom
 	{
 		const FVector2D LastDesiredSize = TextLayoutCache->GetDesiredSize();
 
+		// If we're performing layout caching, it's possible nobody ever called GetDesiredSize(),
+		// which for textblocks is required to be called, since CDS is where it actually generates
+		// a lot for the text layout.
+		if (GSlateLayoutCaching)
+		{
+			GetDesiredSize();
+		}
+
 		// OnPaint will also update the text layout cache if required
 		LayerId = TextLayoutCache->OnPaint(Args, AllottedGeometry, MyCullingRect, OutDrawElements, LayerId, InWidgetStyle, ShouldBeEnabled(bParentEnabled));
 
diff --git a/Engine/Source/Runtime/SlateCore/Private/Widgets/SWidget.cpp b/Engine/Source/Runtime/SlateCore/Private/Widgets/SWidget.cpp
--- a/Engine/Source/Runtime/SlateCore/Private/Widgets/SWidget.cpp
+++ b/Engine/Source/Runtime/SlateCore/Private/Widgets/SWidget.cpp
@@ -493,6 +493,10 @@ void SWidget::SlatePrepass(float InLayoutScaleMultiplier)
 
 		PrepassLayoutScaleMultiplier = InLayoutScaleMultiplier;
 		bNeedsPrepass = false;
+
+		// If the scale changed, that can affect the desired size of some elements that take it into
+		// account, such as text, so when the prepass size changes, so must we invalidate desired size.
+		bNeedsDesiredSize = true;
 	}
 
 	if ( bCanHaveChildren )
