Slate - Making SScrollBars and SSpacer play nicer with invalidation.

Matt.Kuhlenschmidt

diff --git a/Engine/Source/Runtime/Slate/Private/Widgets/Layout/SScrollBar.cpp b/Engine/Source/Runtime/Slate/Private/Widgets/Layout/SScrollBar.cpp
--- a/Engine/Source/Runtime/Slate/Private/Widgets/Layout/SScrollBar.cpp
+++ b/Engine/Source/Runtime/Slate/Private/Widgets/Layout/SScrollBar.cpp
@@ -339,6 +339,8 @@ void SScrollBar::SetStyle(const FScrollBarStyle* InStyle)
 		TopBrush = &Style->HorizontalTopSlotImage;
 		BottomBrush = &Style->HorizontalBottomSlotImage;
 	}
+
+	Invalidate(EInvalidateWidget::Layout);
 }
 
 void SScrollBar::SetDragFocusCause(EFocusCause InDragFocusCause)
@@ -357,18 +359,20 @@ void SScrollBar::SetScrollBarAlwaysVisible(bool InAlwaysVisible)
 
 	if ( InAlwaysVisible )
 	{
-		Visibility = EVisibility::Visible;
+		SetVisibility(EVisibility::Visible);
 	}
 	else
 	{
-		Visibility = TAttribute<EVisibility>(SharedThis(this), &SScrollBar::ShouldBeVisible);
+		SetVisibility(TAttribute<EVisibility>(SharedThis(this), &SScrollBar::ShouldBeVisible));
 	}
+
 	Track->SetIsAlwaysVisible(InAlwaysVisible);
 }
 
 void SScrollBar::SetScrollBarTrackAlwaysVisible(bool InAlwaysVisible)
 {
 	bAlwaysShowScrollbarTrack = InAlwaysVisible;
+	Invalidate(EInvalidateWidget::Layout);
 }
 
 bool SScrollBar::AlwaysShowScrollbar() const
diff --git a/Engine/Source/Runtime/Slate/Private/Widgets/Layout/SScrollBarTrack.cpp b/Engine/Source/Runtime/Slate/Private/Widgets/Layout/SScrollBarTrack.cpp
--- a/Engine/Source/Runtime/Slate/Private/Widgets/Layout/SScrollBarTrack.cpp
+++ b/Engine/Source/Runtime/Slate/Private/Widgets/Layout/SScrollBarTrack.cpp
@@ -128,6 +128,8 @@ void SScrollBarTrack::SetSizes(float InThumbOffsetFraction, float InThumbSizeFra
 	{
 		ThumbSizeFraction = 0.0f;
 	}
+
+	Invalidate(EInvalidateWidget::Paint);
 }
 
 bool SScrollBarTrack::IsNeeded() const
@@ -159,5 +161,9 @@ float SScrollBarTrack::GetThumbSizeFraction() const
 
 void SScrollBarTrack::SetIsAlwaysVisible(bool InIsAlwaysVisible)
 {
-	bIsAlwaysVisible = InIsAlwaysVisible;
+	if (bIsAlwaysVisible != InIsAlwaysVisible)
+	{
+		bIsAlwaysVisible = InIsAlwaysVisible;
+		Invalidate(EInvalidateWidget::Layout);
+	}
 }
diff --git a/Engine/Source/Runtime/Slate/Public/Widgets/Layout/SSpacer.h b/Engine/Source/Runtime/Slate/Public/Widgets/Layout/SSpacer.h
--- a/Engine/Source/Runtime/Slate/Public/Widgets/Layout/SSpacer.h
+++ b/Engine/Source/Runtime/Slate/Public/Widgets/Layout/SSpacer.h
@@ -49,6 +49,7 @@ public:
 	void SetSize( TAttribute<FVector2D> InSpacerSize )
 	{
 		SpacerSize = InSpacerSize;
+		Invalidate(EInvalidateWidget::Layout);
 	}
 
 protected:
