Add methods to SComboBox so it can use non-nullable template types.

diff --git a/Engine/Source/Runtime/Slate/Public/Widgets/Input/SComboBox.h b/Engine/Source/Runtime/Slate/Public/Widgets/Input/SComboBox.h
--- a/Engine/Source/Runtime/Slate/Public/Widgets/Input/SComboBox.h
+++ b/Engine/Source/Runtime/Slate/Public/Widgets/Input/SComboBox.h
@@ -87,11 +87,14 @@ template< typename OptionType >
 class SComboBox : public SComboButton
 {
 public:
+
+	typedef typename TListTypeTraits< OptionType >::NullableType NullableOptionType;
+
 	/** Type of list used for showing menu options. */
 	typedef SListView< OptionType > SComboListType;
 	/** Delegate type used to generate widgets that represent Options */
 	typedef typename TSlateDelegates< OptionType >::FOnGenerateWidget FOnGenerateWidget;
-	typedef typename TSlateDelegates< OptionType >::FOnSelectionChanged FOnSelectionChanged;
+	typedef typename TSlateDelegates< NullableOptionType >::FOnSelectionChanged FOnSelectionChanged;
 
 	SLATE_BEGIN_ARGS( SComboBox )
 		: _Content()
@@ -135,7 +138,7 @@ public:
 		SLATE_ARGUMENT(TSharedPtr<SScrollBar>, CustomScrollbar)
 
 		/** The option that should be selected when the combo box is first created */
-		SLATE_ARGUMENT( OptionType, InitiallySelectedItem )
+		SLATE_ARGUMENT( NullableOptionType, InitiallySelectedItem )
 
 		SLATE_ARGUMENT( TOptional<EPopupMethod>, Method )
 
@@ -237,8 +240,9 @@ public:
 		SelectedItem = InArgs._InitiallySelectedItem;
 		if( TListTypeTraits<OptionType>::IsPtrValid( SelectedItem ) )
 		{
-			ComboListView->Private_SetItemSelection( SelectedItem, true);
-			ComboListView->RequestScrollIntoView(SelectedItem, 0);
+			OptionType ValidatedItem = TListTypeTraits<OptionType>::NullableItemTypeConvertToItemType(SelectedItem);
+			ComboListView->Private_SetItemSelection(ValidatedItem, true);
+			ComboListView->RequestScrollIntoView(ValidatedItem, 0);
 		}
 
 	}
@@ -248,11 +252,12 @@ public:
 		ComboListView->ClearSelection();
 	}
 
-	void SetSelectedItem( OptionType InSelectedItem )
+	void SetSelectedItem(NullableOptionType InSelectedItem)
 	{
 		if (TListTypeTraits<OptionType>::IsPtrValid(InSelectedItem))
 		{
-			ComboListView->SetSelection(InSelectedItem);
+			OptionType InSelected = TListTypeTraits<OptionType>::NullableItemTypeConvertToItemType(InSelectedItem);
+			ComboListView->SetSelection(InSelected);
 		}
 		else
 		{
@@ -261,11 +266,12 @@ public:
 	}
 
 	/** @return the item currently selected by the combo box. */
-	OptionType GetSelectedItem()
+	NullableOptionType GetSelectedItem()
 	{
 		return SelectedItem;
 	}
 
+
 	/** 
 	 * Requests a list refresh after updating options 
 	 * Call SetSelectedItem to update the selected item if required
@@ -362,22 +368,32 @@ protected:
 			{
 				if (KeyPressed == EKeys::Up || KeyPressed == EKeys::Gamepad_DPad_Up || KeyPressed == EKeys::Gamepad_LeftStick_Up)
 				{
-					const int32 SelectionIndex = OptionsSource->Find(GetSelectedItem());
-					if (SelectionIndex >= 1)
+					NullableOptionType NullableSelected = GetSelectedItem();
+					if (TListTypeTraits<OptionType>::IsPtrValid(NullableSelected))
 					{
-						// Select an item on the prev row
-						SetSelectedItem((*OptionsSource)[SelectionIndex - 1]);
+						OptionType ActuallySelected = TListTypeTraits<OptionType>::NullableItemTypeConvertToItemType(NullableSelected);
+						const int32 SelectionIndex = OptionsSource->Find(ActuallySelected);
+						if (SelectionIndex >= 1)
+						{
+							// Select an item on the prev row
+							SetSelectedItem((*OptionsSource)[SelectionIndex - 1]);
+						}
 					}
 
 					return FReply::Handled();
 				}
 				else if (KeyPressed == EKeys::Down || KeyPressed == EKeys::Gamepad_DPad_Down || KeyPressed == EKeys::Gamepad_LeftStick_Down)
 				{
-					const int32 SelectionIndex = OptionsSource->Find(GetSelectedItem());
-					if (SelectionIndex < OptionsSource->Num() - 1)
+					NullableOptionType NullableSelected = GetSelectedItem();
+					if (TListTypeTraits<OptionType>::IsPtrValid(NullableSelected))
 					{
-						// Select an item on the next row
-						SetSelectedItem((*OptionsSource)[SelectionIndex + 1]);
+						OptionType ActuallySelected = TListTypeTraits<OptionType>::NullableItemTypeConvertToItemType(NullableSelected);
+						const int32 SelectionIndex = OptionsSource->Find(ActuallySelected);
+						if (SelectionIndex < OptionsSource->Num() - 1)
+						{
+							// Select an item on the next row
+							SetSelectedItem((*OptionsSource)[SelectionIndex + 1]);
+						}
 					}
 					return FReply::Handled();
 				}
@@ -431,8 +447,10 @@ private:
 			if (TListTypeTraits<OptionType>::IsPtrValid(SelectedItem))
 			{
 				// Ensure the ListView selection is set back to the last committed selection
-				ComboListView->SetSelection(SelectedItem, ESelectInfo::OnNavigation);
-				ComboListView->RequestScrollIntoView(SelectedItem, 0);
+				OptionType ActuallySelected = TListTypeTraits<OptionType>::NullableItemTypeConvertToItemType(SelectedItem);
+
+				ComboListView->SetSelection(ActuallySelected, ESelectInfo::OnNavigation);
+				ComboListView->RequestScrollIntoView(ActuallySelected, 0);
 			}
 
 			// Set focus back to ComboBox for users focusing the ListView that just closed
@@ -443,12 +461,11 @@ private:
 				}
 			});
 
-		}	
+		}
 	}
 
-
 	/** Invoked when the selection in the list changes */
-	void OnSelectionChanged_Internal( OptionType ProposedSelection, ESelectInfo::Type SelectInfo )
+	void OnSelectionChanged_Internal( NullableOptionType ProposedSelection, ESelectInfo::Type SelectInfo )
 	{
 		// Ensure that the proposed selection is different
 		if(SelectInfo != ESelectInfo::OnNavigation)
@@ -512,7 +529,7 @@ private:
 	/** Delegate that is invoked when the selected item in the combo box changes */
 	FOnSelectionChanged OnSelectionChanged;
 	/** The item currently selected in the combo box */
-	OptionType SelectedItem;
+	NullableOptionType SelectedItem;
 	/** The ListView that we pop up; visualized the available options. */
 	TSharedPtr< SComboListType > ComboListView;
 	/** The Scrollbar used in the ListView. */
