Update to SListView to add support for navigating lists with jagged (non-full final rows) ends.  If a suitable list element is not directly below the current one, but there is an incomplete row, navigation will select the final item of the list.

diff --git a/Engine/Source/Runtime/Slate/Public/Widgets/Views/SListView.h b/Engine/Source/Runtime/Slate/Public/Widgets/Views/SListView.h
--- a/Engine/Source/Runtime/Slate/Public/Widgets/Views/SListView.h
+++ b/Engine/Source/Runtime/Slate/Public/Widgets/Views/SListView.h
@@ -409,6 +409,27 @@ public:
 
 			case EUINavigation::Down:
 				AttemptSelectIndex = CurSelectionIndex + NumItemsWide;
+
+				// The list might be jagged so attempt to determine if there's a partially filled row we can move to
+				if (!ItemsSourceRef.IsValidIndex(AttemptSelectIndex))
+				{
+					int32 NumItems = ItemsSourceRef.Num();
+					if (NumItems > 0)
+					{
+						// NumItemsWide should never be 0, ensuring for sanity
+						ensure(NumItemsWide > 0);
+
+						// calculate total number of rows and row of current index (1 index)
+						int32 NumRows = NumItems / NumItemsWide + (NumItems % NumItemsWide != 0 ? 1 : 0);
+						int32 CurRow = CurSelectionIndex / NumItems + 1;
+
+						// if not on final row, assume a jagged list and select the final item
+						if (CurRow < NumRows)
+						{
+							AttemptSelectIndex = NumItems - 1;
+						}
+					}
+				}
 				break;
 
 			default:
