Fixed invalid use of an internal UI_COMMAND function that broke localization support

#rnx

diff --git a/Engine/Source/Editor/Experimental/EditorInteractiveToolsFramework/Public/InteractiveToolsCommands.h b/Engine/Source/Editor/Experimental/EditorInteractiveToolsFramework/Public/InteractiveToolsCommands.h
--- a/Engine/Source/Editor/Experimental/EditorInteractiveToolsFramework/Public/InteractiveToolsCommands.h
+++ b/Engine/Source/Editor/Experimental/EditorInteractiveToolsFramework/Public/InteractiveToolsCommands.h
@@ -18,7 +18,6 @@
  * Usage is as follows:
  *    - in your EdMode Module, create a subclass-instance of this, say FMyToolCommands and call FMyToolCommands::Register() in your ::StartupModule() function
  *        - subclass must implement ::GetToolDefaultObjectList(), here you just add GetMutableDefault<MyToolX> to the input list for all your Tools
- *        - subclass must implement ::GetNamespaceString(), likely you can just return TEXT(LOCTEXT_NAMESPACE)
  *    - add a member TSharedPtr<FUICommandList> UICommandList; to your EdMode impl
  *    - when you start a new Tool, call FMyToolCommands::Get().BindCommandsForCurrentTool(UICommandList, NewTool)
  *    - when you end a Tool, call FMyToolCommands::Get().UnbindActiveCommands()
@@ -60,11 +59,6 @@ public:
 	 */
 	virtual void GetToolDefaultObjectList(TArray<UInteractiveTool*>& ToolCDOs) = 0;
 
-	/**
-	 * The Namespace used to identify this Command set
-	 */
-	virtual const TCHAR * GetNamespaceString() = 0;
-
 
 protected:
 
@@ -283,10 +277,18 @@ bool TInteractiveToolCommands<CommandContextType>::IntializeStandardToolAction(E
 template<typename CommandContextType>
 void TInteractiveToolCommands<CommandContextType>::RegisterUIToolCommand(const FInteractiveToolAction& ToolAction, TSharedPtr<FUICommandInfo>& UICommandInfo)
 {
-	FString TooltipString = ToolAction.ActionName + FString("_ToolTip");
-	FString DotString = FString(".") + ToolAction.ActionName;
-
-	UI_COMMAND_Function(this, UICommandInfo, GetNamespaceString(), *ToolAction.ActionName, *TooltipString, TCHAR_TO_ANSI(*DotString), *ToolAction.ShortName, *ToolAction.Description, EUserInterfaceActionType::Button, FInputChord(ToolAction.DefaultModifiers, ToolAction.DefaultKey));
+	const FString DotString = FString(TEXT(".")) + ToolAction.ActionName;
+
+	FUICommandInfo::MakeCommandInfo(
+		AsShared(),
+		UICommandInfo,
+		*ToolAction.ActionName,
+		ToolAction.ShortName,
+		ToolAction.Description,
+		FSlateIcon(GetStyleSetName(), ISlateStyle::Join(GetContextName(), TCHAR_TO_ANSI(*DotString))),
+		EUserInterfaceActionType::Button,
+		FInputChord(ToolAction.DefaultModifiers, ToolAction.DefaultKey)
+		);
 }
 
 
diff --git a/Engine/Source/Editor/UnrealEd/Private/Commandlets/GatherTextFromSourceCommandlet.cpp b/Engine/Source/Editor/UnrealEd/Private/Commandlets/GatherTextFromSourceCommandlet.cpp
--- a/Engine/Source/Editor/UnrealEd/Private/Commandlets/GatherTextFromSourceCommandlet.cpp
+++ b/Engine/Source/Editor/UnrealEd/Private/Commandlets/GatherTextFromSourceCommandlet.cpp
@@ -1320,12 +1320,20 @@ bool UGatherTextFromSourceCommandlet::FMacroDescriptor::ParseArgsFromMacro(const
 
 	bool Success = false;
 
-	FString RemainingText = Text.RightChop(GetToken().Len()).TrimStart();
-	int32 OpenBracketIdx = RemainingText.Find(TEXT("("));
-	if (0 > OpenBracketIdx)
+	// Step over the token name and any whitespace after it
+	FString RemainingText = Text.RightChop(GetToken().Len());
+	RemainingText.TrimStartInline();
+
+	const int32 OpenBracketIdx = RemainingText.Find(TEXT("("));
+	if (OpenBracketIdx == INDEX_NONE)
 	{
+		// No opening bracket; warn about this, but don't consider it an error as we're likely parsing something we shouldn't be
 		UE_LOG(LogGatherTextFromSourceCommandlet, Warning, TEXT("Missing bracket '(' in %s macro in %s(%d):%s"), *GetToken(), *Context.Filename, Context.LineNumber, *FLocTextHelper::SanitizeLogOutput(Context.LineText));
-		//Dont assume this is an error. It's more likely trying to parse something it shouldn't be.
+		return false;
+	}
+	else if (OpenBracketIdx > 0)
+	{
+		// We stepped over the whitespace when building RemainingText, so if the bracket isn't the first character in the text then it means we only partially matched a longer token and shouldn't parse it
 		return false;
 	}
 	else
diff --git a/Engine/Source/Runtime/Experimental/InteractiveToolsFramework/Private/InteractiveToolActionSet.cpp b/Engine/Source/Runtime/Experimental/InteractiveToolsFramework/Private/InteractiveToolActionSet.cpp
--- a/Engine/Source/Runtime/Experimental/InteractiveToolsFramework/Private/InteractiveToolActionSet.cpp
+++ b/Engine/Source/Runtime/Experimental/InteractiveToolsFramework/Private/InteractiveToolActionSet.cpp
@@ -6,7 +6,7 @@
 
 
 void FInteractiveToolActionSet::RegisterAction(UInteractiveTool* Tool, int32 ActionID,
-	const FString& ActionName, const FString& ShortUIName, const FString& DescriptionText,
+	const FString& ActionName, const FText& ShortUIName, const FText& DescriptionText,
 	EModifierKey::Type Modifiers, const FKey& ShortcutKey,
 	TFunction<void()> ActionFunction )
 {
diff --git a/Engine/Source/Runtime/Experimental/InteractiveToolsFramework/Public/InteractiveToolActionSet.h b/Engine/Source/Runtime/Experimental/InteractiveToolsFramework/Public/InteractiveToolActionSet.h
--- a/Engine/Source/Runtime/Experimental/InteractiveToolsFramework/Public/InteractiveToolActionSet.h
+++ b/Engine/Source/Runtime/Experimental/InteractiveToolsFramework/Public/InteractiveToolActionSet.h
@@ -40,9 +40,9 @@ struct INTERACTIVETOOLSFRAMEWORK_API FInteractiveToolAction
 	/** Internal name for this Action */
 	FString ActionName;
 	/** Short name for this Action */
-	FString ShortName;
+	FText ShortName;
 	/** Descriptive name for this Action */
-	FString Description;
+	FText Description;
 
 	/** Suggested modifier keys for this Action */
 	EModifierKey::Type DefaultModifiers;
@@ -60,7 +60,7 @@ struct INTERACTIVETOOLSFRAMEWORK_API FInteractiveToolAction
 	}
 
 	FInteractiveToolAction(const UClass* ClassTypeIn, int32 ActionIDIn,
-		const FString& ActionNameIn, const FString& ShortNameIn, const FString& DescriptionIn,
+		const FString& ActionNameIn, const FText& ShortNameIn, const FText& DescriptionIn,
 		EModifierKey::Type DefaultModifiersIn, const FKey& DefaultKeyIn )
 	{
 		ClassType = ClassTypeIn;
@@ -89,7 +89,7 @@ public:
 	 * UInteractiveTool::RegisterActions() implementations
 	 */
 	void RegisterAction(UInteractiveTool* Tool, int32 ActionID,
-		const FString& ActionName, const FString& ShortUIName, const FString& DescriptionText,
+		const FString& ActionName, const FText& ShortUIName, const FText& DescriptionText,
 		EModifierKey::Type Modifiers, const FKey& ShortcutKey,
 		TFunction<void()> ActionFunction );
 
diff --git a/Engine/Source/Runtime/Slate/Private/Framework/Commands/Commands.cpp b/Engine/Source/Runtime/Slate/Private/Framework/Commands/Commands.cpp
--- a/Engine/Source/Runtime/Slate/Private/Framework/Commands/Commands.cpp
+++ b/Engine/Source/Runtime/Slate/Private/Framework/Commands/Commands.cpp
@@ -5,18 +5,18 @@
 
 #define LOC_DEFINE_REGION
 
-void UI_COMMAND_Function( FBindingContext* This, TSharedPtr< FUICommandInfo >& OutCommand, const TCHAR* OutSubNamespace, const TCHAR* OutCommandName, const TCHAR* OutCommandNameUnderscoreTooltip, const ANSICHAR* DotOutCommandName, const TCHAR* FriendlyName, const TCHAR* InDescription, const EUserInterfaceActionType CommandType, const FInputChord& InDefaultChord, const FInputChord& InAlternateDefaultChord)
+void MakeUICommand_InternalUseOnly( FBindingContext* This, TSharedPtr< FUICommandInfo >& OutCommand, const TCHAR* InSubNamespace, const TCHAR* InCommandName, const TCHAR* InCommandNameUnderscoreTooltip, const ANSICHAR* DotCommandName, const TCHAR* FriendlyName, const TCHAR* InDescription, const EUserInterfaceActionType CommandType, const FInputChord& InDefaultChord, const FInputChord& InAlternateDefaultChord)
 {
 	static const FString UICommandsStr(TEXT("UICommands"));
-	const FString Namespace = OutSubNamespace && FCString::Strlen(OutSubNamespace) > 0 ?  UICommandsStr + TEXT(".") + OutSubNamespace : UICommandsStr;
+	const FString Namespace = InSubNamespace && FCString::Strlen(InSubNamespace) > 0 ? UICommandsStr + TEXT(".") + InSubNamespace : UICommandsStr;
 
 	FUICommandInfo::MakeCommandInfo(
 		This->AsShared(),
 		OutCommand,
-		OutCommandName,
-		FInternationalization::ForUseOnlyByLocMacroAndGraphNodeTextLiterals_CreateText( FriendlyName, *Namespace, OutCommandName ),
-		FInternationalization::ForUseOnlyByLocMacroAndGraphNodeTextLiterals_CreateText( InDescription, *Namespace, OutCommandNameUnderscoreTooltip ),
-		FSlateIcon( This->GetStyleSetName(), ISlateStyle::Join( This->GetContextName(), DotOutCommandName ) ),
+		InCommandName,
+		FInternationalization::ForUseOnlyByLocMacroAndGraphNodeTextLiterals_CreateText( FriendlyName, *Namespace, InCommandName ),
+		FInternationalization::ForUseOnlyByLocMacroAndGraphNodeTextLiterals_CreateText( InDescription, *Namespace, InCommandNameUnderscoreTooltip ),
+		FSlateIcon( This->GetStyleSetName(), ISlateStyle::Join( This->GetContextName(), DotCommandName ) ),
 		CommandType,
 		InDefaultChord,
 		InAlternateDefaultChord
diff --git a/Engine/Source/Runtime/Slate/Public/Framework/Commands/Commands.h b/Engine/Source/Runtime/Slate/Public/Framework/Commands/Commands.h
--- a/Engine/Source/Runtime/Slate/Public/Framework/Commands/Commands.h
+++ b/Engine/Source/Runtime/Slate/Public/Framework/Commands/Commands.h
@@ -12,14 +12,15 @@
 
 #define LOC_DEFINE_REGION
 
-SLATE_API void UI_COMMAND_Function(FBindingContext* This, TSharedPtr< FUICommandInfo >& OutCommand, const TCHAR* OutSubNamespace, const TCHAR* OutCommandName, const TCHAR* OutCommandNameUnderscoreTooltip, const ANSICHAR* DotOutCommandName, const TCHAR* FriendlyName, const TCHAR* InDescription, const EUserInterfaceActionType CommandType, const FInputChord& InDefaultChord, const FInputChord& InAlternateDefaultChord = FInputChord());
+/** Internal function used by the UI_COMMAND macros to build the command. Do not call this directly as only the macros are gathered for localization; instead use FUICommandInfo::MakeCommandInfo for dynamic content */
+SLATE_API void MakeUICommand_InternalUseOnly(FBindingContext* This, TSharedPtr< FUICommandInfo >& OutCommand, const TCHAR* InSubNamespace, const TCHAR* InCommandName, const TCHAR* InCommandNameUnderscoreTooltip, const ANSICHAR* DotCommandName, const TCHAR* FriendlyName, const TCHAR* InDescription, const EUserInterfaceActionType CommandType, const FInputChord& InDefaultChord, const FInputChord& InAlternateDefaultChord = FInputChord());
 
 // This macro requires LOCTEXT_NAMESPACE to be defined. If you don't want the command to be placed under a sub namespace, provide "" as the namespace.
 #define UI_COMMAND_EXT( BindingContext, OutUICommandInfo, CommandIdName, FriendlyName, InDescription, CommandType, InDefaultChord, ... ) \
-	UI_COMMAND_Function( BindingContext, OutUICommandInfo, TEXT(LOCTEXT_NAMESPACE), TEXT(CommandIdName), TEXT(CommandIdName) TEXT("_ToolTip"), "." CommandIdName, TEXT(FriendlyName), TEXT(InDescription), CommandType, InDefaultChord, ## __VA_ARGS__ );
+	MakeUICommand_InternalUseOnly( BindingContext, OutUICommandInfo, TEXT(LOCTEXT_NAMESPACE), TEXT(CommandIdName), TEXT(CommandIdName) TEXT("_ToolTip"), "." CommandIdName, TEXT(FriendlyName), TEXT(InDescription), CommandType, InDefaultChord, ## __VA_ARGS__ );
 
 #define UI_COMMAND( CommandId, FriendlyName, InDescription, CommandType, InDefaultChord, ... ) \
-	UI_COMMAND_Function( this, CommandId, TEXT(LOCTEXT_NAMESPACE), TEXT(#CommandId), TEXT(#CommandId) TEXT("_ToolTip"), "." #CommandId, TEXT(FriendlyName), TEXT(InDescription), CommandType, InDefaultChord, ## __VA_ARGS__ );
+	MakeUICommand_InternalUseOnly( this, CommandId, TEXT(LOCTEXT_NAMESPACE), TEXT(#CommandId), TEXT(#CommandId) TEXT("_ToolTip"), "." #CommandId, TEXT(FriendlyName), TEXT(InDescription), CommandType, InDefaultChord, ## __VA_ARGS__ );
 
 #undef LOC_DEFINE_REGION
 
