Adding support for removing input preprocessors while they are being processed.
* Scope locking removal of preprocessors and adding a generic iteration function that will remove preprocessors after iteration is complete.

diff --git a/Engine/Source/Runtime/Slate/Private/Framework/Application/SlateApplication.cpp b/Engine/Source/Runtime/Slate/Private/Framework/Application/SlateApplication.cpp
--- a/Engine/Source/Runtime/Slate/Private/Framework/Application/SlateApplication.cpp
+++ b/Engine/Source/Runtime/Slate/Private/Framework/Application/SlateApplication.cpp
@@ -7678,122 +7678,48 @@ void FSlateApplication::InputPreProcessorsHelper::Tick(const float DeltaTime, FS
 
 bool FSlateApplication::InputPreProcessorsHelper::HandleKeyDownEvent(FSlateApplication& SlateApp, const FKeyEvent& InKeyEvent)
 {
-	for (TSharedPtr<IInputProcessor> InputPreProcessor : InputPreProcessorList)
-	{
-		if (InputPreProcessor->HandleKeyDownEvent(SlateApp, InKeyEvent))
-		{
-			return true;
-		}
-	}
-
-	return false;
+	return PreProcessInput([&SlateApp, &InKeyEvent](TSharedPtr<IInputProcessor> Processor) { return Processor->HandleKeyDownEvent(SlateApp, InKeyEvent); });
 }
 
 bool FSlateApplication::InputPreProcessorsHelper::HandleKeyUpEvent(FSlateApplication& SlateApp, const FKeyEvent& InKeyEvent)
 {
-	for (TSharedPtr<IInputProcessor> InputPreProcessor : InputPreProcessorList)
-	{
-		if (InputPreProcessor.IsValid())
-		{
-			if (InputPreProcessor->HandleKeyUpEvent(SlateApp, InKeyEvent))
-			{
-				return true;
-			}
-		}
-	}
-
-	return false;
+	return PreProcessInput([&SlateApp, &InKeyEvent](TSharedPtr<IInputProcessor> Processor) { return Processor->HandleKeyUpEvent(SlateApp, InKeyEvent); });
 }
 
 bool FSlateApplication::InputPreProcessorsHelper::HandleAnalogInputEvent(FSlateApplication& SlateApp, const FAnalogInputEvent& InAnalogInputEvent)
 {
-	for (TSharedPtr<IInputProcessor> InputPreProcessor : InputPreProcessorList)
-	{
-		if (InputPreProcessor->HandleAnalogInputEvent(SlateApp, InAnalogInputEvent))
-		{
-			return true;
-		}
-	}
-
-	return false;
+	return PreProcessInput([&SlateApp, &InAnalogInputEvent](TSharedPtr<IInputProcessor> Processor) { return Processor->HandleAnalogInputEvent(SlateApp, InAnalogInputEvent); });
 }
 
 bool FSlateApplication::InputPreProcessorsHelper::HandleMouseMoveEvent(FSlateApplication& SlateApp, const FPointerEvent& MouseEvent)
 {
-	for (TSharedPtr<IInputProcessor> InputPreProcessor : InputPreProcessorList)
-	{
-		if (InputPreProcessor->HandleMouseMoveEvent(SlateApp, MouseEvent))
-		{
-			return true;
-		}
-	}
-
-	return false;
+	return PreProcessInput([&SlateApp, &MouseEvent](TSharedPtr<IInputProcessor> Processor) { return Processor->HandleMouseMoveEvent(SlateApp, MouseEvent); });
 }
 
 bool FSlateApplication::InputPreProcessorsHelper::HandleMouseButtonDownEvent(FSlateApplication& SlateApp, const FPointerEvent& MouseEvent)
 {
-	for (TSharedPtr<IInputProcessor> InputPreProcessor : InputPreProcessorList)
-	{
-		if (InputPreProcessor->HandleMouseButtonDownEvent(SlateApp, MouseEvent))
-		{
-			return true;
-		}
-	}
-
-	return false;
+	return PreProcessInput([&SlateApp, &MouseEvent](TSharedPtr<IInputProcessor> Processor) { return Processor->HandleMouseButtonDownEvent(SlateApp, MouseEvent); });
 }
 
 bool FSlateApplication::InputPreProcessorsHelper::HandleMouseButtonUpEvent(FSlateApplication& SlateApp, const FPointerEvent& MouseEvent)
 {
-	for (TSharedPtr<IInputProcessor> InputPreProcessor : InputPreProcessorList)
-	{
-		if (InputPreProcessor->HandleMouseButtonUpEvent(SlateApp, MouseEvent))
-		{
-			return true;
-		}
-	}
+	return PreProcessInput([&SlateApp, &MouseEvent](TSharedPtr<IInputProcessor> Processor) { return Processor->HandleMouseButtonUpEvent(SlateApp, MouseEvent); });
 
-	return false;
 }
 
 bool FSlateApplication::InputPreProcessorsHelper::HandleMouseButtonDoubleClickEvent(FSlateApplication& SlateApp, const FPointerEvent& MouseEvent)
 {
-	for (TSharedPtr<IInputProcessor> InputPreProcessor : InputPreProcessorList)
-	{
-		if (InputPreProcessor->HandleMouseButtonDoubleClickEvent(SlateApp, MouseEvent))
-		{
-			return true;
-		}
-	}
-
-	return false;
+	return PreProcessInput([&SlateApp, &MouseEvent](TSharedPtr<IInputProcessor> Processor) { return Processor->HandleMouseButtonDoubleClickEvent(SlateApp, MouseEvent); });
 }
 
 bool FSlateApplication::InputPreProcessorsHelper::HandleMouseWheelOrGestureEvent(FSlateApplication& SlateApp, const FPointerEvent& WheelEvent, const FPointerEvent* GestureEvent)
 {
-	for (TSharedPtr<IInputProcessor> InputPreProcessor : InputPreProcessorList)
-	{
-		if (InputPreProcessor->HandleMouseWheelOrGestureEvent(SlateApp, WheelEvent, GestureEvent))
-		{
-			return true;
-		}
-	}
-
-	return false;
+	return PreProcessInput([&SlateApp, &WheelEvent, &GestureEvent](TSharedPtr<IInputProcessor> Processor) { return Processor->HandleMouseWheelOrGestureEvent(SlateApp, WheelEvent, GestureEvent); });
 }
 
 bool FSlateApplication::InputPreProcessorsHelper::HandleMotionDetectedEvent(FSlateApplication& SlateApp, const FMotionEvent& MotionEvent)
 {
-	for (TSharedPtr<IInputProcessor> InputPreProcessor : InputPreProcessorList)
-	{
-		if (InputPreProcessor->HandleMotionDetectedEvent(SlateApp, MotionEvent))
-		{
-			return true;
-		}
-	}
-
-	return false;
+	return PreProcessInput([&SlateApp, &MotionEvent](TSharedPtr<IInputProcessor> Processor) { return Processor->HandleMotionDetectedEvent(SlateApp, MotionEvent); });
 }
 
 bool FSlateApplication::InputPreProcessorsHelper::Add(TSharedPtr<IInputProcessor> InputProcessor, const int32 Index /*= INDEX_NONE*/)
@@ -7815,11 +7741,45 @@ bool FSlateApplication::InputPreProcessorsHelper::Add(TSharedPtr<IInputProcessor
 
 void FSlateApplication::InputPreProcessorsHelper::Remove(TSharedPtr<IInputProcessor> InputProcessor)
 {
-	InputPreProcessorList.Remove(InputProcessor);
+	if (bIsIteratingPreProcessors)
+	{
+		ProcessorsPendingRemoval.Add(InputProcessor);
+	}
+	else
+	{
+		InputPreProcessorList.Remove(InputProcessor);
+	}
 }
 
 void FSlateApplication::InputPreProcessorsHelper::RemoveAll()
 {
-	InputPreProcessorList.Reset();
+	if (bIsIteratingPreProcessors)
+	{
+		ProcessorsPendingRemoval.Append(InputPreProcessorList);
+	}
+	else
+	{
+		InputPreProcessorList.Reset();
+	}
 }
 
+
+bool FSlateApplication::InputPreProcessorsHelper::PreProcessInput(TFunctionRef<bool(TSharedPtr<IInputProcessor>)> ToRun)
+{
+	TGuardValue<bool> IteratingGuard(bIsIteratingPreProcessors, true);
+
+	bool bShouldExit = false;
+	for (TSharedPtr<IInputProcessor> InputPreProcessor : InputPreProcessorList)
+	{
+		if (ToRun(InputPreProcessor))
+		{
+			bShouldExit = true;
+			break;
+		}
+	}
+
+	InputPreProcessorList.RemoveAll([this](const TSharedPtr<IInputProcessor> Processor){ return ProcessorsPendingRemoval.Contains(Processor); });
+	ProcessorsPendingRemoval.Reset();
+
+	return bShouldExit;
+}
diff --git a/Engine/Source/Runtime/Slate/Public/Framework/Application/SlateApplication.h b/Engine/Source/Runtime/Slate/Public/Framework/Application/SlateApplication.h
--- a/Engine/Source/Runtime/Slate/Public/Framework/Application/SlateApplication.h
+++ b/Engine/Source/Runtime/Slate/Public/Framework/Application/SlateApplication.h
@@ -2187,7 +2187,6 @@ private:
 	class InputPreProcessorsHelper
 	{
 	public:
-
 		// Wrapper functions that call the corresponding function of IInputProcessor for each InputProcessor in the list.
 		void Tick(const float DeltaTime, FSlateApplication& SlateApp, TSharedRef<ICursor> Cursor);
 		bool HandleKeyDownEvent(FSlateApplication& SlateApp, const FKeyEvent& InKeyEvent);
@@ -2220,8 +2219,16 @@ private:
 
 	private:
 
+		bool PreProcessInput(TFunctionRef<bool(TSharedPtr<IInputProcessor>)> ToRun);
+
 		/** The list of input pre-processors. */
 		TArray<TSharedPtr<IInputProcessor>> InputPreProcessorList;
+
+		/** Guard value for if we are currently iterating our preprocessors. */
+		bool bIsIteratingPreProcessors = false;
+
+		/** A list of pre-processors to remove if we are iterating them while removal is requested. */
+		TArray<TSharedPtr<IInputProcessor>> ProcessorsPendingRemoval;
 	};
 
 	/** A list of input pre-processors, gets an opportunity to parse input before anything else. */
