Slate - Provide a static instance accessor to the default rich text writer for escaping unsanitized text.

diff --git a/Engine/Source/Runtime/Slate/Private/Framework/Text/RichTextMarkupProcessing.cpp b/Engine/Source/Runtime/Slate/Private/Framework/Text/RichTextMarkupProcessing.cpp
--- a/Engine/Source/Runtime/Slate/Private/Framework/Text/RichTextMarkupProcessing.cpp
+++ b/Engine/Source/Runtime/Slate/Private/Framework/Text/RichTextMarkupProcessing.cpp
@@ -265,6 +265,12 @@ TSharedRef< FDefaultRichTextMarkupWriter > FDefaultRichTextMarkupWriter::Create(
 	return MakeShareable( new FDefaultRichTextMarkupWriter() );
 }
 
+TSharedRef< FDefaultRichTextMarkupWriter > FDefaultRichTextMarkupWriter::GetStaticInstance()
+{
+	static TSharedRef< FDefaultRichTextMarkupWriter > Writer = MakeShareable(new FDefaultRichTextMarkupWriter());
+	return Writer;
+}
+
 void FDefaultRichTextMarkupWriter::Write(const TArray<FRichTextLine>& InLines, FString& Output)
 {
 	for (int32 LineIndex = 0; LineIndex < InLines.Num(); ++LineIndex)
diff --git a/Engine/Source/Runtime/Slate/Public/Framework/Text/RichTextMarkupProcessing.h b/Engine/Source/Runtime/Slate/Public/Framework/Text/RichTextMarkupProcessing.h
--- a/Engine/Source/Runtime/Slate/Public/Framework/Text/RichTextMarkupProcessing.h
+++ b/Engine/Source/Runtime/Slate/Public/Framework/Text/RichTextMarkupProcessing.h
@@ -34,6 +34,7 @@ class SLATE_API FDefaultRichTextMarkupWriter : public IRichTextMarkupWriter
 {
 public:
 	static TSharedRef< FDefaultRichTextMarkupWriter > Create();
+	static TSharedRef< FDefaultRichTextMarkupWriter > GetStaticInstance();
 
 public:
 	virtual void Write(const TArray<FRichTextLine>& InLines, FString& Output) override;
