Slate Global Invalidation System

diff --git a/Engine/Source/Runtime/Slate/Private/Widgets/SInvalidationPanel.cpp b/Engine/Source/Runtime/Slate/Private/Widgets/SInvalidationPanel.cpp
--- a/Engine/Source/Runtime/Slate/Private/Widgets/SInvalidationPanel.cpp
+++ b/Engine/Source/Runtime/Slate/Private/Widgets/SInvalidationPanel.cpp
@@ -8,124 +8,54 @@
 #include "Layout/WidgetPath.h"
 #include "HAL/IConsoleManager.h"
 #include "Framework/Application/SlateApplication.h"
-#include "Layout/WidgetCaching.h"
 #include "Types/ReflectionMetadata.h"
 #include "Rendering/SlateObjectReferenceCollector.h"
 
-//DECLARE_CYCLE_STAT(TEXT("Invalidation Time"), STAT_InvalidationTime, STATGROUP_Slate);
-DECLARE_DWORD_COUNTER_STAT(TEXT("Num Cached Elements"), STAT_SlateNumCachedElements, STATGROUP_Slate);
-DECLARE_DWORD_COUNTER_STAT(TEXT("Num Invalidated Elements"), STAT_SlateNumInvalidatedElements, STATGROUP_Slate);
-DECLARE_DWORD_COUNTER_STAT(TEXT("Num Volatile Widgets"), STAT_SlateNumVolatileWidgets, STATGROUP_Slate);
-
-DECLARE_CYCLE_STAT(TEXT("SInvalidationPanel::Tick"), STAT_SlateInvalidationTick, STATGROUP_Slate);
 DECLARE_CYCLE_STAT(TEXT("SInvalidationPanel::Paint"), STAT_SlateInvalidationPaint, STATGROUP_Slate);
 
 DEFINE_LOG_CATEGORY_STATIC(LogSlateInvalidationPanel, Log, All);
 
-#if !(UE_BUILD_SHIPPING || UE_BUILD_TEST)
+#if WITH_SLATE_DEBUGGING
 
 /** True if we should allow widgets to be cached in the UI at all. */
-TAutoConsoleVariable<int32> InvalidationDebugging(
-	TEXT("Slate.InvalidationDebugging"),
-	false,
-	TEXT("Whether to show invalidation debugging visualization"));
-
-bool SInvalidationPanel::IsInvalidationDebuggingEnabled()
-{
-	return InvalidationDebugging.GetValueOnGameThread() == 1;
-}
-
-void SInvalidationPanel::EnableInvalidationDebugging(bool bEnable)
-{
-	InvalidationDebugging.AsVariable()->Set(bEnable);
-}
-
-/** True if we should allow widgets to be cached in the UI at all. */
-TAutoConsoleVariable<int32> EnableWidgetCaching(
-	TEXT("Slate.EnableWidgetCaching"),
-	true,
+static int32 InvalidationPanelsEnabled = 1;
+FAutoConsoleVariableRef CVarEnableInvalidationPanels(
+	TEXT("Slate.EnableInvalidationPanels"),
+	InvalidationPanelsEnabled,
 	TEXT("Whether to attempt to cache any widgets through invalidation panels."));
 
-bool SInvalidationPanel::GetEnableWidgetCaching()
-{
-	return EnableWidgetCaching.GetValueOnGameThread() == 1;
-}
-
-void SInvalidationPanel::SetEnableWidgetCaching(bool bEnable)
-{
-	EnableWidgetCaching.AsVariable()->Set(bEnable);
-}
-
-TAutoConsoleVariable<int32> AlwaysInvalidate(
+static int32 AlwaysInvalidate = 0;
+FAutoConsoleVariableRef CVarAlwaysInvalidate(
 	TEXT("Slate.AlwaysInvalidate"),
-	false,
+	AlwaysInvalidate,
 	TEXT("Forces invalidation panels to cache, but to always invalidate."));
 
-#endif
-
-static int32 CacheRenderData = 1;
-static FAutoConsoleVariableRef CVarCacheRenderData(
-	TEXT("Slate.CacheRenderData"),
-	CacheRenderData,
-	TEXT("Invalidation panels will cache render data, otherwise cache only widget draw elements."));
-
-static int32 InvalidationLayerPadding = 10;
-static FAutoConsoleVariableRef CVarInvalidationLayerPadding(
-	TEXT("Slate.InvalidationLayerPadding"),
-	InvalidationLayerPadding,
-	TEXT("The invalidation panel adds this much padding to the layer id to avoid re-invalidating if it fluctuates in a volatile child widget."));
-
-#if SLATE_VERBOSE_NAMED_EVENTS
-
-static int32 ExcessiveInvalidationFrameStreak = 60;
-static FAutoConsoleVariableRef CVarExcessiveInvalidationFrameStreak(
-	TEXT("Slate.ExcessiveInvalidationFrameStreak"),
-	ExcessiveInvalidationFrameStreak,
-	TEXT("."));
-
-static int32 LogExcessiveInvalidation = 1;
-static FAutoConsoleVariableRef CVarLogExcessiveInvalidation(
-	TEXT("Slate.LogExcessiveInvalidation"),
-	LogExcessiveInvalidation,
-	TEXT("."));
+#endif // WITH_SLATE_DEBUGGING
 
-#endif
 
-static bool ShouldCacheRenderData()
-{
-	return WITH_ENGINE && CacheRenderData != 0;
-}
 
 SInvalidationPanel::SInvalidationPanel()
 	: EmptyChildSlot(this)
+	, HittestGrid()
+	, bCanCache(true)
+	, bPaintedSinceLastPrepass(true)
+	, bWasCachable(false)
 {
+	bHasCustomPrepass = true;
+	SetInvalidationRootWidget(*this);
+	SetInvalidationRootHittestGrid(HittestGrid);
+	SetCanTick(false);
+
+	FSlateApplicationBase::Get().OnGlobalInvalidationToggled().AddRaw(this, &SInvalidationPanel::OnGlobalInvalidationToggled);
 }
 
 void SInvalidationPanel::Construct( const FArguments& InArgs )
 {
-	FSlateApplicationBase::Get().OnGlobalInvalidate().AddSP( this, &SInvalidationPanel::OnGlobalInvalidate );
-
 	ChildSlot
 	[
 		InArgs._Content.Widget
 	];
 
-	bNeedsCaching = true;
-	bNeedsCachePrepass = true;
-	bPerformingCachePrepass = false;
-	bIsInvalidating = false;
-	bCanCache = true;
-	RootCacheNode = nullptr;
-	LastUsedCachedNodeIndex = 0;
-	LastHitTestIndex = 0;
-	MaximumLayerIdCachedAt = 0;
-	NumberOfFramesInARowWeInvalidated = 0;
-	LastClippingIntersectionSize = FVector2D::ZeroVector;
-
-	bCacheRelativeTransforms = InArgs._CacheRelativeTransforms;
-
-	bCacheRenderData = ShouldCacheRenderData();
-
 #if SLATE_VERBOSE_NAMED_EVENTS
 	DebugName = InArgs._DebugName;
 	DebugTickName = InArgs._DebugName + TEXT("_Tick");
@@ -135,226 +65,108 @@ void SInvalidationPanel::Construct( const FArguments& InArgs )
 
 SInvalidationPanel::~SInvalidationPanel()
 {
-	for ( int32 i = 0; i < NodePool.Num(); i++ )
-	{
-		delete NodePool[i];
-	}
+	InvalidateRoot();
 
-	if ( FSlateApplication::IsInitialized() )
+	if (FSlateApplicationBase::IsInitialized())
 	{
-		FSlateApplicationBase::Get().OnGlobalInvalidate().RemoveAll(this);
-		FSlateApplication::Get().ReleaseResourcesForLayoutCache(this);
+		FSlateApplicationBase::Get().OnGlobalInvalidationToggled().RemoveAll(this);
 	}
 }
 
-#if !(UE_BUILD_SHIPPING || UE_BUILD_TEST)
-bool SInvalidationPanel::GetCanCache() const
+#if WITH_SLATE_DEBUGGING
+bool SInvalidationPanel::AreInvalidationPanelsEnabled()
 {
-	return bCanCache && EnableWidgetCaching.GetValueOnGameThread() == 1;
+	return InvalidationPanelsEnabled != 0;
 }
-#endif
 
-#if !(UE_BUILD_SHIPPING || UE_BUILD_TEST)
-bool SInvalidationPanel::IsCachingNeeded() const
+void SInvalidationPanel::EnableInvalidationPanels(bool bEnable)
 {
-	return bNeedsCaching || AlwaysInvalidate.GetValueOnGameThread() == 1;
+	InvalidationPanelsEnabled = bEnable;
 }
 #endif
 
-bool SInvalidationPanel::IsCachingNeeded(FSlateWindowElementList& OutDrawElements, const FGeometry& AllottedGeometry, const FSlateRect& MyCullingRect, int32 LayerId) const
+bool SInvalidationPanel::GetCanCache() const
 {
-	// We only need to re-cache if the incoming layer is higher than the maximum layer Id we cached at,
-	// we do this so that widgets that appear and live behind your invalidated UI don't constantly invalidate
-	// everything above it.
-	if (LayerId > MaximumLayerIdCachedAt)
-	{
-		return true;
-	}
-	
-	//TODO NDarnell if clipping state index continues to be an invalidation problem for non-cached render data frames,
-	// we can pad it out some like we do for cached layer id, and pad + invalidate on higher number.
-	
-	if (bCacheRelativeTransforms || CacheRenderData == 0)
-	{
-	if (LastClippingIndex != OutDrawElements.GetClippingIndex())
-	{
-		return true;
-	}
-
-	int32 ClippingStateCount = OutDrawElements.GetClippingManager().GetClippingStates().Num();
-	if (LastClippingStateOffset != ClippingStateCount)
-	{
-		return true;
-	}
-	}
-	
-	if (bCacheRelativeTransforms)
-	{
-		bool bOverlapping;
-		FVector2D IntersectionSize = AllottedGeometry.GetLayoutBoundingRect().IntersectionWith(MyCullingRect, bOverlapping).GetSize();
-		if (!LastClippingIntersectionSize.Equals(IntersectionSize, 1.0f))
-		{
-			return true;
-		}
-
-		const float LayoutScale = AllottedGeometry.Scale;
-		const float LastLayoutScale = LastAllottedGeometry.Scale;
+	// Note: checking for FastPathProxyHandle being valid prevents nested invalidation panels from being a thing.  They are not needed anymore since invalidation panels do not redraw everything inside it just because one thing invalidates
+	// In global invalidation this code makes no sense so we don't bother running it because everything is in an "invalidation panel" at the window level
+#if WITH_SLATE_DEBUGGING
+	// Disable invalidation panels if global invalidation is turned on
+	return bCanCache && !GSlateEnableGlobalInvalidation && !GetProxyHandle().IsValid() && InvalidationPanelsEnabled;
+#else
+	return bCanCache && !GSlateEnableGlobalInvalidation && !GetProxyHandle().IsValid();
+#endif
+}
 
-		// If the layout scale has changed, we need to update things.
-		if (LayoutScale != LastLayoutScale)
-		{
-			return true;
-		}
+void SInvalidationPanel::OnGlobalInvalidationToggled(bool bGlobalInvalidationEnabled)
+{
+	InvalidateRoot();
 
-		const FMatrix2x2& RenderTransformMatrix = AllottedGeometry.GetAccumulatedRenderTransform().GetMatrix();
-		const FMatrix2x2& LastRenderTransformMatrix = LastAllottedGeometry.GetAccumulatedRenderTransform().GetMatrix();
+	ClearAllFastPathData();
+}
 
-		// If the container we're in has changed in the rotation matrix has changed
-		if (RenderTransformMatrix != LastRenderTransformMatrix)
-		{
-			return true;
-		}
+bool SInvalidationPanel::UpdateCachePrequisites(FSlateWindowElementList& OutDrawElements, const FGeometry& AllottedGeometry, const FSlateRect& MyCullingRect, int32 LayerId) const
+{
+	bool bNeedsRecache = false;
+#if WITH_SLATE_DEBUGGING
+	if (AlwaysInvalidate == 1)
+	{
+		bNeedsRecache = true;
 	}
-	else
+#endif
+
+	// We only need to re-cache if the incoming layer is higher than the maximum layer Id we cached at,
+	// we do this so that widgets that appear and live behind your invalidated UI don't constantly invalidate
+	// everything above it.
+	if (LayerId > LastIncomingLayerId)
 	{
-		// If the container we're in has changed in any way we need to invalidate for sure.
-		if ( AllottedGeometry.GetAccumulatedLayoutTransform() != LastAllottedGeometry.GetAccumulatedLayoutTransform() ||
-			AllottedGeometry.GetAccumulatedRenderTransform() != LastAllottedGeometry.GetAccumulatedRenderTransform() )
-		{
-			return true;
-		}
+		LastIncomingLayerId = LayerId;
+		bNeedsRecache = true;
 	}
 
-	if ( AllottedGeometry.GetLocalSize() != LastAllottedGeometry.GetLocalSize() )
+	if ( AllottedGeometry.GetLocalSize() != LastAllottedGeometry.GetLocalSize() || AllottedGeometry.GetAccumulatedRenderTransform() != LastAllottedGeometry.GetAccumulatedRenderTransform() )
 	{
-		return true;
+		LastAllottedGeometry = AllottedGeometry;
+		bNeedsRecache = true;
 	}
 
 	// If our clip rect changes size, we've definitely got to invalidate.
 	const FVector2D ClipRectSize = MyCullingRect.GetSize().RoundToVector();
 	if ( ClipRectSize != LastClipRectSize )
 	{
-		return true;
+		LastClipRectSize = ClipRectSize;
+		bNeedsRecache = true;
 	}
 
 	TOptional<FSlateClippingState> ClippingState = OutDrawElements.GetClippingState();
 	if (LastClippingState != ClippingState)
 	{
-		return true;
+		LastClippingState = ClippingState;
+		bNeedsRecache = true;
 	}
 	
-	return false;
-}
-
-void SInvalidationPanel::AddReferencedObjects(FReferenceCollector& Collector)
-{
-	Collector.AddReferencedObjects(CachedResources);
-
-#if SLATE_VERBOSE_NAMED_EVENTS
-	UE_LOG(LogSlateInvalidationPanel, Verbose, TEXT("SInvalidationPanel(%s): %d References"), *DebugName, CachedResources.Num());
-#endif
-}
-
-FString SInvalidationPanel::GetReferencerName() const
-{
-	return TEXT("SInvalidationPanel");
+	return bNeedsRecache;
 }
 
 void SInvalidationPanel::SetCanCache(bool InCanCache)
 {
 	bCanCache = InCanCache;
-	SetCanTick(GetCanCache());
-	InvalidateCache();
-}
-
-void SInvalidationPanel::Tick( const FGeometry& AllottedGeometry, const double InCurrentTime, const float InDeltaTime )
-{
-#if SLATE_VERBOSE_NAMED_EVENTS
-	SCOPED_NAMED_EVENT_FSTRING(DebugTickName, FColor::Blue);
-#endif
-
-	SCOPE_CYCLE_COUNTER(STAT_SlateInvalidationTick);
-
-	if ( GetCanCache() )
-	{
-		//SCOPE_CYCLE_COUNTER(STAT_InvalidationTime);
-
-		bool bShouldCacheRenderData = ShouldCacheRenderData();
-		if (bCacheRenderData != bShouldCacheRenderData)
-		{
-			bCacheRenderData = bShouldCacheRenderData;
-			InvalidateCache();
-		}
-		
-		if (bNeedsCachePrepass)
-		{
-			TGuardValue<bool> PerformingCachePrepass(bPerformingCachePrepass, true);
-			CachePrepass(SharedThis(this));
-			bNeedsCachePrepass = false;
-		}
-	}
-	else
-	{
-		// Ticking no longer required
-		SetCanTick(false);
-	}
 }
 
 FChildren* SInvalidationPanel::GetChildren()
 {
-	if ( GetCanCache() == false || IsCachingNeeded() || bPerformingCachePrepass )
-	{
-		return SCompoundWidget::GetChildren();
-	}
-	else
+	if (GetCanCache() && !NeedsPrepass())
 	{
 		return &EmptyChildSlot;
 	}
-}
-
-void SInvalidationPanel::InvalidateWidget(SWidget* InvalidateWidget)
-{
-	bNeedsCaching = true;
-
-#if SLATE_VERBOSE_NAMED_EVENTS
-	if (LogExcessiveInvalidation && NumberOfFramesInARowWeInvalidated > 0 && (NumberOfFramesInARowWeInvalidated % ExcessiveInvalidationFrameStreak) == 0)
-	{
-		const FString InvalidationCause = FReflectionMetaData::GetWidgetDebugInfo(InvalidateWidget);
-		const FString InvalidationRoot = FReflectionMetaData::GetWidgetDebugInfo(this);
-
-		UE_LOG(LogSlate, Warning, TEXT("SInvalidationPanel(%s): Excessive Invalidation Detected: %s"), *InvalidationRoot, *InvalidationCause);
-	}
-#endif
-
-#if !(UE_BUILD_SHIPPING || UE_BUILD_TEST)
-	if ( InvalidateWidget != nullptr && IsInvalidationDebuggingEnabled() )
-	{
-		InvalidatorWidgets.Add(InvalidateWidget->AsShared(), 1);
-	}
-#endif
-}
-
-FCachedWidgetNode* SInvalidationPanel::CreateCacheNode() const
-{
-	// If the node pool is empty, allocate a few
-	if ( LastUsedCachedNodeIndex >= NodePool.Num() )
+	else
 	{
-		for ( int32 i = 0; i < 10; i++ )
-		{
-			NodePool.Add(new FCachedWidgetNode());
-		}
+		return SCompoundWidget::GetChildren();
 	}
-
-	// Return one of the preallocated nodes and increment the next node index.
-	FCachedWidgetNode* NewNode = NodePool[LastUsedCachedNodeIndex];
-	++LastUsedCachedNodeIndex;
-
-	return NewNode;
 }
 
-void SInvalidationPanel::OnGlobalInvalidate()
+FChildren* SInvalidationPanel::GetAllChildren()
 {
-	InvalidateCache();
+	return SCompoundWidget::GetChildren();
 }
 
 int32 SInvalidationPanel::OnPaint( const FPaintArgs& Args, const FGeometry& AllottedGeometry, const FSlateRect& MyCullingRect, FSlateWindowElementList& OutDrawElements, int32 LayerId, const FWidgetStyle& InWidgetStyle, bool bParentEnabled ) const
@@ -364,290 +176,48 @@ int32 SInvalidationPanel::OnPaint( const FPaintArgs& Args, const FGeometry& Allo
 #endif
 	SCOPE_CYCLE_COUNTER(STAT_SlateInvalidationPaint);
 
-	if ( GetCanCache() )
-	{
-		const bool bWasCachingNeeded = IsCachingNeeded() || IsCachingNeeded(OutDrawElements, AllottedGeometry, MyCullingRect, LayerId);
+	bPaintedSinceLastPrepass = true;
+	SInvalidationPanel* MutableThis = const_cast<SInvalidationPanel*>(this);
 
-		if ( bWasCachingNeeded )
-		{
-#if SLATE_VERBOSE_NAMED_EVENTS
-				SCOPED_NAMED_EVENT_F(TEXT("%s_Invalidation"), FColor::Red, *DebugPaintName);
-#endif
-				SInvalidationPanel* MutableThis = const_cast<SInvalidationPanel*>(this);
-				TSharedRef<SInvalidationPanel> SharedMutableThis = SharedThis(MutableThis);
-
-				// Always set the caching flag to false first, during the paint / tick pass we may change something
-				// to volatile and need to re-cache.
-				bNeedsCaching = false;
-				NumberOfFramesInARowWeInvalidated++;
-
-				bNeedsCachePrepass = true;
-
-				// Mark that we're in the process of invalidating.
-				bIsInvalidating = true;
-
-			// Record a new maximum layer Id, throw in some padding that avoids needed to recache
-			// if it varies a bit within the padding it avoids us needing to recache to them.
-			MaximumLayerIdCachedAt = LayerId + InvalidationLayerPadding;
-
-				SWindow* Window = OutDrawElements.GetPaintWindow();
-				CachedWindowElements = FSlateApplication::Get().GetCachableElementList(StaticCastSharedRef<SWindow>(Window->AsShared()), this);
-
-
-				// Reset the render data handle in case it was in use, and we're not overriding it this frame.
-				CachedRenderData.Reset();
-
-				// Reset the cached node pool index so that we effectively reset the pool.
-				LastUsedCachedNodeIndex = 0;
-
-				RootCacheNode = CreateCacheNode();
-				RootCacheNode->Initialize(Args, SharedMutableThis, AllottedGeometry);
-
-				// TODO We may be double pre-passing here, if the invalidation happened at the end of last frame,
-				// we'll have already done one pre-pass before getting here.
-				ChildSlot.GetWidget()->SlatePrepass(AllottedGeometry.Scale);
-
-			// If we're caching the render data, we need to make sure the states are mirrored so that the cached
-			// buffer has the current stack.
-			if (bCacheRenderData)
-				{
-				OutDrawElements.GetClippingManager().CopyClippingStateTo(CachedWindowElements->GetClippingManager());
-				}
-
-			//TODO: When SWidget::Paint is called don't draw self if volatile, and we're doing a cache pass.
-				CachedMaxChildLayer = SCompoundWidget::OnPaint(
-					Args.EnableCaching(SharedMutableThis, RootCacheNode, true, false),
-					AllottedGeometry,
-					MyCullingRect,
-					*CachedWindowElements.Get(),
-					MaximumLayerIdCachedAt,
-					InWidgetStyle,
-					bParentEnabled);
-
-				{
-					CachedResources.Reset();
-					FSlateObjectReferenceCollector Collector(CachedResources);
-					CachedWindowElements->AddReferencedObjects(Collector);
-				}
-
-				if (bCacheRelativeTransforms)
-				{
-					CachedAbsolutePosition = AllottedGeometry.GetAccumulatedRenderTransform().GetTranslation();
-				}
-
-				LastClippingStateOffset = OutDrawElements.GetClippingManager().GetClippingStates().Num();
-				LastClippingIndex = OutDrawElements.GetClippingIndex();
-				LastClippingState = OutDrawElements.GetClippingState();
-				
-			if (bCacheRenderData == 0)
-			{
-				const int32 ClippingStateOffset = OutDrawElements.GetClippingManager().MergePartialClippingStates(CachedWindowElements->GetClippingManager().GetClippingStates());
-
-				TArray<FSlateDrawElement>& CachedElements = CachedWindowElements->GetDrawElements();
-				for (int32 Index = 0; Index < CachedElements.Num(); Index++)
-				{
-					FSlateDrawElement& CachedElement = CachedElements[Index];
-					if (CachedElement.GetClippingIndex() == INDEX_NONE)
-				{
-						CachedElement.SetClippingIndex(LastClippingIndex);
-				}
-				else
-				{
-						CachedElement.SetClippingIndex(ClippingStateOffset + CachedElement.GetClippingIndex());
-					}
-				}
-			}
-
-			if (bCacheRenderData)
-			{
-				CachedRenderData = CachedWindowElements->CacheRenderData(this);
-				}
-
-				LastHitTestIndex = Args.GetLastHitTestIndex();
-
-				LastAllottedGeometry = AllottedGeometry;
-				LastClipRectSize = MyCullingRect.GetSize().RoundToVector();
-
-				if (bCacheRelativeTransforms)
-				{
-					LastClippingIntersectionSize = AllottedGeometry.GetLayoutBoundingRect().IntersectionWith(MyCullingRect).GetSize();
-				}
-
-				bIsInvalidating = false;
-			}
-		else
-		{
-			NumberOfFramesInARowWeInvalidated = 0;
-
-			// The full clipping state is stored on the render batches when we cache the render data,
-			// so there's no need to merge clipping states on frames where the data has been cached successfully.
-			if (CacheRenderData == 0)
-			{
-#if SLATE_VERBOSE_NAMED_EVENTS
-			SCOPED_NAMED_EVENT_TEXT("SInvalidationPanel::MergeClippingStates", FColor::Magenta);
-#endif
-				OutDrawElements.GetClippingManager().MergePartialClippingStates(CachedWindowElements->GetClippingManager().GetClippingStates());
-			}
-		}
-
-		FVector2D AbsoluteDeltaPosition = FVector2D::ZeroVector;
-		if ( bCacheRelativeTransforms )
-		{
-			AbsoluteDeltaPosition = AllottedGeometry.GetAccumulatedRenderTransform().GetTranslation() - CachedAbsolutePosition;
-		}
+	const bool bCanCacheThisFrame = GetCanCache();
+	if (bCanCacheThisFrame != bWasCachable)
+	{
+		MutableThis->InvalidateRoot();
 
-		// Record Hit Test Geometry
-		{
-			// The hit test grid is actually populated during the initial cache phase, so don't bother
-			// recording the hit test geometry on the same frame that we regenerate the cache.
-			if (bWasCachingNeeded == false)
-			{
-				INC_DWORD_STAT_BY(STAT_SlateNumCachedElements, CachedWindowElements->GetElementCount());
+		bWasCachable = bCanCacheThisFrame;
+	}
 
-#if SLATE_VERBOSE_NAMED_EVENTS
-				SCOPED_NAMED_EVENT_TEXT("SInvalidationPanel::RecordHittestGeometry", FColor::Magenta);
-#endif
+	if(bCanCacheThisFrame)
+	{
+		// Copy hit test grid settings from the root
+		const bool bHittestCleared = HittestGrid.SetHittestArea(Args.RootGrid.GetGridOrigin(), Args.RootGrid.GetGridSize(), Args.RootGrid.GetGridWindowOrigin());
 
-				RootCacheNode->RecordHittestGeometry(Args.GetGrid(), Args.GetLastHitTestIndex(), MaximumLayerIdCachedAt, AbsoluteDeltaPosition);
-			}
-			else
-			{
-				INC_DWORD_STAT_BY(STAT_SlateNumInvalidatedElements, CachedWindowElements->GetElementCount());
-			}
-		}
+		FPaintArgs NewArgs = Args.WithNewHitTestGrid(HittestGrid);
 
-		int32 OutMaxChildLayer = CachedMaxChildLayer;
+		check(!GSlateEnableGlobalInvalidation);
 
-		if (bCacheRenderData)
+		const bool bRequiresRecache = UpdateCachePrequisites(OutDrawElements, AllottedGeometry, MyCullingRect, LayerId);
+		if (bHittestCleared || bRequiresRecache)
 		{
-			FSlateDrawElement::MakeCachedBuffer(OutDrawElements, MaximumLayerIdCachedAt, CachedRenderData, AbsoluteDeltaPosition);
-			// Merge the resources into the draw element list so it can keep UObjects alive.
-			OutDrawElements.MergeResources(CachedResources);
-		}
-		else
-		{
-			OutDrawElements.MergeElementList(CachedWindowElements.Get(), AbsoluteDeltaPosition);
+			// @todo: Overly aggressive?
+			MutableThis->InvalidateRoot();
 		}
 
-		// Paint the volatile elements
-		if ( CachedWindowElements.IsValid() )
-		{
-#if SLATE_VERBOSE_NAMED_EVENTS
-			SCOPED_NAMED_EVENT_TEXT("Paint Volatile Widgets", FColor::Cyan);
-#endif
-			const TArray<TSharedPtr<FSlateWindowElementList::FVolatilePaint>>& VolatileElements = CachedWindowElements->GetVolatileElements();
-			INC_DWORD_STAT_BY(STAT_SlateNumVolatileWidgets, VolatileElements.Num());
-			
-			int32 VolatileLayerId = 0;
-			if (bCacheRenderData)
-			{
-				VolatileLayerId = CachedWindowElements->PaintVolatile(OutDrawElements, Args.GetCurrentTime(), Args.GetDeltaTime(), AbsoluteDeltaPosition);
-			}
-			else
-			{
-				VolatileLayerId = CachedWindowElements->PaintVolatileRootLayer(OutDrawElements, Args.GetCurrentTime(), Args.GetDeltaTime(), AbsoluteDeltaPosition);
-			}
-			
-			OutMaxChildLayer = FMath::Max(OutMaxChildLayer, VolatileLayerId);
-		}
-
-#if !(UE_BUILD_SHIPPING || UE_BUILD_TEST)
+		// The root widget is our child.  We are not the root because we could be in a parent invalidation panel.  If we are nested in another invalidation panel, our OnPaint was called by that panel
+		FSlateInvalidationContext Context(OutDrawElements, InWidgetStyle);
+		Context.bParentEnabled = bParentEnabled;
+		Context.bAllowFastPathUpdate = true;
+		Context.LayoutScaleMultiplier = GetPrepassLayoutScaleMultiplier();
+		Context.PaintArgs = &NewArgs;
+		Context.IncomingLayerId = LayerId;
+		Context.CullingRect = MyCullingRect;
 
-		if ( IsInvalidationDebuggingEnabled() )
-		{
-			// Draw a green or red border depending on if we were invalidated this frame.
-			{
-				check(Args.IsCaching() == false);
-				//const bool bShowOutlineAsCached = Args.IsCaching() || bWasCachingNeeded == false;
-				const FLinearColor DebugTint = bWasCachingNeeded ? FLinearColor::Red : ( bCacheRelativeTransforms ? FLinearColor::Blue : FLinearColor::Green );
-
-				FGeometry ScaledOutline = AllottedGeometry.MakeChild(FVector2D(0, 0), AllottedGeometry.GetLocalSize() * AllottedGeometry.Scale, Inverse(AllottedGeometry.Scale));
-
-				FSlateDrawElement::MakeBox(
-					OutDrawElements,
-					++OutMaxChildLayer,
-					ScaledOutline.ToPaintGeometry(),
-					FCoreStyle::Get().GetBrush(TEXT("Debug.Border")),
-					ESlateDrawEffect::None,
-					DebugTint
-				);
-			}
-
-			static const FName InvalidationPanelName(TEXT("SInvalidationPanel"));
-
-			const FSlateBrush* VolatileBrush = FCoreStyle::Get().GetBrush(TEXT("FocusRectangle"));
-
-			// Draw a yellow outline around any volatile elements.
-			const TArray<TSharedPtr<FSlateWindowElementList::FVolatilePaint>>& VolatileElements = CachedWindowElements->GetVolatileElements();
-			for ( const TSharedPtr<FSlateWindowElementList::FVolatilePaint>& VolatileElement : VolatileElements )
-			{
-				// Ignore drawing the volatility rect for child invalidation panels, that's not really important, since
-				// they're always volatile and it will make it hard to see when they're invalidated.
-				if ( const SWidget* Widget = VolatileElement->GetWidget() )
-				{
-					if ( Widget->GetType() == InvalidationPanelName )
-					{
-						continue;
-					}
-				}
-
-				FGeometry VolatileGeometry = VolatileElement->GetGeometry();
-				if (!AbsoluteDeltaPosition.IsZero())
-				{
-					// Account for relative translation delta
-					VolatileGeometry.AppendTransform(FSlateLayoutTransform(AbsoluteDeltaPosition));
-				}
-				
-				FSlateDrawElement::MakeBox(
-					OutDrawElements,
-					++OutMaxChildLayer,
-					VolatileGeometry.ToPaintGeometry(),
-					VolatileBrush,
-					ESlateDrawEffect::None,
-					FLinearColor::Yellow
-				);
-			}
-
-			// Draw a red flash for any widget that invalidated us recently, we slowly 
-			// fade out the flashes over time, unless the widget invalidates us again.
-			for ( TMap<TWeakPtr<SWidget>, double>::TIterator It(InvalidatorWidgets); It; ++It )
-			{
-				TSharedPtr<SWidget> SafeInvalidator = It.Key().Pin();
-				if ( SafeInvalidator.IsValid() )
-				{
-					FWidgetPath WidgetPath;
-					if ( FSlateApplication::Get().GeneratePathToWidgetUnchecked(SafeInvalidator.ToSharedRef(), WidgetPath, EVisibility::All) )
-					{
-						FArrangedWidget ArrangedWidget = WidgetPath.FindArrangedWidget(SafeInvalidator.ToSharedRef()).Get(FArrangedWidget::GetNullWidget());
-						ArrangedWidget.Geometry.AppendTransform( FSlateLayoutTransform(Inverse(Args.GetWindowToDesktopTransform())) );
-
-						FSlateDrawElement::MakeBox(
-							OutDrawElements,
-							++OutMaxChildLayer,
-							ArrangedWidget.Geometry.ToPaintGeometry(),
-							FCoreStyle::Get().GetBrush(TEXT("WhiteBrush")),
-							ESlateDrawEffect::None,
-							FLinearColor::Red.CopyWithNewOpacity(0.75f * It.Value())
-						);
-					}
-
-					It.Value() -= FApp::GetDeltaTime();
-
-					if ( It.Value() <= 0 )
-					{
-						It.RemoveCurrent();
-					}
-				}
-				else
-				{
-					It.RemoveCurrent();
-				}
-			}
-		}
+		const FSlateInvalidationResult Result = MutableThis->PaintInvalidationRoot(Context);
 
-#endif
+		// add our widgets to the root hit test grid
+		Args.RootGrid.AppendGrid(HittestGrid);
 
-		return OutMaxChildLayer;
+		return Result.MaxLayerIdPainted;
 	}
 	else
 	{
@@ -660,16 +230,36 @@ int32 SInvalidationPanel::OnPaint( const FPaintArgs& Args, const FGeometry& Allo
 
 void SInvalidationPanel::SetContent(const TSharedRef< SWidget >& InContent)
 {
-	InvalidateCache();
-
 	ChildSlot
 	[
 		InContent
 	];
+
+	InvalidateRoot();
 }
 
-bool SInvalidationPanel::ComputeVolatility() const
+bool SInvalidationPanel::CustomPrepass(float LayoutScaleMultiplier)
+{
+	bPaintedSinceLastPrepass = false;
+
+	if (GetCanCache())
+	{
+		ProcessInvalidation();
+
+		return NeedsPrepass();
+	}
+	else
+	{
+		return true;
+	}
+}
+
+bool SInvalidationPanel::Advanced_IsInvalidationRoot() const
 {
-	// If we are cacheable we must be volatile 
 	return GetCanCache();
 }
+
+int32 SInvalidationPanel::PaintSlowPath(const FSlateInvalidationContext& Context)
+{
+	return SCompoundWidget::OnPaint(*Context.PaintArgs, GetPaintSpaceGeometry(), Context.CullingRect, *Context.WindowElementList, Context.IncomingLayerId, Context.WidgetStyle, Context.bParentEnabled);
+}
diff --git a/Engine/Source/Runtime/Slate/Private/Widgets/SViewport.cpp b/Engine/Source/Runtime/Slate/Private/Widgets/SViewport.cpp
--- a/Engine/Source/Runtime/Slate/Private/Widgets/SViewport.cpp
+++ b/Engine/Source/Runtime/Slate/Private/Widgets/SViewport.cpp
@@ -55,13 +55,17 @@ void SViewport::Construct( const FArguments& InArgs )
 
 void SViewport::SetActive(bool bActive)
 {
-	if (bActive && !ActiveTimerHandle.IsValid())
+	// In game enviroments the viewport is always active
+	if(GIsEditor || IS_PROGRAM)
 	{
-		ActiveTimerHandle = RegisterActiveTimer(0.f, FWidgetActiveTimerDelegate::CreateSP(this, &SViewport::EnsureTick));
-	}
-	else if (!bActive && ActiveTimerHandle.IsValid())
-	{
-		UnRegisterActiveTimer(ActiveTimerHandle.Pin().ToSharedRef());
+		if (bActive && !ActiveTimerHandle.IsValid())
+		{
+			ActiveTimerHandle = RegisterActiveTimer(0.f, FWidgetActiveTimerDelegate::CreateSP(this, &SViewport::EnsureTick));
+		}
+		else if (!bActive && ActiveTimerHandle.IsValid())
+		{
+			UnRegisterActiveTimer(ActiveTimerHandle.Pin().ToSharedRef());
+		}
 	}
 }
 
@@ -79,9 +83,6 @@ int32 SViewport::OnPaint( const FPaintArgs& Args, const FGeometry& AllottedGeome
 	bool bShowDisabledEffect = ShowDisabledEffect.Get();
 	ESlateDrawEffect DrawEffects = bShowDisabledEffect && !bEnabled ? ESlateDrawEffect::DisabledEffect : ESlateDrawEffect::None;
 
-	int32 LastHitTestIndex = Args.GetLastHitTestIndex();
-
-
 	// Viewport texture alpha channels are often in an indeterminate state, even after the resolve,
 	// so we'll tell the shader to not use the alpha channel when blending
 	if( bIgnoreTextureAlpha )
@@ -172,7 +173,7 @@ int32 SViewport::OnPaint( const FPaintArgs& Args, const FGeometry& AllottedGeome
 	// If there are any custom hit testable widgets in the 3D world we need to register their custom hit test path here.
 	if ( CustomHitTestPath.IsValid() )
 	{
-		Args.InsertCustomHitTestPath(CustomHitTestPath.ToSharedRef(), LastHitTestIndex);
+		Args.InsertCustomHitTestPath(this, CustomHitTestPath.ToSharedRef());
 	}
 
 	return Layer;
diff --git a/Engine/Source/Runtime/Slate/Private/Widgets/Text/SRichTextBlock.cpp b/Engine/Source/Runtime/Slate/Private/Widgets/Text/SRichTextBlock.cpp
--- a/Engine/Source/Runtime/Slate/Private/Widgets/Text/SRichTextBlock.cpp
+++ b/Engine/Source/Runtime/Slate/Private/Widgets/Text/SRichTextBlock.cpp
@@ -61,14 +61,6 @@ int32 SRichTextBlock::OnPaint( const FPaintArgs& Args, const FGeometry& Allotted
 {
 	const FVector2D LastDesiredSize = TextLayoutCache->GetDesiredSize();
 
-	// If we're performing layout caching, it's possible nobody ever called GetDesiredSize(),
-	// which for textblocks is required to be called, since CDS is where it actually generates
-	// a lot for the text layout.
-	if (GSlateLayoutCaching)
-	{
-		GetDesiredSize();
-	}
-
 	// OnPaint will also update the text layout cache if required
 	LayerId = TextLayoutCache->OnPaint(Args, AllottedGeometry, MyCullingRect, OutDrawElements, LayerId, InWidgetStyle, ShouldBeEnabled(bParentEnabled));
 
diff --git a/Engine/Source/Runtime/Slate/Private/Widgets/Text/STextBlock.cpp b/Engine/Source/Runtime/Slate/Private/Widgets/Text/STextBlock.cpp
--- a/Engine/Source/Runtime/Slate/Private/Widgets/Text/STextBlock.cpp
+++ b/Engine/Source/Runtime/Slate/Private/Widgets/Text/STextBlock.cpp
@@ -240,14 +240,6 @@ int32 STextBlock::OnPaint( const FPaintArgs& Args, const FGeometry& AllottedGeom
 	{
 		const FVector2D LastDesiredSize = TextLayoutCache->GetDesiredSize();
 
-		// If we're performing layout caching, it's possible nobody ever called GetDesiredSize(),
-		// which for textblocks is required to be called, since CDS is where it actually generates
-		// a lot for the text layout.
-		if (GSlateLayoutCaching)
-		{
-			GetDesiredSize();
-		}
-
 		// OnPaint will also update the text layout cache if required
 		LayerId = TextLayoutCache->OnPaint(Args, AllottedGeometry, MyCullingRect, OutDrawElements, LayerId, InWidgetStyle, ShouldBeEnabled(bParentEnabled));
 
diff --git a/Engine/Source/Runtime/Slate/Private/Widgets/Views/SListPanel.cpp b/Engine/Source/Runtime/Slate/Private/Widgets/Views/SListPanel.cpp
--- a/Engine/Source/Runtime/Slate/Private/Widgets/Views/SListPanel.cpp
+++ b/Engine/Source/Runtime/Slate/Private/Widgets/Views/SListPanel.cpp
@@ -183,6 +183,11 @@ FVector2D SListPanel::ComputeDesiredSize( float ) const
 	}
 }
 
+FChildren* SListPanel::GetAllChildren()
+{
+	return &Children;
+}
+
 /**
  * @return A slot-agnostic representation of this panel's children
  */
diff --git a/Engine/Source/Runtime/Slate/Private/Widgets/Views/SListPanel.h b/Engine/Source/Runtime/Slate/Private/Widgets/Views/SListPanel.h
--- a/Engine/Source/Runtime/Slate/Private/Widgets/Views/SListPanel.h
+++ b/Engine/Source/Runtime/Slate/Private/Widgets/Views/SListPanel.h
@@ -67,6 +67,7 @@ public:
 	// SWidget interface
 	virtual void OnArrangeChildren( const FGeometry& AllottedGeometry, FArrangedChildren& ArrangedChildren ) const override;
 	virtual FVector2D ComputeDesiredSize(float) const override;
+	virtual FChildren* GetAllChildren() override;
 	virtual FChildren* GetChildren() override;
 	virtual void Tick( const FGeometry& AllottedGeometry, const double InCurrentTime, const float InDeltaTime ) override;
 	// End of SWidget interface
diff --git a/Engine/Source/Runtime/Slate/Private/Widgets/Views/STableViewBase.cpp b/Engine/Source/Runtime/Slate/Private/Widgets/Views/STableViewBase.cpp
--- a/Engine/Source/Runtime/Slate/Private/Widgets/Views/STableViewBase.cpp
+++ b/Engine/Source/Runtime/Slate/Private/Widgets/Views/STableViewBase.cpp
@@ -196,7 +196,7 @@ EActiveTimerReturnType STableViewBase::UpdateInertialScroll(double InCurrentTime
 				if (CanUseInertialScroll(ScrollVelocity))
 				{
 					bKeepTicking = true;
-					ScrollBy(GetCachedGeometry(), ScrollVelocity * InDeltaTime, AllowOverscroll);
+					ScrollBy(GetTickSpaceGeometry(), ScrollVelocity * InDeltaTime, AllowOverscroll);
 				}
 				else
 				{
@@ -208,7 +208,7 @@ EActiveTimerReturnType STableViewBase::UpdateInertialScroll(double InCurrentTime
 			{
 				// If we are currently in overscroll, the list will need refreshing.
 				// Do this before UpdateOverscroll, as that could cause GetOverscroll() to be 0
-				if (Overscroll.GetOverscroll(GetCachedGeometry()) != 0.0f)
+				if (Overscroll.GetOverscroll(GetTickSpaceGeometry()) != 0.0f)
 				{
 					bKeepTicking = true;
 					RequestLayoutRefresh();
@@ -277,7 +277,7 @@ void STableViewBase::Tick( const FGeometry& AllottedGeometry, const double InCur
 
 			if (AllowOverscroll == EAllowOverscroll::Yes)
 			{
-				const float OverscrollAmount = Overscroll.GetOverscroll(GetCachedGeometry());
+				const float OverscrollAmount = Overscroll.GetOverscroll(GetTickSpaceGeometry());
 				ItemsPanel->SetOverscrollAmount( OverscrollAmount );
 			}
 
@@ -315,6 +315,8 @@ void STableViewBase::Tick( const FGeometry& AllottedGeometry, const double InCur
 			bItemsNeedRefresh = false;
 			ItemsPanel->SetRefreshPending(false);
 
+			Invalidate(EInvalidateWidget::ChildOrder);
+
 			if (ScrollIntoViewResult == EScrollIntoViewResult::Deferred)
 			{
 				// We call this rather than just leave bItemsNeedRefresh as true to ensure that EnsureTickToRefresh is registered
@@ -910,7 +912,7 @@ TSharedRef<class SWidget> STableViewBase::GetScrollWidget()
 
 bool STableViewBase::CanUseInertialScroll( float ScrollAmount ) const
 {
-	const auto CurrentOverscroll = Overscroll.GetOverscroll(GetCachedGeometry());
+	const auto CurrentOverscroll = Overscroll.GetOverscroll(GetTickSpaceGeometry());
 
 	// We allow sampling for the inertial scroll if we are not in the overscroll region,
 	// Or if we are scrolling outwards of the overscroll region
diff --git a/Engine/Source/Runtime/Slate/Public/Framework/Application/SlateApplication.h b/Engine/Source/Runtime/Slate/Public/Framework/Application/SlateApplication.h
--- a/Engine/Source/Runtime/Slate/Public/Framework/Application/SlateApplication.h
+++ b/Engine/Source/Runtime/Slate/Public/Framework/Application/SlateApplication.h
@@ -729,7 +729,7 @@ public:
 	 * @param InWidget		The widget to find the window for
 	 * @return The window where the widget resides, or null if the widget wasn't found.  Remember, a widget might not be found simply because its parent decided not to report the widget in ArrangeChildren.
 	 */
-	TSharedPtr<SWindow> FindWidgetWindow( TSharedRef<const SWidget> InWidget ) const;
+	virtual TSharedPtr<SWindow> FindWidgetWindow( TSharedRef<const SWidget> InWidget ) const override;
 
 	/**
 	 * Finds the window that the provided widget resides in
@@ -948,18 +948,6 @@ public:
 	 */
 	bool TakeScreenshot(const TSharedRef<SWidget>& Widget, const FIntRect& InnerWidgetArea, TArray<FColor>& OutColorData, FIntVector& OutSize);
 
-	/**
-	 * 
-	 */
-	TSharedPtr< FSlateWindowElementList > GetCachableElementList(const TSharedPtr<SWindow>& CurrentWindow, const ILayoutCache* LayoutCache);
-
-	/**
-	 * Once a layout cache is destroyed it needs to free any resources it was using in a safe way to prevent
-	 * any in-flight rendering from being interrupted by referencing resources that go away.  So when a layout
-	 * cache is destroyed it should call this function so any associated resources can be collected when it's safe.
-	 */
-	void ReleaseResourcesForLayoutCache(const ILayoutCache* LayoutCache);
-
 	/**
 	 * @return a handle for the existing or newly created virtual slate user.  This is handy when you need to create
 	 * virtual hardware users for slate components in the virtual world that may need to be interacted with with virtual hardware.
@@ -1409,7 +1397,7 @@ public:
 	virtual bool IsExternalUIOpened() override;
 	virtual FWidgetPath LocateWindowUnderMouse( FVector2D ScreenspaceMouseCoordinate, const TArray<TSharedRef<SWindow>>& Windows, bool bIgnoreEnabledStatus = false ) override;
 	virtual bool IsWindowHousingInteractiveTooltip(const TSharedRef<const SWindow>& WindowToTest) const override;
-	virtual TSharedRef<SWidget> MakeImage( const TAttribute<const FSlateBrush*>& Image, const TAttribute<FSlateColor>& Color, const TAttribute<EVisibility>& Visibility ) const override;
+	virtual TSharedRef<SImage> MakeImage( const TAttribute<const FSlateBrush*>& Image, const TAttribute<FSlateColor>& Color, const TAttribute<EVisibility>& Visibility ) const override;
 	virtual TSharedRef<SWidget> MakeWindowTitleBar( const TSharedRef<SWindow>& Window, const TSharedPtr<SWidget>& CenterContent, EHorizontalAlignment CenterContentAlignment, TSharedPtr<IWindowTitleBar>& OutTitleBar ) const override;
 	virtual TSharedRef<IToolTip> MakeToolTip( const TAttribute<FText>& ToolTipText ) override;
 	virtual TSharedRef<IToolTip> MakeToolTip( const FText& ToolTipText ) override;
@@ -2129,24 +2117,6 @@ private:
 	// but there could be another monitor on any of the sides.
 	FSlateRect VirtualDesktopRect;
 
-	//
-	// Invalidation Support
-	//
-
-	class FCacheElementPools
-	{
-	public:
-		TSharedPtr< FSlateWindowElementList > GetNextCachableElementList(const TSharedPtr<SWindow>& CurrentWindow );
-		bool IsInUse() const;
-
-	private:
-		TArray< TSharedPtr< FSlateWindowElementList > > ActiveCachedElementListPool;
-		TArray< TSharedPtr< FSlateWindowElementList > > InactiveCachedElementListPool;
-	};
-
-	TMap< const ILayoutCache*, TSharedPtr<FCacheElementPools> > CachedElementLists;
-	TArray< TSharedPtr<FCacheElementPools> > ReleasedCachedElementLists;
-
 	/** This factory function creates a navigation config for each slate user. */
 	TSharedRef<FNavigationConfig> NavigationConfig;
 
diff --git a/Engine/Source/Runtime/Slate/Public/Widgets/Layout/SWidgetSwitcher.h b/Engine/Source/Runtime/Slate/Public/Widgets/Layout/SWidgetSwitcher.h
--- a/Engine/Source/Runtime/Slate/Public/Widgets/Layout/SWidgetSwitcher.h
+++ b/Engine/Source/Runtime/Slate/Public/Widgets/Layout/SWidgetSwitcher.h
@@ -153,7 +153,7 @@ protected:
 	virtual void OnArrangeChildren( const FGeometry& AllottedGeometry, FArrangedChildren& ArrangedChildren ) const override;
 	virtual FVector2D ComputeDesiredSize( float ) const override;
 	virtual FChildren* GetChildren( ) override;
-
+	virtual bool ComputeVolatility() const override { return WidgetIndex.IsBound(); }
 	const FSlot* GetActiveSlot() const;
 
 private:
diff --git a/Engine/Source/Runtime/Slate/Public/Widgets/Notifications/SProgressBar.h b/Engine/Source/Runtime/Slate/Public/Widgets/Notifications/SProgressBar.h
--- a/Engine/Source/Runtime/Slate/Public/Widgets/Notifications/SProgressBar.h
+++ b/Engine/Source/Runtime/Slate/Public/Widgets/Notifications/SProgressBar.h
@@ -53,7 +53,9 @@ public:
 		, _FillImage(nullptr)
 		, _MarqueeImage(nullptr)
 		, _RefreshRate(2.0f)
-		{}
+		{
+			_Visibility = EVisibility::SelfHitTestInvisible;
+		}
 
 		/** Style used for the progress bar */
 		SLATE_STYLE_ARGUMENT( FProgressBarStyle, Style )
diff --git a/Engine/Source/Runtime/Slate/Public/Widgets/SInvalidationPanel.h b/Engine/Source/Runtime/Slate/Public/Widgets/SInvalidationPanel.h
--- a/Engine/Source/Runtime/Slate/Public/Widgets/SInvalidationPanel.h
+++ b/Engine/Source/Runtime/Slate/Public/Widgets/SInvalidationPanel.h
@@ -10,129 +10,72 @@
 #include "Layout/Children.h"
 #include "Widgets/DeclarativeSyntaxSupport.h"
 #include "Widgets/SCompoundWidget.h"
+#include "FastUpdate/SlateInvalidationRoot.h"
+#include "Input/HittestGrid.h"
 
-class FCachedWidgetNode;
 class FPaintArgs;
 class FSlateRenderDataHandle;
 class FSlateWindowElementList;
 class SWindow;
 
-class SLATE_API SInvalidationPanel : public SCompoundWidget, public FGCObject, public ILayoutCache
+class SLATE_API SInvalidationPanel : public SCompoundWidget, public FSlateInvalidationRoot
 {
 public:
 	SLATE_BEGIN_ARGS( SInvalidationPanel )
 	{
 		_Visibility = EVisibility::SelfHitTestInvisible;
-		_CacheRelativeTransforms = false;
 	}
 		SLATE_DEFAULT_SLOT(FArguments, Content)
-		SLATE_ARGUMENT(bool, CacheRelativeTransforms)
 #if !UE_BUILD_SHIPPING
 		SLATE_ARGUMENT(FString, DebugName)
 #endif
 	SLATE_END_ARGS()
 
 	SInvalidationPanel();
-	void Construct( const FArguments& InArgs );
 	~SInvalidationPanel();
 
-#if !(UE_BUILD_SHIPPING || UE_BUILD_TEST)
-	bool GetCanCache() const;
-#else
-	FORCEINLINE bool GetCanCache() const { return bCanCache; }
+#if WITH_SLATE_DEBUGGING
+	static bool AreInvalidationPanelsEnabled();
+	static void EnableInvalidationPanels(bool bEnable);
 #endif
+	void Construct( const FArguments& InArgs );
 
-	virtual void AddReferencedObjects(FReferenceCollector& Collector) override;
-	virtual FString GetReferencerName() const override;
+	bool GetCanCache() const;
 
 	void SetCanCache(bool InCanCache);
 
-	FORCEINLINE void InvalidateCache() { bNeedsCaching = true; }
-
-	// ILayoutCache overrides
-	virtual void InvalidateWidget(SWidget* InvalidateWidget) override;
-	virtual FCachedWidgetNode* CreateCacheNode() const override;
-	// End ILayoutCache
-
-public:
-
 	// SWidget overrides
 	virtual int32 OnPaint( const FPaintArgs& Args, const FGeometry& AllottedGeometry, const FSlateRect& MyCullingRect, FSlateWindowElementList& OutDrawElements, int32 LayerId, const FWidgetStyle& InWidgetStyle, bool bParentEnabled ) const override;
 	virtual FChildren* GetChildren() override;
-	virtual void Tick( const FGeometry& AllottedGeometry, const double InCurrentTime, const float InDeltaTime );
-	virtual bool ComputeVolatility() const override;
+	virtual FChildren* GetAllChildren() override;
 	// End SWidget
 
 	void SetContent(const TSharedRef< SWidget >& InContent);
 
-#if !(UE_BUILD_SHIPPING || UE_BUILD_TEST)
-	static bool IsInvalidationDebuggingEnabled();
-	static void EnableInvalidationDebugging(bool bEnable);
-
-	static bool GetEnableWidgetCaching();
-	static void SetEnableWidgetCaching(bool bEnable);
-#else
-	static bool IsInvalidationDebuggingEnabled() { return false; }
-	static void EnableInvalidationDebugging(bool bEnable) { }
-
-	static bool GetEnableWidgetCaching() { return true; }
-	static void SetEnableWidgetCaching(bool bEnable) { }
-#endif
-
+protected:
+	virtual bool CustomPrepass(float LayoutScaleMultiplier) override;
+	virtual bool Advanced_IsInvalidationRoot() const override;
+	virtual int32 PaintSlowPath(const FSlateInvalidationContext& Context) override;
 private:
-	TSharedPtr< FSlateWindowElementList > GetNextCachedElementList(const TSharedPtr<SWindow>& CurrentWindow) const;
-	void OnGlobalInvalidate();
+	void OnGlobalInvalidationToggled(bool bGlobalInvalidationEnabled);
+	bool UpdateCachePrequisites(FSlateWindowElementList& OutDrawElements, const FGeometry& AllottedGeometry, const FSlateRect& MyCullingRect, int32 LayerId) const;
+private:
+	FSimpleSlot EmptyChildSlot;
 
-#if !(UE_BUILD_SHIPPING || UE_BUILD_TEST)
-	bool IsCachingNeeded() const;
-#else
-	FORCEINLINE bool IsCachingNeeded() const { return bNeedsCaching; }
-#endif
-	
-	bool IsCachingNeeded(FSlateWindowElementList& OutDrawElements, const FGeometry& AllottedGeometry, const FSlateRect& MyCullingRect, int32 LayerId) const;
+	mutable FHittestGrid HittestGrid;
 
-private:
+	mutable TOptional<FSlateClippingState> LastClippingState;
 	mutable FGeometry LastAllottedGeometry;
+	mutable FVector2D LastClipRectSize;
+	mutable int32 LastIncomingLayerId;
 
-	FSimpleSlot EmptyChildSlot;
-	FVector2D CachedDesiredSize;
+	bool bCanCache;
 
+	mutable bool bPaintedSinceLastPrepass;
 #if SLATE_VERBOSE_NAMED_EVENTS
 	FString DebugName;
 	FString DebugTickName;
 	FString DebugPaintName;
 #endif
-
-#if !(UE_BUILD_SHIPPING || UE_BUILD_TEST)
-	mutable TMap<TWeakPtr<SWidget>, double> InvalidatorWidgets;
-#endif
-
-	mutable FCachedWidgetNode* RootCacheNode;
-	mutable TSharedPtr< FSlateWindowElementList > CachedWindowElements;
-	mutable TSharedPtr<FSlateRenderDataHandle, ESPMode::ThreadSafe> CachedRenderData;
-
-	mutable TArray<UObject*> CachedResources;
-	
-	mutable FVector2D CachedAbsolutePosition;
-
-	mutable TArray< FCachedWidgetNode* > NodePool;
-	mutable int32 LastUsedCachedNodeIndex;
-	mutable int32 LastHitTestIndex;
-	mutable FVector2D LastClipRectSize;
-	mutable FVector2D LastClippingIntersectionSize;
-	mutable int32 LastClippingIndex;
-	mutable int32 LastClippingStateOffset;
-	mutable TOptional<FSlateClippingState> LastClippingState;
-	mutable int32 MaximumLayerIdCachedAt;
-
-	mutable int32 CachedMaxChildLayer;
-	mutable bool bNeedsCaching;
-	mutable bool bNeedsCachePrepass;
-	mutable bool bPerformingCachePrepass;
-	mutable bool bIsInvalidating;
-	mutable int32 NumberOfFramesInARowWeInvalidated;
-	bool bCanCache;
-
-	bool bCacheRelativeTransforms;
-	bool bCacheRenderData;
+	mutable bool bWasCachable;
 };
diff --git a/Engine/Source/Runtime/SlateCore/Private/Animation/CurveSequence.cpp b/Engine/Source/Runtime/SlateCore/Private/Animation/CurveSequence.cpp
--- a/Engine/Source/Runtime/SlateCore/Private/Animation/CurveSequence.cpp
+++ b/Engine/Source/Runtime/SlateCore/Private/Animation/CurveSequence.cpp
@@ -58,9 +58,13 @@ FCurveHandle FCurveSequence::AddCurveRelative( const float InOffset, const float
 	return AddCurve(CurveStartTime, InDurationSecond, InEaseFunction);
 }
 
-void FCurveSequence::Play( const TSharedRef<SWidget>& InOwnerWidget, bool bPlayLooped, const float StartAtTime )
+void FCurveSequence::Play(const TSharedRef<SWidget>& InOwnerWidget, bool bPlayLooped, const float StartAtTime, bool bRequiresActiveTimer)
 {
-	RegisterActiveTimerIfNeeded(InOwnerWidget);
+	if (bRequiresActiveTimer)
+	{
+		RegisterActiveTimerIfNeeded(InOwnerWidget);
+	}
+
 	bIsLooping = bPlayLooped;
 	bIsPaused = false;
 	
@@ -89,9 +93,12 @@ void FCurveSequence::Reverse( )
 	SetStartTime(NewStartTime);
 }
 
-void FCurveSequence::PlayReverse( const TSharedRef<SWidget>& InOwnerWidget, bool bPlayLooped, const float StartAtTime )
+void FCurveSequence::PlayReverse(const TSharedRef<SWidget>& InOwnerWidget, bool bPlayLooped, const float StartAtTime, bool bRequiresActiveTimer)
 {
-	RegisterActiveTimerIfNeeded(InOwnerWidget);
+	if (bRequiresActiveTimer)
+	{
+		RegisterActiveTimerIfNeeded(InOwnerWidget);
+	}
 	bIsLooping = bPlayLooped;
 	bIsPaused = false;
 
diff --git a/Engine/Source/Runtime/SlateCore/Private/Application/SlateApplicationBase.cpp b/Engine/Source/Runtime/SlateCore/Private/Application/SlateApplicationBase.cpp
--- a/Engine/Source/Runtime/SlateCore/Private/Application/SlateApplicationBase.cpp
+++ b/Engine/Source/Runtime/SlateCore/Private/Application/SlateApplicationBase.cpp
@@ -152,7 +152,18 @@ bool FSlateApplicationBase::IsSlateAsleep()
 	return bIsSlateAsleep;
 }
 
+void FSlateApplicationBase::ToggleGlobalInvalidation(bool bIsGlobalInvalidationEnabled)
+{
+	if ((!!GSlateEnableGlobalInvalidation) != bIsGlobalInvalidationEnabled)
+	{
+		GSlateEnableGlobalInvalidation = bIsGlobalInvalidationEnabled;
+		OnGlobalInvalidationToggledEvent.Broadcast(GSlateEnableGlobalInvalidation);
+	}
+}
+
 void FSlateApplicationBase::InvalidateAllWidgets() const
 {
-	OnGlobalInvalidateEvent.Broadcast();
+	SCOPED_NAMED_EVENT(Slate_GlobalInvalidate, FColor::Red);
+	UE_LOG(LogSlate, Log, TEXT("All widgets were invalidated"));
+	OnInvalidateAllWidgetsEvent.Broadcast();
 }
diff --git a/Engine/Source/Runtime/SlateCore/Private/Debugging/SlateDebugging.cpp b/Engine/Source/Runtime/SlateCore/Private/Debugging/SlateDebugging.cpp
--- a/Engine/Source/Runtime/SlateCore/Private/Debugging/SlateDebugging.cpp
+++ b/Engine/Source/Runtime/SlateCore/Private/Debugging/SlateDebugging.cpp
@@ -1,10 +1,112 @@
 // Copyright 1998-2019 Epic Games, Inc. All Rights Reserved.
 
 #include "Debugging/SlateDebugging.h"
+
+#if WITH_SLATE_DEBUGGING
+
 #include "SlateGlobals.h"
+#include "FastUpdate/WidgetProxy.h"
+#include "Animation/CurveSequence.h"
+#include "Widgets/SNullWidget.h"
+#include "Rendering/DrawElements.h"
+#include "Styling/CoreStyle.h"
+#include "Widgets/SWidget.h"
+#include "Application/SlateApplicationBase.h"
 
 #define LOCTEXT_NAMESPACE "SlateDebugger"
 
+// Scalar goes from 0 to 1.  0 being all yellow, 1 being all red
+FLinearColor YellowToRedFromScalar(float Scalar)
+{
+	return FLinearColor(1.0f, 1.0f * (1.0f - Scalar), 0.0f);
+}
+
+struct FInvalidatedWidgetDrawer
+{
+	FWidgetProxyHandle ProxyHandle;
+	FCurveSequence FadeCurve;
+	double StartTime;
+	FLinearColor InvalidationColor;
+
+	FInvalidatedWidgetDrawer(const FWidgetProxyHandle& InProxyHandle)
+		: ProxyHandle(InProxyHandle)
+		, FadeCurve(0, 1.0f, ECurveEaseFunction::Linear)
+
+	{
+	}
+
+	/** Widget was invalidated */
+	void Refresh(const FLinearColor* CustomInvalidationColor = nullptr)
+	{
+		if (ProxyHandle.IsValid())
+		{
+			if (CustomInvalidationColor)
+			{
+				InvalidationColor = *CustomInvalidationColor;
+			}
+			else if (FadeCurve.IsPlaying())
+			{
+				// Color more red based on how recently  this was already invalidated
+				InvalidationColor = YellowToRedFromScalar(1.0f - FadeCurve.GetLerp());
+			}
+			else
+			{
+				InvalidationColor = FLinearColor::Yellow;
+			}
+
+			FadeCurve.Play(SNullWidget::NullWidget, false, 0.0f, false);
+		}
+	}
+
+	bool Draw(const FPaintArgs& PaintArgs, FSlateWindowElementList& ElementList)
+	{
+		static const FSlateBrush* WhiteBrush = FCoreStyle::Get().GetBrush(TEXT("FocusRectangle"));
+
+		static const FSlateFontInfo FontInfo = FCoreStyle::Get().GetFontStyle(TEXT("SmallFont"));
+
+		if (ProxyHandle.IsValid() && !FadeCurve.IsAtEnd() && ProxyHandle.GetProxy().Widget)
+		{
+			SWidget* Widget = ProxyHandle.GetProxy().Widget;
+
+			const FSlateWidgetPersistentState& MyState = Widget->GetPersistentState();
+			if (MyState.InitialClipState.IsSet())
+			{
+				ElementList.GetClippingManager().PushClippingState(MyState.InitialClipState.GetValue());
+			}
+		
+			//ElementList.PushAbsoluteBatchPriortyGroup(MyState.BatchPriorityGroup);
+		/*	FSlateDrawElement::MakeText(
+				ElementList,
+				MyState.OutgoingLayerId + 1,
+				MyState.AllottedGeometry.ToPaintGeometry(),
+				FText::Format(FText::FromString(TEXT("{0} {1}")), FText::FromString(Widget->GetTypeAsString()), MyState.LayerId),
+				FontInfo,
+				ESlateDrawEffect::None,
+				InvalidationColor.CopyWithNewOpacity(FMath::Lerp(1.0f, 0.0f, FadeCurve.GetLerp()))
+			);
+*/
+
+			FSlateDrawElement::MakeBox(
+				ElementList,
+				MyState.OutgoingLayerId+1,
+				MyState.AllottedGeometry.ToPaintGeometry(),
+				WhiteBrush,
+				ESlateDrawEffect::None,
+				InvalidationColor.CopyWithNewOpacity(FMath::Lerp(1.0f, 0.0f, FadeCurve.GetLerp()))
+			);
+
+			//ElementList.PopBatchPriortyGroup();
+
+			if (MyState.InitialClipState.IsSet())
+			{
+				ElementList.PopClip();
+			}
+			return true;
+		}
+		return false;
+	}
+};
+
 FSlateDebuggingInputEventArgs::FSlateDebuggingInputEventArgs(ESlateDebuggingInputEvent InInputEventType, const FReply& InReply, const TSharedPtr<SWidget>& InHandlerWidget, const FString& InAdditionalContent)
 	: InputEventType(InInputEventType)
 	, Reply(InReply)
@@ -63,8 +165,6 @@ FSlateDebuggingMouseCaptureEventArgs::FSlateDebuggingMouseCaptureEventArgs(
 {
 }
 
-#if WITH_SLATE_DEBUGGING
-
 FSlateDebugging::FBeginWindow FSlateDebugging::BeginWindow;
 
 FSlateDebugging::FEndWindow FSlateDebugging::EndWindow;
@@ -89,6 +189,10 @@ FSlateDebugging::FWidgetMouseCaptureEvent FSlateDebugging::MouseCaptureEvent;
 
 FSlateDebugging::FUICommandRun FSlateDebugging::CommandRun;
 
+DECLARE_MULTICAST_DELEGATE_OneParam(FWidgetMouseCaptureEvent, const FSlateDebuggingMouseCaptureEventArgs& /*EventArgs*/);
+
+TArray<struct FInvalidatedWidgetDrawer> FSlateDebugging::InvalidatedWidgetDrawers;
+
 void FSlateDebugging::BroadcastWarning(const FText& WarningText, const TSharedPtr<SWidget>& OptionalContextWidget)
 {
 	Warning.Broadcast(FSlateDebuggingWarningEventArgs(WarningText, OptionalContextWidget));
@@ -158,6 +262,68 @@ void FSlateDebugging::BroadcastMouseCaptureLost(uint32 UserIndex, uint32 Pointer
 	MouseCaptureEvent.Broadcast(FSlateDebuggingMouseCaptureEventArgs(false, UserIndex, PointerIndex, InWidgetLostCapture));
 }
 
-#endif
+void FSlateDebugging::WidgetInvalidated(FSlateInvalidationRoot& InvalidationRoot, const class FWidgetProxy& WidgetProxy, const FLinearColor* CustomInvalidationColor)
+{
+	if(FSlateApplicationBase::IsInitialized())
+	{
+		int32 Index = WidgetProxy.Index;
+	
+		FInvalidatedWidgetDrawer* Drawer = InvalidatedWidgetDrawers.FindByPredicate([Index](FInvalidatedWidgetDrawer& InDrawer) { return InDrawer.ProxyHandle.GetIndex() == Index; });
+		if (!Drawer) 
+		{
+			Drawer = &InvalidatedWidgetDrawers.Emplace_GetRef(FWidgetProxyHandle(InvalidationRoot, WidgetProxy.Index));
+		}
+
+		Drawer->Refresh(CustomInvalidationColor);
+	}
+}
+
+void FSlateDebugging::DrawInvalidationRoot(const SWidget& RootWidget, int32 LayerId, FSlateWindowElementList& OutDrawElements)
+{
+	static const FSlateBrush* WhiteBrush = FCoreStyle::Get().GetBrush(TEXT("FocusRectangle"));
+
+
+	FSlateDrawElement::MakeBox(
+		OutDrawElements,
+		LayerId,
+		RootWidget.GetPaintSpaceGeometry().ToPaintGeometry(),
+		WhiteBrush,
+		ESlateDrawEffect::None,
+		FLinearColor(128,0,128)
+	);
+}
+
+void FSlateDebugging::DrawInvalidatedWidgets(const FSlateInvalidationRoot& Root, const FPaintArgs& PaintArgs, FSlateWindowElementList& OutDrawElements)
+{
+	for (int DrawerIdx = 0; DrawerIdx < InvalidatedWidgetDrawers.Num(); ++DrawerIdx)
+	{		
+		FInvalidatedWidgetDrawer& Drawer = InvalidatedWidgetDrawers[DrawerIdx];
+		
+		if(Drawer.ProxyHandle.GetInvalidationRoot() == &Root)
+		{
+			if (!Drawer.Draw(PaintArgs, OutDrawElements))
+			{
+				InvalidatedWidgetDrawers.RemoveAtSwap(DrawerIdx, 1, false);
+				--DrawerIdx;
+			}
+		}
+	}
+}
+
+
+void FSlateDebugging::ClearInvalidatedWidgets(const FSlateInvalidationRoot& Root)
+{
+	for (int DrawerIdx = 0; DrawerIdx < InvalidatedWidgetDrawers.Num(); ++DrawerIdx)
+	{
+		FInvalidatedWidgetDrawer& Drawer = InvalidatedWidgetDrawers[DrawerIdx];
+
+		if (Drawer.ProxyHandle.GetInvalidationRoot() == &Root)
+		{
+			InvalidatedWidgetDrawers.RemoveAtSwap(DrawerIdx, 1, false);
+		}
+	}
+}
+
+#undef LOCTEXT_NAMESPACE
 
-#undef LOCTEXT_NAMESPACE
\ No newline at end of file
+#endif
\ No newline at end of file
diff --git a/Engine/Source/Runtime/SlateCore/Private/FastUpdate/SlateInvalidationRoot.cpp b/Engine/Source/Runtime/SlateCore/Private/FastUpdate/SlateInvalidationRoot.cpp
new file mode 100644
--- /dev/null
+++ b/Engine/Source/Runtime/SlateCore/Private/FastUpdate/SlateInvalidationRoot.cpp
@@ -0,0 +1,547 @@
+// Copyright 1998-2018 Epic Games, Inc. All Rights Reserved.
+
+#include "FastUpdate/SlateInvalidationRoot.h"
+#include "Async/TaskGraphInterfaces.h"
+#include "Application/SlateApplicationBase.h"
+#include "Widgets/SWidget.h"
+#include "Layout/Children.h"
+#include "ProfilingDebugging/CsvProfiler.h"
+
+DECLARE_DWORD_COUNTER_STAT(TEXT("Num Dirty Widgets"), STAT_SlateNumDirtyWidgets, STATGROUP_Slate);
+DECLARE_DWORD_COUNTER_STAT(TEXT("Num Volatile Widgets"), STAT_SlateNumVolatileWidgets, STATGROUP_Slate);
+
+#if WITH_SLATE_DEBUGGING
+bool GDumpUpdateList = false;
+void HandleDumpUpdateList(const TArray<FString>& Args)
+{
+
+	GDumpUpdateList = true;
+}
+
+static FAutoConsoleCommand HandleDumpUpdateListCommand(
+	TEXT("Slate.DumpUpdateList"),
+	TEXT(""),
+	FConsoleCommandWithArgsDelegate::CreateStatic(&HandleDumpUpdateList)
+);
+#endif
+
+
+FSlateInvalidationRoot::FSlateInvalidationRoot()
+	: InvalidationRootWidget(nullptr)
+	, RootHittestGrid(nullptr)
+	, FastPathGenerationNumber(INDEX_NONE)
+	, bChildOrderInvalidated(false)
+	, bNeedsSlowPath(true)
+{
+	FSlateApplicationBase::Get().OnInvalidateAllWidgets().AddRaw(this, &FSlateInvalidationRoot::OnInvalidateAllWidgets);
+
+}
+
+FSlateInvalidationRoot::~FSlateInvalidationRoot()
+{
+	ClearAllFastPathData();
+
+	if (FSlateApplicationBase::IsInitialized())
+	{
+		FSlateApplicationBase::Get().OnInvalidateAllWidgets().RemoveAll(this);
+	}
+}
+
+void FSlateInvalidationRoot::AddReferencedObjects(FReferenceCollector& Collector)
+{
+	CachedElementData.AddReferencedObjects(Collector);
+}
+
+FString FSlateInvalidationRoot::GetReferencerName() const
+{
+	static FString ReferencerName(TEXT("FSlateInvalidationRoot"));
+	return ReferencerName;
+}
+
+void FSlateInvalidationRoot::InvalidateChildOrder()
+{
+	//if (GSlateEnableGlobalInvalidation)
+	{
+		if(!bNeedsSlowPath && !bChildOrderInvalidated)
+		{
+			//UE_LOG(LogSlate, Log, TEXT("Child order invalidated, Slow path needed"));
+			//bNeedsSlowPath = true;
+			bChildOrderInvalidated = true;
+			if(!InvalidationRootWidget->Advanced_IsWindow())
+			{
+				InvalidationRootWidget->InvalidatePrepass();
+			}
+
+			if (!GSlateEnableGlobalInvalidation  && !InvalidationRootWidget->Advanced_IsWindow())
+			{
+				//InvalidationRootWidget->Invalidate(EInvalidateWidgetReason::ChildOrder);
+				InvalidationRootWidget->Invalidate(EInvalidateWidgetReason::Layout);
+			}
+
+#if WITH_SLATE_DEBUGGING
+			if (GSlateInvalidationDebugging && FastWidgetPathList.Num())
+			{
+				FSlateDebugging::WidgetInvalidated(*this, FastWidgetPathList[0], &FLinearColor::Red);
+			}
+#endif
+		}
+	}
+}
+
+int32 RecursiveFindParentWithChildOrderChange(const TArray<FWidgetProxy>& FastWidgetPathList, const FWidgetProxy& Proxy)
+{
+	if (Proxy.bChildOrderInvalid)
+	{
+		return Proxy.Index;
+	}
+	else if (Proxy.ParentIndex == INDEX_NONE)
+	{
+		return INDEX_NONE;
+	}
+	else
+	{
+		return RecursiveFindParentWithChildOrderChange(FastWidgetPathList, FastWidgetPathList[Proxy.ParentIndex]);
+	}
+}
+
+void FSlateInvalidationRoot::RemoveWidgetFromFastPath(FWidgetProxy& Proxy)
+{
+	if (Proxy.Widget->PersistentState.CachedElementListNode)
+	{
+		CachedElementData.RemoveCache(Proxy.Widget->PersistentState.CachedElementListNode);
+	}
+
+	if (Proxy.Index == 0)
+	{
+		InvalidateRoot();
+	}
+	else
+	{
+		InvalidateChildOrder();
+	}
+
+	Proxy.Widget->FastPathProxyHandle = FWidgetProxyHandle();
+	Proxy.Widget = nullptr;
+}
+
+void FSlateInvalidationRoot::InvalidateRoot()
+{
+	// Update the generation number.  This will effectively invalidate all proxy handles
+	++FastPathGenerationNumber;
+
+	InvalidationRootWidget->InvalidatePrepass();
+
+	bNeedsSlowPath = true;
+}
+
+FSlateInvalidationResult FSlateInvalidationRoot::PaintInvalidationRoot(const FSlateInvalidationContext& Context)
+{
+	const int32 LayerId = 0;
+
+	check(InvalidationRootWidget);
+	check(RootHittestGrid);
+
+	FSlateInvalidationResult Result;
+
+	if (Context.bAllowFastPathUpdate)
+	{
+		Context.WindowElementList->PushCachedElementData(CachedElementData);
+	}
+
+	SWidget* RootWidget = InvalidationRootWidget->Advanced_IsWindow() ? InvalidationRootWidget : &(*InvalidationRootWidget->GetAllChildren()->GetChildAt(0));
+
+	//AssignFastPathIndices(RootWidget, true);
+
+	EFlowDirection NewFlowDirection = GSlateFlowDirection;
+	if (RootWidget->GetFlowDirectionPreference() == EFlowDirectionPreference::Inherit)
+	{
+		NewFlowDirection = GSlateFlowDirectionShouldFollowCultureByDefault ? FLayoutLocalization::GetLocalizedLayoutDirection() : EFlowDirection::LeftToRight;
+	}
+
+	TGuardValue<EFlowDirection> FlowGuard(GSlateFlowDirection, NewFlowDirection);
+	if (!Context.bAllowFastPathUpdate || bNeedsSlowPath || GSlateIsInInvalidationSlowPath)
+	{
+		SCOPED_NAMED_EVENT(SWidget_PaintSlowPath, FColor::Red);
+
+		//CSV_EVENT(Basic, "Slate Slow Path update");
+#if WITH_SLATE_DEBUGGING
+		FSlateDebugging::ClearInvalidatedWidgets(*this);
+#endif
+		ClearAllFastPathData();
+
+		GSlateIsOnFastUpdatePath = false;
+		bNeedsSlowPath = false;
+		bChildOrderInvalidated = false;
+
+		{
+			if (Context.bAllowFastPathUpdate)
+			{
+				TGuardValue<bool> InSlowPathGuard(GSlateIsInInvalidationSlowPath, true);
+
+				BuildFastPathList(RootWidget);
+
+				if(GSlateEnableGlobalInvalidation)
+				{
+					InvalidationRootWidget->SlatePrepass(Context.LayoutScaleMultiplier);
+				}
+			}
+
+			CachedMaxLayerId = PaintSlowPath(Context);
+
+		}
+
+		Result.bRepaintedWidgets = true;
+
+	}
+	else if (FastWidgetPathList.Num())
+	{
+		// We should not have been supplied a different root than the one we generated a path to
+		check(RootWidget == FastWidgetPathList[0].Widget);
+
+		// hack, layout caching must be disabled
+		//TGuardValue<int32> LayoutCachingGuard(GSlateLayoutCaching, 0);
+
+		Result.bRepaintedWidgets = PaintFastPath(Context);
+	}
+
+	if (Context.bAllowFastPathUpdate)
+	{
+		Context.WindowElementList->PopCachedElementData();
+
+#if WITH_SLATE_DEBUGGING
+		if (GSlateInvalidationDebugging)
+		{
+			if (!GSlateEnableGlobalInvalidation)
+			{
+				FSlateDebugging::DrawInvalidationRoot(*InvalidationRootWidget, CachedMaxLayerId, *Context.WindowElementList);
+			}
+
+			FSlateDebugging::DrawInvalidatedWidgets(*this, *Context.PaintArgs, *Context.WindowElementList);
+		}
+#endif
+	}
+
+	FinalUpdateList.Empty();
+
+	Result.MaxLayerIdPainted = CachedMaxLayerId;
+	return Result;
+}
+
+bool FSlateInvalidationRoot::PaintFastPath(const FSlateInvalidationContext& Context)
+{
+	SCOPED_NAMED_EVENT(SWidget_FastPathUpdate, FColor::Green);
+
+	check(!bNeedsSlowPath);
+
+	bool bWidgetsNeededRepaint = false;
+	{
+		TGuardValue<bool> OnFastPathGuard(GSlateIsOnFastUpdatePath, true);
+
+
+		bool bDrawElementsUpdated = false;
+
+		int32 NumDirtyWidgets = 0;
+		int32 NumVolatileWidgets = 0;
+
+		int32 LastParentIndex = 0;
+
+#if WITH_SLATE_DEBUGGING
+		if (GDumpUpdateList)
+		{
+			UE_LOG(LogSlate, Log, TEXT("Dumping update list"));
+
+			// The update list is put in reverse order 
+			for (int32 ListIndex = FinalUpdateList.Num() - 1; ListIndex >= 0; --ListIndex)
+			{
+				const int32 MyIndex = FinalUpdateList[ListIndex];
+
+				FWidgetProxy& WidgetProxy = FastWidgetPathList[MyIndex];
+
+				if (EnumHasAnyFlags(WidgetProxy.UpdateFlags, EWidgetUpdateFlags::NeedsRepaint | EWidgetUpdateFlags::NeedsVolatilePaint))
+				{
+					UE_LOG(LogSlate, Log, TEXT("Repaint %s"), *WidgetProxy.Widget->GetTypeAsString());
+				}
+				else if (!WidgetProxy.bInvisibleDueToParentOrSelfVisibility)
+				{
+					if (EnumHasAnyFlags(WidgetProxy.UpdateFlags, EWidgetUpdateFlags::NeedsActiveTimerUpdate))
+					{
+						UE_LOG(LogSlate, Log, TEXT("ActiveTimer %s"), *WidgetProxy.Widget->GetTypeAsString());
+					}
+
+					if (EnumHasAnyFlags(WidgetProxy.UpdateFlags, EWidgetUpdateFlags::NeedsTick))
+					{
+						UE_LOG(LogSlate, Log, TEXT("Tick %s"), *WidgetProxy.Widget->GetTypeAsString());
+					}
+				}
+			}
+
+			GDumpUpdateList = false;
+		}
+#endif
+
+		{
+			// The update list is put in reverse order by ProcessInvalidation
+			for (int32 ListIndex = FinalUpdateList.Num() - 1; ListIndex >= 0; --ListIndex)
+			{
+				const int32 MyIndex = FinalUpdateList[ListIndex];
+				FWidgetProxy& WidgetProxy = FastWidgetPathList[MyIndex];
+
+				// Check visibility, it may have been in the update list but a parent who was also in the update list already updated it.
+				if (!WidgetProxy.bInvisibleDueToParentOrSelfVisibility && !WidgetProxy.bUpdatedSinceLastInvalidate && ensure(WidgetProxy.Widget))
+				{
+					const int32 NewLayerId = WidgetProxy.Update(*Context.PaintArgs, MyIndex, *Context.WindowElementList);
+					CachedMaxLayerId = FMath::Max(NewLayerId, CachedMaxLayerId);
+
+					FWidgetProxy::MarkProxyUpdatedThisFrame(WidgetProxy, WidgetsNeedingUpdate);
+
+					if (bNeedsSlowPath)
+					{
+						break;
+					}
+				}
+			}
+		}
+	}
+
+	if (bNeedsSlowPath)
+	{
+		SCOPED_NAMED_EVENT(SWidget_PaintSlowPath, FColor::Red);
+		CachedMaxLayerId = PaintSlowPath(Context);
+	}
+
+	return bWidgetsNeededRepaint;
+
+	INC_DWORD_STAT_BY(STAT_SlateNumDirtyWidgets, FinalUpdateList.Num());
+	//SET_DWORD_STAT(STAT_SlateNumVolatileWidgets, NumVolatileWidgets);
+}
+
+void FSlateInvalidationRoot::BuildNewFastPathList_Recursive(FSlateInvalidationRoot& Root, FWidgetProxy& Proxy, int32 ParentIndex, int32& NextTreeIndex, TArray<FWidgetProxy>& CurrentFastPathList, TArray<FWidgetProxy, TMemStackAllocator<>>& NewFastPathList)
+{
+	if (Proxy.bChildOrderInvalid)
+	{
+		NextTreeIndex = Proxy.LeafMostChildIndex != INDEX_NONE ? Proxy.LeafMostChildIndex + 1 : NextTreeIndex + 1;
+		Proxy.Widget->AssignIndicesToChildren(*this, ParentIndex, NewFastPathList, !Proxy.bInvisibleDueToParentOrSelfVisibility, Proxy.Widget->IsVolatile() || Proxy.Widget->IsVolatileIndirectly());
+	}
+	else
+	{ 
+		const int32 PrevIndex = Proxy.Index;
+		const int32 PrevParentIndex = Proxy.ParentIndex;
+		Proxy.Index = NewFastPathList.Num();
+		Proxy.ParentIndex = ParentIndex;
+		//if (PrevIndex != Proxy.Index)
+		{
+			// Update the proxy handle
+			Proxy.Widget->FastPathProxyHandle = FWidgetProxyHandle(Root, Proxy.Index);
+		}
+
+		 NewFastPathList.Add(MoveTemp(Proxy));
+
+		for (int32 LocalChildIndex = 0; LocalChildIndex < Proxy.NumChildren; ++LocalChildIndex)
+		{
+			FWidgetProxy& ChildProxy = CurrentFastPathList[NextTreeIndex];
+
+			if (!ChildProxy.bChildOrderInvalid)
+			{
+				++NextTreeIndex;
+			}
+
+			BuildNewFastPathList_Recursive(Root, ChildProxy, Proxy.Index, NextTreeIndex, CurrentFastPathList, NewFastPathList);
+		}
+
+		{
+			FWidgetProxy& MyProxyRef = NewFastPathList[Proxy.Index];
+			int32 LastIndex = NewFastPathList.Num() - 1;
+			MyProxyRef.LeafMostChildIndex = LastIndex != Proxy.Index ? LastIndex : INDEX_NONE;
+		}
+	}
+	
+}
+
+#define VERIFY_CHILD_ORDER 0
+void FSlateInvalidationRoot::BuildFastPathList(SWidget* RootWidget)
+{
+	SCOPED_NAMED_EVENT_TEXT("AssignFastPathIndices", FColor::Magenta);
+
+	TSharedPtr<SWidget> Parent = RootWidget->GetParentWidget();
+	// If the widget has no parent it is likely a window
+	const bool bParentVisible = Parent.IsValid() ? Parent->GetVisibility().IsVisible() : true;
+	const bool bParentVolatile = false;
+
+	FMemMark Mark(FMemStack::Get());
+	{
+		// Update the generation number.  This will effectively invalidate all proxy handles
+		++FastPathGenerationNumber;
+
+		WidgetsNeedingUpdate.Empty();
+
+
+#if VERIFY_CHILD_ORDER
+		TArray<FWidgetProxy, TMemStackAllocator<>> Copy;
+		Copy.Reserve(FastWidgetPathList.Num());
+		RootWidget->AssignIndicesToChildren(*this, INDEX_NONE, Copy, bParentVisible, bParentVolatile);
+
+		// Update the generation number.  This will effectively invalidate all proxy handles
+		//++FastPathGenerationNumber;
+#endif
+		TArray<FWidgetProxy, TMemStackAllocator<>> TempList;
+		TempList.Reserve(FastWidgetPathList.Num());
+
+		if (FastWidgetPathList.Num() > 0)
+		{
+			int32 NextTreeIndex = 1;
+			BuildNewFastPathList_Recursive(*this, FastWidgetPathList[0], INDEX_NONE, NextTreeIndex, FastWidgetPathList, TempList);
+		}
+		else
+		{
+			RootWidget->AssignIndicesToChildren(*this, INDEX_NONE, TempList, bParentVisible, bParentVolatile);
+		}
+
+#if VERIFY_CHILD_ORDER
+		for (int i = 0; i < Copy.Num(); ++i)
+		{
+			ensureAlways(Copy[i].Widget == TempList[i].Widget);
+		}
+#endif
+		FastWidgetPathList = MoveTemp(TempList);
+	}
+}
+
+bool FSlateInvalidationRoot::ProcessInvalidation()
+{
+	SCOPED_NAMED_EVENT(Slate_InvalidationProcessing, FColor::Blue);
+
+	FinalUpdateList.Empty();
+
+	bool bWidgetsNeedRepaint = false;
+
+	if (!bNeedsSlowPath)
+	{
+		if (bChildOrderInvalidated)
+		{
+			CSV_EVENT_GLOBAL(TEXT("Slate Sort Children"));
+			SCOPED_NAMED_EVENT(Slate_InvalidationProcessing_SortChildren, FColor::Orange);
+
+			struct FWidgetNeedingUpdate
+			{
+				FWidgetNeedingUpdate(SWidget* InWidget, EInvalidateWidgetReason InCurrentInvalidateReason, EWidgetUpdateFlags InUpdateFlags)
+					: Widget(InWidget)
+					, CurrentInvalidateReason(InCurrentInvalidateReason)
+					, UpdateFlags(InUpdateFlags)
+				{}
+
+				SWidget* Widget;
+				EInvalidateWidgetReason CurrentInvalidateReason;
+				EWidgetUpdateFlags UpdateFlags;
+			};
+
+			// We need to store off all widgets needing update as the child order change may invalidate them all
+			TArray<FWidgetNeedingUpdate, TMemStackAllocator<>> WidgetsNeedingUpdateCache;
+			WidgetsNeedingUpdateCache.Reserve(WidgetsNeedingUpdate.Num());
+
+			for(int32 WidgetIndex : WidgetsNeedingUpdate.GetRawData())
+			{
+				FWidgetProxy& WidgetProxy = FastWidgetPathList[WidgetIndex];
+				if (WidgetProxy.Widget)
+				{
+					check(WidgetProxy.Widget->FastPathProxyHandle.GetInvalidationRoot() == this);
+					WidgetsNeedingUpdateCache.Emplace(WidgetProxy.Widget, WidgetProxy.CurrentInvalidateReason, WidgetProxy.UpdateFlags);
+				}
+			}
+
+			SWidget* RootWidget = InvalidationRootWidget->Advanced_IsWindow() ? InvalidationRootWidget : &(*InvalidationRootWidget->GetAllChildren()->GetChildAt(0));
+			BuildFastPathList(RootWidget);
+
+			for (FWidgetNeedingUpdate& WidgetNeedingUpdate : WidgetsNeedingUpdateCache)
+			{
+				const int32 NewIndex = WidgetNeedingUpdate.Widget->FastPathProxyHandle.GetIndex();
+				FWidgetProxy& WidgetProxy = FastWidgetPathList[NewIndex];
+				check(WidgetProxy.Widget == WidgetNeedingUpdate.Widget);
+				WidgetProxy.CurrentInvalidateReason = WidgetNeedingUpdate.CurrentInvalidateReason;
+				WidgetProxy.UpdateFlags = WidgetNeedingUpdate.UpdateFlags;
+				WidgetProxy.bInUpdateList = true;
+				WidgetsNeedingUpdate.Push(WidgetProxy);
+			}
+
+			bChildOrderInvalidated = false;
+		}
+
+		while (WidgetsNeedingUpdate.Num() && !bNeedsSlowPath)
+		{
+			int32 MyIndex = WidgetsNeedingUpdate.Pop();
+			FinalUpdateList.Add(MyIndex);
+			FWidgetProxy& WidgetProxy = FastWidgetPathList[MyIndex];
+
+			// Reset each widgets paint state
+			// Must be done before actual painting because children can repaint 
+			WidgetProxy.bUpdatedSinceLastInvalidate = false;
+			WidgetProxy.bInUpdateList = false;
+
+			// Widget could be null if it was removed and we are on the slow path
+			if (WidgetProxy.Widget)
+			{
+				if (!GSlateEnableGlobalInvalidation && !InvalidationRootWidget->NeedsPrepass() && WidgetProxy.Widget->Advanced_IsInvalidationRoot())
+				{
+					// Prepass on child invalidation panels to give them a chance to process their invalidations since our panel is not going to prepass down all the children
+					WidgetProxy.Widget->SlatePrepass(WidgetProxy.Widget->GetPrepassLayoutScaleMultiplier());
+				}
+				else
+				{
+					bWidgetsNeedRepaint |= WidgetProxy.ProcessInvalidation(WidgetsNeedingUpdate, FastWidgetPathList, *this);
+				}
+
+#if WITH_SLATE_DEBUGGING
+				if (GSlateInvalidationDebugging)
+				{
+					if (EnumHasAnyFlags(WidgetProxy.UpdateFlags, EWidgetUpdateFlags::NeedsRepaint))
+					{
+						FSlateDebugging::WidgetInvalidated(*this, WidgetProxy);
+					}
+					else if (EnumHasAnyFlags(WidgetProxy.UpdateFlags, EWidgetUpdateFlags::NeedsVolatilePaint) && !WidgetProxy.Widget->Advanced_IsInvalidationRoot())
+					{
+						FSlateDebugging::WidgetInvalidated(*this, WidgetProxy, &FLinearColor::Blue);
+					}
+					else if (EnumHasAnyFlags(WidgetProxy.UpdateFlags, EWidgetUpdateFlags::NeedsTick))
+					{
+						FSlateDebugging::WidgetInvalidated(*this, WidgetProxy, &FLinearColor::White);
+					}
+					else if (EnumHasAnyFlags(WidgetProxy.UpdateFlags, EWidgetUpdateFlags::NeedsActiveTimerUpdate))
+					{
+						FSlateDebugging::WidgetInvalidated(*this, WidgetProxy, &FLinearColor::Green);
+					}
+				}
+#endif
+			}
+		}
+
+		WidgetsNeedingUpdate.Reset();
+	}
+	else
+	{
+		bWidgetsNeedRepaint = true;
+	}
+
+	return bWidgetsNeedRepaint;
+}
+
+void FSlateInvalidationRoot::ClearAllFastPathData()
+{
+	for (const FWidgetProxy& Proxy : FastWidgetPathList)
+	{
+		if (Proxy.Widget)
+		{
+			Proxy.Widget->PersistentState.CachedElementListNode = nullptr;
+		}
+	}
+
+	FastWidgetPathList.Empty();
+	WidgetsNeedingUpdate.Empty();
+	CachedElementData.Empty();
+	FinalUpdateList.Empty();
+}
+
+void FSlateInvalidationRoot::OnInvalidateAllWidgets()
+{
+	InvalidateChildOrder();
+
+	bNeedsSlowPath = true;
+}
diff --git a/Engine/Source/Runtime/SlateCore/Private/FastUpdate/WidgetProxy.cpp b/Engine/Source/Runtime/SlateCore/Private/FastUpdate/WidgetProxy.cpp
new file mode 100644
--- /dev/null
+++ b/Engine/Source/Runtime/SlateCore/Private/FastUpdate/WidgetProxy.cpp
@@ -0,0 +1,276 @@
+// Copyright 1998-2018 Epic Games, Inc. All Rights Reserved.
+
+#include "FastUpdate/WidgetProxy.h"
+#include "Widgets/SWidget.h"
+#include "Rendering/DrawElements.h"
+#include "Widgets/SWindow.h"
+#include "ProfilingDebugging/CsvProfiler.h"
+#include "Types/ReflectionMetadata.h"
+
+const FSlateWidgetPersistentState FSlateWidgetPersistentState::NoState;
+
+FWidgetProxy::FWidgetProxy(SWidget* InWidget)
+	: Widget(InWidget)
+	, Index(INDEX_NONE)
+	, ParentIndex(INDEX_NONE)
+	, NumChildren(0)
+	, LeafMostChildIndex(INDEX_NONE)
+	, UpdateFlags(EWidgetUpdateFlags::None)
+	, CurrentInvalidateReason(EInvalidateWidget::None)
+	// Potentially unsafe to update visibility from the widget due to attribute bindings.  This is updated later when the widgets are sorted in ProcessInvalidation
+	, Visibility(EVisibility::Collapsed) 
+	, bUpdatedSinceLastInvalidate(false)
+	, bInUpdateList(false)
+	, bInvisibleDueToParentOrSelfVisibility(false)
+	, bChildOrderInvalid(false)
+{
+	check(Widget != nullptr);
+}
+
+int32 FWidgetProxy::Update(const FPaintArgs& PaintArgs, int32 MyIndex, FSlateWindowElementList& OutDrawElements)
+{
+	// If Outgoing layer id remains index none, there was no change
+	int32 OutgoingLayerId = INDEX_NONE;
+	if (EnumHasAnyFlags(UpdateFlags,  EWidgetUpdateFlags::NeedsRepaint|EWidgetUpdateFlags::NeedsVolatilePaint))
+	{
+		ensure(!bInvisibleDueToParentOrSelfVisibility);
+		OutgoingLayerId = Repaint(PaintArgs, MyIndex, OutDrawElements);
+	}
+	else if(!bInvisibleDueToParentOrSelfVisibility)
+	{
+		if (EnumHasAnyFlags(UpdateFlags, EWidgetUpdateFlags::NeedsActiveTimerUpdate))
+		{
+			SCOPE_CYCLE_COUNTER(STAT_SlateExecuteActiveTimers);
+			Widget->ExecuteActiveTimers(PaintArgs.GetCurrentTime(), PaintArgs.GetDeltaTime());
+		}
+
+		if (EnumHasAnyFlags(UpdateFlags, EWidgetUpdateFlags::NeedsTick))
+		{
+			const FSlateWidgetPersistentState& MyState = Widget->GetPersistentState();
+
+			INC_DWORD_STAT(STAT_SlateNumTickedWidgets);
+			SCOPE_CYCLE_COUNTER(STAT_SlateTickWidgets);
+			Widget->Tick(MyState.AllottedGeometry, PaintArgs.GetCurrentTime(), PaintArgs.GetDeltaTime());
+		}
+
+	}
+
+	return OutgoingLayerId;
+}
+
+bool FWidgetProxy::ProcessInvalidation(FWidgetUpdateList& UpdateList, TArray<FWidgetProxy>& FastWidgetPathList, FSlateInvalidationRoot& Root)
+{
+	bool bWidgetNeedsRepaint = false;
+	if (EnumHasAnyFlags(CurrentInvalidateReason, EInvalidateWidget::RenderTransform))
+	{
+		UpdateFlags |= EWidgetUpdateFlags::NeedsRepaint;
+
+		FWidgetProxy& ParentProxy = FastWidgetPathList[ParentIndex];
+		ParentProxy.CurrentInvalidateReason |= EInvalidateWidget::Layout;
+		UpdateList.Push(ParentProxy);
+
+		bWidgetNeedsRepaint = true;
+	}
+	else if (EnumHasAnyFlags(CurrentInvalidateReason, EInvalidateWidget::Layout|EInvalidateWidget::Visibility|EInvalidateWidget::ChildOrder))
+	{
+		// When layout changes compute a new desired size for this widget
+		FVector2D CurrentDesiredSize = Widget->GetDesiredSize();
+		FVector2D NewDesiredSize = FVector2D::ZeroVector;
+		if (Visibility != EVisibility::Collapsed)
+		{
+			if (Widget->NeedsPrepass())
+			{
+				Widget->SlatePrepass(Widget->PrepassLayoutScaleMultiplier);
+			}
+			else
+			{
+				Widget->CacheDesiredSize(Widget->PrepassLayoutScaleMultiplier);
+			}
+
+			NewDesiredSize = Widget->GetDesiredSize();
+		}
+
+		// Note even if volatile we need to recompute desired size. We do not need to invalidate parents though if they are volatile since they will naturally redraw this widget
+		if (!Widget->IsVolatileIndirectly() && Visibility.IsVisible())
+		{
+			UpdateFlags |= EWidgetUpdateFlags::NeedsRepaint;
+		}
+
+		// If the desired size changed, invalidate the parent if it is visible
+		if (NewDesiredSize != CurrentDesiredSize || EnumHasAnyFlags(CurrentInvalidateReason, EInvalidateWidget::Visibility))
+		{
+			if(ParentIndex != INDEX_NONE)
+			{
+				FWidgetProxy& ParentProxy = FastWidgetPathList[ParentIndex];
+				if (ParentIndex == 0)
+				{
+					// root of the invalidation panel just invalidate the whole thing
+					Root.InvalidateRoot();
+				}
+				else if (ParentProxy.Visibility.IsVisible())
+				{
+					ParentProxy.CurrentInvalidateReason |= EInvalidateWidget::Layout;
+					UpdateList.Push(ParentProxy);
+				}
+			}
+			else if(!GSlateEnableGlobalInvalidation && Widget->IsParentValid())
+			{
+				TSharedPtr<SWidget> ParentWidget = Widget->GetParentWidget();
+				if (ParentWidget->Advanced_IsInvalidationRoot())
+				{
+					Root.InvalidateRoot();
+				}
+			}
+		}
+		
+		bWidgetNeedsRepaint = true;
+	}
+	else if (EnumHasAnyFlags(CurrentInvalidateReason, EInvalidateWidget::Paint) && !Widget->IsVolatileIndirectly())
+	{
+		UpdateFlags |= EWidgetUpdateFlags::NeedsRepaint;
+
+		bWidgetNeedsRepaint = true;
+	}
+
+
+	CurrentInvalidateReason = EInvalidateWidget::None;
+
+	return bWidgetNeedsRepaint;
+}
+
+void FWidgetProxy::MarkProxyUpdatedThisFrame(FWidgetProxy& Proxy, FWidgetUpdateList& UpdateList)
+{
+	Proxy.bUpdatedSinceLastInvalidate = true;
+
+	if(EnumHasAnyFlags(Proxy.UpdateFlags, EWidgetUpdateFlags::AnyUpdate))
+	{
+		if (!Proxy.bInUpdateList && !Proxy.bInvisibleDueToParentOrSelfVisibility)
+		{
+			// If there are any updates still needed add them to the next update list
+			UpdateList.Push(Proxy);
+		}
+	}
+	else
+	{
+		Proxy.bInUpdateList = false;
+	}
+}
+
+int32 FWidgetProxy::Repaint(const FPaintArgs& PaintArgs, int32 MyIndex, FSlateWindowElementList& OutDrawElements) const
+{
+#if WITH_SLATE_DEBUGGING
+	SCOPED_NAMED_EVENT_FSTRING(FReflectionMetaData::GetWidgetDebugInfo(Widget), EnumHasAnyFlags(UpdateFlags, EWidgetUpdateFlags::NeedsRepaint) ? FColor::Orange : FColor::Red);
+#endif
+
+	const FSlateWidgetPersistentState& MyState = Widget->GetPersistentState();
+
+	const int32 StartingClipIndex = OutDrawElements.GetClippingIndex();
+
+	// Get the clipping manager into the correct state
+	const bool bNeedsNewClipState = MyState.InitialClipState.IsSet();
+	if (bNeedsNewClipState)
+	{
+		OutDrawElements.GetClippingManager().PushClippingState(MyState.InitialClipState.GetValue());
+	}
+
+	GSlateFlowDirection = MyState.IncomingFlowDirection;
+	
+	FPaintArgs UpdatedArgs = PaintArgs.WithNewParent(MyState.PaintParent.Pin().Get());
+	UpdatedArgs.SetInheritedHittestability(MyState.bInheritedHittestability);
+
+	int32 PrevLayerId = MyState.OutgoingLayerId;
+
+	if (GSlateEnableGlobalInvalidation)
+	{
+		if (EnumHasAnyFlags(UpdateFlags, EWidgetUpdateFlags::NeedsVolatilePaint))
+		{
+			if (Widget->ShouldInvalidatePrepassDueToVolatility())
+			{
+				Widget->InvalidatePrepass();
+			}
+			Widget->SlatePrepass(Widget->GetPrepassLayoutScaleMultiplier());
+		}
+	}
+	const int32 NewLayerId = Widget->Paint(UpdatedArgs, MyState.AllottedGeometry, MyState.CullingBounds, OutDrawElements, MyState.LayerId, MyState.WidgetStyle, MyState.bParentEnabled);
+
+	if (bNeedsNewClipState)
+	{
+		OutDrawElements.PopClip();
+		// clip index should be what it was before.  if this assert fails something internal inside the above paint call did not pop clip properly
+		check(StartingClipIndex == OutDrawElements.GetClippingIndex());
+	}
+
+	return NewLayerId;
+}
+
+FWidgetProxyHandle::FWidgetProxyHandle(FSlateInvalidationRoot& InInvalidationRoot, int32 InIndex)
+	: InvalidationRoot(&InInvalidationRoot)
+	, MyIndex(InIndex)
+	, GenerationNumber(InInvalidationRoot.GetFastPathGenerationNumber())
+{
+
+}
+
+bool FWidgetProxyHandle::IsValid() const
+{
+	return InvalidationRoot && InvalidationRoot->GetFastPathGenerationNumber() == GenerationNumber && MyIndex != INDEX_NONE;
+}
+
+FWidgetProxy& FWidgetProxyHandle::GetProxy()
+{
+	return InvalidationRoot->FastWidgetPathList[MyIndex];
+}
+
+const FWidgetProxy& FWidgetProxyHandle::GetProxy() const
+{
+	return InvalidationRoot->FastWidgetPathList[MyIndex];
+}
+
+void FWidgetProxyHandle::MarkWidgetUpdatedThisFrame()
+{
+	FWidgetProxy::MarkProxyUpdatedThisFrame(GetProxy(), InvalidationRoot->WidgetsNeedingUpdate);
+}
+
+void FWidgetProxyHandle::MarkWidgetDirty(EInvalidateWidget InvalidateReason)
+{
+	FWidgetProxy& Proxy = GetProxy();
+
+	if (EnumHasAnyFlags(InvalidateReason, EInvalidateWidget::ChildOrder))
+	{
+		/*
+				CSV_EVENT_GLOBAL(TEXT("Slow Path Needed"));
+		#if WITH_SLATE_DEBUGGING
+				UE_LOG(LogSlate, Log, TEXT("Slow Widget Path Needed: %s %s"), *Proxy.Widget->ToString(), *Proxy.Widget->GetTag().ToString());
+		#endif*/
+		Proxy.bChildOrderInvalid = true;
+		InvalidationRoot->InvalidateChildOrder();
+	}
+
+	if (Proxy.CurrentInvalidateReason == EInvalidateWidget::None)
+	{
+		InvalidationRoot->WidgetsNeedingUpdate.Push(Proxy);
+	}
+#if WITH_SLATE_DEBUGGING
+	else
+	{
+		ensure(InvalidationRoot->WidgetsNeedingUpdate.Contains(Proxy));
+	}
+#endif
+	Proxy.CurrentInvalidateReason |= InvalidateReason;
+}
+
+void FWidgetProxyHandle::UpdateWidgetFlags(EWidgetUpdateFlags NewFlags)
+{
+	FWidgetProxy& Proxy = GetProxy();
+
+	if (!Proxy.bInvisibleDueToParentOrSelfVisibility)
+	{
+		Proxy.UpdateFlags = NewFlags;
+
+		// Add to update list if the widget is now tickable or has an active timer.  Disregard if dirty, it was already in the list
+		if (!Proxy.bInUpdateList && EnumHasAnyFlags(NewFlags, EWidgetUpdateFlags::AnyUpdate))
+		{
+			InvalidationRoot->WidgetsNeedingUpdate.Push(Proxy);
+		}
+	}
+}
