Fix widget switcher crash if initialized to an invalid index at construction time

diff --git a/Engine/Source/Runtime/Slate/Private/Widgets/Layout/SWidgetSwitcher.cpp b/Engine/Source/Runtime/Slate/Private/Widgets/Layout/SWidgetSwitcher.cpp
--- a/Engine/Source/Runtime/Slate/Private/Widgets/Layout/SWidgetSwitcher.cpp
+++ b/Engine/Source/Runtime/Slate/Private/Widgets/Layout/SWidgetSwitcher.cpp
@@ -118,13 +118,16 @@ void SWidgetSwitcher::SetActiveWidgetIndex( int32 Index )
 {
 	if (WidgetIndex.Get() != Index)
 	{
-		const FSlot* ActiveSlot = GetActiveSlot();
-
-		SWidget& OldActiveWidget = ActiveSlot->GetWidget().Get();
-
 		Invalidate(EInvalidateWidget::ChildOrder);
 
-		InvalidateChildRemovedFromTree(OldActiveWidget);
+		const FSlot* ActiveSlot = GetActiveSlot();
+
+		// Active slot can be null if the widget switcher was initialized to an invalid index.
+		if (ActiveSlot)
+		{
+			SWidget& OldActiveWidget = ActiveSlot->GetWidget().Get();
+			InvalidateChildRemovedFromTree(OldActiveWidget);
+		}
 
 		WidgetIndex = Index;
 
