Fix editor tests failing in ricky bobby mode because prepassing was never happening with nullrhi...

diff --git a/Engine/Source/Runtime/Slate/Private/Framework/Application/SlateApplication.cpp b/Engine/Source/Runtime/Slate/Private/Framework/Application/SlateApplication.cpp
--- a/Engine/Source/Runtime/Slate/Private/Framework/Application/SlateApplication.cpp
+++ b/Engine/Source/Runtime/Slate/Private/Framework/Application/SlateApplication.cpp
@@ -1471,7 +1471,7 @@ static bool DoAnyWindowDescendantsNeedPrepass(TSharedRef<SWindow> WindowToPrepas
 
 static void PrepassWindowAndChildren( TSharedRef<SWindow> WindowToPrepass )
 {
-	if (UNLIKELY(!FApp::CanEverRender()))
+	if (IsRunningDedicatedServer())
 	{
 		return;
 	}
