Fix for crash in SScrollBar when the track dosn't need scrolling.

The root cause is that the contents of the HittestGrid are different in the disabled state. More investigation is likely needed.

The fix however is to avoid looking up the widgets when we know they arn't there to be found.

diff --git a/Engine/Source/Runtime/Slate/Private/Widgets/Layout/SScrollBar.cpp b/Engine/Source/Runtime/Slate/Private/Widgets/Layout/SScrollBar.cpp
--- a/Engine/Source/Runtime/Slate/Private/Widgets/Layout/SScrollBar.cpp
+++ b/Engine/Source/Runtime/Slate/Private/Widgets/Layout/SScrollBar.cpp
@@ -106,11 +106,11 @@ void SScrollBar::SetState( float InOffsetFraction, float InThumbSizeFraction )
 
 FReply SScrollBar::OnMouseButtonDown( const FGeometry& MyGeometry, const FPointerEvent& MouseEvent )
 {
-	if ( MouseEvent.GetEffectingButton() == EKeys::LeftMouseButton )
+	if (MouseEvent.GetEffectingButton() == EKeys::LeftMouseButton && IsNeeded())
 	{
 		FGeometry ThumbGeometry = FindChildGeometry(MyGeometry, DragThumb.ToSharedRef());
 
-		if( DragThumb->IsHovered() )
+		if (DragThumb->IsHovered())
 		{
 			// Clicking on the scrollbar drag thumb
 			if( Orientation == Orient_Horizontal )
@@ -124,7 +124,7 @@ FReply SScrollBar::OnMouseButtonDown( const FGeometry& MyGeometry, const FPointe
 
 			bDraggingThumb = true;
 		}
-		else if( OnUserScrolled.IsBound() )
+		else if (OnUserScrolled.IsBound())
 		{
 			// Clicking in the non drag thumb area of the scrollbar
 			DragGrabOffset = Orientation == Orient_Horizontal ? (ThumbGeometry.GetLocalSize().X * 0.5f) : (ThumbGeometry.GetLocalSize().Y * 0.5f);
