Fixed a problem with mouse click/drag events being received by incorrect window/widget in certain conditions while dragging windows on Mac

diff --git a/Engine/Source/Runtime/ApplicationCore/Private/Mac/MacApplication.cpp b/Engine/Source/Runtime/ApplicationCore/Private/Mac/MacApplication.cpp
--- a/Engine/Source/Runtime/ApplicationCore/Private/Mac/MacApplication.cpp
+++ b/Engine/Source/Runtime/ApplicationCore/Private/Mac/MacApplication.cpp
@@ -283,7 +283,7 @@ bool FMacApplication::IsCursorDirectlyOverSlateWindow() const
 
 TSharedPtr<FGenericWindow> FMacApplication::GetWindowUnderCursor()
 {
-	if (WindowUnderCursor && WindowUnderCursor != DraggedWindow)
+	if (WindowUnderCursor)
 	{
 		return StaticCastSharedPtr<FGenericWindow>(FindWindowByNSWindow(WindowUnderCursor));
 	}
diff --git a/Engine/Source/Runtime/Slate/Private/Framework/Application/SlateApplication.cpp b/Engine/Source/Runtime/Slate/Private/Framework/Application/SlateApplication.cpp
--- a/Engine/Source/Runtime/Slate/Private/Framework/Application/SlateApplication.cpp
+++ b/Engine/Source/Runtime/Slate/Private/Framework/Application/SlateApplication.cpp
@@ -1221,11 +1221,7 @@ FWidgetPath FSlateApplication::LocateWindowUnderMouse( FVector2D ScreenspaceMous
 		TSharedPtr<SWindow> Window = FSlateWindowHelper::FindWindowByPlatformWindow(Windows, NativeWindowUnderMouse.ToSharedRef());
 		if (Window.IsValid())
 		{
-			FWidgetPath PathToLocatedWidget = LocateWidgetInWindow(ScreenspaceMouseCoordinate, Window.ToSharedRef(), bIgnoreEnabledStatus);
-			if (PathToLocatedWidget.IsValid())
-			{
-				return PathToLocatedWidget;
-			}
+			return LocateWidgetInWindow(ScreenspaceMouseCoordinate, Window.ToSharedRef(), bIgnoreEnabledStatus);
 		}
 	}
 
