SlateApplication - When processing the mouse button up event leave responsiveness mode before dispatching the event rather than after to prevent being "stuck" in responsive mode when a modal dialog is opened.

diff --git a/Engine/Source/Runtime/Slate/Private/Framework/Application/SlateApplication.cpp b/Engine/Source/Runtime/Slate/Private/Framework/Application/SlateApplication.cpp
--- a/Engine/Source/Runtime/Slate/Private/Framework/Application/SlateApplication.cpp
+++ b/Engine/Source/Runtime/Slate/Private/Framework/Application/SlateApplication.cpp
@@ -6063,6 +6063,13 @@ bool FSlateApplication::ProcessMouseButtonUpEvent( FPointerEvent& MouseEvent )
 {
 	SCOPE_CYCLE_COUNTER(STAT_ProcessMouseButtonUp);
 
+	// If in responsive mode throttle, leave it on mouse up.  Release this before dispatching the event to prevent being stuck in this mode
+	// until the next click if a modal dialog is opened.
+	if (MouseButtonDownResponsivnessThrottle.IsValid())
+	{
+		FSlateThrottleManager::Get().LeaveResponsiveMode(MouseButtonDownResponsivnessThrottle);
+	}
+
 	QueueSynthesizedMouseMove();
 	SetLastUserInteractionTime(this->GetCurrentTime());
 	LastUserInteractionTimeForThrottling = LastUserInteractionTime;
@@ -6078,12 +6085,6 @@ bool FSlateApplication::ProcessMouseButtonUpEvent( FPointerEvent& MouseEvent )
 	FWidgetPath EmptyPath;
 	const bool bHandled = RoutePointerUpEvent( EmptyPath, MouseEvent ).IsEventHandled();
 
-	// If in responsive mode throttle, leave it on mouse up.
-	if( MouseButtonDownResponsivnessThrottle.IsValid() )
-	{
-		FSlateThrottleManager::Get().LeaveResponsiveMode( MouseButtonDownResponsivnessThrottle );
-	}
-
 	if ( PressedMouseButtons.Num() == 0 )
 	{
 		// Release Capture
