Removed the generation of a mouse leave on touch input. This caused issues with simulated touch. Basically a revert of CL 3965402.
Copied from Dev-Editor, CL 4474420

diff --git a/Engine/Source/Runtime/Slate/Private/Framework/Application/SlateApplication.cpp b/Engine/Source/Runtime/Slate/Private/Framework/Application/SlateApplication.cpp
--- a/Engine/Source/Runtime/Slate/Private/Framework/Application/SlateApplication.cpp
+++ b/Engine/Source/Runtime/Slate/Private/Framework/Application/SlateApplication.cpp
@@ -6512,20 +6512,6 @@ void FSlateApplication::ProcessTouchStartedEvent( const TSharedPtr< FGenericWind
 	FUserAndPointer UserAndIndex(InTouchEvent.GetUserIndex(), InTouchEvent.GetPointerIndex());
 	PointerIndexLastPositionMap.Add(UserAndIndex, InTouchEvent.GetScreenSpacePosition());
 
-	const FWeakWidgetPath& LastWidgetsUnderCursor = WidgetsUnderCursorLastEvent.FindRef(UserAndIndex);
-	if (LastWidgetsUnderCursor.IsValid())
-	{
-		FWidgetPath SafeWidgetPath = LastWidgetsUnderCursor.ToWidgetPath();
-
-		FEventRouter::Route<FNoReply>(this, FEventRouter::FBubblePolicy(SafeWidgetPath), InTouchEvent, [] (const FArrangedWidget& SomeWidget, const FPointerEvent& PointerEvent)
-		{
-			SomeWidget.Widget->OnMouseLeave(PointerEvent);
-			return FNoReply();
-		});
-
-		WidgetsUnderCursorLastEvent.Remove(UserAndIndex);
-	}
-
 	ProcessMouseButtonDownEvent(PlatformWindow, InTouchEvent);
 }
 
