Fix for scrollbars that are set to always display not clamping their size correctly

diff --git a/Engine/Source/Runtime/Slate/Private/Widgets/Layout/SScrollBar.cpp b/Engine/Source/Runtime/Slate/Private/Widgets/Layout/SScrollBar.cpp
--- a/Engine/Source/Runtime/Slate/Private/Widgets/Layout/SScrollBar.cpp
+++ b/Engine/Source/Runtime/Slate/Private/Widgets/Layout/SScrollBar.cpp
@@ -27,8 +27,6 @@ void SScrollBar::Construct(const FArguments& InArgs)
 
 	SBorder::Construct( SBorder::FArguments()
 		.BorderImage(FCoreStyle::Get().GetBrush("NoBorder"))
-		.BorderBackgroundColor( this, &SScrollBar::GetTrackOpacity )
-		.ColorAndOpacity( this, &SScrollBar::GetThumbOpacity )
 		[
 			SNew(SVerticalBox)
 
@@ -50,6 +48,7 @@ void SScrollBar::Construct(const FArguments& InArgs)
 						.VAlign(VerticalAlignment)
 						[
 							SNew(SImage)
+							.ColorAndOpacity(this, &SScrollBar::GetTrackOpacity)
 							.Image(TopBrush)
 						]
 					]
@@ -57,6 +56,7 @@ void SScrollBar::Construct(const FArguments& InArgs)
 					[
 						SAssignNew(DragThumb, SBorder)
 						.BorderImage( this, &SScrollBar::GetDragThumbImage )
+						.ColorAndOpacity(this, &SScrollBar::GetThumbOpacity)
 						.HAlign(HAlign_Center)
 						.VAlign(VAlign_Center)
 						[
@@ -71,6 +71,7 @@ void SScrollBar::Construct(const FArguments& InArgs)
 						.VAlign(VerticalAlignment)
 						[
 							SNew(SImage)
+							.ColorAndOpacity(this, &SScrollBar::GetTrackOpacity)
 							.Image(BottomBrush)
 						]
 					]
@@ -216,7 +217,11 @@ SScrollBar::SScrollBar()
 
 FSlateColor SScrollBar::GetTrackOpacity() const
 {
-	if ( bDraggingThumb || IsHovered() )
+	if (AlwaysShowScrollbar() && (Track->GetThumbSizeFraction() == 0.0f))
+	{
+		return FLinearColor(1, 1, 1, 0.5f);
+	}
+	else if (bDraggingThumb || IsHovered())
 	{
 		return FLinearColor(1,1,1,1);
 	}
@@ -228,7 +233,11 @@ FSlateColor SScrollBar::GetTrackOpacity() const
 
 FLinearColor SScrollBar::GetThumbOpacity() const
 {
-	if ( bDraggingThumb || IsHovered() )
+	if (Track->GetThumbSizeFraction() <= 0.0f)
+	{
+		return FLinearColor(1, 1, 1, 0);
+	}
+	else if ( bDraggingThumb || IsHovered() )
 	{
 		return FLinearColor(1,1,1,1);
 	}
@@ -241,7 +250,7 @@ FLinearColor SScrollBar::GetThumbOpacity() const
 			float ThumbOpacity = FMath::Lerp(1.0f, 0.0f, FMath::Clamp((float)( ( LastInteractionDelta - 0.2 ) / 0.2 ), 0.0f, 1.0f));
 			return FLinearColor(1, 1, 1, ThumbOpacity);
 		}
-		else
+		else 
 		{
 			return FLinearColor(1, 1, 1, 0.75f);
 		}
diff --git a/Engine/Source/Runtime/Slate/Private/Widgets/Layout/SScrollBarTrack.cpp b/Engine/Source/Runtime/Slate/Private/Widgets/Layout/SScrollBarTrack.cpp
--- a/Engine/Source/Runtime/Slate/Private/Widgets/Layout/SScrollBarTrack.cpp
+++ b/Engine/Source/Runtime/Slate/Private/Widgets/Layout/SScrollBarTrack.cpp
@@ -36,7 +36,8 @@ void SScrollBarTrack::Construct(const FArguments& InArgs)
 
 SScrollBarTrack::FTrackSizeInfo SScrollBarTrack::GetTrackSizeInfo(const FGeometry& InTrackGeometry) const
 {
-	return FTrackSizeInfo(InTrackGeometry, Orientation, MinThumbSize, this->ThumbSizeFraction, this->OffsetFraction);
+	const float CurrentMinThumbSize = ThumbSizeFraction <= 0.0f ? 0.0f : MinThumbSize;
+	return FTrackSizeInfo(InTrackGeometry, Orientation, CurrentMinThumbSize, this->ThumbSizeFraction, this->OffsetFraction);
 }
 
 void SScrollBarTrack::OnArrangeChildren(const FGeometry& AllottedGeometry, FArrangedChildren& ArrangedChildren) const
@@ -119,10 +120,14 @@ void SScrollBarTrack::SetSizes(float InThumbOffsetFraction, float InThumbSizeFra
 	ThumbSizeFraction = InThumbSizeFraction;
 
 	// If you have no thumb, then it's effectively the size of the whole track
-	if (ThumbSizeFraction == 0.0f)
+	if (ThumbSizeFraction == 0.0f && !bIsAlwaysVisible)
 	{
 		ThumbSizeFraction = 1.0f;
 	}
+	else if (ThumbSizeFraction > 1.0f && bIsAlwaysVisible)
+	{
+		ThumbSizeFraction = 0.0f;
+	}
 }
 
 bool SScrollBarTrack::IsNeeded() const
