Fixed issue where AnalogCursor update events didnt have the proper pointer index

diff --git a/Engine/Source/Runtime/Slate/Private/Framework/Application/AnalogCursor.cpp b/Engine/Source/Runtime/Slate/Private/Framework/Application/AnalogCursor.cpp
--- a/Engine/Source/Runtime/Slate/Private/Framework/Application/AnalogCursor.cpp
+++ b/Engine/Source/Runtime/Slate/Private/Framework/Application/AnalogCursor.cpp
@@ -280,7 +280,7 @@ void FAnalogCursor::UpdateCursorPosition(FSlateApplication& SlateApp, TSharedRef
 
 		//create a new mouse event
 		FPointerEvent MouseEvent(
-			0,
+			SlateApp.CursorPointerIndex,
 			UpdatedPosition,
 			OldPosition,
 			SlateApp.GetPressedMouseButtons(),
