Fix for deprecation warnings.

diff --git a/Engine/Source/Runtime/Slate/Private/Widgets/Input/SEditableLabel.cpp b/Engine/Source/Runtime/Slate/Private/Widgets/Input/SEditableLabel.cpp
--- a/Engine/Source/Runtime/Slate/Private/Widgets/Input/SEditableLabel.cpp
+++ b/Engine/Source/Runtime/Slate/Private/Widgets/Input/SEditableLabel.cpp
@@ -12,7 +12,7 @@
 
 /* SEditableLabel interface
  *****************************************************************************/
-
+PRAGMA_DISABLE_DEPRECATION_WARNINGS
 void SEditableLabel::Construct(const FArguments& InArgs)
 {
 	CanEditAttribute = InArgs._CanEdit;
@@ -137,5 +137,6 @@ FReply SEditableLabel::HandleTextBlockDoubleClicked()
 	return FReply::Handled().SetUserFocus(EditableText.ToSharedRef(), EFocusCause::Navigation);
 }
 
+PRAGMA_ENABLE_DEPRECATION_WARNINGS
 
 #undef LOCTEXT_NAMESPACE
