Disabled deprecation warnings in SEditableLabel.h, the class is fully deprecated and is causing warnings with internal code.

diff --git a/Engine/Source/Runtime/Slate/Public/Widgets/Input/SEditableLabel.h b/Engine/Source/Runtime/Slate/Public/Widgets/Input/SEditableLabel.h
--- a/Engine/Source/Runtime/Slate/Public/Widgets/Input/SEditableLabel.h
+++ b/Engine/Source/Runtime/Slate/Public/Widgets/Input/SEditableLabel.h
@@ -17,6 +17,7 @@
 class SEditableText;
 class STextBlock;
 
+PRAGMA_DISABLE_DEPRECATION_WARNINGS
 
 class DEPRECATED(4.21, "SEditableLabel is deprecated, please use SInlineEditableTextBlock instead") SLATE_API SEditableLabel
 	: public SCompoundWidget
@@ -115,3 +116,5 @@ private:
 	TSharedPtr<STextBlock> TextBlock;
 	TSharedPtr<SEditableText> EditableText;
 };
+
+PRAGMA_ENABLE_DEPRECATION_WARNINGS
\ No newline at end of file
