PR #5072: Fix SInlineEditableTextBlock OnTextBoxCommitted (Contributed by [REDACTED_PERSON])

diff --git a/Engine/Source/Runtime/Slate/Private/Widgets/Text/SInlineEditableTextBlock.cpp b/Engine/Source/Runtime/Slate/Private/Widgets/Text/SInlineEditableTextBlock.cpp
--- a/Engine/Source/Runtime/Slate/Private/Widgets/Text/SInlineEditableTextBlock.cpp
+++ b/Engine/Source/Runtime/Slate/Private/Widgets/Text/SInlineEditableTextBlock.cpp
@@ -264,9 +264,10 @@ void SInlineEditableTextBlock::OnTextBoxCommitted(const FText& InText, ETextComm
 {
 	if(InCommitType == ETextCommit::OnCleared)
 	{
+		FText SourceTextInEditMode = Text.Get();
 		CancelEditMode();
 		// Commit the name, certain actions might need to be taken by the bound function
-		OnTextCommittedDelegate.ExecuteIfBound(Text.Get(), InCommitType);
+		OnTextCommittedDelegate.ExecuteIfBound(SourceTextInEditMode, InCommitType);
 	}
 	else if(IsInEditMode())
 	{
