Added MouseWheel and Gesture API to the Input Preprocessor.

diff --git a/Engine/Source/Runtime/Slate/Private/Framework/Application/SlateApplication.cpp b/Engine/Source/Runtime/Slate/Private/Framework/Application/SlateApplication.cpp
--- a/Engine/Source/Runtime/Slate/Private/Framework/Application/SlateApplication.cpp
+++ b/Engine/Source/Runtime/Slate/Private/Framework/Application/SlateApplication.cpp
@@ -6062,6 +6062,12 @@ bool FSlateApplication::ProcessMouseWheelOrGestureEvent( FPointerEvent& InWheelE
 	}
 
 	SetLastUserInteractionTime(this->GetCurrentTime());
+
+	// Input preprocessors get the first chance at the input
+	if (InputPreProcessors.HandleMouseWheelOrGestureEvent(*this, InWheelEvent, InGestureEvent))
+	{
+		return true;
+	}
 	
 	// NOTE: We intentionally don't reset LastUserInteractionTimeForThrottling here so that the UI can be responsive while scrolling
 
@@ -7320,6 +7326,19 @@ bool FSlateApplication::InputPreProcessorsHelper::HandleMouseButtonUpEvent(FSlat
 	return false;
 }
 
+bool FSlateApplication::InputPreProcessorsHelper::HandleMouseWheelOrGestureEvent(FSlateApplication& SlateApp, const FPointerEvent& WheelEvent, const FPointerEvent* GestureEvent)
+{
+	for (TSharedPtr<IInputProcessor> InputPreProcessor : InputPreProcessorList)
+	{
+		if (InputPreProcessor->HandleMouseWheelOrGestureEvent(SlateApp, WheelEvent, GestureEvent))
+		{
+			return true;
+		}
+	}
+
+	return false;
+}
+
 bool FSlateApplication::InputPreProcessorsHelper::HandleMotionDetectedEvent(FSlateApplication& SlateApp, const FMotionEvent& MotionEvent)
 {
 	for (TSharedPtr<IInputProcessor> InputPreProcessor : InputPreProcessorList)
diff --git a/Engine/Source/Runtime/Slate/Public/Framework/Application/IInputProcessor.h b/Engine/Source/Runtime/Slate/Public/Framework/Application/IInputProcessor.h
--- a/Engine/Source/Runtime/Slate/Public/Framework/Application/IInputProcessor.h
+++ b/Engine/Source/Runtime/Slate/Public/Framework/Application/IInputProcessor.h
@@ -40,6 +40,9 @@ public:
 	/** Mouse button release */
 	virtual bool HandleMouseButtonUpEvent( FSlateApplication& SlateApp, const FPointerEvent& MouseEvent) { return false; }
 
+	/** Mouse wheel input */
+	virtual bool HandleMouseWheelOrGestureEvent(FSlateApplication& SlateApp, const FPointerEvent& InWheelEvent, const FPointerEvent* InGestureEvent) { return false; }
+
 	/** Called when a motion-driven device has new input */
 	virtual bool HandleMotionDetectedEvent(FSlateApplication& SlateApp, const FMotionEvent& MotionEvent) { return false; };
 };
\ No newline at end of file
diff --git a/Engine/Source/Runtime/Slate/Public/Framework/Application/SlateApplication.h b/Engine/Source/Runtime/Slate/Public/Framework/Application/SlateApplication.h
--- a/Engine/Source/Runtime/Slate/Public/Framework/Application/SlateApplication.h
+++ b/Engine/Source/Runtime/Slate/Public/Framework/Application/SlateApplication.h
@@ -2168,6 +2168,7 @@ private:
 		bool HandleMouseMoveEvent(FSlateApplication& SlateApp, const FPointerEvent& MouseEvent);
 		bool HandleMouseButtonDownEvent(FSlateApplication& SlateApp, const FPointerEvent& MouseEvent);
 		bool HandleMouseButtonUpEvent(FSlateApplication& SlateApp, const FPointerEvent& MouseEvent);
+		bool HandleMouseWheelOrGestureEvent(FSlateApplication& SlateApp, const FPointerEvent& WheelEvent, const FPointerEvent* GestureEvent);
 		bool HandleMotionDetectedEvent(FSlateApplication& SlateApp, const FMotionEvent& MotionEvent);
 
 		/**
