Protect against rare window creation issue.
Hopeful work around for rare crash.

diff --git a/Engine/Source/Runtime/Slate/Private/Widgets/Input/SMenuAnchor.cpp b/Engine/Source/Runtime/Slate/Private/Widgets/Input/SMenuAnchor.cpp
--- a/Engine/Source/Runtime/Slate/Private/Widgets/Input/SMenuAnchor.cpp
+++ b/Engine/Source/Runtime/Slate/Private/Widgets/Input/SMenuAnchor.cpp
@@ -365,11 +365,15 @@ void SMenuAnchor::SetIsOpen( bool InIsOpen, const bool bFocusMenu, const int32 F
 						{
 							// Open the pop-up
 							TSharedPtr<IMenu> NewMenu = FSlateApplication::Get().PushMenu(AsShared(), MyWidgetPath, MenuContentRef, NewPosition, TransitionEffect, bFocusMenu, MyGeometry.GetLocalSize(), MethodInUse.GetPopupMethod(), bIsCollapsedByParent);
+							
+							if (ensure(NewMenu.IsValid()))
+							{
+								check(NewMenu->GetOwnedWindow().IsValid());
 
-							PopupMenuPtr = NewMenu;
-							check(NewMenu.IsValid() && NewMenu->GetOwnedWindow().IsValid());
-							NewMenu->GetOnMenuDismissed().AddSP(this, &SMenuAnchor::OnMenuClosed);
-							PopupWindowPtr = NewMenu->GetOwnedWindow();
+								PopupMenuPtr = NewMenu;
+								NewMenu->GetOnMenuDismissed().AddSP(this, &SMenuAnchor::OnMenuClosed);
+								PopupWindowPtr = NewMenu->GetOwnedWindow();
+							}
 						}
 						else
 						{
