Scrollbox keyboard navigation will now recursively search each element for a focusable widget instead of checking the top-level widget and moving on

diff --git a/Engine/Source/Runtime/Slate/Private/Widgets/Layout/SScrollBox.cpp b/Engine/Source/Runtime/Slate/Private/Widgets/Layout/SScrollBox.cpp
--- a/Engine/Source/Runtime/Slate/Private/Widgets/Layout/SScrollBox.cpp
+++ b/Engine/Source/Runtime/Slate/Private/Widgets/Layout/SScrollBox.cpp
@@ -997,8 +997,8 @@ FNavigationReply SScrollBox::OnNavigation(const FGeometry& MyGeometry, const FNa
 			// Search in the direction we need to move for the next focusable child of the scrollbox.
 			for ( int32 ChildIndex = FocusedChildIndex + FocusedChildDirection; ChildIndex >= 0 && ChildIndex < Children.Num(); ChildIndex += FocusedChildDirection )
 			{
-				TSharedRef<SWidget> PossiblyFocusableChild = Children[ChildIndex].GetWidget();
-				if ( PossiblyFocusableChild->SupportsKeyboardFocus() )
+				TSharedPtr<SWidget> PossiblyFocusableChild = GetKeyboardFocusableWidget(Children[ChildIndex].GetWidget());
+				if ( PossiblyFocusableChild.IsValid() )
 				{
 					NextFocusableChild = PossiblyFocusableChild;
 					break;
@@ -1017,6 +1017,28 @@ FNavigationReply SScrollBox::OnNavigation(const FGeometry& MyGeometry, const FNa
 	return SCompoundWidget::OnNavigation(MyGeometry, InNavigationEvent);
 }
 
+TSharedPtr<SWidget> SScrollBox::GetKeyboardFocusableWidget(TSharedPtr<SWidget> InWidget)
+{
+	if (InWidget->SupportsKeyboardFocus())
+	{
+		return InWidget;
+	}
+	else
+	{
+		FChildren* Children = InWidget->GetChildren();
+		for (int32 i = 0; i < Children->Num(); ++i)
+		{
+			TSharedPtr<SWidget> ChildWidget = Children->GetChildAt(i);
+			TSharedPtr<SWidget> FoucusableWidget = GetKeyboardFocusableWidget(ChildWidget);
+			if (FoucusableWidget.IsValid())
+			{
+				return FoucusableWidget;
+			}
+		}
+	}
+	return nullptr;
+}
+
 int32 SScrollBox::OnPaint( const FPaintArgs& Args, const FGeometry& AllottedGeometry, const FSlateRect& MyCullingRect, FSlateWindowElementList& OutDrawElements, int32 LayerId, const FWidgetStyle& InWidgetStyle, bool bParentEnabled ) const
 {
 	int32 NewLayerId = SCompoundWidget::OnPaint( Args, AllottedGeometry, MyCullingRect, OutDrawElements, LayerId, InWidgetStyle, bParentEnabled );
diff --git a/Engine/Source/Runtime/Slate/Public/Widgets/Layout/SScrollBox.h b/Engine/Source/Runtime/Slate/Public/Widgets/Layout/SScrollBox.h
--- a/Engine/Source/Runtime/Slate/Public/Widgets/Layout/SScrollBox.h
+++ b/Engine/Source/Runtime/Slate/Public/Widgets/Layout/SScrollBox.h
@@ -282,6 +282,8 @@ private:
 
 	void EndInertialScrolling();
 
+	TSharedPtr<SWidget> GetKeyboardFocusableWidget(TSharedPtr<SWidget> InWidget);
+
 private:
 
 	/** The panel which stacks the child slots */
