Fixed Slider widget not restoring cursor after dragging when a custom cursor is set

diff --git a/Engine/Source/Runtime/Slate/Private/Widgets/Input/SSlider.cpp b/Engine/Source/Runtime/Slate/Private/Widgets/Input/SSlider.cpp
--- a/Engine/Source/Runtime/Slate/Private/Widgets/Input/SSlider.cpp
+++ b/Engine/Source/Runtime/Slate/Private/Widgets/Input/SSlider.cpp
@@ -244,6 +244,7 @@ FReply SSlider::OnMouseButtonDown( const FGeometry& MyGeometry, const FPointerEv
 {
 	if ((MouseEvent.GetEffectingButton() == EKeys::LeftMouseButton) && !IsLocked())
 	{
+		CachedCursor = Cursor.Get().Get(EMouseCursor::Default);
 		OnMouseCaptureBegin.ExecuteIfBound();
 		CommitValue(PositionToValue(MyGeometry, MouseEvent.GetLastScreenSpacePosition()));
 		
@@ -260,7 +261,7 @@ FReply SSlider::OnMouseButtonUp( const FGeometry& MyGeometry, const FPointerEven
 {
 	if ((MouseEvent.GetEffectingButton() == EKeys::LeftMouseButton) && HasMouseCaptureByUser(MouseEvent.GetUserIndex(), MouseEvent.GetPointerIndex()))
 	{
-		SetCursor(EMouseCursor::Default);
+		SetCursor(CachedCursor);
 		OnMouseCaptureEnd.ExecuteIfBound();
 		
 		// Release capture for controller/keyboard when switching to mouse.
@@ -292,6 +293,7 @@ FReply SSlider::OnTouchStarted(const FGeometry& MyGeometry, const FPointerEvent&
 {
 	if (!IsLocked())
 	{
+		CachedCursor = Cursor.Get().Get(EMouseCursor::Default);
 		OnMouseCaptureBegin.ExecuteIfBound();
 		CommitValue(PositionToValue(MyGeometry, InTouchEvent.GetLastScreenSpacePosition()));
 
@@ -323,7 +325,7 @@ FReply SSlider::OnTouchEnded(const FGeometry& MyGeometry, const FPointerEvent& I
 {
 	if (HasMouseCaptureByUser(InTouchEvent.GetUserIndex(), InTouchEvent.GetPointerIndex()))
 	{
-		SetCursor(EMouseCursor::Default);
+		SetCursor(CachedCursor);
 		OnMouseCaptureEnd.ExecuteIfBound();
 
 		// Release capture for controller/keyboard when switching to mouse.
diff --git a/Engine/Source/Runtime/Slate/Public/Widgets/Input/SSlider.h b/Engine/Source/Runtime/Slate/Public/Widgets/Input/SSlider.h
--- a/Engine/Source/Runtime/Slate/Public/Widgets/Input/SSlider.h
+++ b/Engine/Source/Runtime/Slate/Public/Widgets/Input/SSlider.h
@@ -197,6 +197,9 @@ protected:
 	// Holds the slider's current value.
 	TAttribute<float> ValueAttribute;
 
+	// Holds the initial cursor in case a custom cursor has been specified, so we can restore it after dragging the slider
+	EMouseCursor::Type CachedCursor;
+
 	/** Holds the amount to adjust the value by when using a controller or keyboard */
 	TAttribute<float> StepSize;
 
