Input Key Selector now behaves like a button while not selecting.

diff --git a/Engine/Source/Runtime/Slate/Private/Widgets/Input/SInputKeySelector.cpp b/Engine/Source/Runtime/Slate/Private/Widgets/Input/SInputKeySelector.cpp
--- a/Engine/Source/Runtime/Slate/Private/Widgets/Input/SInputKeySelector.cpp
+++ b/Engine/Source/Runtime/Slate/Private/Widgets/Input/SInputKeySelector.cpp
@@ -151,10 +151,17 @@ FReply SInputKeySelector::OnPreviewKeyDown( const FGeometry& MyGeometry, const F
 
 FReply SInputKeySelector::OnKeyDown(const FGeometry& MyGeometry, const FKeyEvent& InKeyEvent)
 {
-	if (!bIsSelectingKey && SelectedKey.IsSet() && SelectedKey.Get().Key.IsValid() && (bAllowGamepadKeys && InKeyEvent.GetKey() == EKeys::Gamepad_FaceButton_Left))
+	if (!bIsSelectingKey)
 	{
-		SelectedKey = FInputChord();
-		return FReply::Handled();
+		if (SelectedKey.IsSet() && SelectedKey.Get().Key.IsValid() && (bAllowGamepadKeys && InKeyEvent.GetKey() == EKeys::Gamepad_FaceButton_Left))
+		{
+			SelectedKey = FInputChord();
+			return FReply::Handled();
+		}
+		else if (Button.IsValid())
+		{
+			return Button->OnKeyDown(MyGeometry, InKeyEvent);
+		}
 	}
 	return SCompoundWidget::OnKeyDown(MyGeometry, InKeyEvent);
 }
@@ -187,6 +194,10 @@ FReply SInputKeySelector::OnKeyUp( const FGeometry& MyGeometry, const FKeyEvent&
 			ModifierKey == EModifierKey::Command);
 		return FReply::Handled();
 	}
+	else if (!bIsSelectingKey && Button.IsValid())
+	{
+		return Button->OnKeyUp(MyGeometry, InKeyEvent);
+	}
 
 	return SCompoundWidget::OnKeyUp( MyGeometry, InKeyEvent );
 }
