Don't force refresh when updating text style

diff --git a/Engine/Source/Runtime/Slate/Private/Widgets/Input/SEditableText.cpp b/Engine/Source/Runtime/Slate/Private/Widgets/Input/SEditableText.cpp
--- a/Engine/Source/Runtime/Slate/Private/Widgets/Input/SEditableText.cpp
+++ b/Engine/Source/Runtime/Slate/Private/Widgets/Input/SEditableText.cpp
@@ -443,7 +443,7 @@ void SEditableText::SynchronizeTextStyle()
 	if (bTextStyleChanged)
 	{
 		EditableTextLayout->SetTextStyle(NewTextStyle);
-		EditableTextLayout->Refresh();
+		EditableTextLayout->ForceRefreshTextLayout(EditableTextLayout->GetEditableText());
 	}
 }
 
