Fix CIS deprecation warning

diff --git a/Engine/Source/Runtime/Slate/Public/Widgets/Input/SEditableLabel.h b/Engine/Source/Runtime/Slate/Public/Widgets/Input/SEditableLabel.h
--- a/Engine/Source/Runtime/Slate/Public/Widgets/Input/SEditableLabel.h
+++ b/Engine/Source/Runtime/Slate/Public/Widgets/Input/SEditableLabel.h
@@ -19,7 +19,7 @@ class STextBlock;
 
 PRAGMA_DISABLE_DEPRECATION_WARNINGS
 
-class DEPRECATED(4.21, "SEditableLabel is deprecated, please use SInlineEditableTextBlock instead") SLATE_API SEditableLabel
+class UE_DEPRECATED(4.21, "SEditableLabel is deprecated, please use SInlineEditableTextBlock instead") SLATE_API SEditableLabel
 	: public SCompoundWidget
 {
 public:
