Attempt CIS fix

diff --git a/Engine/Source/Runtime/Slate/Public/Widgets/Input/SEditableLabel.h b/Engine/Source/Runtime/Slate/Public/Widgets/Input/SEditableLabel.h
--- a/Engine/Source/Runtime/Slate/Public/Widgets/Input/SEditableLabel.h
+++ b/Engine/Source/Runtime/Slate/Public/Widgets/Input/SEditableLabel.h
@@ -13,13 +13,14 @@
 #include "Styling/SlateTypes.h"
 #include "Styling/CoreStyle.h"
 #include "Framework/SlateDelegates.h"
+#include "Misc/CoreMiscDefines.h"
 
 class SEditableText;
 class STextBlock;
 
 PRAGMA_DISABLE_DEPRECATION_WARNINGS
 
-class UE_DEPRECATED(4.21, "SEditableLabel is deprecated, please use SInlineEditableTextBlock instead") SLATE_API SEditableLabel
+class SLATE_API UE_DEPRECATED(4.21, "SEditableLabel is deprecated, please use SInlineEditableTextBlock instead") SEditableLabel
 	: public SCompoundWidget
 {
 public:
