PR #5323: Preserve opened tab when closing another one (Contributed by [REDACTED_PERSON])

diff --git a/Engine/Source/Runtime/Slate/Private/Framework/Docking/SDockingTabWell.cpp b/Engine/Source/Runtime/Slate/Private/Framework/Docking/SDockingTabWell.cpp
--- a/Engine/Source/Runtime/Slate/Private/Framework/Docking/SDockingTabWell.cpp
+++ b/Engine/Source/Runtime/Slate/Private/Framework/Docking/SDockingTabWell.cpp
@@ -536,6 +536,8 @@ void SDockingTabWell::RemoveAndDestroyTab(const TSharedRef<SDockTab>& TabToRemov
 
 		// Remove the old tab from the list of tabs and activate the new tab.
 		{
+			int32 OldTabIndex = ForegroundTabIndex;
+
 			BringTabToFront(TabIndex);
 			Tabs.RemoveAt(TabIndex);
 			// We no longer have a tab in the foreground.
@@ -543,7 +545,14 @@ void SDockingTabWell::RemoveAndDestroyTab(const TSharedRef<SDockTab>& TabToRemov
 			ForegroundTabIndex = INDEX_NONE;
 
 			// Now bring the last tab that we were on to the foreground
-			BringTabToFront(FMath::Max(TabIndex-1, 0));
+			if (OldTabIndex == INDEX_NONE || TabIndex <= OldTabIndex)
+			{
+				BringTabToFront(FMath::Max(OldTabIndex - 1, 0));
+			}
+			else
+			{
+				BringTabToFront(OldTabIndex);
+			}
 		}
 		
 		if ( ensure(ParentTabStack.IsValid()) )
