Fixed a crash caused by mouse capture being instigated by a keyboard event

diff --git a/Engine/Source/Runtime/Slate/Private/Framework/Application/SlateApplication.cpp b/Engine/Source/Runtime/Slate/Private/Framework/Application/SlateApplication.cpp
--- a/Engine/Source/Runtime/Slate/Private/Framework/Application/SlateApplication.cpp
+++ b/Engine/Source/Runtime/Slate/Private/Framework/Application/SlateApplication.cpp
@@ -3257,8 +3257,18 @@ void FSlateApplication::ProcessReply( const FWidgetPath& CurrentEventPath, const
 							{
 								if (SomeWidgetPreviouslyUnderCursor != RequestedMouseCaptor)
 								{
-									// Note that the event's pointer position is not translated.
-									SomeWidgetPreviouslyUnderCursor->OnMouseLeave(*InMouseEvent);
+									// It's possible for mouse event to be null if we end up here from a keyboard event. If so, we should synthesize an event
+									if (InMouseEvent)
+									{
+										// Note that the event's pointer position is not translated.
+										SomeWidgetPreviouslyUnderCursor->OnMouseLeave(*InMouseEvent);
+									}
+									else
+									{
+										const FPointerEvent& SimulatedPointer = FPointerEvent();
+										SomeWidgetPreviouslyUnderCursor->OnMouseLeave(SimulatedPointer);
+									}
+									
 
 #if WITH_SLATE_DEBUGGING
 									FSlateDebugging::BroadcastInputEvent(ESlateDebuggingInputEvent::MouseLeave, SomeWidgetPreviouslyUnderCursor);
