Disable window title bar ticking

diff --git a/Engine/Source/Runtime/Slate/Public/Framework/Application/SWindowTitleBar.h b/Engine/Source/Runtime/Slate/Public/Framework/Application/SWindowTitleBar.h
--- a/Engine/Source/Runtime/Slate/Public/Framework/Application/SWindowTitleBar.h
+++ b/Engine/Source/Runtime/Slate/Public/Framework/Application/SWindowTitleBar.h
@@ -52,6 +52,8 @@ class SAppIconWidget
 	
 	void Construct( const FArguments& Args )
 	{
+		SetCanTick(false);
+
 		this->ChildSlot
 			[
 				SNew(SVerticalBox)
@@ -105,6 +107,8 @@ public:
 	 */
 	void Construct( const FArguments& InArgs, const TSharedRef<SWindow>& InWindow, const TSharedPtr<SWidget>& InCenterContent, EHorizontalAlignment InCenterContentAlignment )
 	{
+		SetCanTick(false);
+
 		OwnerWindowPtr = InWindow;
 		Style = InArgs._Style;
 		ShowAppIcon = InArgs._ShowAppIcon;
