Slate: Added proper volatility support for SButton and SBorder.

diff --git a/Engine/Source/Runtime/Slate/Private/Widgets/Input/SButton.cpp b/Engine/Source/Runtime/Slate/Private/Widgets/Input/SButton.cpp
--- a/Engine/Source/Runtime/Slate/Private/Widgets/Input/SButton.cpp
+++ b/Engine/Source/Runtime/Slate/Private/Widgets/Input/SButton.cpp
@@ -383,6 +383,12 @@ bool SButton::IsInteractable() const
 	return IsEnabled();
 }
 
+bool SButton::ComputeVolatility() const
+{
+	// Note: we need to be careful with button volatility.  The parent SBorder class always has bound delegates to button but the following are the only thing that actually would be bound that would not be caught by an Invalidate call alone
+	return ContentScale.IsBound() || DesiredSizeScale.IsBound() || BorderBackgroundColor.IsBound() || ContentPadding.IsBound() || ForegroundColor.IsBound();
+}
+
 bool SButton::IsPreciseTapOrClick(const FPointerEvent& MouseEvent) const
 {
 	return ( TouchMethod == EButtonTouchMethod::PreciseTap && MouseEvent.IsTouchEvent() ) ||
diff --git a/Engine/Source/Runtime/Slate/Private/Widgets/Layout/SBorder.cpp b/Engine/Source/Runtime/Slate/Private/Widgets/Layout/SBorder.cpp
--- a/Engine/Source/Runtime/Slate/Private/Widgets/Layout/SBorder.cpp
+++ b/Engine/Source/Runtime/Slate/Private/Widgets/Layout/SBorder.cpp
@@ -102,6 +102,11 @@ int32 SBorder::OnPaint( const FPaintArgs& Args, const FGeometry& AllottedGeometr
 	return SCompoundWidget::OnPaint(Args, AllottedGeometry, MyCullingRect, OutDrawElements, LayerId, InWidgetStyle, bEnabled );
 }
 
+bool SBorder::ComputeVolatility() const
+{
+	return BorderImage.IsBound() || BorderBackgroundColor.IsBound() || DesiredSizeScale.IsBound() || ShowDisabledEffect.IsBound();
+}
+
 FVector2D SBorder::ComputeDesiredSize(float LayoutScaleMultiplier) const
 {
 	return DesiredSizeScale.Get() * SCompoundWidget::ComputeDesiredSize(LayoutScaleMultiplier);
diff --git a/Engine/Source/Runtime/Slate/Public/Widgets/Input/SButton.h b/Engine/Source/Runtime/Slate/Public/Widgets/Input/SButton.h
--- a/Engine/Source/Runtime/Slate/Public/Widgets/Input/SButton.h
+++ b/Engine/Source/Runtime/Slate/Public/Widgets/Input/SButton.h
@@ -177,6 +177,7 @@ public:
 	virtual void OnMouseLeave( const FPointerEvent& MouseEvent ) override;
 	virtual void OnMouseCaptureLost(const FCaptureLostEvent& CaptureLostEvent) override;
 	virtual bool IsInteractable() const override;
+	virtual bool ComputeVolatility() const override;
 	// SWidget
 
 protected:
diff --git a/Engine/Source/Runtime/Slate/Public/Widgets/Layout/SBorder.h b/Engine/Source/Runtime/Slate/Public/Widgets/Layout/SBorder.h
--- a/Engine/Source/Runtime/Slate/Public/Widgets/Layout/SBorder.h
+++ b/Engine/Source/Runtime/Slate/Public/Widgets/Layout/SBorder.h
@@ -122,6 +122,7 @@ public:
 public:
 	// SWidget interface
 	virtual int32 OnPaint( const FPaintArgs& Args, const FGeometry& AllottedGeometry, const FSlateRect& MyCullingRect, FSlateWindowElementList& OutDrawElements, int32 LayerId, const FWidgetStyle& InWidgetStyle, bool bParentEnabled ) const override;
+	virtual bool ComputeVolatility() const override;
 	// End of SWidget interface
 
 protected:
