Prevent navigation on invalid Widget
#rnx

diff --git a/Engine/Source/Runtime/Slate/Private/Widgets/LayerManager/STooltipPresenter.cpp b/Engine/Source/Runtime/Slate/Private/Widgets/LayerManager/STooltipPresenter.cpp
--- a/Engine/Source/Runtime/Slate/Private/Widgets/LayerManager/STooltipPresenter.cpp
+++ b/Engine/Source/Runtime/Slate/Private/Widgets/LayerManager/STooltipPresenter.cpp
@@ -40,7 +40,14 @@ void STooltipPresenter::OnArrangeChildren(const FGeometry& AllottedGeometry, FAr
 
 FVector2D STooltipPresenter::ComputeDesiredSize( float ) const
 {
-	return ChildSlot.GetWidget()->GetDesiredSize();
+	if (ChildSlot.GetWidget().IsValid())
+	{
+		return ChildSlot.GetWidget()->GetDesiredSize();
+	}
+	else
+	{
+		return FVector2D(0, 0);
+	}
 }
 
 FChildren* STooltipPresenter::GetChildren()
