Revert CL 4744377
TSharedRef validity test is intentionally private because shared references must always be valid.
#rnx

diff --git a/Engine/Source/Runtime/Slate/Private/Widgets/LayerManager/STooltipPresenter.cpp b/Engine/Source/Runtime/Slate/Private/Widgets/LayerManager/STooltipPresenter.cpp
--- a/Engine/Source/Runtime/Slate/Private/Widgets/LayerManager/STooltipPresenter.cpp
+++ b/Engine/Source/Runtime/Slate/Private/Widgets/LayerManager/STooltipPresenter.cpp
@@ -40,14 +40,7 @@ void STooltipPresenter::OnArrangeChildren(const FGeometry& AllottedGeometry, FAr
 
 FVector2D STooltipPresenter::ComputeDesiredSize( float ) const
 {
-	if (ChildSlot.GetWidget().IsValid())
-	{
-		return ChildSlot.GetWidget()->GetDesiredSize();
-	}
-	else
-	{
-		return FVector2D(0, 0);
-	}
+	return ChildSlot.GetWidget()->GetDesiredSize();
 }
 
 FChildren* STooltipPresenter::GetChildren()
