Skip highlight text string search when there is nothing to search for.

diff --git a/Engine/Source/Runtime/Slate/Private/Widgets/Text/SlateTextBlockLayout.cpp b/Engine/Source/Runtime/Slate/Private/Widgets/Text/SlateTextBlockLayout.cpp
--- a/Engine/Source/Runtime/Slate/Private/Widgets/Text/SlateTextBlockLayout.cpp
+++ b/Engine/Source/Runtime/Slate/Private/Widgets/Text/SlateTextBlockLayout.cpp
@@ -209,24 +209,27 @@ void FSlateTextBlockLayout::UpdateTextHighlights(const FText& InHighlightText)
 	const TArray< FTextLayout::FLineModel >& LineModels = TextLayout->GetLineModels();
 
 	TArray<FTextRunRenderer> TextHighlights;
-	for(int32 LineIndex = 0; LineIndex < LineModels.Num(); ++LineIndex)
+	if (HighlightTextString.Len() > 0)
 	{
-		const FTextLayout::FLineModel& LineModel = LineModels[LineIndex];
-
-		int32 FindBegin = 0;
-		int32 CurrentHighlightBegin;
-		const int32 TextLength = LineModel.Text->Len();
-		while(FindBegin < TextLength && (CurrentHighlightBegin = LineModel.Text->Find(HighlightTextString, ESearchCase::IgnoreCase, ESearchDir::FromStart, FindBegin)) != INDEX_NONE)
+		for(int32 LineIndex = 0; LineIndex < LineModels.Num(); ++LineIndex)
 		{
-			FindBegin = CurrentHighlightBegin + HighlightTextLength;
+			const FTextLayout::FLineModel& LineModel = LineModels[LineIndex];
 
-			if(TextHighlights.Num() > 0 && TextHighlights.Last().LineIndex == LineIndex && TextHighlights.Last().Range.EndIndex == CurrentHighlightBegin)
-			{
-				TextHighlights[TextHighlights.Num() - 1] = FTextRunRenderer(LineIndex, FTextRange(TextHighlights.Last().Range.BeginIndex, FindBegin), TextHighlighter.ToSharedRef());
-			}
-			else
+			int32 FindBegin = 0;
+			int32 CurrentHighlightBegin;
+			const int32 TextLength = LineModel.Text->Len();
+			while(FindBegin < TextLength && (CurrentHighlightBegin = LineModel.Text->Find(HighlightTextString, ESearchCase::IgnoreCase, ESearchDir::FromStart, FindBegin)) != INDEX_NONE)
 			{
-				TextHighlights.Add(FTextRunRenderer(LineIndex, FTextRange(CurrentHighlightBegin, FindBegin), TextHighlighter.ToSharedRef()));
+				FindBegin = CurrentHighlightBegin + HighlightTextLength;
+
+				if(TextHighlights.Num() > 0 && TextHighlights.Last().LineIndex == LineIndex && TextHighlights.Last().Range.EndIndex == CurrentHighlightBegin)
+				{
+					TextHighlights[TextHighlights.Num() - 1] = FTextRunRenderer(LineIndex, FTextRange(TextHighlights.Last().Range.BeginIndex, FindBegin), TextHighlighter.ToSharedRef());
+				}
+				else
+				{
+					TextHighlights.Add(FTextRunRenderer(LineIndex, FTextRange(CurrentHighlightBegin, FindBegin), TextHighlighter.ToSharedRef()));
+				}
 			}
 		}
 	}
