Fixed pinned viewport commands being searchable

Slate Changes: Allowed menu builder to be optionally searchable. Added a default parameter to control this.

diff --git a/Engine/Source/Editor/PinnedCommandList/Private/SPinnedCommandList.cpp b/Engine/Source/Editor/PinnedCommandList/Private/SPinnedCommandList.cpp
--- a/Engine/Source/Editor/PinnedCommandList/Private/SPinnedCommandList.cpp
+++ b/Engine/Source/Editor/PinnedCommandList/Private/SPinnedCommandList.cpp
@@ -97,8 +97,8 @@ public:
 
 		// Using a menu builder here is slightly wasteful, but as we cant construct
 		// a menu item individually it will have to do for now.
-		const bool bInShouldCloseWindowAfterMenuSelection = true;
-		FMenuBuilder MenuBuilder(bInShouldCloseWindowAfterMenuSelection, CommandList.Pin());
+		const bool bInShouldCloseWindowAfterMenuSelection = false;
+		FMenuBuilder MenuBuilder(bInShouldCloseWindowAfterMenuSelection, CommandList.Pin(), nullptr, false, InArgs._StyleSet, false);
 		MenuBuilder.SetStyle(InArgs._StyleSet, InArgs._StyleName);
 
 		if(CommandInfo.IsValid() && CommandList.IsValid())
@@ -136,7 +136,7 @@ public:
 					]
 				];
 
-			MenuBuilder.AddWidget(CustomWidgetContainer, FText(), true);
+			MenuBuilder.AddWidget(CustomWidgetContainer, FText(), true, false);
 		}
 		else
 		{
diff --git a/Engine/Source/Runtime/Slate/Private/Framework/MultiBox/MultiBoxBuilder.cpp b/Engine/Source/Runtime/Slate/Private/Framework/MultiBox/MultiBoxBuilder.cpp
--- a/Engine/Source/Runtime/Slate/Private/Framework/MultiBox/MultiBoxBuilder.cpp
+++ b/Engine/Source/Runtime/Slate/Private/Framework/MultiBox/MultiBoxBuilder.cpp
@@ -168,8 +168,8 @@ void FBaseMenuBuilder::AddMenuEntry( const FUIAction& UIAction, const TSharedRef
 
 TSharedRef< class SWidget > FMenuBuilder::MakeWidget( FMultiBox::FOnMakeMultiBoxBuilderOverride* InMakeMultiBoxBuilderOverride /* = nullptr */ )
 {
-	// Make menu builders searchable
-	return MultiBox->MakeWidget( true, InMakeMultiBoxBuilderOverride );
+	// Make menu builders searchable (by default)
+	return MultiBox->MakeWidget( bSearchable, InMakeMultiBoxBuilderOverride );
 }
 
 void FMenuBuilder::BeginSection( FName InExtensionHook, const TAttribute< FText >& InHeadingText )
diff --git a/Engine/Source/Runtime/Slate/Public/Framework/MultiBox/MultiBoxBuilder.h b/Engine/Source/Runtime/Slate/Public/Framework/MultiBox/MultiBoxBuilder.h
--- a/Engine/Source/Runtime/Slate/Public/Framework/MultiBox/MultiBoxBuilder.h
+++ b/Engine/Source/Runtime/Slate/Public/Framework/MultiBox/MultiBoxBuilder.h
@@ -210,11 +210,15 @@ public:
 	 * @param	InCommandList	The action list that maps command infos to delegates that should be called for each command associated with a multiblock widget
 	 * @param	bInCloseSelfOnly	True if clicking on a menu entry closes itself only and its children but not the entire stack
 	 */
-	FMenuBuilder( const bool bInShouldCloseWindowAfterMenuSelection, TSharedPtr< const FUICommandList > InCommandList, TSharedPtr<FExtender> InExtender = TSharedPtr<FExtender>(), const bool bCloseSelfOnly = false, const ISlateStyle* InStyleSet = &FCoreStyle::Get() )
+	FMenuBuilder( const bool bInShouldCloseWindowAfterMenuSelection, TSharedPtr< const FUICommandList > InCommandList, TSharedPtr<FExtender> InExtender = TSharedPtr<FExtender>(), const bool bCloseSelfOnly = false, const ISlateStyle* InStyleSet = &FCoreStyle::Get(), bool bInSearchable = true )
 		: FBaseMenuBuilder( EMultiBoxType::Menu, bInShouldCloseWindowAfterMenuSelection, InCommandList, bCloseSelfOnly, InExtender, InStyleSet )
 		, bSectionNeedsToBeApplied(false)
+		, bSearchable(bInSearchable)
 	{
-		AddSearchWidget();
+		if(bSearchable)
+		{
+			AddSearchWidget();
+		}
 	}
 
 	/**
@@ -317,6 +321,9 @@ private:
 	
 	/** True if there is a pending section that needs to be applied */
 	bool bSectionNeedsToBeApplied;
+
+	/** Whether this menu is searchable */
+	bool bSearchable;
 };
 
 
