Fix shadowed variable warning

diff --git a/Engine/Source/Runtime/Slate/Private/Framework/MultiBox/MultiBoxBuilder.cpp b/Engine/Source/Runtime/Slate/Private/Framework/MultiBox/MultiBoxBuilder.cpp
--- a/Engine/Source/Runtime/Slate/Private/Framework/MultiBox/MultiBoxBuilder.cpp
+++ b/Engine/Source/Runtime/Slate/Private/Framework/MultiBox/MultiBoxBuilder.cpp
@@ -279,12 +279,12 @@ void FMenuBuilder::AddWrapperSubMenu( const FText& InMenuLabel, const FText& InT
 	MultiBox->AddMultiBlock( NewMenuEntryBlock );
 }
 
-void FMenuBuilder::AddWidget( TSharedRef<SWidget> InWidget, const FText& Label, bool bNoIndent, bool bSearchable )
+void FMenuBuilder::AddWidget( TSharedRef<SWidget> InWidget, const FText& Label, bool bNoIndent, bool bInSearchable )
 {
 	ApplySectionBeginning();
 
 	TSharedRef< FWidgetBlock > NewWidgetBlock(new FWidgetBlock( InWidget, Label, bNoIndent ));
-	NewWidgetBlock->SetSearchable( bSearchable );
+	NewWidgetBlock->SetSearchable( bInSearchable );
 
 	MultiBox->AddMultiBlock( NewWidgetBlock );
 }
