Slate - Fixing an ensure from accessing a weak widget before checking if it's valid.

diff --git a/Engine/Source/Runtime/Slate/Private/Widgets/SWeakWidget.cpp b/Engine/Source/Runtime/Slate/Private/Widgets/SWeakWidget.cpp
--- a/Engine/Source/Runtime/Slate/Private/Widgets/SWeakWidget.cpp
+++ b/Engine/Source/Runtime/Slate/Private/Widgets/SWeakWidget.cpp
@@ -7,8 +7,8 @@
 
 SWeakWidget::SWeakWidget()
 	: WeakChild(this)
-{ }
-
+{
+}
 
 void SWeakWidget::Construct(const FArguments& InArgs)
 {
@@ -16,42 +16,44 @@ void SWeakWidget::Construct(const FArguments& InArgs)
 	WeakChild.AttachWidget( InArgs._PossiblyNullContent );
 }
 
-
 FVector2D SWeakWidget::ComputeDesiredSize( float ) const
-{	
-	TSharedRef<SWidget> ReferencedWidget = WeakChild.GetWidget();
-
-	if ( ReferencedWidget != SNullWidget::NullWidget && ReferencedWidget->GetVisibility() != EVisibility::Collapsed )
+{
+	if (WeakChild.Num() > 0)
 	{
-		return ReferencedWidget->GetDesiredSize();
+		TSharedRef<SWidget> ReferencedWidget = WeakChild.GetWidget();
+
+		if (ReferencedWidget != SNullWidget::NullWidget && ReferencedWidget->GetVisibility() != EVisibility::Collapsed)
+		{
+			return ReferencedWidget->GetDesiredSize();
+		}
 	}
 
 	return FVector2D::ZeroVector;
 }
 
-
 void SWeakWidget::OnArrangeChildren( const FGeometry& AllottedGeometry, FArrangedChildren& ArrangedChildren ) const
 {
-	// We just want to show the child that we are presenting. Always stretch it to occupy all of the space.
-	TSharedRef<SWidget> MyContent = WeakChild.GetWidget();
-
-	if( MyContent!=SNullWidget::NullWidget && ArrangedChildren.Accepts(MyContent->GetVisibility()) )
+	if (WeakChild.Num() > 0)
 	{
-		ArrangedChildren.AddWidget( AllottedGeometry.MakeChild(
-			MyContent,
-			FVector2D(0,0),
-			AllottedGeometry.GetLocalSize()
-			) );
+		// We just want to show the child that we are presenting. Always stretch it to occupy all of the space.
+		TSharedRef<SWidget> MyContent = WeakChild.GetWidget();
+
+		if (MyContent != SNullWidget::NullWidget && ArrangedChildren.Accepts(MyContent->GetVisibility()))
+		{
+			ArrangedChildren.AddWidget(AllottedGeometry.MakeChild(
+				MyContent,
+				FVector2D(0, 0),
+				AllottedGeometry.GetLocalSize()
+			));
+		}
 	}
 }
 
-
 FChildren* SWeakWidget::GetChildren()
 {
 	return &WeakChild;
 }
 
-
 int32 SWeakWidget::OnPaint( const FPaintArgs& Args, const FGeometry& AllottedGeometry, const FSlateRect& MyCullingRect, FSlateWindowElementList& OutDrawElements, int32 LayerId, const FWidgetStyle& InWidgetStyle, bool bParentEnabled ) const
 {
 	// Just draw the children.
@@ -76,14 +78,12 @@ int32 SWeakWidget::OnPaint( const FPaintArgs& Args, const FGeometry& AllottedGeo
 	return LayerId;
 }
 
-
 void SWeakWidget::SetContent(const TSharedRef<SWidget>& InWidget)
 {
 	WeakChild.AttachWidget( InWidget );
 }
 
-
 bool SWeakWidget::ChildWidgetIsValid() const
 {
-	return WeakChild.GetWidget() != SNullWidget::NullWidget;
+	return WeakChild.Num() > 0;
 }
