Adding "Down" touch click method to match mouse click types.

diff --git a/Engine/Source/Runtime/Slate/Private/Widgets/Input/SButton.cpp b/Engine/Source/Runtime/Slate/Private/Widgets/Input/SButton.cpp
--- a/Engine/Source/Runtime/Slate/Private/Widgets/Input/SButton.cpp
+++ b/Engine/Source/Runtime/Slate/Private/Widgets/Input/SButton.cpp
@@ -216,8 +216,10 @@ FReply SButton::OnMouseButtonDown( const FGeometry& MyGeometry, const FPointerEv
 	{
 		Press();
 		PressedScreenSpacePosition = MouseEvent.GetScreenSpacePosition();
+
+		EButtonClickMethod::Type InputClickMethod = GetClickMethodFromInputType(MouseEvent);
 		
-		if( ClickMethod == EButtonClickMethod::MouseDown )
+		if(InputClickMethod == EButtonClickMethod::MouseDown)
 		{
 			//get the reply from the execute function
 			Reply = OnClicked.IsBound() ? OnClicked.Execute() : FReply::Handled();
@@ -225,7 +227,7 @@ FReply SButton::OnMouseButtonDown( const FGeometry& MyGeometry, const FPointerEv
 			//You should ALWAYS handle the OnClicked event.
 			ensure(Reply.IsEventHandled() == true);
 		}
-		else if ( IsPreciseTapOrClick(MouseEvent) )
+		else if (InputClickMethod == EButtonClickMethod::PreciseClick)
 		{
 			// do not capture the pointer for precise taps or clicks
 			// 
@@ -252,7 +254,8 @@ FReply SButton::OnMouseButtonDoubleClick( const FGeometry& InMyGeometry, const F
 FReply SButton::OnMouseButtonUp( const FGeometry& MyGeometry, const FPointerEvent& MouseEvent )
 {
 	FReply Reply = FReply::Unhandled();
-	const bool bMustBePressed = ClickMethod == EButtonClickMethod::DownAndUp || IsPreciseTapOrClick(MouseEvent);
+	const EButtonClickMethod::Type InputClickMethod = GetClickMethodFromInputType(MouseEvent);
+	const bool bMustBePressed = InputClickMethod == EButtonClickMethod::DownAndUp || InputClickMethod == EButtonClickMethod::PreciseClick;
 	const bool bMeetsPressedRequirements = (!bMustBePressed || (bIsPressed && bMustBePressed));
 
 	if (bMeetsPressedRequirements && ( ( MouseEvent.GetEffectingButton() == EKeys::LeftMouseButton || MouseEvent.IsTouchEvent())))
@@ -261,7 +264,7 @@ FReply SButton::OnMouseButtonUp( const FGeometry& MyGeometry, const FPointerEven
 
 		if ( IsEnabled() )
 		{
-			if ( ClickMethod == EButtonClickMethod::MouseDown )
+			if (InputClickMethod == EButtonClickMethod::MouseDown )
 			{
 				// NOTE: If we're configured to click on mouse-down/precise-tap, then we never capture the mouse thus
 				//       may never receive an OnMouseButtonUp() call.  We make sure that our bIsPressed
@@ -279,11 +282,11 @@ FReply SButton::OnMouseButtonUp( const FGeometry& MyGeometry, const FPointerEven
 				if ( bEventOverButton )
 				{
 					// If we asked for a precise tap, all we need is for the user to have not moved their pointer very far.
-					const bool bTriggerForTouchEvent = IsPreciseTapOrClick(MouseEvent);
+					const bool bTriggerForTouchEvent = InputClickMethod == EButtonClickMethod::PreciseClick;
 
 					// If we were asked to allow the button to be clicked on mouse up, regardless of whether the user
 					// pressed the button down first, then we'll allow the click to proceed without an active capture
-					const bool bTriggerForMouseEvent = ( ClickMethod == EButtonClickMethod::MouseUp || HasMouseCapture() );
+					const bool bTriggerForMouseEvent = (InputClickMethod == EButtonClickMethod::MouseUp || HasMouseCapture() );
 
 					if ( ( bTriggerForTouchEvent || bTriggerForMouseEvent ) && OnClicked.IsBound() == true )
 					{
@@ -389,10 +392,27 @@ bool SButton::ComputeVolatility() const
 	return ContentScale.IsBound() || DesiredSizeScale.IsBound() || BorderBackgroundColor.IsBound() || ContentPadding.IsBound() || ForegroundColor.IsBound();
 }
 
+TEnumAsByte<EButtonClickMethod::Type> SButton::GetClickMethodFromInputType(const FPointerEvent& MouseEvent) const
+{
+	if (MouseEvent.IsTouchEvent())
+	{
+		switch (TouchMethod)
+		{
+			case EButtonTouchMethod::Down:
+				return EButtonClickMethod::MouseDown;
+			case EButtonTouchMethod::DownAndUp:
+				return EButtonClickMethod::DownAndUp;
+			case EButtonTouchMethod::PreciseTap:
+				return EButtonClickMethod::PreciseClick;
+		}
+	}
+
+	return ClickMethod;
+}
+
 bool SButton::IsPreciseTapOrClick(const FPointerEvent& MouseEvent) const
 {
-	return ( TouchMethod == EButtonTouchMethod::PreciseTap && MouseEvent.IsTouchEvent() ) ||
-		   ( ClickMethod == EButtonClickMethod::PreciseClick && !MouseEvent.IsTouchEvent() );
+	return GetClickMethodFromInputType(MouseEvent) == EButtonClickMethod::PreciseClick;
 }
 
 void SButton::PlayPressedSound() const
diff --git a/Engine/Source/Runtime/Slate/Public/Widgets/Input/SButton.h b/Engine/Source/Runtime/Slate/Public/Widgets/Input/SButton.h
--- a/Engine/Source/Runtime/Slate/Public/Widgets/Input/SButton.h
+++ b/Engine/Source/Runtime/Slate/Public/Widgets/Input/SButton.h
@@ -231,6 +231,9 @@ protected:
 	/** Release the button */
 	virtual void Release();
 
+	/** Utility function to translate other input click methods to regular ones. */
+	TEnumAsByte<EButtonClickMethod::Type> GetClickMethodFromInputType(const FPointerEvent& MouseEvent) const;
+
 	/** Utility function to determine if the incoming mouse event is for a precise tap or click */
 	bool IsPreciseTapOrClick(const FPointerEvent& MouseEvent) const;
 
diff --git a/Engine/Source/Runtime/SlateCore/Public/Types/SlateEnums.h b/Engine/Source/Runtime/SlateCore/Public/Types/SlateEnums.h
--- a/Engine/Source/Runtime/SlateCore/Public/Types/SlateEnums.h
+++ b/Engine/Source/Runtime/SlateCore/Public/Types/SlateEnums.h
@@ -51,6 +51,11 @@ namespace EButtonTouchMethod
 		/** Most buttons behave this way. */
 		DownAndUp,
 
+		/**
+		 * Click will be triggered immediately on touch down, and touch will not be captured.
+		 */
+		Down,
+
 		/**
 		 * Inside a list, buttons can only be clicked with precise tap.
 		 * Moving the pointer will scroll the list.
