Make sure to preprocess windows if any of their descendants are visible. This solves issues on Mac with child window contents not updating properly if its parent is minimized.

diff --git a/Engine/Source/Runtime/Slate/Private/Framework/Application/SlateApplication.cpp b/Engine/Source/Runtime/Slate/Private/Framework/Application/SlateApplication.cpp
--- a/Engine/Source/Runtime/Slate/Private/Framework/Application/SlateApplication.cpp
+++ b/Engine/Source/Runtime/Slate/Private/Framework/Application/SlateApplication.cpp
@@ -1387,6 +1387,23 @@ void FSlateApplication::DrawWindowAndChildren( const TSharedRef<SWindow>& Window
 	}
 }
 
+static bool DoAnyWindowDescendantsNeedPrepass(TSharedRef<SWindow> WindowToPrepass)
+{
+	for (const TSharedRef<SWindow>& ChildWindow : WindowToPrepass->GetChildWindows())
+	{
+		if (ChildWindow->IsVisible() && !ChildWindow->IsWindowMinimized())
+		{
+			return true;
+		}
+		else
+		{
+			return DoAnyWindowDescendantsNeedPrepass(ChildWindow);
+		}
+	}
+
+	return false;
+}
+
 static void PrepassWindowAndChildren( TSharedRef<SWindow> WindowToPrepass )
 {
 	if (UNLIKELY(!FApp::CanEverRender()))
@@ -1394,7 +1411,9 @@ static void PrepassWindowAndChildren( TSharedRef<SWindow> WindowToPrepass )
 		return;
 	}
 
-	if ( WindowToPrepass->IsVisible() && !WindowToPrepass->IsWindowMinimized() )
+	const bool bIsWindowVisible = WindowToPrepass->IsVisible() && !WindowToPrepass->IsWindowMinimized();
+
+	if (bIsWindowVisible || DoAnyWindowDescendantsNeedPrepass(WindowToPrepass))
 	{
 		FScopedSwitchWorldHack SwitchWorld(WindowToPrepass);
 		
@@ -1403,7 +1422,7 @@ static void PrepassWindowAndChildren( TSharedRef<SWindow> WindowToPrepass )
 			WindowToPrepass->SlatePrepass(FSlateApplication::Get().GetApplicationScale() * WindowToPrepass->GetNativeWindow()->GetDPIScaleFactor());
 		}
 
-		if ( WindowToPrepass->IsAutosized() )
+		if ( bIsWindowVisible && WindowToPrepass->IsAutosized() )
 		{
 			WindowToPrepass->Resize(WindowToPrepass->GetDesiredSizeDesktopPixels());
 		}
