PR #5398: UE-45365: Scroll to selected item on construction (Contributed by [REDACTED_PERSON])

diff --git a/Engine/Source/Runtime/Slate/Public/Widgets/Input/SComboBox.h b/Engine/Source/Runtime/Slate/Public/Widgets/Input/SComboBox.h
--- a/Engine/Source/Runtime/Slate/Public/Widgets/Input/SComboBox.h
+++ b/Engine/Source/Runtime/Slate/Public/Widgets/Input/SComboBox.h
@@ -238,6 +238,7 @@ public:
 		if( TListTypeTraits<OptionType>::IsPtrValid( SelectedItem ) )
 		{
 			ComboListView->Private_SetItemSelection( SelectedItem, true);
+			ComboListView->RequestScrollIntoView(SelectedItem, 0);
 		}
 
 	}
