PR #5404: Option to always show scrollbar track (Contributed by [REDACTED_PERSON])

diff --git a/Engine/Source/Runtime/Slate/Private/Widgets/Layout/SScrollBar.cpp b/Engine/Source/Runtime/Slate/Private/Widgets/Layout/SScrollBar.cpp
--- a/Engine/Source/Runtime/Slate/Private/Widgets/Layout/SScrollBar.cpp
+++ b/Engine/Source/Runtime/Slate/Private/Widgets/Layout/SScrollBar.cpp
@@ -82,6 +82,7 @@ void SScrollBar::Construct(const FArguments& InArgs)
 
 	this->EnabledState = TAttribute<bool>( Track.ToSharedRef(), &SScrollBarTrack::IsNeeded );
 	SetScrollBarAlwaysVisible(InArgs._AlwaysShowScrollbar);
+	bAlwaysShowScrollbarTrack = InArgs._AlwaysShowScrollbarTrack;
 }
 
 void SScrollBar::SetOnUserScrolled( const FOnUserScrolled& InHandler )
@@ -217,17 +218,17 @@ SScrollBar::SScrollBar()
 
 FSlateColor SScrollBar::GetTrackOpacity() const
 {
-	if (AlwaysShowScrollbar() && (Track->GetThumbSizeFraction() == 0.0f))
+	if (bDraggingThumb || IsHovered())
 	{
-		return FLinearColor(1, 1, 1, 0.5f);
+		return FLinearColor(1,1,1,1);
 	}
-	else if (bDraggingThumb || IsHovered())
+	else if ((bAlwaysShowScrollbarTrack  && Track->GetThumbSizeFraction() > KINDA_SMALL_NUMBER) || AlwaysShowScrollbar())
 	{
-		return FLinearColor(1,1,1,1);
+		return FLinearColor(1, 1, 1, 0.5f);
 	}
 	else
 	{
-		return FLinearColor(1,1,1,0.0f);
+		return FLinearColor(1,1,1,0);
 	}
 }
 
@@ -365,6 +366,11 @@ void SScrollBar::SetScrollBarAlwaysVisible(bool InAlwaysVisible)
 	Track->SetIsAlwaysVisible(InAlwaysVisible);
 }
 
+void SScrollBar::SetScrollBarTrackAlwaysVisible(bool InAlwaysVisible)
+{
+	bAlwaysShowScrollbarTrack = InAlwaysVisible;
+}
+
 bool SScrollBar::AlwaysShowScrollbar() const
 {
 	return bAlwaysShowScrollbar;
diff --git a/Engine/Source/Runtime/Slate/Private/Widgets/Layout/SScrollBox.cpp b/Engine/Source/Runtime/Slate/Private/Widgets/Layout/SScrollBox.cpp
--- a/Engine/Source/Runtime/Slate/Private/Widgets/Layout/SScrollBox.cpp
+++ b/Engine/Source/Runtime/Slate/Private/Widgets/Layout/SScrollBox.cpp
@@ -527,6 +527,11 @@ void SScrollBox::SetScrollBarAlwaysVisible(bool InAlwaysVisible)
 	ScrollBar->SetScrollBarAlwaysVisible(InAlwaysVisible);
 }
 
+void SScrollBox::SetScrollBarTrackAlwaysVisible(bool InAlwaysVisible)
+{
+	ScrollBar->SetScrollBarTrackAlwaysVisible(InAlwaysVisible);
+}
+
 void SScrollBox::SetScrollBarThickness(FVector2D InThickness)
 {
 	ScrollBar->SetThickness(InThickness);
diff --git a/Engine/Source/Runtime/Slate/Public/Widgets/Layout/SScrollBar.h b/Engine/Source/Runtime/Slate/Public/Widgets/Layout/SScrollBar.h
--- a/Engine/Source/Runtime/Slate/Public/Widgets/Layout/SScrollBar.h
+++ b/Engine/Source/Runtime/Slate/Public/Widgets/Layout/SScrollBar.h
@@ -31,6 +31,7 @@ public:
 		: _Style( &FCoreStyle::Get().GetWidgetStyle<FScrollBarStyle>("Scrollbar") )
 		, _OnUserScrolled()
 		, _AlwaysShowScrollbar(false)
+		, _AlwaysShowScrollbarTrack(true)
 #if PLATFORM_UI_HAS_MOBILE_SCROLLBARS
 		, _HideWhenNotInUse(true)
 #else
@@ -45,6 +46,7 @@ public:
 		SLATE_STYLE_ARGUMENT( FScrollBarStyle, Style )
 		SLATE_EVENT( FOnUserScrolled, OnUserScrolled )
 		SLATE_ARGUMENT( bool, AlwaysShowScrollbar )
+		SLATE_ARGUMENT( bool, AlwaysShowScrollbarTrack )
 		SLATE_ARGUMENT( bool, HideWhenNotInUse )
 		SLATE_ARGUMENT( EOrientation, Orientation )
 		SLATE_ARGUMENT( EFocusCause, DragFocusCause )
@@ -102,21 +104,24 @@ public:
 	/** @return the orientation in which the scrollbar is scrolling. */
 	EOrientation GetOrientation() const; 
 
-	/** See argument Style */
+	/** Set argument Style */
 	void SetStyle(const FScrollBarStyle* InStyle);
 
-	/** See UserVisibility attribute */
+	/** Set UserVisibility attribute */
 	void SetUserVisibility(TAttribute<EVisibility> InUserVisibility) { UserVisibility = InUserVisibility; }
 
-	/** See DragFocusCause attribute */
+	/** Set DragFocusCause attribute */
 	void SetDragFocusCause(EFocusCause InDragFocusCause);
 
-	/** See Thickness attribute */
+	/** Set Thickness attribute */
 	void SetThickness(TAttribute<FVector2D> InThickness);
 
-	/** See ScrollBarAlwaysVisible attribute */
+	/** Set ScrollBarAlwaysVisible attribute */
 	void SetScrollBarAlwaysVisible(bool InAlwaysVisible);
 
+	/** Set ScrollBarTrackAlwaysVisible attribute */
+	void SetScrollBarTrackAlwaysVisible(bool InAlwaysVisible);
+
 	/** Returns True when the scrollbar should always be shown, else False */
 	bool AlwaysShowScrollbar() const;
 
@@ -153,6 +158,7 @@ protected:
 	float DragGrabOffset;
 	EOrientation Orientation;
 	bool bAlwaysShowScrollbar;
+	bool bAlwaysShowScrollbarTrack;
 	EFocusCause DragFocusCause;
 	bool bHideWhenNotInUse;
 	bool bIsScrolling;
diff --git a/Engine/Source/Runtime/Slate/Public/Widgets/Layout/SScrollBox.h b/Engine/Source/Runtime/Slate/Public/Widgets/Layout/SScrollBox.h
--- a/Engine/Source/Runtime/Slate/Public/Widgets/Layout/SScrollBox.h
+++ b/Engine/Source/Runtime/Slate/Public/Widgets/Layout/SScrollBox.h
@@ -193,6 +193,8 @@ public:
 	void SetScrollBarVisibility(EVisibility InVisibility);
 
 	void SetScrollBarAlwaysVisible(bool InAlwaysVisible);
+	
+	void SetScrollBarTrackAlwaysVisible(bool InAlwaysVisible);
 
 	void SetScrollBarThickness(FVector2D InThickness);
 
diff --git a/Engine/Source/Runtime/UMG/Private/Components/ScrollBar.cpp b/Engine/Source/Runtime/UMG/Private/Components/ScrollBar.cpp
--- a/Engine/Source/Runtime/UMG/Private/Components/ScrollBar.cpp
+++ b/Engine/Source/Runtime/UMG/Private/Components/ScrollBar.cpp
@@ -13,6 +13,7 @@ UScrollBar::UScrollBar(const FObjectInitializer& ObjectInitializer)
 	bIsVariable = false;
 
 	bAlwaysShowScrollbar = true;
+	bAlwaysShowScrollbarTrack = true;
 	Orientation = Orient_Vertical;
 	Thickness = FVector2D(12.0f, 12.0f);
 
@@ -32,6 +33,7 @@ TSharedRef<SWidget> UScrollBar::RebuildWidget()
 	MyScrollBar = SNew(SScrollBar)
 		.Style(&WidgetStyle)
 		.AlwaysShowScrollbar(bAlwaysShowScrollbar)
+		.AlwaysShowScrollbarTrack(bAlwaysShowScrollbarTrack)
 		.Orientation(Orientation)
 		.Thickness(Thickness);
 
diff --git a/Engine/Source/Runtime/UMG/Private/Components/ScrollBox.cpp b/Engine/Source/Runtime/UMG/Private/Components/ScrollBox.cpp
--- a/Engine/Source/Runtime/UMG/Private/Components/ScrollBox.cpp
+++ b/Engine/Source/Runtime/UMG/Private/Components/ScrollBox.cpp
@@ -16,6 +16,7 @@ UScrollBox::UScrollBox(const FObjectInitializer& ObjectInitializer)
 	, ConsumeMouseWheel(EConsumeMouseWheel::WhenScrollingPossible)
 	, ScrollbarThickness(5.0f, 5.0f)
 	, AlwaysShowScrollbar(false)
+	, AlwaysShowScrollbarTrack(false)
 	, AllowOverscroll(true)
 	, NavigationDestination(EDescendantScrollDestination::IntoView)
 	, NavigationScrollPadding(0.0f)
@@ -97,6 +98,7 @@ void UScrollBox::SynchronizeProperties()
 	MyScrollBox->SetScrollBarVisibility(UWidget::ConvertSerializedVisibilityToRuntime(ScrollBarVisibility));
 	MyScrollBox->SetScrollBarThickness(ScrollbarThickness);
 	MyScrollBox->SetScrollBarAlwaysVisible(AlwaysShowScrollbar);
+	MyScrollBox->SetScrollBarTrackAlwaysVisible(AlwaysShowScrollbarTrack);
 	MyScrollBox->SetAllowOverscroll(AllowOverscroll ? EAllowOverscroll::Yes : EAllowOverscroll::No);
 	MyScrollBox->SetScrollBarRightClickDragAllowed(bAllowRightClickDragScrolling);
 	MyScrollBox->SetConsumeMouseWheel(ConsumeMouseWheel);
diff --git a/Engine/Source/Runtime/UMG/Public/Components/ScrollBar.h b/Engine/Source/Runtime/UMG/Public/Components/ScrollBar.h
--- a/Engine/Source/Runtime/UMG/Public/Components/ScrollBar.h
+++ b/Engine/Source/Runtime/UMG/Public/Components/ScrollBar.h
@@ -29,6 +29,10 @@ public:
 	UPROPERTY(EditAnywhere, Category="Behavior")
 	bool bAlwaysShowScrollbar;
 
+	/**  */
+	UPROPERTY(EditAnywhere, Category = "Behavior")
+	bool bAlwaysShowScrollbarTrack;
+
 	/**  */
 	UPROPERTY(EditAnywhere, Category="Behavior")
 	TEnumAsByte<EOrientation> Orientation;
diff --git a/Engine/Source/Runtime/UMG/Public/Components/ScrollBox.h b/Engine/Source/Runtime/UMG/Public/Components/ScrollBox.h
--- a/Engine/Source/Runtime/UMG/Public/Components/ScrollBox.h
+++ b/Engine/Source/Runtime/UMG/Public/Components/ScrollBox.h
@@ -56,6 +56,10 @@ public:
 	UPROPERTY(EditAnywhere, BlueprintReadOnly, Category="Scroll")
 	bool AlwaysShowScrollbar;
 
+	/**  */
+	UPROPERTY(EditAnywhere, BlueprintReadOnly, Category = "Scroll")
+	bool AlwaysShowScrollbarTrack;
+
 	/**  Disable to stop scrollbars from activating inertial overscrolling */
 	UPROPERTY(EditAnywhere, BlueprintReadOnly, Category = "Scroll")
 	bool AllowOverscroll;
