PR #5401: Prevent entering edit mode when drag/dropping SInlineEditableTextBlocks (Contributed by [REDACTED_PERSON])

diff --git a/Engine/Source/Runtime/Slate/Private/Widgets/Text/SInlineEditableTextBlock.cpp b/Engine/Source/Runtime/Slate/Private/Widgets/Text/SInlineEditableTextBlock.cpp
--- a/Engine/Source/Runtime/Slate/Private/Widgets/Text/SInlineEditableTextBlock.cpp
+++ b/Engine/Source/Runtime/Slate/Private/Widgets/Text/SInlineEditableTextBlock.cpp
@@ -189,7 +189,7 @@ FReply SInlineEditableTextBlock::OnMouseButtonDown( const FGeometry& MyGeometry,
 	{
 		if (IsSelected.Execute() && !bIsReadOnly.Get() && !ActiveTimerHandle.IsValid())
 		{
-			RegisterActiveTimer(0.5f, FWidgetActiveTimerDelegate::CreateSP(this, &SInlineEditableTextBlock::TriggerEditMode));
+			ActiveTimerHandle = RegisterActiveTimer(0.5f, FWidgetActiveTimerDelegate::CreateSP(this, &SInlineEditableTextBlock::TriggerEditMode));
 		}
 	}
 	else
