Added more persision in the editor for float fields.

diff --git a/Engine/Source/Editor/PropertyEditor/Private/PropertyHandleImpl.cpp b/Engine/Source/Editor/PropertyEditor/Private/PropertyHandleImpl.cpp
--- a/Engine/Source/Editor/PropertyEditor/Private/PropertyHandleImpl.cpp
+++ b/Engine/Source/Editor/PropertyEditor/Private/PropertyHandleImpl.cpp
@@ -3374,7 +3374,7 @@ FPropertyAccess::Result FPropertyHandleFloat::SetValue( const float& NewValue, E
 	// Clamp the value from any meta data ranges stored on the property value
 	float FinalValue = ClampValueFromMetaData<float>( NewValue, *Implementation->GetPropertyNode() );
 
-	const FString ValueStr = FString::Printf( TEXT("%f"), FinalValue );
+	const FString ValueStr = FString::Printf( TEXT("%.8f"), FinalValue );
 	Res = Implementation->ImportText( ValueStr, Flags );
 
 	return Res;
diff --git a/Engine/Source/Runtime/Slate/Public/Widgets/Input/NumericTypeInterface.h b/Engine/Source/Runtime/Slate/Public/Widgets/Input/NumericTypeInterface.h
--- a/Engine/Source/Runtime/Slate/Public/Widgets/Input/NumericTypeInterface.h
+++ b/Engine/Source/Runtime/Slate/Public/Widgets/Input/NumericTypeInterface.h
@@ -36,7 +36,7 @@ struct TDefaultNumericTypeInterface : INumericTypeInterface<NumericType>
 		static const FNumberFormattingOptions NumberFormattingOptions = FNumberFormattingOptions()
 			.SetUseGrouping(false)
 			.SetMinimumFractionalDigits(TIsIntegral<NumericType>::Value ? 0 : 1)
-			.SetMaximumFractionalDigits(TIsIntegral<NumericType>::Value ? 0 : 6);
+			.SetMaximumFractionalDigits(TIsIntegral<NumericType>::Value ? 0 : 8);
 		return FastDecimalFormat::NumberToString(Value, ExpressionParser::GetLocalizedNumberFormattingRules(), NumberFormattingOptions);
 	}
 	virtual TOptional<NumericType> FromString(const FString& InString, const NumericType& InExistingValue) override
