Fix For the tab close not aligning properly.

diff --git a/Engine/Source/Runtime/Slate/Private/Widgets/Docking/SDockTab.cpp b/Engine/Source/Runtime/Slate/Private/Widgets/Docking/SDockTab.cpp
--- a/Engine/Source/Runtime/Slate/Private/Widgets/Docking/SDockTab.cpp
+++ b/Engine/Source/Runtime/Slate/Private/Widgets/Docking/SDockTab.cpp
@@ -531,7 +531,7 @@ void SDockTab::Construct( const FArguments& InArgs )
 			+ SOverlay::Slot()
 			.Padding(TAttribute<FMargin>::Create(TAttribute<FMargin>::FGetter::CreateSP(this, &SDockTab::GetTabPadding)))
 			.VAlign(VAlign_Center)
-			.HAlign(HAlign_Left)
+			.HAlign(HAlign_Fill)
 			[
 				SNew(SHorizontalBox)
 				.Visibility(EVisibility::Visible)
@@ -556,31 +556,46 @@ void SDockTab::Construct( const FArguments& InArgs )
 						.Image(this, &SDockTab::GetTabIcon)
 					]
 				]
-				// Tab Label
+
 				+ SHorizontalBox::Slot()
 				.FillWidth(1.0f)
-				.Padding( 0.0f, 1.0f )
-				.VAlign(VAlign_Center)
-				[
-					SAssignNew(LabelWidget, STextBlock)
-					.TextStyle( FCoreStyle::Get(), "Docking.TabFont" )					
-					.Text( this, &SDockTab::GetTabLabel )
-				]
-				+ SHorizontalBox::Slot()
-				.AutoWidth()
-				.Padding( 0.0f, 1.0f )
+				.Padding(0.0f, 1.0f)
+				.HAlign(HAlign_Left)
 				.VAlign(VAlign_Center)
 				[
-					SAssignNew(LabelWidget, STextBlock)
-					.TextStyle( FCoreStyle::Get(), "Docking.TabFont" )					
-					.Text(this, &SDockTab::GetTabLabelSuffix)
+					// Label sub HBox
+					SNew(SHorizontalBox)
+					.Visibility(EVisibility::Visible)
+					.ToolTip(InArgs._ToolTip)
+
+					// Tab Label
+					+ SHorizontalBox::Slot()
+					.FillWidth(1.0f)
+					.Padding(0.0f, 1.0f)
+					.VAlign(VAlign_Center)
+					[
+						SAssignNew(LabelWidget, STextBlock)
+						.TextStyle(FCoreStyle::Get(), "Docking.TabFont")
+						.Text(this, &SDockTab::GetTabLabel)
+					]
+
+					// Tab Label Suffix
+					+ SHorizontalBox::Slot()
+					.AutoWidth()
+					.Padding(0.0f, 1.0f)
+					.VAlign(VAlign_Center)
+					[
+						SAssignNew(LabelWidget, STextBlock)
+						.TextStyle(FCoreStyle::Get(), "Docking.TabFont")					
+						.Text(this, &SDockTab::GetTabLabelSuffix)
+					]
 				]
 				
 				// @todo toolkit major: Could inject inline content here into tab for standalone asset editing dropdown/dirty state, etc.
 
 				+ SHorizontalBox::Slot()
 				.AutoWidth()
-				.HAlign(HAlign_Center)
+				.HAlign(HAlign_Right)
 				.VAlign(VAlign_Center)
 				[
 					SNew(SButton)
