Only include SlateDebug content when packaging with debug files included

diff --git a/Engine/Source/Runtime/Slate/Slate.Build.cs b/Engine/Source/Runtime/Slate/Slate.Build.cs
--- a/Engine/Source/Runtime/Slate/Slate.Build.cs
+++ b/Engine/Source/Runtime/Slate/Slate.Build.cs
@@ -58,12 +58,12 @@ public class Slate : ModuleRules
 		if (Target.bUsesSlate)
 		{
 			RuntimeDependencies.Add("$(EngineDir)/Content/Slate/...", StagedFileType.UFS);
-			RuntimeDependencies.Add("$(EngineDir)/Content/SlateDebug/...", StagedFileType.UFS);
+			RuntimeDependencies.Add("$(EngineDir)/Content/SlateDebug/...", StagedFileType.DebugNonUFS);
 
 			if (Target.ProjectFile != null)
 			{
 				RuntimeDependencies.Add("$(ProjectDir)/Content/Slate/...", StagedFileType.UFS);
-				RuntimeDependencies.Add("$(ProjectDir)/Content/SlateDebug/...", StagedFileType.UFS);
+				RuntimeDependencies.Add("$(ProjectDir)/Content/SlateDebug/...", StagedFileType.DebugNonUFS);
 			}
 		}
 	}
