Copied 4869342 from Dev-Enterprise

Prevent tabs in the same tab stack or docking splitter from being resizable when they all specify to be sized to content.

diff --git a/Engine/Source/Runtime/Slate/Private/Framework/Docking/SDockingSplitter.cpp b/Engine/Source/Runtime/Slate/Private/Framework/Docking/SDockingSplitter.cpp
--- a/Engine/Source/Runtime/Slate/Private/Framework/Docking/SDockingSplitter.cpp
+++ b/Engine/Source/Runtime/Slate/Private/Framework/Docking/SDockingSplitter.cpp
@@ -383,3 +383,23 @@ TSharedRef<SDockingNode> SDockingSplitter::FindTabStack(ETabStackToFind FindMe)
 	}
 
 }
+
+SSplitter::ESizeRule SDockingSplitter::GetSizeRule() const
+{
+	if (Children.Num() > 0)
+	{
+		for (const TSharedRef<SDockingNode>& ChildNode : Children)
+		{
+			if (ChildNode->GetSizeRule() == SSplitter::FractionOfParent)
+			{
+				return SSplitter::FractionOfParent;
+			}
+		}
+		// If all nodes in this docking splitter are sized to content, then it should size to Content.
+		return SSplitter::SizeToContent;
+	}
+	else
+	{
+		return SSplitter::FractionOfParent;
+	}
+}
diff --git a/Engine/Source/Runtime/Slate/Private/Framework/Docking/SDockingSplitter.h b/Engine/Source/Runtime/Slate/Private/Framework/Docking/SDockingSplitter.h
--- a/Engine/Source/Runtime/Slate/Private/Framework/Docking/SDockingSplitter.h
+++ b/Engine/Source/Runtime/Slate/Private/Framework/Docking/SDockingSplitter.h
@@ -76,6 +76,9 @@ public:
 	/** Recursively searches through all children looking for child tabs */
 	virtual TArray< TSharedRef<SDockTab> > GetAllChildTabs() const override;
 
+	/** Should this node auto-size or be a percentage of its parent size */
+	virtual SSplitter::ESizeRule GetSizeRule() const override;
+
 	/** Gets the size coefficient of a given child dock node */
 	float GetSizeCoefficientForSlot(int32 Index) const;
 
diff --git a/Engine/Source/Runtime/Slate/Private/Framework/Docking/SDockingTabStack.cpp b/Engine/Source/Runtime/Slate/Private/Framework/Docking/SDockingTabStack.cpp
--- a/Engine/Source/Runtime/Slate/Private/Framework/Docking/SDockingTabStack.cpp
+++ b/Engine/Source/Runtime/Slate/Private/Framework/Docking/SDockingTabStack.cpp
@@ -942,10 +942,17 @@ bool SDockingTabStack::CanCloseAllButForegroundTab() const
 
 SSplitter::ESizeRule SDockingTabStack::GetSizeRule() const
 {
-	if (this->GetNumTabs() == 1 && this->GetTabs()[0]->ShouldAutosize() )
+	int32 NumTabs = this->GetNumTabs();
+	if (NumTabs > 0)
 	{
-		// If there is a single tab in this stack and it is
-		// sized to content, then the stack's cell should size to Content.
+		for (int32 Index = 0; Index < NumTabs; ++Index)
+		{
+			if (!this->GetTabs()[Index]->ShouldAutosize())
+			{
+				return SSplitter::FractionOfParent;
+			}
+		}
+		// If all tabs in this stack are sized to content, then the stack's cell should size to Content.
 		return SSplitter::SizeToContent;
 	}
 	else
