Make buttons stextblock hit test invisible by default

diff --git a/Engine/Source/Runtime/Slate/Private/Widgets/Input/SButton.cpp b/Engine/Source/Runtime/Slate/Private/Widgets/Input/SButton.cpp
--- a/Engine/Source/Runtime/Slate/Private/Widgets/Input/SButton.cpp
+++ b/Engine/Source/Runtime/Slate/Private/Widgets/Input/SButton.cpp
@@ -26,6 +26,7 @@ void SButton::Construct( const FArguments& InArgs )
 			if ((InOpArgs._Content.Widget == SNullWidget::NullWidget) && (InOpArgs._Text.IsBound() || !InOpArgs._Text.Get().IsEmpty()) )
 			{
 				return SNew(STextBlock)
+					.Visibility(EVisibility::HitTestInvisible)
 					.Text( InOpArgs._Text )
 					.TextStyle( InOpArgs._TextStyle )
 					.TextShapingMethod( InOpArgs._TextShapingMethod )
