Remove double click event from STextBlock.  It was redundant with the one in base SWidget

diff --git a/Engine/Plugins/FX/Niagara/Source/NiagaraEditorWidgets/Private/Stack/SNiagaraStackFunctionInputValue.cpp b/Engine/Plugins/FX/Niagara/Source/NiagaraEditorWidgets/Private/Stack/SNiagaraStackFunctionInputValue.cpp
--- a/Engine/Plugins/FX/Niagara/Source/NiagaraEditorWidgets/Private/Stack/SNiagaraStackFunctionInputValue.cpp
+++ b/Engine/Plugins/FX/Niagara/Source/NiagaraEditorWidgets/Private/Stack/SNiagaraStackFunctionInputValue.cpp
@@ -416,7 +416,7 @@ FText SNiagaraStackFunctionInputValue::GetInvalidValueToolTipText() const
 	}
 }
 
-FReply SNiagaraStackFunctionInputValue::DynamicInputTextDoubleClicked()
+FReply SNiagaraStackFunctionInputValue::DynamicInputTextDoubleClicked(const FGeometry& MyGeometry, const FPointerEvent& PointerEvent)
 {
 	UNiagaraNodeFunctionCall* DynamicInputNode = FunctionInput->GetDynamicInputNode();
 	if (DynamicInputNode->FunctionScript != nullptr && DynamicInputNode->FunctionScript->IsAsset())
@@ -427,7 +427,7 @@ FReply SNiagaraStackFunctionInputValue::DynamicInputTextDoubleClicked()
 	return FReply::Unhandled();
 }
 
-FReply SNiagaraStackFunctionInputValue::OnLinkedInputDoubleClicked()
+FReply SNiagaraStackFunctionInputValue::OnLinkedInputDoubleClicked(const FGeometry& MyGeometry, const FPointerEvent& PointerEvent)
 {
 	FString ParamCollection;
 	FString ParamName;
diff --git a/Engine/Plugins/FX/Niagara/Source/NiagaraEditorWidgets/Private/Stack/SNiagaraStackFunctionInputValue.h b/Engine/Plugins/FX/Niagara/Source/NiagaraEditorWidgets/Private/Stack/SNiagaraStackFunctionInputValue.h
--- a/Engine/Plugins/FX/Niagara/Source/NiagaraEditorWidgets/Private/Stack/SNiagaraStackFunctionInputValue.h
+++ b/Engine/Plugins/FX/Niagara/Source/NiagaraEditorWidgets/Private/Stack/SNiagaraStackFunctionInputValue.h
@@ -59,8 +59,8 @@ private:
 
 	FText GetInvalidValueToolTipText() const;
 
-	FReply DynamicInputTextDoubleClicked();
-	FReply OnLinkedInputDoubleClicked();
+	FReply DynamicInputTextDoubleClicked(const FGeometry& MyGeometry, const FPointerEvent& PointerEvent);
+	FReply OnLinkedInputDoubleClicked(const FGeometry& MyGeometry, const FPointerEvent& PointerEvent);
 
 	class SNiagaraFunctionInputActionMenuExpander: public SExpanderArrow
 	{
diff --git a/Engine/Source/Editor/AnimationModifiers/Private/SModifierItemRow.cpp b/Engine/Source/Editor/AnimationModifiers/Private/SModifierItemRow.cpp
--- a/Engine/Source/Editor/AnimationModifiers/Private/SModifierItemRow.cpp
+++ b/Engine/Source/Editor/AnimationModifiers/Private/SModifierItemRow.cpp
@@ -35,7 +35,7 @@ void SModifierItemRow::Construct(const FArguments& InArgs, const TSharedRef<STab
 	];
 }
 
-FReply SModifierItemRow::OnDoubleClicked()
+FReply SModifierItemRow::OnDoubleClicked(const FGeometry& MyGeometry, const FPointerEvent& PointerEvent)
 {
 	OnOpenModifier.ExecuteIfBound(InternalItem->Instance);
 	return FReply::Handled();
diff --git a/Engine/Source/Editor/AnimationModifiers/Private/SModifierItemRow.h b/Engine/Source/Editor/AnimationModifiers/Private/SModifierItemRow.h
--- a/Engine/Source/Editor/AnimationModifiers/Private/SModifierItemRow.h
+++ b/Engine/Source/Editor/AnimationModifiers/Private/SModifierItemRow.h
@@ -13,7 +13,7 @@ public:
 	SLATE_END_ARGS()
 
 	void Construct(const FArguments& InArgs, const TSharedRef<STableViewBase>& InOwnerTableView, const ModifierListviewItem& Item);
-	FReply OnDoubleClicked();
+	FReply OnDoubleClicked(const FGeometry& MyGeometry, const FPointerEvent& PointerEvent);
 
 protected:
 	FText GetInstanceText() const;
diff --git a/Engine/Source/Runtime/Slate/Private/Widgets/Input/SEditableLabel.cpp b/Engine/Source/Runtime/Slate/Private/Widgets/Input/SEditableLabel.cpp
--- a/Engine/Source/Runtime/Slate/Private/Widgets/Input/SEditableLabel.cpp
+++ b/Engine/Source/Runtime/Slate/Private/Widgets/Input/SEditableLabel.cpp
@@ -34,7 +34,6 @@ void SEditableLabel::Construct(const FArguments& InArgs)
 					.HighlightShape(InArgs._HighlightShape)
 					.HighlightText(InArgs._HighlightText)
 					.MinDesiredWidth(InArgs._MinDesiredWidth)
-					.OnDoubleClicked(this, &SEditableLabel::HandleTextBlockDoubleClicked)
 					.ShadowColorAndOpacity(InArgs._ShadowColorAndOpacity)
 					.ShadowOffset(InArgs._ShadowOffset)
 					.TextStyle(InArgs._TextStyle)
@@ -115,6 +114,16 @@ FReply SEditableLabel::OnKeyDown(const FGeometry& MyGeometry, const FKeyEvent& I
 }
 
 
+FReply SEditableLabel::OnMouseButtonDoubleClick(const FGeometry& InMyGeometry, const FPointerEvent& InMouseEvent)
+{
+	if (InMouseEvent.GetEffectingButton() == EKeys::LeftMouseButton)
+	{
+		return HandleTextBlockDoubleClicked();
+	}
+
+	return FReply::Unhandled();
+}
+
 bool SEditableLabel::SupportsKeyboardFocus() const
 {
 	return CanEditAttribute.Get();
diff --git a/Engine/Source/Runtime/Slate/Private/Widgets/Text/STextBlock.cpp b/Engine/Source/Runtime/Slate/Private/Widgets/Text/STextBlock.cpp
--- a/Engine/Source/Runtime/Slate/Private/Widgets/Text/STextBlock.cpp
+++ b/Engine/Source/Runtime/Slate/Private/Widgets/Text/STextBlock.cpp
@@ -49,7 +49,7 @@ void STextBlock::Construct( const FArguments& InArgs )
 
 	bSimpleTextMode = InArgs._SimpleTextMode;
 
-	OnDoubleClicked = InArgs._OnDoubleClicked;
+	SetOnMouseDoubleClick(InArgs._OnDoubleClicked);
 
 	BoundText = InArgs._Text;
 
@@ -238,19 +238,6 @@ int32 STextBlock::OnPaint( const FPaintArgs& Args, const FGeometry& AllottedGeom
 	return LayerId;
 }
 
-FReply STextBlock::OnMouseButtonDoubleClick( const FGeometry& InMyGeometry, const FPointerEvent& InMouseEvent )
-{
-	if ( InMouseEvent.GetEffectingButton() == EKeys::LeftMouseButton )
-	{
-		if( OnDoubleClicked.IsBound() )
-		{
-			return OnDoubleClicked.Execute();
-		}
-	}
-
-	return FReply::Unhandled();
-}
-
 FVector2D STextBlock::ComputeDesiredSize(float LayoutScaleMultiplier) const
 {
 	SCOPE_CYCLE_COUNTER(Stat_SlateTextBlockCDS);
diff --git a/Engine/Source/Runtime/Slate/Public/Widgets/Input/SEditableLabel.h b/Engine/Source/Runtime/Slate/Public/Widgets/Input/SEditableLabel.h
--- a/Engine/Source/Runtime/Slate/Public/Widgets/Input/SEditableLabel.h
+++ b/Engine/Source/Runtime/Slate/Public/Widgets/Input/SEditableLabel.h
@@ -101,6 +101,7 @@ public:
 
 	virtual bool HasKeyboardFocus() const override;
 	virtual FReply OnKeyDown(const FGeometry& MyGeometry, const FKeyEvent& InKeyEvent) override;
+	virtual FReply OnMouseButtonDoubleClick(const FGeometry& InMyGeometry, const FPointerEvent& InMouseEvent) override;
 	virtual bool SupportsKeyboardFocus() const override;
 
 private:
diff --git a/Engine/Source/Runtime/Slate/Public/Widgets/Text/STextBlock.h b/Engine/Source/Runtime/Slate/Public/Widgets/Text/STextBlock.h
--- a/Engine/Source/Runtime/Slate/Public/Widgets/Text/STextBlock.h
+++ b/Engine/Source/Runtime/Slate/Public/Widgets/Text/STextBlock.h
@@ -137,7 +137,7 @@ public:
 		SLATE_ARGUMENT( bool, SimpleTextMode )
 
 		/** Called when this text is double clicked */
-		SLATE_EVENT( FOnClicked, OnDoubleClicked )
+		SLATE_EVENT(FPointerEventHandler, OnDoubleClicked)
 
 	SLATE_END_ARGS()
 
@@ -232,7 +232,6 @@ public:
 
 	// SWidget interface
 	virtual int32 OnPaint( const FPaintArgs& Args, const FGeometry& AllottedGeometry, const FSlateRect& MyCullingRect, FSlateWindowElementList& OutDrawElements, int32 LayerId, const FWidgetStyle& InWidgetStyle, bool bParentEnabled ) const override;
-	virtual FReply OnMouseButtonDoubleClick( const FGeometry& InMyGeometry, const FPointerEvent& InMouseEvent ) override;
 	virtual FVector2D ComputeDesiredSize(float) const override;
 	// End of SWidget interface
 
@@ -322,9 +321,6 @@ private:
 	/** Prevents the text block from being smaller than desired in certain cases (e.g. when it is empty) */
 	TAttribute<float> MinDesiredWidth;
 
-	/** The delegate to execute when this text is double clicked */
-	FOnClicked OnDoubleClicked;
-
 	/** If this is enabled, text shaping, wrapping, justification are disabled in favor of much faster text layout and measurement. */
 	mutable TOptional<FVector2D> CachedSimpleDesiredSize;
 	bool bSimpleTextMode;
