Added better warning messages for when users try to rename rows with invalid names in the row editor.
Updated the editable text layout to prevent it from calling OnTextCommitted delegates when the text has not been changed. This fixes a bug where committing a text change by pressing enter can call the delegate twice.

diff --git a/Engine/Source/Editor/DataTableEditor/Private/SRowEditor.cpp b/Engine/Source/Editor/DataTableEditor/Private/SRowEditor.cpp
--- a/Engine/Source/Editor/DataTableEditor/Private/SRowEditor.cpp
+++ b/Engine/Source/Editor/DataTableEditor/Private/SRowEditor.cpp
@@ -346,11 +346,22 @@ void SRowEditor::OnRowRenamed(const FText& Text, ETextCommit::Type CommitType)
 			return;
 		}
 		const FName NewName = DataTableUtils::MakeValidName(Text.ToString());
+		if (NewName == NAME_None)
+		{
+			// popup an error dialog here
+			const FText Message = FText::Format(LOCTEXT("InvalidRowName", "'{0}' is not a valid row name"), Text);
+			FMessageDialog::Open(EAppMsgType::Ok, Message);
+
+			return;
+		}
 		for (auto Name : CachedRowNames)
 		{
 			if (Name.IsValid() && (*Name == NewName))
 			{
-				 //the name already exists
+				//the name already exists
+				// popup an error dialog here
+				const FText Message = FText::Format(LOCTEXT("DuplicateRowName", "'{0}' is already used as a row name in this table"), Text);
+				FMessageDialog::Open(EAppMsgType::Ok, Message);
 				return;
 			}
 		}
diff --git a/Engine/Source/Runtime/Slate/Private/Widgets/Text/SlateEditableTextLayout.cpp b/Engine/Source/Runtime/Slate/Private/Widgets/Text/SlateEditableTextLayout.cpp
--- a/Engine/Source/Runtime/Slate/Private/Widgets/Text/SlateEditableTextLayout.cpp
+++ b/Engine/Source/Runtime/Slate/Private/Widgets/Text/SlateEditableTextLayout.cpp
@@ -105,6 +105,9 @@ FSlateEditableTextLayout::FSlateEditableTextLayout(ISlateEditableTextWidget& InO
 		const FText& InitialTextToSet = BoundText.Get(FText::GetEmpty());
 		SetEditableText(InitialTextToSet, true);
 
+		// treat the initial text as uncommitted
+		bHasUncommittedText = true;
+
 		// Update the cached BoundText value to prevent it triggering another SetEditableText update again next Tick
 		BoundTextLastTick = FTextSnapshot(InitialTextToSet);
 		bWasPasswordLastTick = bIsPassword;
@@ -776,16 +779,7 @@ bool FSlateEditableTextLayout::HandleFocusLost(const FFocusEvent& InFocusEvent)
 		break;
 	}
 
-	// Always clear the local undo chain on commit
-	ClearUndoStates();
-
-	const FText EditedText = GetEditableText();
-
-	OwnerWidget->OnTextCommitted(EditedText, TextAction);
-
-	// Reload underlying value now it is committed  (commit may alter the value) 
-	// so it can be re-displayed in the edit box
-	LoadText();
+	CommitTextChanges(TextAction);
 
 	UpdateCursorHighlight();
 
@@ -799,6 +793,38 @@ bool FSlateEditableTextLayout::HandleFocusLost(const FFocusEvent& InFocusEvent)
 	return true;
 }
 
+void FSlateEditableTextLayout::CommitTextChanges(ETextCommit::Type CommitReason)
+{
+	// don't try to commit changes if we haven't changed anything
+	if (bHasUncommittedText)
+	{
+		bHasUncommittedText = false;
+
+		// Always clear the local undo chain on commit
+		ClearUndoStates();
+
+		const FText EditedText = GetEditableText();
+
+		OwnerWidget->OnTextCommitted(EditedText, CommitReason);
+
+		// Reload underlying value now it is committed  (commit may alter the value) 
+		// so it can be re-displayed in the edit box
+		LoadText();
+
+		// Select all text?
+		if (OwnerWidget->ShouldSelectAllTextOnCommit())
+		{
+			SelectAllText();
+		}
+
+		// Release input focus?
+		if (CommitReason == ETextCommit::OnEnter && OwnerWidget->ShouldClearKeyboardFocusOnCommit())
+		{
+			FSlateApplication::Get().ClearKeyboardFocus(EFocusCause::Cleared);
+		}
+	}
+}
+
 FReply FSlateEditableTextLayout::HandleKeyChar(const FCharacterEvent& InCharacterEvent)
 {
 	// Check for special characters
@@ -1442,29 +1468,7 @@ bool FSlateEditableTextLayout::HandleCarriageReturn()
 	}
 	else
 	{
-		// Always clear the local undo chain on commit.
-		ClearUndoStates();
-
-		const FText EditedText = GetEditableText();
-
-		// When enter is pressed text is committed.  Let anyone interested know about it.
-		OwnerWidget->OnTextCommitted(EditedText, ETextCommit::OnEnter);
-
-		// Reload underlying value now it is committed  (commit may alter the value) 
-		// so it can be re-displayed in the edit box if it retains focus
-		LoadText();
-
-		// Select all text?
-		if (OwnerWidget->ShouldSelectAllTextOnCommit())
-		{
-			SelectAllText();
-		}
-
-		// Release input focus?
-		if (OwnerWidget->ShouldClearKeyboardFocusOnCommit())
-		{
-			FSlateApplication::Get().ClearKeyboardFocus(EFocusCause::Cleared);
-		}
+		CommitTextChanges(ETextCommit::OnEnter);
 	}
 
 	return true;
@@ -2893,6 +2897,7 @@ void FSlateEditableTextLayout::RestoreOriginalText()
 
 		// Let outsiders know that the text content has been changed
 		OwnerWidget->OnTextCommitted(OriginalText.Text, ETextCommit::OnCleared);
+		bHasUncommittedText = false;
 	}
 }
 
@@ -2939,6 +2944,8 @@ void FSlateEditableTextLayout::EndEditTransaction()
 	const bool bHasTextChanged = !EditedText.ToString().Equals(StateBeforeChangingText.GetValue().Text.ToString(), ESearchCase::CaseSensitive);
 	if (bHasTextChanged)
 	{
+		bHasUncommittedText = true;
+
 		// Save text state
 		SaveText(EditedText);
 
@@ -3159,6 +3166,7 @@ void FSlateEditableTextLayout::Tick(const FGeometry& AllottedGeometry, const dou
 		// Let outsiders know that the text content has been changed
 		OwnerWidget->OnTextCommitted(GetEditableText(), VirtualKeyboardTextCommitType);
 		bTextCommittedByVirtualKeyboard = false;
+		bHasUncommittedText = false;
 	}
 
 	if (TextInputMethodChangeNotifier.IsValid() && TextInputMethodContext.IsValid() && TextInputMethodContext->UpdateCachedGeometry(AllottedGeometry))
@@ -3487,6 +3495,8 @@ void FSlateEditableTextLayout::FVirtualKeyboardEntry::SetTextFromVirtualKeyboard
 		OwnerLayout->BoundText.Set(InNewText);
 	}
 
+	OwnerLayout->bHasUncommittedText = true;
+
 	// Update the internal editable text
 	// This method is called from the main thread (i.e. not the game thread) of the device with the virtual keyboard
 	// This causes the app to crash on those devices, so we're using polling here to ensure delegates are
diff --git a/Engine/Source/Runtime/Slate/Public/Widgets/Text/SlateEditableTextLayout.h b/Engine/Source/Runtime/Slate/Public/Widgets/Text/SlateEditableTextLayout.h
--- a/Engine/Source/Runtime/Slate/Public/Widgets/Text/SlateEditableTextLayout.h
+++ b/Engine/Source/Runtime/Slate/Public/Widgets/Text/SlateEditableTextLayout.h
@@ -365,6 +365,10 @@ public:
 
 	bool HasActiveContextMenu() const;
 
+protected:
+
+	void CommitTextChanges(ETextCommit::Type CommitReason);
+
 private:
 	/** Insert the given text at the current cursor position, correctly taking into account new line characters */
 	void InsertTextAtCursorImpl(const FString& InString);
@@ -605,6 +609,9 @@ private:
 	/** Whether the text has been committed by a virtual keyboard */
 	bool bTextCommittedByVirtualKeyboard;
 
+	/** Whether the widget contains uncommitted text */
+	bool bHasUncommittedText;
+
 	/** What text was submitted by a virtual keyboard */
 	FText VirtualKeyboardText;
 
