Fix dock tabs being dragged having the wrong parent pointers which causes problems for the hit test grid

diff --git a/Engine/Source/Runtime/Slate/Private/Framework/Docking/FDockingDragOperation.cpp b/Engine/Source/Runtime/Slate/Private/Framework/Docking/FDockingDragOperation.cpp
--- a/Engine/Source/Runtime/Slate/Private/Framework/Docking/FDockingDragOperation.cpp
+++ b/Engine/Source/Runtime/Slate/Private/Framework/Docking/FDockingDragOperation.cpp
@@ -94,6 +94,8 @@ void FDockingDragOperation::OnTabWellLeft( const TSharedRef<class SDockingTabWel
 
 	LastContentSize = DesiredSizeFrom( DockNodeGeometry.GetLocalSize() );
 
+	CursorDecoratorStackNode->OpenTab(TabBeingDragged.ToSharedRef());
+
 	TabBeingDragged->SetDraggedOverDockArea( NULL );
 }
 
@@ -245,9 +247,6 @@ FDockingDragOperation::FDockingDragOperation( const TSharedRef<SDockTab>& InTabT
 	CursorDecoratorWindow->SetSizingRule( ESizingRule::FixedSize );
 	CursorDecoratorWindow->SetOpacity(0.45f);
 	
-	//TSharedPtr<SDockingArea> OriginalDockArea = TabBeingDragged->GetOriginalDockArea();
-
-	TSharedPtr<SDockingTabStack> CursorDecoratorStackNode;
 	CursorDecoratorWindow->SetContent
 	(
 		SNew(SBorder)
@@ -262,8 +261,6 @@ FDockingDragOperation::FDockingDragOperation( const TSharedRef<SDockTab>& InTabT
 		]
 	);
 
-	CursorDecoratorStackNode->OpenTab( TabBeingDragged.ToSharedRef() );
-
 	if ( TabBeingDragged->IsActive() )
 	{
 		FGlobalTabmanager::Get()->SetActiveTab(NULL);
diff --git a/Engine/Source/Runtime/Slate/Private/Framework/Docking/FDockingDragOperation.h b/Engine/Source/Runtime/Slate/Private/Framework/Docking/FDockingDragOperation.h
--- a/Engine/Source/Runtime/Slate/Private/Framework/Docking/FDockingDragOperation.h
+++ b/Engine/Source/Runtime/Slate/Private/Framework/Docking/FDockingDragOperation.h
@@ -148,6 +148,9 @@ protected:
 	/** The area from which we initially started dragging */
 	TSharedPtr<SDockingArea> TabOwnerAreaOfOrigin;
 
+	/** Decorator widget where we add temp doc tabs to */
+	TSharedPtr<SDockingTabStack> CursorDecoratorStackNode;
+
 	/** Tab stack from which we started dragging this tab. */
 	TWeakPtr<class SDockingTabStack> TabStackOfOrigin;
 
diff --git a/Engine/Source/Runtime/Slate/Private/Framework/Docking/SDockingTabWell.cpp b/Engine/Source/Runtime/Slate/Private/Framework/Docking/SDockingTabWell.cpp
--- a/Engine/Source/Runtime/Slate/Private/Framework/Docking/SDockingTabWell.cpp
+++ b/Engine/Source/Runtime/Slate/Private/Framework/Docking/SDockingTabWell.cpp
@@ -24,7 +24,7 @@ SDockingTabWell::SDockingTabWell()
 void SDockingTabWell::Construct( const FArguments& InArgs )
 {
 	ForegroundTabIndex = INDEX_NONE;
-	TabBeingDraggedPtr = TSharedPtr<SDockTab>(nullptr);
+	TabBeingDraggedPtr = nullptr;
 	ChildBeingDraggedOffset = 0.0f;
 	TabGrabOffsetFraction = FVector2D::ZeroVector;
 		
@@ -72,7 +72,7 @@ void SDockingTabWell::OnArrangeChildren( const FGeometry& AllottedGeometry, FArr
 {
 	// The specialized TabWell is dedicated to arranging tabs.
 	// Tabs have uniform sizing (all tabs the same size).
-	// TabWell also ignores widget visibilit, as it is not really
+	// TabWell also ignores widget visibility, as it is not really
 	// relevant.
 
 
@@ -94,6 +94,12 @@ void SDockingTabWell::OnArrangeChildren( const FGeometry& AllottedGeometry, FArr
 		const TSharedRef<SDockTab> CurTab = Tabs[TabIndex];
 		const float ChildWidthWithOverlap = ChildSize.X - CurTab->GetOverlapWidth();
 
+		// This tab being dragged is arranged later.  It should not be arranged twice
+		if (CurTab == TabBeingDragged)
+		{
+			continue;
+		}
+
 		// Is this spot reserved from the tab that is being dragged?
 		if ( TabBeingDragged.IsValid() && XOffset <= DraggedChildCenter && DraggedChildCenter < (XOffset + ChildWidthWithOverlap) )
 		{
@@ -165,16 +171,22 @@ FVector2D SDockingTabWell::ComputeDesiredSize( float ) const
 {
 	FVector2D DesiredSizeResult(0,0);
 
+	TSharedPtr<SDockTab> TabBeingDragged = TabBeingDraggedPtr;
+
 	for ( int32 TabIndex=0; TabIndex < Tabs.Num(); ++TabIndex )
 	{
 		// Currently not respecting Visibility because tabs cannot be invisible.
 		const TSharedRef<SDockTab>& SomeTab = Tabs[TabIndex];
-		const FVector2D SomeTabDesiredSize = SomeTab->GetDesiredSize();
-		DesiredSizeResult.X += SomeTabDesiredSize.X;
-		DesiredSizeResult.Y = FMath::Max(SomeTabDesiredSize.Y, DesiredSizeResult.Y);
+
+		// Tab being dragged is computed later
+		if(SomeTab != TabBeingDragged)
+		{
+			const FVector2D SomeTabDesiredSize = SomeTab->GetDesiredSize();
+			DesiredSizeResult.X += SomeTabDesiredSize.X;
+			DesiredSizeResult.Y = FMath::Max(SomeTabDesiredSize.Y, DesiredSizeResult.Y);
+		}
 	}
 
-	TSharedPtr<SDockTab> TabBeingDragged = TabBeingDraggedPtr;
 	if ( TabBeingDragged.IsValid() )
 	{
 		const FVector2D SomeTabDesiredSize = TabBeingDragged->GetDesiredSize();
@@ -305,6 +317,10 @@ void SDockingTabWell::OnDragEnter( const FGeometry& MyGeometry, const FDragDropE
 
 			// Preview the position of the tab in the TabWell
 			this->TabBeingDraggedPtr = DragDropOperation->GetTabBeingDragged();
+
+			// Add the tab widget to the well when the tab is dragged in
+			Tabs.Add(TabBeingDraggedPtr.ToSharedRef());
+
 			this->TabGrabOffsetFraction = DragDropOperation->GetTabGrabOffsetFraction();
 			
 			// The user should see the contents of the tab that we're dragging.
@@ -326,6 +342,9 @@ void SDockingTabWell::OnDragLeave( const FDragDropEvent& DragDropEvent )
 			// Update the DragAndDrop operation based on this change.
 			const int32 LastForegroundTabIndex = Tabs.Find(TabBeingDragged.ToSharedRef());
 
+			// Remove the tab from the well when it is dragged out
+			Tabs.Remove(TabBeingDraggedPtr.ToSharedRef());
+
 			// The user is pulling a tab out of this TabWell.
 			TabBeingDragged->SetParent();
 
@@ -393,6 +412,9 @@ FReply SDockingTabWell::OnDrop( const FGeometry& MyGeometry, const FDragDropEven
 
 				ensure( DragDropOperation->GetTabBeingDragged().ToSharedRef() == TabBeingDragged );
 
+				// Remove the tab when dropped.  If it was being dragged in this it will be added again in a more permanent way by OpenTab
+				Tabs.Remove(TabBeingDraggedPtr.ToSharedRef());
+
 				// Actually insert the new tab.
 				ParentTabStackPtr.Pin()->OpenTab(TabBeingDragged, DropLocationIndex);
 
@@ -422,7 +444,11 @@ FReply SDockingTabWell::OnMouseButtonUp( const FGeometry& MyGeometry, const FPoi
 		const TSharedRef<SDockTab> TabBeingDragged = TabBeingDraggedPtr.ToSharedRef();
 		this->TabBeingDraggedPtr.Reset();
 		const int32 DropLocationIndex = ComputeChildDropIndex(MyGeometry, TabBeingDragged);
-		Tabs.Insert( TabBeingDragged, DropLocationIndex );
+	
+		// Reorder the tab
+		Tabs.Remove(TabBeingDragged);
+		Tabs.Insert(TabBeingDragged, DropLocationIndex);
+		
 		BringTabToFront(TabBeingDragged);
 		// We are no longer dragging a tab in this tab well, so stop showing it in the TabWell.
 		return FReply::Handled().ReleaseMouseCapture();
