Fix CIS

diff --git a/Engine/Source/Runtime/Slate/Private/Framework/MultiBox/SMenuEntryBlock.cpp b/Engine/Source/Runtime/Slate/Private/Framework/MultiBox/SMenuEntryBlock.cpp
--- a/Engine/Source/Runtime/Slate/Private/Framework/MultiBox/SMenuEntryBlock.cpp
+++ b/Engine/Source/Runtime/Slate/Private/Framework/MultiBox/SMenuEntryBlock.cpp
@@ -932,7 +932,7 @@ void SMenuEntryBlock::BuildMultiBlockWidget(const ISlateStyle* StyleSet, const F
 	if(TutorialName != NAME_None)
 	{
 		TSharedRef<SWidget> ChildWidget = ChildSlot.GetWidget();
-		ChildWidget->AddMetaData<FTagMetaData>(MakeShared<FTagMetaData>(TutorialName));
+		ChildWidget->AddMetadata<FTagMetaData>(MakeShared<FTagMetaData>(TutorialName));
 	}
 
 	// Bind our widget's enabled state to whether or not our action can execute
