Slate - Adding cross thread access protection to some key points in Slate, any time a widget is invalidated or a virtual keyboard accesses slate resources.  Need to probably add similar checks in a few other places.

Matt.Kuhlenschmidt

diff --git a/Engine/Source/Runtime/Slate/Private/Widgets/Text/SlateEditableTextLayout.cpp b/Engine/Source/Runtime/Slate/Private/Widgets/Text/SlateEditableTextLayout.cpp
--- a/Engine/Source/Runtime/Slate/Private/Widgets/Text/SlateEditableTextLayout.cpp
+++ b/Engine/Source/Runtime/Slate/Private/Widgets/Text/SlateEditableTextLayout.cpp
@@ -260,7 +260,7 @@ void FSlateEditableTextLayout::SetText(const TAttribute<FText>& InText)
 
 FText FSlateEditableTextLayout::GetText() const
 {
-	check(IsInGameThread());
+	SLATE_CROSS_THREAD_CHECK();
 
 	return BoundText.Get(FText::GetEmpty());
 }
@@ -330,7 +330,7 @@ void FSlateEditableTextLayout::SetCompositionBrush(const TAttribute<const FSlate
 
 FText FSlateEditableTextLayout::GetPlainText() const
 {
-	check(IsInGameThread());
+	SLATE_CROSS_THREAD_CHECK();
 
 	const TArray< FTextLayout::FLineModel >& Lines = TextLayout->GetLineModels();
 
@@ -349,7 +349,7 @@ FText FSlateEditableTextLayout::GetPlainText() const
 
 bool FSlateEditableTextLayout::SetEditableText(const FText& TextToSet, const bool bForce)
 {
-	check(IsInGameThread());
+	SLATE_CROSS_THREAD_CHECK();
 
 	bool bHasTextChanged = bForce;
 	if (!bHasTextChanged)
@@ -523,7 +523,7 @@ bool FSlateEditableTextLayout::Refresh()
 
 bool FSlateEditableTextLayout::RefreshImpl(const FText* InTextToSet, const bool bForce)
 {
-	check(IsInGameThread());
+	SLATE_CROSS_THREAD_CHECK();
 
 	bool bHasSetText = false;
 
diff --git a/Engine/Source/Runtime/SlateCore/Private/Widgets/SWidget.cpp b/Engine/Source/Runtime/SlateCore/Private/Widgets/SWidget.cpp
--- a/Engine/Source/Runtime/SlateCore/Private/Widgets/SWidget.cpp
+++ b/Engine/Source/Runtime/SlateCore/Private/Widgets/SWidget.cpp
@@ -79,7 +79,7 @@ struct FScopeCycleCounterSWidget
 	{
 	}
 	~FScopeCycleCounterSWidget()
-	{;
+	{
 	}
 };
 
@@ -855,6 +855,8 @@ bool SWidget::IsDirectlyHovered() const
 
 void SWidget::Invalidate(EInvalidateWidget InvalidateReason)
 {
+	SLATE_CROSS_THREAD_CHECK();
+
 	SCOPED_NAMED_EVENT_TEXT("SWidget::Invalidate", FColor::Orange);
 
 	const bool bWasVolatile = IsVolatileIndirectly() || IsVolatile();
@@ -866,7 +868,6 @@ void SWidget::Invalidate(EInvalidateWidget InvalidateReason)
 	}
 
 	LayoutChanged(InvalidateReason);
-
 }
 
 void SWidget::SetCursor( const TAttribute< TOptional<EMouseCursor::Type> >& InCursor )
diff --git a/Engine/Source/Runtime/SlateCore/Public/SlateGlobals.h b/Engine/Source/Runtime/SlateCore/Public/SlateGlobals.h
--- a/Engine/Source/Runtime/SlateCore/Public/SlateGlobals.h
+++ b/Engine/Source/Runtime/SlateCore/Public/SlateGlobals.h
@@ -58,3 +58,7 @@ extern SLATECORE_API int32 GSlateFastWidgetPath;
 *****************************************************************************/
 class FActiveTimerHandle;
 enum class EActiveTimerReturnType : uint8;
+
+
+/** Used to guard access across slate to specific threads */
+#define SLATE_CROSS_THREAD_CHECK() checkf(IsInGameThread() || IsInSlateThread(), TEXT("Slate can only be accessed from the GameThread or the SlateLoadingThread!"));
\ No newline at end of file
