Slate - Temporarily changing MeasuredRangesSizes to FVector4's to get around a compiler bug that leads to sampling off the end of the array.

diff --git a/Engine/Source/Runtime/Slate/Private/Framework/Text/TextLayout.cpp b/Engine/Source/Runtime/Slate/Private/Framework/Text/TextLayout.cpp
--- a/Engine/Source/Runtime/Slate/Private/Framework/Text/TextLayout.cpp
+++ b/Engine/Source/Runtime/Slate/Private/Framework/Text/TextLayout.cpp
@@ -2680,10 +2680,7 @@ TSharedRef< ILayoutBlock > FTextLayout::FRunModel::CreateBlock( const FBlockDefi
 		if ( MeasuredRanges[ StartRangeIndex ].BeginIndex == SizeRange.BeginIndex && 
 			MeasuredRanges[ StartRangeIndex ].EndIndex == SizeRange.EndIndex )
 		{
-			if (MeasuredRangeSizes.IsValidIndex(StartRangeIndex))
-			{
-				BlockSize += MeasuredRangeSizes[StartRangeIndex];
-			}
+			BlockSize += FVector2D(MeasuredRangeSizes[StartRangeIndex]);
 		}
 		else
 		{
@@ -2694,10 +2691,7 @@ TSharedRef< ILayoutBlock > FTextLayout::FRunModel::CreateBlock( const FBlockDefi
 	{
 		if ( MeasuredRanges[ StartRangeIndex ].BeginIndex == SizeRange.BeginIndex )
 		{
-			if (MeasuredRangeSizes.IsValidIndex(StartRangeIndex))
-			{
-				BlockSize += MeasuredRangeSizes[StartRangeIndex];
-			}
+			BlockSize += FVector2D(MeasuredRangeSizes[StartRangeIndex]);
 		}
 		else
 		{
@@ -2706,20 +2700,14 @@ TSharedRef< ILayoutBlock > FTextLayout::FRunModel::CreateBlock( const FBlockDefi
 
 		for (int32 Index = StartRangeIndex + 1; Index < EndRangeIndex; Index++)
 		{
-			if (MeasuredRangeSizes.IsValidIndex(Index))
-			{
-				BlockSize.X += MeasuredRangeSizes[Index].X;
-				BlockSize.Y = FMath::Max(MeasuredRangeSizes[Index].Y, BlockSize.Y);
-			}
+			BlockSize.X += MeasuredRangeSizes[Index].X;
+			BlockSize.Y = FMath::Max(MeasuredRangeSizes[Index].Y, BlockSize.Y);
 		}
 
 		if ( MeasuredRanges[ EndRangeIndex ].EndIndex == SizeRange.EndIndex )
 		{
-			if (MeasuredRangeSizes.IsValidIndex(EndRangeIndex))
-			{
-				BlockSize.X += MeasuredRangeSizes[EndRangeIndex].X;
-				BlockSize.Y = FMath::Max(MeasuredRangeSizes[EndRangeIndex].Y, BlockSize.Y);
-			}
+			BlockSize.X += MeasuredRangeSizes[EndRangeIndex].X;
+			BlockSize.Y = FMath::Max(MeasuredRangeSizes[EndRangeIndex].Y, BlockSize.Y);
 		}
 		else
 		{
@@ -2792,7 +2780,7 @@ FVector2D FTextLayout::FRunModel::Measure(int32 BeginIndex, int32 EndIndex, floa
 	FVector2D Size = Run->Measure(BeginIndex, EndIndex, InScale, InTextContext);
 
 	MeasuredRanges.Add( FTextRange( BeginIndex, EndIndex ) );
-	MeasuredRangeSizes.Add( Size );
+	MeasuredRangeSizes.Add(FVector4(Size, FVector2D::ZeroVector));
 
 	return Size;
 }
@@ -2832,9 +2820,10 @@ TSharedRef< IRun > FTextLayout::FRunModel::GetRun() const
 	return Run;
 }
 
-FTextLayout::FRunModel::FRunModel( const TSharedRef< IRun >& InRun ) : Run( InRun )
+FTextLayout::FRunModel::FRunModel( const TSharedRef< IRun >& InRun )
+	: Run( InRun )
 {
-
+	SLATE_CROSS_THREAD_CHECK();
 }
 
 int32 FTextLayout::FTextOffsetLocations::TextLocationToOffset(const FTextLocation& InLocation) const
diff --git a/Engine/Source/Runtime/Slate/Public/Framework/Text/TextLayout.h b/Engine/Source/Runtime/Slate/Public/Framework/Text/TextLayout.h
--- a/Engine/Source/Runtime/Slate/Public/Framework/Text/TextLayout.h
+++ b/Engine/Source/Runtime/Slate/Public/Framework/Text/TextLayout.h
@@ -262,7 +262,8 @@ public:
 
 		TSharedRef< IRun > Run;
 		TArray< FTextRange > MeasuredRanges;
-		TArray< FVector2D > MeasuredRangeSizes;
+		// HACK: This should be a Vector2D, but changing to a FVector4 to overcome a compiler issue with vectorization running off the end of the array.
+		TArray< FVector4 > MeasuredRangeSizes;
 	};
 
 	struct ELineModelDirtyState
