This is a significant overhaul to the Curve Editor used by Sequencer which adds a plugin-based architecture and extensibility. New tools and toolbar buttons can be added to all usages of the curve edi

diff --git a/Engine/Source/Editor/Sequencer/Private/Tools/SequencerEditTool_Movement.cpp b/Engine/Source/Editor/Sequencer/Private/Tools/SequencerEditTool_Movement.cpp
--- a/Engine/Source/Editor/Sequencer/Private/Tools/SequencerEditTool_Movement.cpp
+++ b/Engine/Source/Editor/Sequencer/Private/Tools/SequencerEditTool_Movement.cpp
@@ -156,10 +156,7 @@ TSharedPtr<ISequencerEditToolDragOperation> FSequencerEditTool_Movement::CreateD
 		// check for that first before trying to figure out if they're resizing or dilating.
 		if (bSectionsSelected && bKeySelected && !bIsDuplicateEvent)
 		{
-			TArray<FSectionHandle> SelectedSectionHandles = SequencerWidget->GetSectionHandles(Selection.GetSelectedSections());
-			TSet<FSequencerSelectedKey> SelectedKeys = Selection.GetSelectedKeys();
-
-			return MakeShareable(new FMoveKeysAndSections(Sequencer, SelectedKeys, SelectedSectionHandles, bHotspotIsSection));
+			return MakeShareable(new FMoveKeysAndSections(Sequencer, Selection.GetSelectedKeys(), Selection.GetSelectedSections(), bHotspotIsSection));
 		}
 		else if (bIsDuplicateEvent)
 		{
@@ -187,57 +184,50 @@ TSharedPtr<ISequencerEditToolDragOperation> FSequencerEditTool_Movement::CreateD
 			}
 			else if (HotspotType == ESequencerHotspot::Section)
 			{
-				FSectionHotspot HoveredSection = StaticCastSharedPtr<FSectionHotspot>(DelayedDrag->Hotspot)->Section;
+				UMovieSceneSection* HoveredSection = StaticCastSharedPtr<FSectionHotspot>(DelayedDrag->Hotspot)->WeakSection.Get();
 
-				if (!Selection.IsSelected(HoveredSection.Section.GetSectionObject()))
+				if (!Selection.IsSelected(HoveredSection))
 				{
 					Selection.EmptySelectedKeys();
 					Selection.EmptySelectedSections();
 					Selection.EmptyNodesWithSelectedKeysOrSections();
-					Selection.AddToSelection(HoveredSection.Section.GetSectionObject());
+					Selection.AddToSelection(HoveredSection);
 					SequencerHelpers::UpdateHoveredNodeFromSelectedSections(Sequencer);
 				}
 			}
 
-			return MakeShareable(new FDuplicateKeysAndSections(Sequencer, Selection.GetSelectedKeys(), SequencerWidget->GetSectionHandles(Selection.GetSelectedSections()), bHotspotIsSection));
+			return MakeShareable(new FDuplicateKeysAndSections(Sequencer, Selection.GetSelectedKeys(), Selection.GetSelectedSections(), bHotspotIsSection));
 		}
 
 
-		TOptional<FSectionHandle> SectionToDrag;
+		UMovieSceneSection* SectionToDrag = nullptr;
 		if (HotspotType == ESequencerHotspot::Section || HotspotType == ESequencerHotspot::EasingArea)
 		{
-			SectionToDrag = StaticCastSharedPtr<FSectionHotspot>(DelayedDrag->Hotspot)->Section;
+			SectionToDrag = StaticCastSharedPtr<FSectionHotspot>(DelayedDrag->Hotspot)->WeakSection.Get();
 		}
 
 		// Moving section(s)?
-		if (SectionToDrag.IsSet())
+		if (SectionToDrag)
 		{
-			UMovieSceneSection* ThisSection = SectionToDrag->GetSectionObject();
-
-			TArray<FSectionHandle> SectionHandles;
-			if (Selection.IsSelected(ThisSection))
-			{
-				SectionHandles = SequencerWidget->GetSectionHandles(Selection.GetSelectedSections());
-			}
-			else
+			if (!Selection.IsSelected(SectionToDrag))
 			{
 				Selection.EmptySelectedKeys();
 				Selection.EmptySelectedSections();
 				Selection.EmptyNodesWithSelectedKeysOrSections();
-				Selection.AddToSelection(ThisSection);
+				Selection.AddToSelection(SectionToDrag);
 				SequencerHelpers::UpdateHoveredNodeFromSelectedSections(Sequencer);
-				SectionHandles.Add(SectionToDrag.GetValue());
 			}
+
 			if (MouseEvent.IsShiftDown())
 			{
 				const bool bDraggingByEnd = false;
 				const bool bIsSlipping = true;
-				return MakeShareable( new FResizeSection( Sequencer, SectionHandles, bDraggingByEnd, bIsSlipping ) );
+				return MakeShareable( new FResizeSection( Sequencer, Selection.GetSelectedSections(), bDraggingByEnd, bIsSlipping ) );
 			}
 			else
 			{
 				TSet<FSequencerSelectedKey> EmptyKeySet;
-				return MakeShareable( new FMoveKeysAndSections( Sequencer, EmptyKeySet, SectionHandles, true) );
+				return MakeShareable( new FMoveKeysAndSections( Sequencer, EmptyKeySet, Selection.GetSelectedSections(), true) );
 			}
 		}
 		// Moving key(s)?
@@ -263,20 +253,20 @@ TSharedPtr<ISequencerEditToolDragOperation> FSequencerEditTool_Movement::CreateD
 				SequencerHelpers::UpdateHoveredNodeFromSelectedKeys(Sequencer);
 			}
 
-			TArray<FSectionHandle> EmptySectionHandles;
-			return MakeShareable( new FMoveKeysAndSections( Sequencer, Selection.GetSelectedKeys(), EmptySectionHandles, false) );
+			TSet<TWeakObjectPtr<UMovieSceneSection>> NoSections;
+			return MakeShareable( new FMoveKeysAndSections( Sequencer, Selection.GetSelectedKeys(), NoSections, false) );
 		}
 	}
 	// If we're not dragging a hotspot, sections take precedence over keys
 	else if (Selection.GetSelectedSections().Num())
 	{
 		TSet<FSequencerSelectedKey> EmptyKeySet;
-		return MakeShareable( new FMoveKeysAndSections( Sequencer, EmptyKeySet, SequencerWidget->GetSectionHandles(Selection.GetSelectedSections()), true ) );
+		return MakeShareable( new FMoveKeysAndSections( Sequencer, EmptyKeySet, Selection.GetSelectedSections(), true ) );
 	}
 	else if (Selection.GetSelectedKeys().Num())
 	{
-		TArray<FSectionHandle> EmptySectionHandles;
-		return MakeShareable( new FMoveKeysAndSections( Sequencer, Selection.GetSelectedKeys(), EmptySectionHandles, false) );
+		TSet<TWeakObjectPtr<UMovieSceneSection>> NoSections;
+		return MakeShareable( new FMoveKeysAndSections( Sequencer, Selection.GetSelectedKeys(), NoSections, false) );
 	}
 
 	return nullptr;
diff --git a/Engine/Source/Editor/Sequencer/Private/Tools/SequencerEditTool_Movement.h b/Engine/Source/Editor/Sequencer/Private/Tools/SequencerEditTool_Movement.h
--- a/Engine/Source/Editor/Sequencer/Private/Tools/SequencerEditTool_Movement.h
+++ b/Engine/Source/Editor/Sequencer/Private/Tools/SequencerEditTool_Movement.h
@@ -53,7 +53,9 @@ private:
 		FDelayedDrag_Hotspot(FVector2D InInitialPosition, FKey InApplicableKey, TSharedPtr<ISequencerHotspot> InHotspot)
 			: FDelayedDrag(InInitialPosition, InApplicableKey)
 			, Hotspot(MoveTemp(InHotspot))
-		{ }
+		{
+			SetTriggerScaleFactor(0.1f);
+		}
 
 		TSharedPtr<ISequencerHotspot> Hotspot;
 	};
diff --git a/Engine/Source/Editor/Sequencer/Private/Tools/SequencerEditTool_Selection.h b/Engine/Source/Editor/Sequencer/Private/Tools/SequencerEditTool_Selection.h
--- a/Engine/Source/Editor/Sequencer/Private/Tools/SequencerEditTool_Selection.h
+++ b/Engine/Source/Editor/Sequencer/Private/Tools/SequencerEditTool_Selection.h
@@ -49,7 +49,9 @@ private:
 		FDelayedDrag_Hotspot(FVector2D InInitialPosition, FKey InApplicableKey, TSharedPtr<ISequencerHotspot> InHotspot)
 			: FDelayedDrag(InInitialPosition, InApplicableKey)
 			, Hotspot(MoveTemp(InHotspot))
-		{}
+		{
+			SetTriggerScaleFactor(0.1f);
+		}
 
 		TSharedPtr<ISequencerHotspot> Hotspot;
 	};
diff --git a/Engine/Source/Editor/Sequencer/Public/IKeyArea.h b/Engine/Source/Editor/Sequencer/Public/IKeyArea.h
--- a/Engine/Source/Editor/Sequencer/Public/IKeyArea.h
+++ b/Engine/Source/Editor/Sequencer/Public/IKeyArea.h
@@ -41,7 +41,15 @@ public:
 	 * @param InSection The section that owns the channel that this key area represents
 	 * @param InChannel Handle to the channel this key area represents
 	 */
-	IKeyArea(UMovieSceneSection& InSection, FMovieSceneChannelHandle InChannel);
+	IKeyArea(UMovieSceneSection* InSection, FMovieSceneChannelHandle InChannel);
+
+	/**
+	 * Reinitialize this key area with a new channel
+	 *
+	 * @param InSection The section that owns the channel that this key area represents
+	 * @param InChannel Handle to the channel this key area represents
+	 */
+	void Reinitialize(UMovieSceneSection* InSection, FMovieSceneChannelHandle InChannel);
 
 public:
 
diff --git a/Engine/Source/Editor/Sequencer/Public/ISectionLayoutBuilder.h b/Engine/Source/Editor/Sequencer/Public/ISectionLayoutBuilder.h
--- a/Engine/Source/Editor/Sequencer/Public/ISectionLayoutBuilder.h
+++ b/Engine/Source/Editor/Sequencer/Public/ISectionLayoutBuilder.h
@@ -6,6 +6,8 @@
 
 class IKeyArea;
 
+struct FMovieSceneChannelHandle;
+
 /**
  * Builds an inner layout for a section
  */
@@ -24,18 +26,16 @@ public:
 	
 	/**
 	 * Sets the section as a key area itself
-	 * @param KeyArea		Interface for accessing and drawing keys
+	 * @param Channel		The channel that is to be assigned as the top level channel for this section
 	 */
-	virtual void SetSectionAsKeyArea( TSharedRef<IKeyArea> KeyArea ) = 0;
+	virtual void SetTopLevelChannel( const FMovieSceneChannelHandle& Channel ) = 0;
 
 	/**
-	 * Adds a key area onto the layout.  If a category is pushed, the key area will appear as a child of the current category
+	 * Adds a channel onto the layout. If a category is pushed, the key area will appear as a child of the current category
 	 *
-	 * @param KeyAreaName	The name of the key area
-	 * @param DisplayLabel	The localized display label for the key area
-	 * @param KeyArea		Interface for accessing and drawing keys
+	 * @param Channel		A handle to the channel to be added to the layout
 	 */
-	virtual void AddKeyArea( FName KeyAreaName, const FText& DisplayLabel, TSharedRef<IKeyArea> KeyArea ) = 0;
+	virtual void AddChannel( const FMovieSceneChannelHandle& Channel ) = 0;
 
 	/**
 	 * Pops a category off the stack
diff --git a/Engine/Source/Editor/Sequencer/Public/ISequencerChannelInterface.h b/Engine/Source/Editor/Sequencer/Public/ISequencerChannelInterface.h
--- a/Engine/Source/Editor/Sequencer/Public/ISequencerChannelInterface.h
+++ b/Engine/Source/Editor/Sequencer/Public/ISequencerChannelInterface.h
@@ -141,6 +141,11 @@ struct ISequencerChannelInterface
 	 */
 	virtual void DrawKeys_Raw(FMovieSceneChannel* Channel, TArrayView<const FKeyHandle> InKeyHandles, TArrayView<FKeyDrawParams> OutKeyDrawParams) const = 0;
 
+	/**
+	 * Whether this channel supports curve models
+	 */
+	virtual bool SupportsCurveEditorModels_Raw(const FMovieSceneChannelHandle& InChannel) const = 0;
+
 	/**
 	 * Create a new model for this channel that can be used on the curve editor interface
 	 *
diff --git a/Engine/Source/Editor/Sequencer/Public/SequencerChannelInterface.h b/Engine/Source/Editor/Sequencer/Public/SequencerChannelInterface.h
--- a/Engine/Source/Editor/Sequencer/Public/SequencerChannelInterface.h
+++ b/Engine/Source/Editor/Sequencer/Public/SequencerChannelInterface.h
@@ -139,6 +139,15 @@ struct TSequencerChannelInterfaceCommon : ISequencerChannelInterface
 		DrawKeys(static_cast<ChannelType*>(InChannel), InKeyHandles, OutKeyDrawParams);
 	}
 
+	/**
+	 * Whether this channel supports curve models
+	 */
+	virtual bool SupportsCurveEditorModels_Raw(const FMovieSceneChannelHandle& InChannel) const override
+	{
+		using namespace Sequencer;
+		return SupportsCurveEditorModels(InChannel.Cast<ChannelType>());
+	}
+
 	/**
 	 * Create a new model for this channel that can be used on the curve editor interface
 	 *
diff --git a/Engine/Source/Editor/Sequencer/Public/SequencerChannelTraits.h b/Engine/Source/Editor/Sequencer/Public/SequencerChannelTraits.h
--- a/Engine/Source/Editor/Sequencer/Public/SequencerChannelTraits.h
+++ b/Engine/Source/Editor/Sequencer/Public/SequencerChannelTraits.h
@@ -335,7 +335,10 @@ namespace Sequencer
 		Section->SetRange(NewRange);
 	}
 
-
+	/**
+	 * Whether the specified channel handle supports curve models or not
+	 */
+	SEQUENCER_API bool SupportsCurveEditorModels(const FMovieSceneChannelHandle& ChannelHandle);
 
 	/**
 	 * Create a new model for the specified channel that can be used on the curve editor interface
diff --git a/Engine/Source/Editor/Sequencer/Public/SequencerCommands.h b/Engine/Source/Editor/Sequencer/Public/SequencerCommands.h
--- a/Engine/Source/Editor/Sequencer/Public/SequencerCommands.h
+++ b/Engine/Source/Editor/Sequencer/Public/SequencerCommands.h
@@ -343,10 +343,6 @@ public:
 	/** Toggle whether we should evaluate sub sequences in isolation */
 	TSharedPtr< FUICommandInfo > ToggleEvaluateSubSequencesInIsolation;
 
-	TSharedPtr<FUICommandInfo> SetAllCurveVisibility;
-	TSharedPtr<FUICommandInfo> SetSelectedCurveVisibility;
-	TSharedPtr<FUICommandInfo> SetAnimatedCurveVisibility;
-
 	/**
 	 * Initialize commands
 	 */
diff --git a/Engine/Source/Editor/Sequencer/Public/SequencerSettings.h b/Engine/Source/Editor/Sequencer/Public/SequencerSettings.h
--- a/Engine/Source/Editor/Sequencer/Public/SequencerSettings.h
+++ b/Engine/Source/Editor/Sequencer/Public/SequencerSettings.h
@@ -15,20 +15,6 @@ enum class EAllowEditsMode : uint8;
 enum class EKeyGroupMode : uint8;
 enum class EMovieSceneKeyInterpolation : uint8;
 
-
-
-/** Defines visibility states for the curves in the curve editor. */
-UENUM()
-enum class ECurveEditorCurveVisibility : uint8
-{
-	/** All curves should be visible. */
-	AllCurves,
-	/** Only curves from selected nodes should be visible. */
-	SelectedCurves,
-	/** Only curves which have keyframes should be visible. */
-	AnimatedCurves
-};
-
 UENUM()
 enum ESequencerSpawnPosition
 {
@@ -106,7 +92,6 @@ public:
 	DECLARE_MULTICAST_DELEGATE( FOnEvaluateSubSequencesInIsolationChanged );
 	DECLARE_MULTICAST_DELEGATE( FOnShowSelectedNodesOnlyChanged );
 	DECLARE_MULTICAST_DELEGATE_OneParam( FOnAllowEditsModeChanged, EAllowEditsMode );
-	DECLARE_MULTICAST_DELEGATE( FOnCurveEditorCurveVisibilityChanged );
 	DECLARE_MULTICAST_DELEGATE(FOnLoopStateChanged);
 
 	virtual void PostEditChangeProperty(FPropertyChangedEvent& PropertyChangedEvent) override;
@@ -321,13 +306,6 @@ public:
 
 	uint32 GetTrajectoryPathCap() const { return TrajectoryPathCap; }
 
-	/** Gets the current curve visibility. */
-	ECurveEditorCurveVisibility GetCurveVisibility() const;
-	/** Sets the current curve visibility. */
-	void SetCurveVisibility(ECurveEditorCurveVisibility InCurveVisibility);
-
-	FOnCurveEditorCurveVisibilityChanged& GetOnCurveEditorCurveVisibilityChanged();
-
 	FOnLoopStateChanged& GetOnLoopStateChanged();
 
 	/** What format should we display the UI controls in when representing time in a sequence? */
@@ -503,10 +481,6 @@ protected:
 	UPROPERTY(config, EditAnywhere, Category=General)
 	EFrameNumberDisplayFormats FrameNumberDisplayFormat;
 
-	/** Specifies which curves to show in the curve editor */
-	UPROPERTY(config, EditAnywhere, Category=General)
-	ECurveEditorCurveVisibility CurveVisibility;
-	FOnCurveEditorCurveVisibilityChanged OnCurveEditorCurveVisibilityChanged;
 	FOnEvaluateSubSequencesInIsolationChanged OnEvaluateSubSequencesInIsolationChangedEvent;
 	FOnShowSelectedNodesOnlyChanged OnShowSelectedNodesOnlyChangedEvent;
 	FOnAllowEditsModeChanged OnAllowEditsModeChangedEvent;
diff --git a/Engine/Source/Editor/Sequencer/Public/SequencerTimeSliderController.h b/Engine/Source/Editor/Sequencer/Public/SequencerTimeSliderController.h
--- a/Engine/Source/Editor/Sequencer/Public/SequencerTimeSliderController.h
+++ b/Engine/Source/Editor/Sequencer/Public/SequencerTimeSliderController.h
@@ -15,41 +15,6 @@ struct FSlateBrush;
 class FSequencer;
 class IPropertyTypeCustomization;
 
-struct FPaintPlaybackRangeArgs
-{
-	FPaintPlaybackRangeArgs()
-		: StartBrush(nullptr), EndBrush(nullptr), BrushWidth(0.f), SolidFillOpacity(0.f)
-	{}
-
-	FPaintPlaybackRangeArgs(const FSlateBrush* InStartBrush, const FSlateBrush* InEndBrush, float InBrushWidth)
-		: StartBrush(InStartBrush), EndBrush(InEndBrush), BrushWidth(InBrushWidth)
-	{}
-	/** Brush to use for the start bound */
-	const FSlateBrush* StartBrush;
-	/** Brush to use for the end bound */
-	const FSlateBrush* EndBrush;
-	/** The width of the above brushes, in slate units */
-	float BrushWidth;
-	/** level of opacity for the fill color between the range markers */
-	float SolidFillOpacity;
-};
-
-struct FPaintSectionAreaViewArgs
-{
-	FPaintSectionAreaViewArgs()
-		: bDisplayTickLines(false), bDisplayScrubPosition(false), bDisplayMarkedFrames(false)
-	{}
-
-	/** Whether to display tick lines */
-	bool bDisplayTickLines;
-	/** Whether to display the scrub position */
-	bool bDisplayScrubPosition;
-	/** Whether to display the marked frames */
-	bool bDisplayMarkedFrames;
-	/** Optional Paint args for the playback range*/
-	TOptional<FPaintPlaybackRangeArgs> PlaybackRangeArgs;
-};
-
 /**
  * A time slider controller for sequencer
  * Draws and manages time data for a Sequencer
@@ -94,13 +59,19 @@ public:
 	/** Get the current play range for this controller */
 	virtual TRange<FFrameNumber> GetPlayRange() const override { return TimeSliderArgs.PlaybackRange.Get(TRange<FFrameNumber>()); }
 
+	/** Get the current time for the Scrub handle which indicates what range is being evaluated. */
+	virtual FFrameTime GetScrubPosition() const override { return TimeSliderArgs.ScrubPosition.Get(FFrameTime()); }
+
+	/** Get the current time for the Scrub handle which indicates what range is being evaluated. */
+	virtual void SetScrubPosition(FFrameTime InTime) override { CommitScrubPosition(InTime, false); } 
+
 	/**
 	 * Clamp the given range to the clamp range 
 	 *
 	 * @param NewRangeMin		The new lower bound of the range
 	 * @param NewRangeMax		The new upper bound of the range
 	 */	
-	void ClampViewRange(double& NewRangeMin, double& NewRangeMax);
+	virtual void ClampViewRange(double& NewRangeMin, double& NewRangeMax);
 
 	/**
 	 * Set a new range based on a min, max and an interpolation mode
@@ -145,7 +116,9 @@ public:
 	/**
 	 * Draws major tick lines in the section view                                                              
 	 */
-	int32 OnPaintSectionView( const FGeometry& AllottedGeometry, const FSlateRect& MyCullingRect, FSlateWindowElementList& OutDrawElements, int32 LayerId, bool bEnabled, const FPaintSectionAreaViewArgs& Args ) const;
+	int32 OnPaintViewArea( const FGeometry& AllottedGeometry, const FSlateRect& MyCullingRect, FSlateWindowElementList& OutDrawElements, int32 LayerId, bool bEnabled, const FPaintViewAreaArgs& Args ) const;
+
+
 
 public:
 
diff --git a/Engine/Source/Editor/SequencerWidgets/Public/ITimeSlider.h b/Engine/Source/Editor/SequencerWidgets/Public/ITimeSlider.h
--- a/Engine/Source/Editor/SequencerWidgets/Public/ITimeSlider.h
+++ b/Engine/Source/Editor/SequencerWidgets/Public/ITimeSlider.h
@@ -169,11 +169,51 @@ struct FTimeSliderArgs
 	TSharedPtr<INumericTypeInterface<double>> NumericTypeInterface;
 };
 
+struct FPaintPlaybackRangeArgs
+{
+	FPaintPlaybackRangeArgs()
+		: StartBrush(nullptr), EndBrush(nullptr), BrushWidth(0.f), SolidFillOpacity(0.f)
+	{}
+
+	FPaintPlaybackRangeArgs(const FSlateBrush* InStartBrush, const FSlateBrush* InEndBrush, float InBrushWidth)
+		: StartBrush(InStartBrush), EndBrush(InEndBrush), BrushWidth(InBrushWidth)
+	{}
+	/** Brush to use for the start bound */
+	const FSlateBrush* StartBrush;
+	/** Brush to use for the end bound */
+	const FSlateBrush* EndBrush;
+	/** The width of the above brushes, in slate units */
+	float BrushWidth;
+	/** level of opacity for the fill color between the range markers */
+	float SolidFillOpacity;
+};
+
+struct FPaintViewAreaArgs
+{
+	FPaintViewAreaArgs()
+		: bDisplayTickLines(false), bDisplayScrubPosition(false), bDisplayMarkedFrames(false)
+	{}
+
+	/** Whether to display tick lines */
+	bool bDisplayTickLines;
+	/** Whether to display the scrub position */
+	bool bDisplayScrubPosition;
+	/** Whether to display the marked frames */
+	bool bDisplayMarkedFrames;
+	/** Optional Paint args for the playback range*/
+	TOptional<FPaintPlaybackRangeArgs> PlaybackRangeArgs;
+};
+
+
 class ITimeSliderController : public ISequencerInputHandler
 {
 public:
 	virtual ~ITimeSliderController(){}
+
 	virtual int32 OnPaintTimeSlider( bool bMirrorLabels, const FGeometry& AllottedGeometry, const FSlateRect& MyCullingRect, FSlateWindowElementList& OutDrawElements, int32 LayerId, const FWidgetStyle& InWidgetStyle, bool bParentEnabled ) const = 0;
+
+	virtual int32 OnPaintViewArea( const FGeometry& AllottedGeometry, const FSlateRect& MyCullingRect, FSlateWindowElementList& OutDrawElements, int32 LayerId, bool bEnabled, const FPaintViewAreaArgs& Args ) const = 0;
+
 	virtual FCursorReply OnCursorQuery( TSharedRef<const SWidget> WidgetOwner, const FGeometry& MyGeometry, const FPointerEvent& CursorEvent ) const = 0;
 
 	/** Get the current play rate for this controller */
@@ -191,6 +231,12 @@ public:
 	/** Get the current play range for this controller */
 	virtual TRange<FFrameNumber> GetPlayRange() const { return TRange<FFrameNumber>(); }
 
+	/** Get the current time for the Scrub handle which indicates what range is being evaluated. */
+	virtual FFrameTime GetScrubPosition() const { return FFrameTime(); }
+
+	/** Set the current time for the Scrub handle which indicates what range is being evaluated. */
+	virtual void SetScrubPosition(FFrameTime InTime) {}
+
 	/**
 	 * Set a new range based on a min, max and an interpolation mode
 	 * 
diff --git a/Engine/Source/Editor/SkeletonEditor/Private/SSkeletonTree.cpp b/Engine/Source/Editor/SkeletonEditor/Private/SSkeletonTree.cpp
--- a/Engine/Source/Editor/SkeletonEditor/Private/SSkeletonTree.cpp
+++ b/Engine/Source/Editor/SkeletonEditor/Private/SSkeletonTree.cpp
@@ -589,6 +589,7 @@ void SSkeletonTree::CreateTreeColumns()
 		.OnMouseButtonDoubleClick(this, &SSkeletonTree::OnTreeDoubleClick)
 		.OnSetExpansionRecursive(this, &SSkeletonTree::SetTreeItemExpansionRecursive)
 		.ItemHeight(24)
+		.HighlightParentNodesForSelection(true)
 		.HeaderRow
 		(
 			TreeHeaderRow
diff --git a/Engine/Source/Editor/UnrealEd/Private/Toolkits/AssetEditorToolkit.cpp b/Engine/Source/Editor/UnrealEd/Private/Toolkits/AssetEditorToolkit.cpp
--- a/Engine/Source/Editor/UnrealEd/Private/Toolkits/AssetEditorToolkit.cpp
+++ b/Engine/Source/Editor/UnrealEd/Private/Toolkits/AssetEditorToolkit.cpp
@@ -41,25 +41,27 @@ FAssetEditorToolkit::FAssetEditorToolkit()
 	, bCheckDirtyOnAssetSave(false)
 	, AssetEditorModeManager(nullptr)
 	, bIsToolbarFocusable(false)
+	, bIsToolbarUsingSmallIcons(false)
 {
 	WorkspaceMenuCategory = FWorkspaceItem::NewGroup(LOCTEXT("WorkspaceMenu_BaseAssetEditor", "Asset Editor"));
 }
 
-void FAssetEditorToolkit::InitAssetEditor( const EToolkitMode::Type Mode, const TSharedPtr< class IToolkitHost >& InitToolkitHost, const FName AppIdentifier, const TSharedRef<FTabManager::FLayout>& StandaloneDefaultLayout, const bool bCreateDefaultStandaloneMenu, const bool bCreateDefaultToolbar, UObject* ObjectToEdit, const bool bInIsToolbarFocusable )
+void FAssetEditorToolkit::InitAssetEditor( const EToolkitMode::Type Mode, const TSharedPtr< class IToolkitHost >& InitToolkitHost, const FName AppIdentifier, const TSharedRef<FTabManager::FLayout>& StandaloneDefaultLayout, const bool bCreateDefaultStandaloneMenu, const bool bCreateDefaultToolbar, UObject* ObjectToEdit, const bool bInIsToolbarFocusable, const bool bInUseSmallToolbarIcons )
 {
 	TArray< UObject* > ObjectsToEdit;
 	ObjectsToEdit.Add( ObjectToEdit );
 
-	InitAssetEditor( Mode, InitToolkitHost, AppIdentifier, StandaloneDefaultLayout, bCreateDefaultStandaloneMenu, bCreateDefaultToolbar, ObjectsToEdit, bInIsToolbarFocusable );
+	InitAssetEditor( Mode, InitToolkitHost, AppIdentifier, StandaloneDefaultLayout, bCreateDefaultStandaloneMenu, bCreateDefaultToolbar, ObjectsToEdit, bInIsToolbarFocusable, bInUseSmallToolbarIcons );
 }
 
-void FAssetEditorToolkit::InitAssetEditor( const EToolkitMode::Type Mode, const TSharedPtr< class IToolkitHost >& InitToolkitHost, const FName AppIdentifier, const TSharedRef<FTabManager::FLayout>& StandaloneDefaultLayout, const bool bCreateDefaultStandaloneMenu, const bool bCreateDefaultToolbar, const TArray<UObject*>& ObjectsToEdit, const bool bInIsToolbarFocusable )
+void FAssetEditorToolkit::InitAssetEditor( const EToolkitMode::Type Mode, const TSharedPtr< class IToolkitHost >& InitToolkitHost, const FName AppIdentifier, const TSharedRef<FTabManager::FLayout>& StandaloneDefaultLayout, const bool bCreateDefaultStandaloneMenu, const bool bCreateDefaultToolbar, const TArray<UObject*>& ObjectsToEdit, const bool bInIsToolbarFocusable, const bool bInUseSmallToolbarIcons )
 {
 	// Must not already be editing an object
 	check( ObjectsToEdit.Num() > 0 );
 	check( EditingObjects.Num() == 0 );
 
 	bIsToolbarFocusable = bInIsToolbarFocusable;
+	bIsToolbarUsingSmallIcons = bInUseSmallToolbarIcons;
 
 	// cache reference to ToolkitManager; also ensure it was initialized.
 	FToolkitManager& ToolkitManager = FToolkitManager::Get();
@@ -951,7 +953,7 @@ void FAssetEditorToolkit::GenerateToolbar()
 {
 	TSharedPtr<FExtender> Extender = FExtender::Combine(ToolbarExtenders);
 
-	FToolBarBuilder ToolbarBuilder( GetToolkitCommands(), FMultiBoxCustomization::AllowCustomization( GetToolkitFName() ), Extender);
+	FToolBarBuilder ToolbarBuilder( GetToolkitCommands(), FMultiBoxCustomization::AllowCustomization( GetToolkitFName() ), Extender, Orient_Horizontal, bIsToolbarUsingSmallIcons);
 	ToolbarBuilder.SetIsFocusable(bIsToolbarFocusable);
 	ToolbarBuilder.BeginSection("Asset");
 	{
diff --git a/Engine/Source/Editor/UnrealEd/Public/SColorGradientEditor.h b/Engine/Source/Editor/UnrealEd/Public/SColorGradientEditor.h
--- a/Engine/Source/Editor/UnrealEd/Public/SColorGradientEditor.h
+++ b/Engine/Source/Editor/UnrealEd/Public/SColorGradientEditor.h
@@ -42,7 +42,7 @@ public:
 
 };
 
-class SColorGradientEditor : public SLeafWidget
+class UNREALED_API SColorGradientEditor : public SLeafWidget
 {
 public:
 	SLATE_BEGIN_ARGS( SColorGradientEditor ) 
diff --git a/Engine/Source/Editor/UnrealEd/Public/Toolkits/AssetEditorToolkit.h b/Engine/Source/Editor/UnrealEd/Public/Toolkits/AssetEditorToolkit.h
--- a/Engine/Source/Editor/UnrealEd/Public/Toolkits/AssetEditorToolkit.h
+++ b/Engine/Source/Editor/UnrealEd/Public/Toolkits/AssetEditorToolkit.h
@@ -105,9 +105,10 @@ public:
 	 * @param	bCreateDefaultStandaloneMenu	True if in standalone mode, the asset editor should automatically generate a default "asset" menu, or false if you're going to do this yourself in your derived asset editor's implementation
 	 * @param	ObjectToEdit			The object to edit
 	 * @param	bInIsToolbarFocusable	Whether the buttons on the default toolbar can receive keyboard focus
+	 * @param	bUseSmallToolbarIcons	Whether the buttons on the default toolbar use the small icons
 	 */
-	virtual void InitAssetEditor(const EToolkitMode::Type Mode, const TSharedPtr<IToolkitHost>& InitToolkitHost, const FName AppIdentifier, const TSharedRef<FTabManager::FLayout>& StandaloneDefaultLayout, const bool bCreateDefaultStandaloneMenu, const bool bCreateDefaultToolbar, const TArray<UObject*>& ObjectsToEdit, const bool bInIsToolbarFocusable = false);
-	virtual void InitAssetEditor(const EToolkitMode::Type Mode, const TSharedPtr<IToolkitHost>& InitToolkitHost, const FName AppIdentifier, const TSharedRef<FTabManager::FLayout>& StandaloneDefaultLayout, const bool bCreateDefaultStandaloneMenu, const bool bCreateDefaultToolbar, UObject* ObjectToEdit, const bool bInIsToolbarFocusable = false);
+	virtual void InitAssetEditor(const EToolkitMode::Type Mode, const TSharedPtr<IToolkitHost>& InitToolkitHost, const FName AppIdentifier, const TSharedRef<FTabManager::FLayout>& StandaloneDefaultLayout, const bool bCreateDefaultStandaloneMenu, const bool bCreateDefaultToolbar, const TArray<UObject*>& ObjectsToEdit, const bool bInIsToolbarFocusable = false, const bool bInUseSmallToolbarIcons = false);
+	virtual void InitAssetEditor(const EToolkitMode::Type Mode, const TSharedPtr<IToolkitHost>& InitToolkitHost, const FName AppIdentifier, const TSharedRef<FTabManager::FLayout>& StandaloneDefaultLayout, const bool bCreateDefaultStandaloneMenu, const bool bCreateDefaultToolbar, UObject* ObjectToEdit, const bool bInIsToolbarFocusable = false, const bool bInUseSmallToolbarIcons = false);
 
 	/** Virtual destructor, so that we can clean up our app when destroyed */
 	virtual ~FAssetEditorToolkit();
@@ -392,6 +393,9 @@ private:
 	/** Whether the buttons on the default toolbar can receive keyboard focus */
 	bool bIsToolbarFocusable;
 
+	/** Whether the buttons on the default toolbar use small icons */
+	bool bIsToolbarUsingSmallIcons;
+
 	/**	The tab ids for all the tabs used */
 	static const FName ToolbarTabId;
 };
diff --git a/Engine/Source/Runtime/AudioMixer/Public/DSP/PassiveFilter.h b/Engine/Source/Runtime/AudioMixer/Public/DSP/PassiveFilter.h
--- a/Engine/Source/Runtime/AudioMixer/Public/DSP/PassiveFilter.h
+++ b/Engine/Source/Runtime/AudioMixer/Public/DSP/PassiveFilter.h
@@ -5,10 +5,11 @@
 #include "CoreMinimal.h"
 #include "DSP/AudioFFT.h"
 #include "Async/ParallelFor.h"
+#include "Audio.h"
 
 namespace Audio
 {
-	struct PassiveFilterParams
+	struct FPassiveFilterParams
 	{
 		enum class EClass
 		{
@@ -31,7 +32,7 @@ namespace Audio
 
 		bool bRemoveDC;
 
-		PassiveFilterParams()
+		FPassiveFilterParams()
 			: Class(EClass::Butterworth)
 			, Type(EType::Lowpass)
 			, Order(4)
@@ -87,20 +88,20 @@ namespace Audio
 		}
 	}
 
-	static float GetGainForFrequency(float NormalizedFreq, const PassiveFilterParams& InParams)
+	static float GetGainForFrequency(float NormalizedFreq, const FPassiveFilterParams& InParams)
 	{
-		const float FrequencyRatio = (InParams.Type == PassiveFilterParams::EType::Lowpass) ? (NormalizedFreq / InParams.NormalizedCutoffFrequency) : (InParams.NormalizedCutoffFrequency / NormalizedFreq);
+		const float FrequencyRatio = (InParams.Type == FPassiveFilterParams::EType::Lowpass) ? (NormalizedFreq / InParams.NormalizedCutoffFrequency) : (InParams.NormalizedCutoffFrequency / NormalizedFreq);
 
 		switch (InParams.Class)
 		{
-		case PassiveFilterParams::EClass::Chebyshev:
+		case FPassiveFilterParams::EClass::Chebyshev:
 		{
 			const float ChebyshevPolynomial = EvaluateChebyshevPolynomial(FrequencyRatio, InParams.Order);
 			return InParams.UnitGain / (FMath::Sqrt(1.0f + ChebyshevPolynomial * ChebyshevPolynomial));
 			break;
 		}
 		default:
-		case PassiveFilterParams::EClass::Butterworth:
+		case FPassiveFilterParams::EClass::Butterworth:
 		{
 			const float Denominator = FMath::Sqrt((1 + FMath::Pow(FrequencyRatio, InParams.Order * 2.0f)));
 			return InParams.UnitGain / Denominator;
@@ -112,7 +113,7 @@ namespace Audio
 	/**
 	* This can be called on any TArrayView<float> whose length is a power of 2.
 	*/
-	static void Filter(TArrayView<float>& Signal, const PassiveFilterParams& InParams)
+	static void Filter(TArrayView<float>& Signal, const FPassiveFilterParams& InParams)
 	{
 		if (!FMath::IsPowerOfTwo(Signal.Num()))
 		{
@@ -165,7 +166,7 @@ namespace Audio
 	/**
 	* Static function for applying a filter to any time series.
 	*/
-	static void Filter(TArray<float>& Signal, const PassiveFilterParams& InParams)
+	static void Filter(TArray<float>& Signal, const FPassiveFilterParams& InParams)
 	{
 		if (FMath::IsPowerOfTwo(Signal.Num()))
 		{
diff --git a/Engine/Source/Runtime/Engine/Classes/Curves/KeyHandle.h b/Engine/Source/Runtime/Engine/Classes/Curves/KeyHandle.h
--- a/Engine/Source/Runtime/Engine/Classes/Curves/KeyHandle.h
+++ b/Engine/Source/Runtime/Engine/Classes/Curves/KeyHandle.h
@@ -31,6 +31,11 @@ struct FKeyHandle
 		return A.Index < B.Index;
 	}
 
+	friend bool operator>(FKeyHandle A, FKeyHandle B)
+	{
+		return A.Index > B.Index;
+	}
+
 	friend uint32 GetTypeHash(const FKeyHandle& Handle)
 	{
 		return GetTypeHash(Handle.Index);
diff --git a/Engine/Source/Runtime/MovieScene/Public/Channels/MovieSceneChannelHandle.h b/Engine/Source/Runtime/MovieScene/Public/Channels/MovieSceneChannelHandle.h
--- a/Engine/Source/Runtime/MovieScene/Public/Channels/MovieSceneChannelHandle.h
+++ b/Engine/Source/Runtime/MovieScene/Public/Channels/MovieSceneChannelHandle.h
@@ -25,6 +25,43 @@ struct MOVIESCENE_API FMovieSceneChannelHandle
 	 */
 	FMovieSceneChannelHandle(TWeakPtr<FMovieSceneChannelProxy> InWeakChannelProxy, FName InChannelTypeName, int32 InChannelIndex);
 
+public:
+
+	/**
+	 * Compare this channel handle for equality against another. Will return false for a handle that is to a different proxy even if it might have pointed to the same data (ie, one handle has expired)
+	 */
+	friend bool operator==(const FMovieSceneChannelHandle& A, const FMovieSceneChannelHandle& B)
+	{
+		if (A.ChannelIndex == B.ChannelIndex && A.ChannelTypeName == B.ChannelTypeName)
+		{
+			FMovieSceneChannelProxy* ProxyA = A.WeakChannelProxy.Pin().Get();
+			FMovieSceneChannelProxy* ProxyB = B.WeakChannelProxy.Pin().Get();
+
+			// Cannot be equal if either handle is expired
+			return ProxyA && ProxyA == ProxyB;
+		}
+		return false;
+	}
+
+	/**
+	 * Compare this channel handle for inequality against another. Will return true for a handle that is to a different proxy even if it might have pointed to the same data (ie, one handle has expired)
+	 */
+	friend bool operator!=(const FMovieSceneChannelHandle& A, const FMovieSceneChannelHandle& B)
+	{
+		if (A.ChannelIndex != B.ChannelIndex || A.ChannelTypeName != B.ChannelTypeName)
+		{
+			return false;
+		}
+
+		FMovieSceneChannelProxy* ProxyA = A.WeakChannelProxy.Pin().Get();
+		FMovieSceneChannelProxy* ProxyB = B.WeakChannelProxy.Pin().Get();
+
+		// Expired channel handles always compare falsey
+		return !ProxyA || !ProxyB || ProxyA != ProxyB;
+	}
+
+public:
+
 	/**
 	 * Cast this handle to a handle of a related type. Callee is responsible for ensuring that the type's are compatible.
 	 */
diff --git a/Engine/Source/Runtime/MovieSceneTracks/Public/Channels/MovieSceneEvent.h b/Engine/Source/Runtime/MovieSceneTracks/Public/Channels/MovieSceneEvent.h
--- a/Engine/Source/Runtime/MovieSceneTracks/Public/Channels/MovieSceneEvent.h
+++ b/Engine/Source/Runtime/MovieSceneTracks/Public/Channels/MovieSceneEvent.h
@@ -27,7 +27,7 @@ class UMovieSceneEventSectionBase;
  *         - Objects provided by the playback context.
  *     - Will only trigger if the specified object is of the same type as the parameter (or interface)
  */
-USTRUCT()
+USTRUCT(BlueprintType)
 struct MOVIESCENETRACKS_API FMovieSceneEvent
 {
 	GENERATED_BODY()
@@ -54,7 +54,7 @@ struct MOVIESCENETRACKS_API FMovieSceneEvent
 	 * The function that should be called to invoke this event.
 	 * Functions must have either no parameters, or a single, pass-by-value object/interface parameter, with no return parameter.
 	 */
-	UPROPERTY()
+	UPROPERTY(EditDefaultsOnly, Category="Sequencer")
 	FName FunctionName;
 
 #if WITH_EDITORONLY_DATA
diff --git a/Engine/Source/Runtime/Slate/Public/Framework/DelayedDrag.h b/Engine/Source/Runtime/Slate/Public/Framework/DelayedDrag.h
--- a/Engine/Source/Runtime/Slate/Public/Framework/DelayedDrag.h
+++ b/Engine/Source/Runtime/Slate/Public/Framework/DelayedDrag.h
@@ -15,6 +15,7 @@ public:
 	FDelayedDrag(FVector2D InInitialPosition, FKey InEffectiveKey)
 		: InitialPosition(InInitialPosition), bHasInitiatedDrag(false), DistanceDragged(0), EffectiveKey(InEffectiveKey)
 	{
+		TriggerDistance = FSlateApplication::Get().GetDragTriggerDistance();
 	}
 
 	/** Get the initial start position (before any drag has started) */
@@ -32,7 +33,7 @@ public:
 		if (!bHasInitiatedDrag && MouseEvent.IsMouseButtonDown(EffectiveKey))
 		{
 			DistanceDragged += MouseEvent.GetCursorDelta().Size();
-			if (!bHasInitiatedDrag && DistanceDragged > FSlateApplication::Get().GetDragTriggerDistance())
+			if (!bHasInitiatedDrag && DistanceDragged > TriggerDistance)
 			{
 				ForceDragStart();
 				return true;
@@ -41,6 +42,12 @@ public:
 		return false;
 	}
 
+	/** Assign a new scale factor to apply to the drag trigger distance */
+	void SetTriggerScaleFactor(float InTriggerScaleFactor)
+	{
+		TriggerDistance = FMath::Max(FSlateApplication::Get().GetDragTriggerDistance() * InTriggerScaleFactor, 1.f);
+	}
+
 protected:
 
 	/** The initial position of the drag start */
@@ -49,6 +56,8 @@ protected:
 	bool bHasInitiatedDrag;
 	/** The amount we have dragged */
 	float DistanceDragged;
+	/** The minimum distance that must be moved before the drag initiates. */
+	float TriggerDistance;
 	/** The key that must be pressed to initiate the drag */
 	FKey EffectiveKey;
 };
diff --git a/Engine/Source/Runtime/Slate/Public/Framework/Views/ITypedTableView.h b/Engine/Source/Runtime/Slate/Public/Framework/Views/ITypedTableView.h
--- a/Engine/Source/Runtime/Slate/Public/Framework/Views/ITypedTableView.h
+++ b/Engine/Source/Runtime/Slate/Public/Framework/Views/ITypedTableView.h
@@ -72,6 +72,15 @@ public:
 	 */
 	virtual bool Private_IsItemSelected( const ItemType& TheItem ) const = 0;
 
+	/**
+	* Test if the current item should be highlighted. This is separate from hover highlights.
+	*
+	* @param TheItem The item to test.
+	*
+	* @return true if the item is marked as highlighted in this list; false otherwise.
+	*/
+	virtual bool Private_IsItemHighlighted(const ItemType& TheItem) const = 0;
+	
 	/**
 	 * Set the selection state of an item. Does not cause an OnSelectionChanged()!
 	 *
@@ -144,6 +153,19 @@ public:
 	/** @return the number of items that are selected */
 	virtual int32 Private_GetNumSelectedItems() const = 0;
 
+	/**
+	* Enable a soft highlight on the element. This is useful for explaining parent/child relationships without actually modifying selection. Unrelated to any hover highlights.
+	*
+	* @param TheItem         		The item whose highlighted state to change.
+	* @param bShouldBeHighlighted   Enables a highlight on this item if true, otherwise disables the highlight.
+	*/
+	virtual void Private_SetItemHighlighted(ItemType TheItem, bool bShouldBeHighlighted) = 0;
+
+	/**
+	* Empty the set of highlighted items.
+	*/
+	virtual void Private_ClearHighlightedItems() = 0;
+
 	/**
 	 * @param ItemIndexInList  The index of the data item in the linearized array.
 	 *
diff --git a/Engine/Source/Runtime/Slate/Public/Framework/Views/TableViewTypeTraits.h b/Engine/Source/Runtime/Slate/Public/Framework/Views/TableViewTypeTraits.h
--- a/Engine/Source/Runtime/Slate/Public/Framework/Views/TableViewTypeTraits.h
+++ b/Engine/Source/Runtime/Slate/Public/Framework/Views/TableViewTypeTraits.h
@@ -126,6 +126,11 @@ public:
 		return InPtr;
 	}
 
+	static FString DebugDump( TSharedPtr<T> InPtr )
+	{
+		return InPtr.IsValid() ? FString::Printf(TEXT("0x%08x"), InPtr.Get()) : FString(TEXT("nullptr"));
+	}
+
 	class SerializerType{};
 };
 
@@ -163,6 +168,11 @@ public:
 		return InPtr;
 	}
 
+	static FString DebugDump( TSharedPtr<T, ESPMode::ThreadSafe> InPtr )
+	{
+		return InPtr.IsValid() ? FString::Printf(TEXT("0x%08x"), InPtr.Get()) : FString(TEXT("nullptr"));
+	}
+
 	class SerializerType{};
 };
 
@@ -201,6 +211,11 @@ public:
 		return InPtr.ToSharedRef();
 	}
 
+	static FString DebugDump( TSharedRef<T> InPtr )
+	{
+		return FString::Printf(TEXT("0x%08x"), &InPtr.Get());
+	}
+
 	class SerializerType{};
 };
 
@@ -238,6 +253,11 @@ public:
 		return InPtr.ToSharedRef();
 	}
 
+	static FString DebugDump( TSharedRef<T, ESPMode::ThreadSafe> InPtr )
+	{
+		return FString::Printf(TEXT("0x%08x"), &InPtr.Get());
+	}
+
 	class SerializerType{};
 };
 
@@ -278,6 +298,12 @@ public:
 		return InPtr;
 	}
 
+	static FString DebugDump( TWeakObjectPtr<T> InPtr )
+	{
+		T* ObjPtr = InPtr.Get();
+		return ObjPtr ? FString::Printf(TEXT("0x%08x [%s]"), ObjPtr, *ObjPtr->GetName()) : FString(TEXT("nullptr"));
+	}
+
 	class SerializerType{};
 };
 
@@ -314,6 +340,11 @@ public:
 
 	static T* NullableItemTypeConvertToItemType( T* InPtr ) { return InPtr; }
 
+	static FString DebugDump( T* InPtr )
+	{
+		return InPtr ? FString::Printf(TEXT("0x%08x [%s]"), InPtr, *InPtr->GetName()) : FString(TEXT("nullptr"));
+	}
+
 	typedef FGCObject SerializerType;
 };
 
@@ -344,5 +375,10 @@ public:
 
 	static const T* NullableItemTypeConvertToItemType( const T* InPtr ) { return InPtr; }
 
+	static FString DebugDump( const T* InPtr )
+	{
+		return InPtr ? FString::Printf(TEXT("0x%08x [%s]"), InPtr, *InPtr->GetName()) : FString(TEXT("nullptr"));
+	}
+
 	typedef FGCObject SerializerType;
 };
diff --git a/Engine/Source/Runtime/Slate/Public/Widgets/Views/SListView.h b/Engine/Source/Runtime/Slate/Public/Widgets/Views/SListView.h
--- a/Engine/Source/Runtime/Slate/Public/Widgets/Views/SListView.h
+++ b/Engine/Source/Runtime/Slate/Public/Widgets/Views/SListView.h
@@ -772,6 +772,11 @@ public:
 		return nullptr != SelectedItems.Find(TheItem);
 	}
 
+	virtual bool Private_IsItemHighlighted(const ItemType& TheItem) const override
+	{
+		return nullptr != HighlightedItems.Find(TheItem);
+	}
+
 	virtual bool Private_IsItemExpanded( const ItemType& TheItem ) const override
 	{
 		// List View does not support item expansion.
@@ -799,6 +804,23 @@ public:
 		return SelectedItems.Num();
 	}
 
+	virtual void Private_SetItemHighlighted(ItemType TheItem, bool bShouldBeHighlighted)
+	{
+		if (bShouldBeHighlighted)
+		{
+			HighlightedItems.Add(TheItem);
+		}
+		else
+		{
+			HighlightedItems.Remove(TheItem);
+		}
+	}
+
+	virtual void Private_ClearHighlightedItems()
+	{
+		HighlightedItems.Empty();
+	}
+
 	virtual int32 Private_GetNestingDepth( int32 ItemIndexInList ) const override
 	{
 		// List View items are not indented
@@ -1188,6 +1210,27 @@ public:
 		Private_SignalSelectionChanged(ESelectInfo::Direct);
 	}
 
+
+
+	/**
+	* Set the highlighted state of an item.
+	*
+	* @param TheItem      The Item whose highlight state you wish to modify
+	* @param bHighlighted True to enable the soft parent highlight, false to disable it.
+	*/
+	void SetItemHighlighted(const ItemType& TheItem, bool bHighlighted)
+	{
+		Private_SetItemHighlighted(TheItem, bHighlighted);
+	}
+
+	/**
+	* Empty the highlighted item set.
+	*/
+	void ClearHighlightedItems()
+	{
+		Private_ClearHighlightedItems();
+	}
+
 	/**
 	 * Gets the number of selected items.
 	 *
@@ -1385,20 +1428,9 @@ public:
 
 protected:
 
-	FOnItemToString_Debug GetDefaultDebugDelegate()
+	static FOnItemToString_Debug GetDefaultDebugDelegate()
 	{
-		return
-		FOnItemToString_Debug::CreateLambda([](ItemType InItem)
-		{
-			if (TListTypeTraits<ItemType>::IsPtrValid(InItem))
-			{
-				return FString::Printf(TEXT("0x%08x"), &(*InItem));
-			}
-			else
-			{
-				return FString(TEXT("nullptr"));
-			}
-		});
+		return FOnItemToString_Debug::CreateStatic(TListTypeTraits<ItemType>::DebugDump);
 	}
 
 	/**
@@ -1740,6 +1772,9 @@ protected:
 	/** The item which was last manipulated; used as a start for shift-click selection */
 	NullableItemType RangeSelectionStart;
 
+	/** A set of which items should be highlighted */
+	TItemSet HighlightedItems;
+
 	/** Pointer to the array of data items that we are observing */
 	const TArray<ItemType>* ItemsSource;
 
diff --git a/Engine/Source/Runtime/Slate/Public/Widgets/Views/STableRow.h b/Engine/Source/Runtime/Slate/Public/Widgets/Views/STableRow.h
--- a/Engine/Source/Runtime/Slate/Public/Widgets/Views/STableRow.h
+++ b/Engine/Source/Runtime/Slate/Public/Widgets/Views/STableRow.h
@@ -126,7 +126,7 @@ DECLARE_DELEGATE_RetVal_OneParam(FReply, FOnTableRowDrop, FDragDropEvent const&)
 
 /**
  * The ListView is populated by Selectable widgets.
- * A Selectable widget is away of the ListView containing it (OwnerWidget) and holds arbitrary Content (Content).
+ * A Selectable widget is a way of the ListView containing it (OwnerWidget) and holds arbitrary Content (Content).
  * A Selectable works with its corresponding ListView to provide selection functionality.
  */
 template<typename ItemType>
@@ -839,6 +839,10 @@ public:
 		check(MyItem);
 
 		const bool bIsSelected = OwnerWidget->Private_IsItemSelected( *MyItem );
+		const bool bIsHighlighted = OwnerWidget->Private_IsItemHighlighted( *MyItem );
+
+		const bool bAllowSelection = GetSelectionMode() != ESelectionMode::None;
+		const bool bEvenEntryIndex = (IndexInList % 2 == 0);
 
 		if (bIsSelected && bShowSelection)
 		{
@@ -855,11 +859,25 @@ public:
 					: &Style->InactiveBrush;
 			}
 		}
+		else if (!bIsSelected && bIsHighlighted)
+		{
+			if (bIsActive)
+			{
+				return IsHovered()
+					? (bEvenEntryIndex ? &Style->EvenRowBackgroundHoveredBrush : &Style->OddRowBackgroundHoveredBrush)
+					: &Style->ActiveHighlightedBrush;
+			}
+			else
+			{
+				return IsHovered()
+					? (bEvenEntryIndex ? &Style->EvenRowBackgroundHoveredBrush : &Style->OddRowBackgroundHoveredBrush)
+					: &Style->InactiveHighlightedBrush;
+			}
+		}
 		else
 		{
 			// Add a slightly lighter background for even rows
-			const bool bAllowSelection = GetSelectionMode() != ESelectionMode::None;
-			if( IndexInList % 2 == 0 )
+			if(bEvenEntryIndex)
 			{
 				return ( IsHovered() && bAllowSelection )
 					? &Style->EvenRowBackgroundHoveredBrush
diff --git a/Engine/Source/Runtime/Slate/Public/Widgets/Views/STreeView.h b/Engine/Source/Runtime/Slate/Public/Widgets/Views/STreeView.h
--- a/Engine/Source/Runtime/Slate/Public/Widgets/Views/STreeView.h
+++ b/Engine/Source/Runtime/Slate/Public/Widgets/Views/STreeView.h
@@ -45,10 +45,11 @@ struct FItemInfo
 	{		
 	}
 
-	FItemInfo(TBitArray<> InNeedsVerticalWire, bool InHasChildren, bool InIsLastChild )
+	FItemInfo(TBitArray<> InNeedsVerticalWire, bool InHasChildren, bool InIsLastChild, int32 InParentIndex )
 	: NeedsVerticalWire(InNeedsVerticalWire)
 	, bHasChildren( InHasChildren )
 	, bIsLastChild( InIsLastChild )
+	, ParentIndex( InParentIndex )
 	{
 	}
 
@@ -66,6 +67,9 @@ struct FItemInfo
 
 	/** Is this the last child of its parent? If so, it gets a special kind of wire/connector. */
 	uint32 bIsLastChild : 1;
+
+	/** Index into the linearized tree of the parent for this item, if any, otherwise INDEX_NONE. */
+	int32 ParentIndex;
 };
 
 
@@ -129,6 +133,7 @@ public:
 		, _HandleGamepadEvents(true)
 		, _HandleDirectionalNavigation(true)
 		, _AllowInvisibleItemSelection(false)
+		, _HighlightParentNodesForSelection(false)
 		{
 			//_Clipping = EWidgetClipping::ClipToBounds;
 		}
@@ -186,6 +191,7 @@ public:
 
 		SLATE_ARGUMENT(bool, AllowInvisibleItemSelection);
 
+		SLATE_ARGUMENT(bool, HighlightParentNodesForSelection);
 
 	SLATE_END_ARGS()
 
@@ -225,6 +231,7 @@ public:
 		this->bHandleGamepadEvents = InArgs._HandleGamepadEvents;
 		this->bHandleDirectionalNavigation = InArgs._HandleDirectionalNavigation;
 		this->bAllowInvisibleItemSelection = InArgs._AllowInvisibleItemSelection;
+		this->bHighlightParentNodesForSelection = InArgs._HighlightParentNodesForSelection;
 
 		// Check for any parameters that the coder forgot to specify.
 		FString ErrorString;
@@ -433,7 +440,47 @@ private:
 			: false;
 	}
 
+	/**
+	 * This clears the highlight state from all nodes and then traverses the parents of each selected node to add it to the highlighted set.
+	 */
+	virtual void Private_UpdateParentHighlights()
+	{
+		this->Private_ClearHighlightedItems();
+		
+		for( typename TItemSet::TConstIterator SelectedItemIt(this->SelectedItems); SelectedItemIt; ++SelectedItemIt )
+		{
+			// Sometimes selection events can come through before the Linearized List is built, so the item may not exist yet.
+			int32 ItemIndex = LinearizedItems.Find(*SelectedItemIt);
+			if (ItemIndex == INDEX_NONE)
+			{
+				continue;
+			}
+
+			const FItemInfo& ItemInfo = DenseItemInfos[ItemIndex];
+			int32 ParentIndex = ItemInfo.ParentIndex;
+			while (ParentIndex != INDEX_NONE)
+			{
+				const ItemType& ParentItem = this->LinearizedItems[ParentIndex];
+				this->Private_SetItemHighlighted(ParentItem, true);
+
+				const FItemInfo& ParentItemInfo = DenseItemInfos[ParentIndex];
+				ParentIndex = ParentItemInfo.ParentIndex;
+			}
+		}
+	}
+
 public:
+	virtual void Private_SignalSelectionChanged(ESelectInfo::Type SelectInfo) override
+	{
+		SListView< ItemType >::Private_SignalSelectionChanged(SelectInfo);
+		
+		if (bHighlightParentNodesForSelection)
+		{
+			this->Private_UpdateParentHighlights();
+		}
+	}
+
+
 
 	/**
 	 * See SWidget::Tick()
@@ -462,7 +509,7 @@ public:
 					
 					// Rebuild the linearized view of the tree data.
 					LinearizedItems.Empty();
-					PopulateLinearizedItems( *TreeItemsSource, LinearizedItems, TempDenseItemInfos, TBitArray<>(), TempSelectedItemsMap, TempSparseItemInfo, true );
+					PopulateLinearizedItems( *TreeItemsSource, LinearizedItems, TempDenseItemInfos, TBitArray<>(), TempSelectedItemsMap, TempSparseItemInfo, true, INDEX_NONE );
 
 					if( !bAllowInvisibleItemSelection &&
 						(this->SelectedItems.Num() != TempSelectedItemsMap.Num() ||
@@ -492,6 +539,12 @@ public:
 					// and could be invalid)
 					SparseItemInfos = TempSparseItemInfo;
 					DenseItemInfos  = TempDenseItemInfos;
+
+					// Once the selection changed events have gone through we can update the parent highlight statuses, which are based on your current selection.
+					if (bHighlightParentNodesForSelection)
+					{
+						this->Private_UpdateParentHighlights();
+					}
 				}
 			}
 		}
@@ -503,7 +556,7 @@ public:
 		
 	/** 
 	 * Given: an array of items (ItemsSource) each of which potentially has a child.
-	 * Task: populate the LinearezedItems array with a flattened version of the visible data items.
+	 * Task: populate the LinearizedItems array with a flattened version of the visible data items.
 	 *       In the process, remove any items that are not visible while maintaining any collapsed
 	 *       items that may have expanded children.
 	 *
@@ -514,6 +567,7 @@ public:
 	 * @param OutNewSelectedItems  Selected items minus any items that are no longer observed by the list.
 	 * @param NewSparseItemInfo    Expanded items and items that have expanded children minus any items that are no longer observed by the list.
 	 * @param bAddingItems         Are we adding encountered items to the linearized items list or just testing them for existence.
+	 * @param ParentIndex		   The index in the resulting linearized item list of the parent node for the currently processed level.
 	 *
 	 * @return true if we encountered expanded children; false otherwise.
 	 */
@@ -524,7 +578,8 @@ public:
 		TBitArray<> NeedsParentWire,
 		TItemSet& OutNewSelectedItems,
 		TSparseItemMap& NewSparseItemInfo,
-		bool bAddingItems )
+		bool bAddingItems,
+		int32 ParentIndex)
 	{
 
 		NeedsParentWire.Add(false);
@@ -556,7 +611,7 @@ public:
 			{
 				InLinearizedItems.Add( CurItem );
 
-				NewDenseItemInfos.Add( FItemInfo(NeedsParentWire, bHasChildren, bIsLastChild) );
+				NewDenseItemInfos.Add( FItemInfo(NeedsParentWire, bHasChildren, bIsLastChild, ParentIndex) );
 
 				const bool bIsSelected = this->IsItemSelected( CurItem );
 				if (bIsSelected)
@@ -570,7 +625,7 @@ public:
 				// If this item is expanded, we should process all of its children at the next indentation level.
 				// If it is collapsed, process its children but do not add them to the linearized list.
 				const bool bAddChildItems = bAddingItems && bIsExpanded;
-				bHasExpandedChildren = PopulateLinearizedItems( ChildItems, InLinearizedItems, NewDenseItemInfos, NeedsParentWire, OutNewSelectedItems, NewSparseItemInfo, bAddChildItems );
+				bHasExpandedChildren = PopulateLinearizedItems( ChildItems, InLinearizedItems, NewDenseItemInfos, NeedsParentWire, OutNewSelectedItems, NewSparseItemInfo, bAddChildItems, InLinearizedItems.Num() - 1);
 			}
 
 			if ( bIsExpanded || bHasExpandedChildren )
@@ -729,4 +784,7 @@ private:
 
 	/** true if we allow invisible items to stay selected. */
 	bool bAllowInvisibleItemSelection = false;
+
+	/** true if we should highlight all parents for each of the currently selected items */
+	bool bHighlightParentNodesForSelection = false;
 };
diff --git a/Engine/Source/Runtime/SlateCore/Private/Styling/CoreStyle.cpp b/Engine/Source/Runtime/SlateCore/Private/Styling/CoreStyle.cpp
--- a/Engine/Source/Runtime/SlateCore/Private/Styling/CoreStyle.cpp
+++ b/Engine/Source/Runtime/SlateCore/Private/Styling/CoreStyle.cpp
@@ -35,6 +35,7 @@ public:
 		, SelectionColor_LinearRef(MakeShareable(new FLinearColor(0.728f, 0.364f, 0.003f)))
 		, SelectionColor_Inactive_LinearRef(MakeShareable(new FLinearColor(0.25f, 0.25f, 0.25f)))
 		, SelectionColor_Pressed_LinearRef(MakeShareable(new FLinearColor(0.701f, 0.225f, 0.003f)))
+		, HighlightColor_LinearRef(MakeShareable(new FLinearColor(0.068f, 0.068f, 0.068f)))
 	{
 	}
 
@@ -53,6 +54,7 @@ public:
 	const TSharedRef<FLinearColor> SelectionColor_LinearRef;
 	const TSharedRef<FLinearColor> SelectionColor_Inactive_LinearRef;
 	const TSharedRef<FLinearColor> SelectionColor_Pressed_LinearRef;
+	const TSharedRef<FLinearColor> HighlightColor_LinearRef;
 };
 
 
@@ -174,6 +176,7 @@ TSharedRef<ISlateStyle> FCoreStyle::Create( const FName& InStyleSetName )
 	const FSlateColor SelectionColor( Style->SelectionColor_LinearRef );
 	const FSlateColor SelectionColor_Inactive( Style->SelectionColor_Inactive_LinearRef );
 	const FSlateColor SelectionColor_Pressed( Style->SelectionColor_Pressed_LinearRef );
+	const FSlateColor HighlightColor( Style->HighlightColor_LinearRef );
 
 	Style->Set("DefaultAppIcon", new IMAGE_BRUSH("Icons/DefaultAppIcon", Icon24x24));
 
@@ -759,6 +762,9 @@ TSharedRef<ISlateStyle> FCoreStyle::Create( const FName& InStyleSetName )
 			.SetActiveHoveredBrush(IMAGE_BRUSH("Common/Selection", Icon8x8, SelectionColor))
 			.SetInactiveBrush(IMAGE_BRUSH("Common/Selection", Icon8x8, SelectionColor_Inactive))
 			.SetInactiveHoveredBrush(IMAGE_BRUSH("Common/Selection", Icon8x8, SelectionColor_Inactive))
+			.SetActiveHighlightedBrush(IMAGE_BRUSH("Common/Selection", Icon8x8, HighlightColor))
+			.SetInactiveHighlightedBrush(IMAGE_BRUSH("Common/Selection", Icon8x8, FSlateColor(FLinearColor(.1f, .1f, .1f))))
+
 			.SetTextColor(DefaultForeground)
 			.SetSelectedTextColor(InvertedForeground)
 			.SetDropIndicator_Above(BOX_BRUSH("Common/DropZoneIndicator_Above", FMargin(10.0f / 16.0f, 10.0f / 16.0f, 0, 0), SelectionColor))
diff --git a/Engine/Source/Runtime/SlateCore/Public/Rendering/DrawElements.h b/Engine/Source/Runtime/SlateCore/Public/Rendering/DrawElements.h
--- a/Engine/Source/Runtime/SlateCore/Public/Rendering/DrawElements.h
+++ b/Engine/Source/Runtime/SlateCore/Public/Rendering/DrawElements.h
@@ -35,6 +35,12 @@ struct FSlateGradientStop
 	FVector2D Position;
 	FLinearColor Color;
 
+	/**
+	* Construct a Gradient Stop from a Position and a Color.
+	* @param InPosition - The position in widget space for this stop. Both X and Y are used for a single-axis gradient.
+						  A two stop gradient should go from (0,0), to (Width,Height).
+	* @param InColor	- The color to lerp towards at this stop.
+	*/
 	FSlateGradientStop( const FVector2D& InPosition, const FLinearColor& InColor )
 		: Position(InPosition)
 		, Color(InColor)
diff --git a/Engine/Source/Runtime/SlateCore/Public/Styling/SlateTypes.h b/Engine/Source/Runtime/SlateCore/Public/Styling/SlateTypes.h
--- a/Engine/Source/Runtime/SlateCore/Public/Styling/SlateTypes.h
+++ b/Engine/Source/Runtime/SlateCore/Public/Styling/SlateTypes.h
@@ -1187,6 +1187,16 @@ struct SLATECORE_API FTableRowStyle : public FSlateWidgetStyle
 	UPROPERTY(EditAnywhere, BlueprintReadWrite, Category = Appearance)
 	FSlateBrush DropIndicator_Below;
 	FTableRowStyle& SetDropIndicator_Below(const FSlateBrush& InValue){ DropIndicator_Below = InValue; return *this; }
+	
+	/** Brush used when a highlighted row is active */
+	UPROPERTY(EditAnywhere, BlueprintReadWrite, Category = Appearance)
+	FSlateBrush ActiveHighlightedBrush;
+	FTableRowStyle& SetActiveHighlightedBrush( const FSlateBrush& InActiveHighlightedBrush ){ ActiveHighlightedBrush = InActiveHighlightedBrush; return *this; }
+	
+	/** Brush used when a highlighted row is inactive and hovered */
+	UPROPERTY(EditAnywhere, BlueprintReadWrite, Category = Appearance)
+	FSlateBrush InactiveHighlightedBrush;
+	FTableRowStyle& SetInactiveHighlightedBrush( const FSlateBrush& InInactiveHighlightedBrush){ InactiveHighlightedBrush = InInactiveHighlightedBrush; return *this; }
 };
 
 
diff --git a/Engine/Source/Runtime/TimeManagement/Public/FrameNumberNumericInterface.h b/Engine/Source/Runtime/TimeManagement/Public/FrameNumberNumericInterface.h
--- a/Engine/Source/Runtime/TimeManagement/Public/FrameNumberNumericInterface.h
+++ b/Engine/Source/Runtime/TimeManagement/Public/FrameNumberNumericInterface.h
@@ -137,6 +137,34 @@ struct FFrameNumberInterface : public INumericTypeInterface<double>
 				return TOptional<double>(FrameResult.GetValue().GetFrame().Value);
 			}
 
+			static FBasicMathExpressionEvaluator Parser;
+
+			// If not parsed, try the math expression evaluator
+			if (FallbackFormat == EFrameNumberDisplayFormats::Seconds)
+			{
+				double TimeInSeconds = DestinationFrameRate.AsSeconds(FFrameTime::FromDecimal(InExistingValue));
+
+				TValueOrError<double, FExpressionError> Result = Parser.Evaluate(*InString, TimeInSeconds);
+				if (Result.IsValid())
+				{
+					FFrameTime ResultTime = DestinationFrameRate.AsFrameTime(Result.GetValue());
+
+					return TOptional<double>(ResultTime.GetFrame().Value);
+				}
+			}
+			else if (FallbackFormat == EFrameNumberDisplayFormats::Frames)
+			{
+				FFrameTime ExistingTime = FFrameRate::TransformTime(FFrameTime::FromDecimal(InExistingValue), DestinationFrameRate, SourceFrameRate);
+
+				TValueOrError<double, FExpressionError> Result = Parser.Evaluate(*InString, (double)(ExistingTime.GetFrame().Value));
+				if (Result.IsValid())
+				{
+					FFrameTime ResultTime = FFrameRate::TransformTime(FFrameTime::FromDecimal(Result.GetValue()), SourceFrameRate, DestinationFrameRate);
+
+					return TOptional<double>(ResultTime.GetFrame().Value);
+				}
+			}
+
 			// Whatever they entered wasn't understood by any of our parsers, so it was probably malformed or had letters, etc.
 			return TOptional<double>();
 		}
