Add focus change callback to navigation config and fixed some places sending wrong focus reason. This can be used to override focus changes or play a sound
Add NavigationAction enum and functions to navigation config to replace use of Virtual_Accept/back. This can be overridden by a game to do something like look up an input action mapping and allow multiple accept keys
Fixed SSlider to use navigation instead of key input, this fixes it so analog deadzone works properly on them and they handle axis remapping
Added some functions to KismetInputLibrary to query navigation, this is useful for UMG widgets that handle input themselves

diff --git a/Engine/Plugins/Runtime/Synthesis/Source/Synthesis/Private/UI/SSynth2DSlider.cpp b/Engine/Plugins/Runtime/Synthesis/Source/Synthesis/Private/UI/SSynth2DSlider.cpp
--- a/Engine/Plugins/Runtime/Synthesis/Source/Synthesis/Private/UI/SSynth2DSlider.cpp
+++ b/Engine/Plugins/Runtime/Synthesis/Source/Synthesis/Private/UI/SSynth2DSlider.cpp
@@ -2,6 +2,7 @@
 
 #include "UI/SSynth2DSlider.h"
 #include "Rendering/DrawElements.h"
+#include "Framework/Application/SlateApplication.h"
 
 void SSynth2DSlider::Construct(const SSynth2DSlider::FArguments& InDeclaration)
 {
@@ -125,7 +126,7 @@ FReply SSynth2DSlider::OnKeyDown(const FGeometry& MyGeometry, const FKeyEvent& I
 		// The controller's bottom face button must be pressed once to begin manipulating the slider's value.
 		// Navigation away from the widget is prevented until the button has been pressed again or focus is lost.
 		// The value can be manipulated by using the game pad's directional arrows ( relative to slider orientation ).
-		if (KeyPressed == EKeys::Enter || KeyPressed == EKeys::SpaceBar || KeyPressed == EKeys::Virtual_Accept)
+		if (FSlateApplication::Get().GetNavigationActionForKey(KeyPressed) == EUINavigationAction::Accept)
 		{
 			if (bControllerInputCaptured == false)
 			{
diff --git a/Engine/Plugins/Runtime/Synthesis/Source/Synthesis/Private/UI/SSynthKnob.cpp b/Engine/Plugins/Runtime/Synthesis/Source/Synthesis/Private/UI/SSynthKnob.cpp
--- a/Engine/Plugins/Runtime/Synthesis/Source/Synthesis/Private/UI/SSynthKnob.cpp
+++ b/Engine/Plugins/Runtime/Synthesis/Source/Synthesis/Private/UI/SSynthKnob.cpp
@@ -2,6 +2,7 @@
 
 #include "UI/SSynthKnob.h"
 #include "Rendering/DrawElements.h"
+#include "Framework/Application/SlateApplication.h"
 
 #define LOCTEXT_NAMESPACE "SynthKnob"
 
@@ -122,7 +123,7 @@ FReply SSynthKnob::OnKeyDown(const FGeometry& MyGeometry, const FKeyEvent& InKey
 		// The controller's bottom face button must be pressed once to begin manipulating the slider's value.
 		// Navigation away from the widget is prevented until the button has been pressed again or focus is lost.
 		// The value can be manipulated by using the game pad's directional arrows ( relative to slider orientation ).
-		if (KeyPressed == EKeys::Enter || KeyPressed == EKeys::SpaceBar || KeyPressed == EKeys::Virtual_Accept)
+		if (FSlateApplication::Get().GetNavigationActionForKey(KeyPressed) == EUINavigationAction::Accept)
 		{
 			if (bControllerInputCaptured == false)
 			{
diff --git a/Engine/Source/Runtime/Engine/Classes/Kismet/KismetInputLibrary.h b/Engine/Source/Runtime/Engine/Classes/Kismet/KismetInputLibrary.h
--- a/Engine/Source/Runtime/Engine/Classes/Kismet/KismetInputLibrary.h
+++ b/Engine/Source/Runtime/Engine/Classes/Kismet/KismetInputLibrary.h
@@ -90,6 +90,18 @@ class ENGINE_API UKismetInputLibrary : public UBlueprintFunctionLibrary
 	UFUNCTION(BlueprintPure, meta = (DisplayName = "Is Valid"), Category = "Utilities|Key")
 	static bool Key_IsValid(const FKey& Key);
 
+	/** Returns the navigation action corresponding to this key, or Invalid if not found */
+	UFUNCTION(BlueprintPure, meta = (DisplayName = "Get Key Navigation Action"), Category = "Utilities|Key")
+	static EUINavigationAction Key_GetNavigationAction(const FKey& InKey);
+
+	/** Returns the navigation action corresponding to this key, or Invalid if not found */
+	UFUNCTION(BlueprintPure, meta = (DisplayName = "Get Key Event Navigation Direction"), Category = "Utilities|KeyEvent")
+	static EUINavigation Key_GetNavigationDirectionFromKey(const FKeyEvent& InKeyEvent);
+
+	/** Returns the navigation action corresponding to this key, or Invalid if not found */
+	UFUNCTION(BlueprintPure, meta = (DisplayName = "Get Analog Event Navigation Direction"), Category = "Utilities|AnalogEvent")
+	static EUINavigation Key_GetNavigationDirectionFromAnalog(const FAnalogInputEvent& InAnalogEvent);
+
 	/**
 	 * Returns the display name of the key.
 	 */
diff --git a/Engine/Source/Runtime/Engine/Private/KismetInputLibrary.cpp b/Engine/Source/Runtime/Engine/Private/KismetInputLibrary.cpp
--- a/Engine/Source/Runtime/Engine/Private/KismetInputLibrary.cpp
+++ b/Engine/Source/Runtime/Engine/Private/KismetInputLibrary.cpp
@@ -3,6 +3,7 @@
 #include "Kismet/KismetInputLibrary.h"
 #include "EngineGlobals.h"
 #include "Engine/Engine.h"
+#include "Framework/Application/SlateApplication.h"
 
 
 //////////////////////////////////////////////////////////////////////////
@@ -67,6 +68,37 @@ bool UKismetInputLibrary::Key_IsValid(const FKey& Key)
 	return Key.IsValid();
 }
 
+EUINavigationAction UKismetInputLibrary::Key_GetNavigationAction(const FKey& InKey)
+{
+	if (FSlateApplication::IsInitialized())
+	{
+		return FSlateApplication::Get().GetNavigationActionForKey(InKey);
+	}
+
+	return EUINavigationAction::Invalid;
+}
+
+EUINavigation UKismetInputLibrary::Key_GetNavigationDirectionFromKey(const FKeyEvent& InKeyEvent)
+{
+	if (FSlateApplication::IsInitialized())
+	{
+		return FSlateApplication::Get().GetNavigationDirectionFromKey(InKeyEvent);
+	}
+
+	return EUINavigation::Invalid;
+}
+
+EUINavigation UKismetInputLibrary::Key_GetNavigationDirectionFromAnalog(const FAnalogInputEvent& InAnalogEvent)
+{
+	if (FSlateApplication::IsInitialized())
+	{
+		return FSlateApplication::Get().GetNavigationDirectionFromAnalog(InAnalogEvent);
+	}
+
+	return EUINavigation::Invalid;
+}
+
+
 FText UKismetInputLibrary::Key_GetDisplayName(const FKey& Key)
 {
 	return Key.GetDisplayName();
diff --git a/Engine/Source/Runtime/Engine/Private/Slate/SceneViewport.cpp b/Engine/Source/Runtime/Engine/Private/Slate/SceneViewport.cpp
--- a/Engine/Source/Runtime/Engine/Private/Slate/SceneViewport.cpp
+++ b/Engine/Source/Runtime/Engine/Private/Slate/SceneViewport.cpp
@@ -487,7 +487,7 @@ FReply FSceneViewport::OnMouseButtonDown( const FGeometry& InGeometry, const FPo
 			!bNewMenuWasOpened && // We should not focus the viewport if a menu was opened as it would close the menu
 			(bPermanentCapture || bTemporaryCapture))
 		{
-			CurrentReplyState = AcquireFocusAndCapture(FIntPoint(InMouseEvent.GetScreenSpacePosition().X, InMouseEvent.GetScreenSpacePosition().Y));
+			CurrentReplyState = AcquireFocusAndCapture(FIntPoint(InMouseEvent.GetScreenSpacePosition().X, InMouseEvent.GetScreenSpacePosition().Y), EFocusCause::Mouse);
 		}
 	}
 
@@ -497,7 +497,7 @@ FReply FSceneViewport::OnMouseButtonDown( const FGeometry& InGeometry, const FPo
 	return CurrentReplyState;
 }
 
-FReply FSceneViewport::AcquireFocusAndCapture(FIntPoint MousePosition)
+FReply FSceneViewport::AcquireFocusAndCapture(FIntPoint MousePosition, EFocusCause FocusCause)
 {
 	bShouldCaptureMouseOnActivate = false;
 
@@ -506,7 +506,7 @@ FReply FSceneViewport::AcquireFocusAndCapture(FIntPoint MousePosition)
 	TSharedRef<SViewport> ViewportWidgetRef = ViewportWidget.Pin().ToSharedRef();
 
 	// Mouse down should focus viewport for user input
-	ReplyState.SetUserFocus(ViewportWidgetRef, EFocusCause::SetDirectly, true);
+	ReplyState.SetUserFocus(ViewportWidgetRef, FocusCause, true);
 
 	UWorld* World = ViewportClient->GetWorld();
 	if (World && World->IsGameWorld() && World->GetGameInstance() && (World->GetGameInstance()->GetFirstLocalPlayerController() || World->IsPlayInEditor()))
@@ -741,7 +741,7 @@ FReply FSceneViewport::OnTouchStarted( const FGeometry& MyGeometry, const FPoint
 			const bool bTemporaryCapture = ViewportClient->CaptureMouseOnClick() == EMouseCaptureMode::CaptureDuringMouseDown;
 			if (bTemporaryCapture)
 			{
-				CurrentReplyState = AcquireFocusAndCapture(FIntPoint(TouchEvent.GetScreenSpacePosition().X, TouchEvent.GetScreenSpacePosition().Y));
+				CurrentReplyState = AcquireFocusAndCapture(FIntPoint(TouchEvent.GetScreenSpacePosition().X, TouchEvent.GetScreenSpacePosition().Y), EFocusCause::Mouse);
 			}
 		}
 		else
@@ -1090,7 +1090,7 @@ FReply FSceneViewport::OnFocusReceived(const FFocusEvent& InFocusEvent)
 				FWidgetPath PathToWidget;
 				SlateApp.GeneratePathToWidgetUnchecked(ViewportWidgetRef, PathToWidget);
 
-				return AcquireFocusAndCapture(GetSizeXY() / 2);
+				return AcquireFocusAndCapture(GetSizeXY() / 2, EFocusCause::Mouse);
 			}
 		}
 	}
@@ -1169,7 +1169,7 @@ FReply FSceneViewport::OnViewportActivated(const FWindowActivateEvent& InActivat
 		//    - the user clicked in our window but not an area our viewport covers.
 		if (InActivateEvent.GetActivationType() == FWindowActivateEvent::EA_Activate && (bShouldCaptureMouseOnActivate || bPermanentCapture))
 		{
-			return AcquireFocusAndCapture(GetSizeXY() / 2);
+			return AcquireFocusAndCapture(GetSizeXY() / 2, EFocusCause::WindowActivate);
 		}
 	}
 
diff --git a/Engine/Source/Runtime/Engine/Public/Slate/SceneViewport.h b/Engine/Source/Runtime/Engine/Public/Slate/SceneViewport.h
--- a/Engine/Source/Runtime/Engine/Public/Slate/SceneViewport.h
+++ b/Engine/Source/Runtime/Engine/Public/Slate/SceneViewport.h
@@ -344,7 +344,7 @@ private:
 	void ApplyModifierKeys( const FModifierKeysState& InKeysState );
 
 	/** Utility function to create an FReply that properly gets Focus and capture based on the settings*/
-	FReply AcquireFocusAndCapture(FIntPoint MousePosition);
+	FReply AcquireFocusAndCapture(FIntPoint MousePosition, EFocusCause FocusCause = EFocusCause::SetDirectly);
 
 	/** Utility function to figure out if we are currently a game viewport */
 	bool IsCurrentlyGameViewport();
diff --git a/Engine/Source/Runtime/GameMenuBuilder/Private/SGameMenuPageWidget.cpp b/Engine/Source/Runtime/GameMenuBuilder/Private/SGameMenuPageWidget.cpp
--- a/Engine/Source/Runtime/GameMenuBuilder/Private/SGameMenuPageWidget.cpp
+++ b/Engine/Source/Runtime/GameMenuBuilder/Private/SGameMenuPageWidget.cpp
@@ -943,6 +943,7 @@ FReply SGameMenuPageWidget::OnKeyDown(const FGeometry& MyGeometry, const FKeyEve
 	{
 		bool bNavigationLocked = PendingMainMenu.IsValid() || PendingSubMenu.IsValid();
 		const FKey Key = InKeyEvent.GetKey();
+		EUINavigationAction NavAction = FSlateApplication::Get().GetNavigationActionForKey(Key);
 
 		if (bNavigationLocked == false)
 		{
@@ -973,12 +974,12 @@ FReply SGameMenuPageWidget::OnKeyDown(const FGeometry& MyGeometry, const FKeyEve
 				Result = FReply::Handled();
 			}
 		}
-		if (Key == EKeys::Enter || Key == EKeys::Virtual_Accept)
+		if (NavAction == EUINavigationAction::Accept)
 		{
 			ConfirmMenuItem();
 			Result = FReply::Handled();
 		} 
-		else if (Key == EKeys::Escape || Key == EKeys::Virtual_Back || Key == EKeys::Gamepad_Special_Left)
+		else if (NavAction == EUINavigationAction::Back || Key == EKeys::Gamepad_Special_Left)
 		{
 			MenuGoBack(true);
 			Result = FReply::Handled();
diff --git a/Engine/Source/Runtime/Slate/Private/Framework/Application/NavigationConfig.cpp b/Engine/Source/Runtime/Slate/Private/Framework/Application/NavigationConfig.cpp
--- a/Engine/Source/Runtime/Slate/Private/Framework/Application/NavigationConfig.cpp
+++ b/Engine/Source/Runtime/Slate/Private/Framework/Application/NavigationConfig.cpp
@@ -145,4 +145,20 @@ float FNavigationConfig::GetRepeatRateForPressure(float InPressure, int32 InRepe
 	}
 
 	return RepeatRate;
-}
\ No newline at end of file
+}
+
+EUINavigationAction FNavigationConfig::GetNavigationActionForKey(const FKey& InKey) const
+{
+	if (InKey == EKeys::Enter || InKey == EKeys::SpaceBar || InKey == EKeys::Virtual_Accept)
+	{
+		// By default, enter, space, and gamepad accept are all counted as accept
+		return EUINavigationAction::Accept;
+	}
+	else if (InKey == EKeys::Escape || InKey == EKeys::Virtual_Back)
+	{
+		// By default, escape and gamepad back count as leaving current scope
+		return EUINavigationAction::Back;
+	}
+
+	return EUINavigationAction::Invalid;
+}
diff --git a/Engine/Source/Runtime/Slate/Private/Framework/Application/SlateApplication.cpp b/Engine/Source/Runtime/Slate/Private/Framework/Application/SlateApplication.cpp
--- a/Engine/Source/Runtime/Slate/Private/Framework/Application/SlateApplication.cpp
+++ b/Engine/Source/Runtime/Slate/Private/Framework/Application/SlateApplication.cpp
@@ -2087,6 +2087,11 @@ EUINavigation FSlateApplication::GetNavigationDirectionFromAnalog(const FAnalogI
 	return NavigationConfig->GetNavigationDirectionFromAnalog(InAnalogEvent);
 }
 
+EUINavigationAction FSlateApplication::GetNavigationActionForKey(const FKey& InKey) const
+{
+	return NavigationConfig->GetNavigationActionForKey(InKey);
+}
+
 void FSlateApplication::AddModalWindow( TSharedRef<SWindow> InSlateWindow, const TSharedPtr<const SWidget> InParentWidget, bool bSlowTaskWindow )
 {
 	if( !CanAddModalWindow() )
@@ -2977,11 +2982,14 @@ bool FSlateApplication::SetUserFocus(FSlateUser* User, const FWidgetPath& InFocu
 		FSlateDebugging::BroadcastFocusReceived(FocusEvent, OldFocusedWidgetPath, OldFocusedWidget, NewFocusedWidgetPath, NewFocusedWidget);
 #endif
 
+		FFocusEvent FocusEvent = FFocusEvent(InCause, User->GetUserIndex());
 		FReply Reply = NewFocusedWidget->OnFocusReceived(WidgetToFocus.Geometry, FocusEvent);
 		if (Reply.IsEventHandled())
 		{
 			ProcessReply(InFocusPath, Reply, nullptr, nullptr, User->GetUserIndex());
 		}
+
+		NavigationConfig->OnNavigationChangedFocus(OldFocusedWidget, NewFocusedWidget, FocusEvent);
 	}
 
 	return true;
diff --git a/Engine/Source/Runtime/Slate/Private/Widgets/Input/SButton.cpp b/Engine/Source/Runtime/Slate/Private/Widgets/Input/SButton.cpp
--- a/Engine/Source/Runtime/Slate/Private/Widgets/Input/SButton.cpp
+++ b/Engine/Source/Runtime/Slate/Private/Widgets/Input/SButton.cpp
@@ -154,7 +154,7 @@ void SButton::OnFocusLost( const FFocusEvent& InFocusEvent )
 FReply SButton::OnKeyDown( const FGeometry& MyGeometry, const FKeyEvent& InKeyEvent )
 {
 	FReply Reply = FReply::Unhandled();
-	if (IsEnabled() && (InKeyEvent.GetKey() == EKeys::Enter || InKeyEvent.GetKey() == EKeys::SpaceBar || InKeyEvent.GetKey() == EKeys::Virtual_Accept))
+	if (IsEnabled() && FSlateApplication::Get().GetNavigationActionForKey(InKeyEvent.GetKey()) == EUINavigationAction::Accept)
 	{
 		Press();
 
@@ -184,8 +184,7 @@ FReply SButton::OnKeyUp(const FGeometry& MyGeometry, const FKeyEvent& InKeyEvent
 {
 	FReply Reply = FReply::Unhandled();
 
-
-	if (IsEnabled() && (InKeyEvent.GetKey() == EKeys::Enter || InKeyEvent.GetKey() == EKeys::SpaceBar || InKeyEvent.GetKey() == EKeys::Virtual_Accept))
+	if (IsEnabled() && FSlateApplication::Get().GetNavigationActionForKey(InKeyEvent.GetKey()) == EUINavigationAction::Accept)
 	{
 		const bool bWasPressed = bIsPressed;
 
diff --git a/Engine/Source/Runtime/Slate/Private/Widgets/Input/SCheckBox.cpp b/Engine/Source/Runtime/Slate/Private/Widgets/Input/SCheckBox.cpp
--- a/Engine/Source/Runtime/Slate/Private/Widgets/Input/SCheckBox.cpp
+++ b/Engine/Source/Runtime/Slate/Private/Widgets/Input/SCheckBox.cpp
@@ -69,7 +69,7 @@ bool SCheckBox::SupportsKeyboardFocus() const
 
 FReply SCheckBox::OnKeyUp( const FGeometry& MyGeometry, const FKeyEvent& InKeyEvent )
 {
-	if (InKeyEvent.GetKey() == EKeys::Enter || InKeyEvent.GetKey() == EKeys::SpaceBar || InKeyEvent.GetKey() == EKeys::Virtual_Accept)
+	if (FSlateApplication::Get().GetNavigationActionForKey(InKeyEvent.GetKey()) == EUINavigationAction::Accept)
 	{
 		ToggleCheckedState();
 
diff --git a/Engine/Source/Runtime/Slate/Private/Widgets/Input/SComboButton.cpp b/Engine/Source/Runtime/Slate/Private/Widgets/Input/SComboButton.cpp
--- a/Engine/Source/Runtime/Slate/Private/Widgets/Input/SComboButton.cpp
+++ b/Engine/Source/Runtime/Slate/Private/Widgets/Input/SComboButton.cpp
@@ -4,6 +4,7 @@
 #include "Widgets/Layout/SBorder.h"
 #include "Widgets/Images/SImage.h"
 #include "Widgets/Input/SButton.h"
+#include "Framework/Application/SlateApplication.h"
 
 
 void SComboButton::Construct( const FArguments& InArgs )
@@ -118,7 +119,7 @@ FReply SComboButton::OnButtonClicked()
 FReply SComboButton::OnKeyDown(const FGeometry& MyGeometry, const FKeyEvent& InKeyEvent)
 {
 	FReply Reply = FReply::Unhandled();
-	if (InKeyEvent.GetKey() == EKeys::Virtual_Accept)
+	if (FSlateApplication::Get().GetNavigationActionForKey(InKeyEvent.GetKey()) == EUINavigationAction::Accept)
 	{
 		// Handle menu open with controller.
 		Reply = OnButtonClicked();
diff --git a/Engine/Source/Runtime/Slate/Private/Widgets/Input/SSlider.cpp b/Engine/Source/Runtime/Slate/Private/Widgets/Input/SSlider.cpp
--- a/Engine/Source/Runtime/Slate/Private/Widgets/Input/SSlider.cpp
+++ b/Engine/Source/Runtime/Slate/Private/Widgets/Input/SSlider.cpp
@@ -2,6 +2,7 @@
 
 #include "Widgets/Input/SSlider.h"
 #include "Rendering/DrawElements.h"
+#include "Framework/Application/SlateApplication.h"
 
 void SSlider::Construct( const SSlider::FArguments& InDeclaration )
 {
@@ -145,67 +146,75 @@ void SSlider::ResetControllerState()
 	}
 }
 
-FReply SSlider::OnKeyDown(const FGeometry& MyGeometry, const FKeyEvent& InKeyEvent)
+FNavigationReply SSlider::OnNavigation(const FGeometry& MyGeometry, const FNavigationEvent& InNavigationEvent)
 {
-	FReply Reply = FReply::Unhandled();
-	const FKey KeyPressed = InKeyEvent.GetKey();
-
-	if (IsInteractable())
-	{
-		// The controller's bottom face button must be pressed once to begin manipulating the slider's value.
-		// Navigation away from the widget is prevented until the button has been pressed again or focus is lost.
-		// The value can be manipulated by using the game pad's directional arrows ( relative to slider orientation ).
-		if ((KeyPressed == EKeys::Enter || KeyPressed == EKeys::SpaceBar || KeyPressed == EKeys::Virtual_Accept)
-			&& bRequiresControllerLock)
-		{
-			if (bControllerInputCaptured == false)
-			{
-				// Begin capturing controller input and allow user to modify the slider's value.
-				bControllerInputCaptured = true;
-				OnControllerCaptureBegin.ExecuteIfBound();
-				Reply = FReply::Handled();
-			}
-			else
-			{
-				ResetControllerState();
-				Reply = FReply::Handled();
-			}
-		}
-
+	FNavigationReply Reply = FNavigationReply::Escape();
 		if (bControllerInputCaptured || !bRequiresControllerLock)
 		{
 			float NewValue = ValueAttribute.Get();
 			if (Orientation == EOrientation::Orient_Horizontal)
 			{
-				if (KeyPressed == EKeys::Left || KeyPressed == EKeys::Gamepad_DPad_Left || KeyPressed == EKeys::Gamepad_LeftStick_Left)
+			if (InNavigationEvent.GetNavigationType() == EUINavigation::Left)
 				{
 					NewValue -= StepSize.Get();
-					Reply = FReply::Handled();
+				Reply = FNavigationReply::Stop();
 				}
-				else if (KeyPressed == EKeys::Right || KeyPressed == EKeys::Gamepad_DPad_Right || KeyPressed == EKeys::Gamepad_LeftStick_Right)
+			else if (InNavigationEvent.GetNavigationType() == EUINavigation::Right)
 				{
 					NewValue += StepSize.Get();
-					Reply = FReply::Handled();
+				Reply = FNavigationReply::Stop();
 				}
 			}
 			else
 			{
-				if (KeyPressed == EKeys::Down || KeyPressed == EKeys::Gamepad_DPad_Down || KeyPressed == EKeys::Gamepad_LeftStick_Down)
+			if (InNavigationEvent.GetNavigationType() == EUINavigation::Down)
 				{
 					NewValue -= StepSize.Get();
-					Reply = FReply::Handled();
+				Reply = FNavigationReply::Stop();
 				}
-				else if (KeyPressed == EKeys::Up || KeyPressed == EKeys::Gamepad_DPad_Up || KeyPressed == EKeys::Gamepad_LeftStick_Up)
+			if (InNavigationEvent.GetNavigationType() == EUINavigation::Up)
 				{
 					NewValue += StepSize.Get();
-					Reply = FReply::Handled();
+				Reply = FNavigationReply::Stop();
 				}
 			}
-
+		if (ValueAttribute.Get() != NewValue)
+		{
 			CommitValue(FMath::Clamp(NewValue, 0.0f, 1.0f));
-			if (!Reply.IsEventHandled())
+		}
+	}
+
+	if (Reply.GetBoundaryRule() != EUINavigationRule::Escape)
+			{
+		Reply = SLeafWidget::OnNavigation(MyGeometry, InNavigationEvent);
+	}
+
+	return Reply;
+}
+
+FReply SSlider::OnKeyDown(const FGeometry& MyGeometry, const FKeyEvent& InKeyEvent)
+{
+	FReply Reply = FReply::Unhandled();
+	const FKey KeyPressed = InKeyEvent.GetKey();
+
+	if (IsInteractable())
+	{
+		// The controller's bottom face button must be pressed once to begin manipulating the slider's value.
+		// Navigation away from the widget is prevented until the button has been pressed again or focus is lost.
+		// The value can be manipulated by using the game pad's directional arrows ( relative to slider orientation ).
+		if (FSlateApplication::Get().GetNavigationActionForKey(KeyPressed) == EUINavigationAction::Accept && bRequiresControllerLock)
+		{
+			if (bControllerInputCaptured == false)
+			{
+				// Begin capturing controller input and allow user to modify the slider's value.
+				bControllerInputCaptured = true;
+				OnControllerCaptureBegin.ExecuteIfBound();
+				Reply = FReply::Handled();
+			}
+			else
 			{
-				Reply = SLeafWidget::OnKeyDown(MyGeometry, InKeyEvent);
+				ResetControllerState();
+				Reply = FReply::Handled();
 			}
 		}
 		else
diff --git a/Engine/Source/Runtime/Slate/Private/Widgets/Text/SlateEditableTextLayout.cpp b/Engine/Source/Runtime/Slate/Private/Widgets/Text/SlateEditableTextLayout.cpp
--- a/Engine/Source/Runtime/Slate/Private/Widgets/Text/SlateEditableTextLayout.cpp
+++ b/Engine/Source/Runtime/Slate/Private/Widgets/Text/SlateEditableTextLayout.cpp
@@ -1059,7 +1059,7 @@ FReply FSlateEditableTextLayout::HandleKeyDown(const FKeyEvent& InKeyEvent)
 
 FReply FSlateEditableTextLayout::HandleKeyUp(const FKeyEvent& InKeyEvent)
 {
-	if (FPlatformApplicationMisc::RequiresVirtualKeyboard() && InKeyEvent.GetKey() == EKeys::Virtual_Accept)
+	if (FPlatformApplicationMisc::RequiresVirtualKeyboard() && FSlateApplication::Get().GetNavigationActionForKey(InKeyEvent.GetKey()) == EUINavigationAction::Accept)
 	{
 		if (!OwnerWidget->IsTextReadOnly())
 		{
diff --git a/Engine/Source/Runtime/Slate/Public/Framework/Application/NavigationConfig.h b/Engine/Source/Runtime/Slate/Public/Framework/Application/NavigationConfig.h
--- a/Engine/Source/Runtime/Slate/Public/Framework/Application/NavigationConfig.h
+++ b/Engine/Source/Runtime/Slate/Public/Framework/Application/NavigationConfig.h
@@ -55,6 +55,12 @@ public:
 	/** Notified when users are removed from the system, good chance to clean up any user specific state. */
 	virtual void OnUserRemoved(int32 UserIndex);
 
+	/** Notified when navigation has caused a widget change to occur */
+	virtual void OnNavigationChangedFocus(TSharedPtr<SWidget> OldWidget, TSharedPtr<SWidget> NewWidget, FFocusEvent FocusEvent) {}
+
+	/** Returns the navigation action corresponding to this key, or Invalid if not found */
+	virtual EUINavigationAction GetNavigationActionForKey(const FKey& InKey) const;
+
 public:
 	/** Should the Tab key perform next and previous style navigation. */
 	bool bTabNavigation;
diff --git a/Engine/Source/Runtime/Slate/Public/Framework/Application/SlateApplication.h b/Engine/Source/Runtime/Slate/Public/Framework/Application/SlateApplication.h
--- a/Engine/Source/Runtime/Slate/Public/Framework/Application/SlateApplication.h
+++ b/Engine/Source/Runtime/Slate/Public/Framework/Application/SlateApplication.h
@@ -370,6 +370,9 @@ public:
 	virtual EUINavigation GetNavigationDirectionFromKey(const FKeyEvent& InKeyEvent) const override;
 	virtual EUINavigation GetNavigationDirectionFromAnalog(const FAnalogInputEvent& InAnalogEvent) override;
 
+	/** Returns the navigation action corresponding to this key, or Invalid if not found */
+	virtual EUINavigationAction GetNavigationActionForKey(const FKey& InKey) const override;
+
 	/**
 	 * Adds a modal window to the application.  
 	 * In most cases, this function does not return until the modal window is closed (the only exception is a modal window for slow tasks)  
diff --git a/Engine/Source/Runtime/Slate/Public/Widgets/Input/SComboBox.h b/Engine/Source/Runtime/Slate/Public/Widgets/Input/SComboBox.h
--- a/Engine/Source/Runtime/Slate/Public/Widgets/Input/SComboBox.h
+++ b/Engine/Source/Runtime/Slate/Public/Widgets/Input/SComboBox.h
@@ -280,7 +280,8 @@ protected:
 	/** Handle key presses that SListView ignores */
 	virtual FReply OnHandleKeyPressed(FKey KeyPressed)
 	{
-		if (KeyPressed == EKeys::Enter || KeyPressed == EKeys::SpaceBar || KeyPressed == EKeys::Virtual_Accept)
+		EUINavigationAction NavAction = FSlateApplication::Get().GetNavigationActionForKey(KeyPressed);
+		if (NavAction == EUINavigationAction::Accept)
 		{
 			TArray<OptionType> SelectedItems = ComboListView->GetSelectedItems();
 			if (SelectedItems.Num() > 0)
@@ -289,7 +290,7 @@ protected:
 			}
 			return FReply::Handled();
 		}
-		else if (KeyPressed == EKeys::Escape)
+		else if (NavAction == EUINavigationAction::Back)
 		{
 			this->SetIsOpen(false);
 			return FReply::Handled();
@@ -307,9 +308,10 @@ protected:
 
 			if (EnableGamepadNavigationMode)
 			{
+				EUINavigationAction NavAction = FSlateApplication::Get().GetNavigationActionForKey(KeyPressed);
 				// The controller's bottom face button must be pressed once to begin manipulating the combobox's value.
 				// Navigation away from the widget is prevented until the button has been pressed again or focus is lost.
-				if (KeyPressed == EKeys::Enter || KeyPressed == EKeys::SpaceBar || KeyPressed == EKeys::Virtual_Accept)
+				if (NavAction == EUINavigationAction::Accept)
 				{
 					if (bControllerInputCaptured == false)
 					{
@@ -338,7 +340,7 @@ protected:
 					}
 
 				}
-				else if (KeyPressed == EKeys::Escape || KeyPressed == EKeys::Virtual_Back || KeyPressed == EKeys::BackSpace)
+				else if (NavAction == EUINavigationAction::Back || KeyPressed == EKeys::BackSpace)
 				{
 					OnMenuOpenChanged(false);
 				}
diff --git a/Engine/Source/Runtime/Slate/Public/Widgets/Input/SSlider.h b/Engine/Source/Runtime/Slate/Public/Widgets/Input/SSlider.h
--- a/Engine/Source/Runtime/Slate/Public/Widgets/Input/SSlider.h
+++ b/Engine/Source/Runtime/Slate/Public/Widgets/Input/SSlider.h
@@ -146,6 +146,7 @@ public:
 	virtual FReply OnKeyDown(const FGeometry& MyGeometry, const FKeyEvent& InKeyEvent) override;
 	virtual FReply OnKeyUp(const FGeometry& MyGeometry, const FKeyEvent& InKeyEvent) override;
 	virtual void OnFocusLost(const FFocusEvent& InFocusEvent) override;
+	virtual FNavigationReply OnNavigation(const FGeometry& MyGeometry, const FNavigationEvent& InNavigationEvent) override;
 
 	virtual bool SupportsKeyboardFocus() const override;
 	virtual bool IsInteractable() const override;
diff --git a/Engine/Source/Runtime/SlateCore/Private/Widgets/SWindow.cpp b/Engine/Source/Runtime/SlateCore/Private/Widgets/SWindow.cpp
--- a/Engine/Source/Runtime/SlateCore/Private/Widgets/SWindow.cpp
+++ b/Engine/Source/Runtime/SlateCore/Private/Widgets/SWindow.cpp
@@ -1583,7 +1583,7 @@ bool SWindow::OnIsActiveChanged( const FWindowActivateEvent& ActivateEvent )
 				FWidgetPath WidgetToFocusPath;
 				if (FSlateWindowHelper::FindPathToWidget(JustThisWindow, PinnedWidgetToFocus.ToSharedRef(), WidgetToFocusPath))
 				{
-					FSlateApplicationBase::Get().SetAllUserFocus(WidgetToFocusPath, EFocusCause::SetDirectly);
+					FSlateApplicationBase::Get().SetAllUserFocus(WidgetToFocusPath, EFocusCause::WindowActivate);
 				}
 			}
 
@@ -1595,7 +1595,7 @@ bool SWindow::OnIsActiveChanged( const FWindowActivateEvent& ActivateEvent )
 				TSharedRef<SWidget> WindowWidgetToFocus = WidgetFocusedOnDeactivate.IsValid() ? WidgetFocusedOnDeactivate.Pin().ToSharedRef() : AsShared();
 				if (FSlateWindowHelper::FindPathToWidget(JustThisWindow, WindowWidgetToFocus, WindowWidgetPath))
 				{
-					FSlateApplicationBase::Get().SetAllUserFocusAllowingDescendantFocus(WindowWidgetPath, EFocusCause::SetDirectly);
+					FSlateApplicationBase::Get().SetAllUserFocusAllowingDescendantFocus(WindowWidgetPath, EFocusCause::WindowActivate);
 				}
 			}
 		}
diff --git a/Engine/Source/Runtime/SlateCore/Public/Application/SlateApplicationBase.h b/Engine/Source/Runtime/SlateCore/Public/Application/SlateApplicationBase.h
--- a/Engine/Source/Runtime/SlateCore/Public/Application/SlateApplicationBase.h
+++ b/Engine/Source/Runtime/SlateCore/Public/Application/SlateApplicationBase.h
@@ -224,6 +224,7 @@ public:
 
 	virtual EUINavigation GetNavigationDirectionFromKey( const FKeyEvent& InKeyEvent ) const = 0;
 	virtual EUINavigation GetNavigationDirectionFromAnalog(const FAnalogInputEvent& InAnalogEvent) = 0;
+	virtual EUINavigationAction GetNavigationActionForKey(const FKey& InKey) const = 0;
 
 	/** @return	Returns true if there are any pop-up menus summoned */
 	virtual bool AnyMenusVisible() const = 0;
diff --git a/Engine/Source/Runtime/SlateCore/Public/Types/SlateEnums.h b/Engine/Source/Runtime/SlateCore/Public/Types/SlateEnums.h
--- a/Engine/Source/Runtime/SlateCore/Public/Types/SlateEnums.h
+++ b/Engine/Source/Runtime/SlateCore/Public/Types/SlateEnums.h
@@ -109,12 +109,34 @@ enum class EUINavigation : uint8
 
 	/** Number of navigation types*/
 	Num UMETA(Hidden),
+
 	/** Denotes an invalid navigation, more important used to denote no specified navigation*/
-	Invalid UMETA(Hidden),
+	Invalid,
 };
 
 ENUM_RANGE_BY_COUNT(EUINavigation, EUINavigation::Num);
 
+/**
+ * Generic UI navigation action that are usually bound to a gamepad or other input device
+ */
+UENUM(BlueprintType)
+enum class EUINavigationAction : uint8
+{
+	/** Accept, normally acts like a mouse click. Corresponds to Virtual_Accept key */
+	Accept,
+
+	/** Back/Cancel, normally acts like escape. Corresponds to Virtual_Back key */
+	Back,
+
+	/** Number of navigation actions */
+	Num UMETA(Hidden),
+
+	/** Invalid action */
+	Invalid,
+};
+
+ENUM_RANGE_BY_COUNT(EUINavigationAction, EUINavigationAction::Num);
+
 /**
  * Enumerates the source of the navigation
  */
