Fix CIS warning

diff --git a/Engine/Source/Runtime/Slate/Private/Framework/Application/SlateApplication.cpp b/Engine/Source/Runtime/Slate/Private/Framework/Application/SlateApplication.cpp
--- a/Engine/Source/Runtime/Slate/Private/Framework/Application/SlateApplication.cpp
+++ b/Engine/Source/Runtime/Slate/Private/Framework/Application/SlateApplication.cpp
@@ -2982,7 +2982,6 @@ bool FSlateApplication::SetUserFocus(FSlateUser* User, const FWidgetPath& InFocu
 		FSlateDebugging::BroadcastFocusReceived(FocusEvent, OldFocusedWidgetPath, OldFocusedWidget, NewFocusedWidgetPath, NewFocusedWidget);
 #endif
 
-		FFocusEvent FocusEvent = FFocusEvent(InCause, User->GetUserIndex());
 		FReply Reply = NewFocusedWidget->OnFocusReceived(WidgetToFocus.Geometry, FocusEvent);
 		if (Reply.IsEventHandled())
 		{
