Fix for SCombobox returning unhandled when input is captured and menu is closed

diff --git a/Engine/Source/Runtime/Slate/Public/Widgets/Input/SComboBox.h b/Engine/Source/Runtime/Slate/Public/Widgets/Input/SComboBox.h
--- a/Engine/Source/Runtime/Slate/Public/Widgets/Input/SComboBox.h
+++ b/Engine/Source/Runtime/Slate/Public/Widgets/Input/SComboBox.h
@@ -340,7 +340,13 @@ protected:
 				}
 				else if (KeyPressed == EKeys::Escape || KeyPressed == EKeys::Virtual_Back || KeyPressed == EKeys::BackSpace)
 				{
+					const bool bWasInputCaptured = bControllerInputCaptured;
+
 					OnMenuOpenChanged(false);
+					if (bWasInputCaptured)
+					{
+						return FReply::Handled();
+					}
 				}
 				else
 				{
