PR #5513: Correct OnTextChanged tooltips (Contributed by [REDACTED_PERSON])

diff --git a/Engine/Source/Runtime/Slate/Public/Widgets/Input/SEditableText.h b/Engine/Source/Runtime/Slate/Public/Widgets/Input/SEditableText.h
--- a/Engine/Source/Runtime/Slate/Public/Widgets/Input/SEditableText.h
+++ b/Engine/Source/Runtime/Slate/Public/Widgets/Input/SEditableText.h
@@ -129,7 +129,7 @@ public:
 		 */
 		SLATE_EVENT( FOnIsTypedCharValid, OnIsTypedCharValid )
 
-		/** Called whenever the text is changed interactively by the user */
+		/** Called whenever the text is changed programmatically or interactively by the user */
 		SLATE_EVENT( FOnTextChanged, OnTextChanged )
 
 		/** Called whenever the text is committed.  This happens when the user presses enter or the text box loses focus. */
@@ -436,7 +436,7 @@ protected:
 	/** Called when a character is typed and we want to know if the text field supports typing this character. */
 	FOnIsTypedCharValid OnIsTypedCharValid;
 
-	/** Called whenever the text is changed interactively by the user */
+	/** Called whenever the text is changed programmatically or interactively by the user */
 	FOnTextChanged OnTextChangedCallback;
 
 	/** Called whenever the text is committed.  This happens when the user presses enter or the text box loses focus. */
diff --git a/Engine/Source/Runtime/Slate/Public/Widgets/Input/SEditableTextBox.h b/Engine/Source/Runtime/Slate/Public/Widgets/Input/SEditableTextBox.h
--- a/Engine/Source/Runtime/Slate/Public/Widgets/Input/SEditableTextBox.h
+++ b/Engine/Source/Runtime/Slate/Public/Widgets/Input/SEditableTextBox.h
@@ -102,7 +102,7 @@ public:
 		/** Delegate to call before a context menu is opened. User returns the menu content or null to the disable context menu */
 		SLATE_EVENT(FOnContextMenuOpening, OnContextMenuOpening)
 
-		/** Called whenever the text is changed interactively by the user */
+		/** Called whenever the text is changed programmatically or interactively by the user */
 		SLATE_EVENT( FOnTextChanged, OnTextChanged )
 
 		/** Called whenever the text is committed.  This happens when the user presses enter or the text box loses focus. */
diff --git a/Engine/Source/Runtime/Slate/Public/Widgets/Input/SMultiLineEditableTextBox.h b/Engine/Source/Runtime/Slate/Public/Widgets/Input/SMultiLineEditableTextBox.h
--- a/Engine/Source/Runtime/Slate/Public/Widgets/Input/SMultiLineEditableTextBox.h
+++ b/Engine/Source/Runtime/Slate/Public/Widgets/Input/SMultiLineEditableTextBox.h
@@ -156,7 +156,7 @@ public:
 		 */
 		SLATE_EVENT( FOnIsTypedCharValid, OnIsTypedCharValid )
 
-		/** Called whenever the text is changed interactively by the user */
+		/** Called whenever the text is changed programmatically or interactively by the user */
 		SLATE_EVENT( FOnTextChanged, OnTextChanged )
 
 		/** Called whenever the text is committed.  This happens when the user presses enter or the text box loses focus. */
diff --git a/Engine/Source/Runtime/Slate/Public/Widgets/Input/SNumericEntryBox.h b/Engine/Source/Runtime/Slate/Public/Widgets/Input/SNumericEntryBox.h
--- a/Engine/Source/Runtime/Slate/Public/Widgets/Input/SNumericEntryBox.h
+++ b/Engine/Source/Runtime/Slate/Public/Widgets/Input/SNumericEntryBox.h
@@ -132,11 +132,11 @@ public:
 		SLATE_ATTRIBUTE( float, MinDesiredValueWidth )
 		/** The text margin to use if overridden. */
 		SLATE_ATTRIBUTE( FMargin, OverrideTextMargin )
-		/** Called whenever the text is changed interactively by the user */
+		/** Called whenever the text is changed programmatically or interactively by the user */
 		SLATE_EVENT( FOnValueChanged, OnValueChanged )
 		/** Called whenever the text is committed.  This happens when the user presses enter or the text box loses focus. */
 		SLATE_EVENT( FOnValueCommitted, OnValueCommitted )
-		/** Called whenever the text is changed interactively by the user */
+		/** Called whenever the text is changed programmatically or interactively by the user */
 		SLATE_EVENT( FOnUndeterminedValueChanged, OnUndeterminedValueChanged )
 		/** Called whenever the text is committed.  This happens when the user presses enter or the text box loses focus. */
 		SLATE_EVENT( FOnUndeterminedValueCommitted, OnUndeterminedValueCommitted )
diff --git a/Engine/Source/Runtime/Slate/Public/Widgets/Input/SSuggestionTextBox.h b/Engine/Source/Runtime/Slate/Public/Widgets/Input/SSuggestionTextBox.h
--- a/Engine/Source/Runtime/Slate/Public/Widgets/Input/SSuggestionTextBox.h
+++ b/Engine/Source/Runtime/Slate/Public/Widgets/Input/SSuggestionTextBox.h
@@ -96,7 +96,7 @@ public:
 		/** Called before the suggestion list is shown. */
 		SLATE_EVENT(FOnShowingSuggestions, OnShowingSuggestions)
 
-		/** Called whenever the text is changed interactively by the user. */
+		/** Called whenever the text is changed programmatically or interactively by the user. */
 		SLATE_EVENT(FOnTextChanged, OnTextChanged)
 
 		/** Called when the text has been committed. */
diff --git a/Engine/Source/Runtime/Slate/Public/Widgets/Input/STextEntryPopup.h b/Engine/Source/Runtime/Slate/Public/Widgets/Input/STextEntryPopup.h
--- a/Engine/Source/Runtime/Slate/Public/Widgets/Input/STextEntryPopup.h
+++ b/Engine/Source/Runtime/Slate/Public/Widgets/Input/STextEntryPopup.h
@@ -26,7 +26,7 @@ public:
 		/** Test to place into text entry box before anything is typed */
 		SLATE_ARGUMENT( FText, DefaultText )
 
-		/** Called whenever the text is changed interactively by the user */
+		/** Called whenever the text is changed programmatically or interactively by the user */
 		SLATE_EVENT( FOnTextChanged, OnTextChanged )
 
 		/** Called when the text is committed. */
diff --git a/Engine/Source/Runtime/Slate/Public/Widgets/Input/SVirtualKeyboardEntry.h b/Engine/Source/Runtime/Slate/Public/Widgets/Input/SVirtualKeyboardEntry.h
--- a/Engine/Source/Runtime/Slate/Public/Widgets/Input/SVirtualKeyboardEntry.h
+++ b/Engine/Source/Runtime/Slate/Public/Widgets/Input/SVirtualKeyboardEntry.h
@@ -52,7 +52,7 @@ public:
 		/** Whether to clear keyboard focus when pressing enter to commit changes */
 		SLATE_ATTRIBUTE( bool, ClearKeyboardFocusOnCommit )
 
-		/** Called whenever the text is changed interactively by the user */
+		/** Called whenever the text is changed programmatically or interactively by the user */
 		SLATE_EVENT( FOnTextChanged, OnTextChanged )
 
 		/** Called whenever the text is committed.  This happens when the user presses enter or the text box loses focus. */
@@ -169,7 +169,7 @@ private:
 	/** Whether to clear keyboard focus when pressing enter to commit changes */
 	TAttribute< bool > ClearKeyboardFocusOnCommit;
 
-	/** Called whenever the text is changed interactively by the user */
+	/** Called whenever the text is changed programmatically or interactively by the user */
 	FOnTextChanged OnTextChanged;
 
 	/** Called whenever the text is committed.  This happens when the user presses enter or the text box loses focus. */
diff --git a/Engine/Source/Runtime/Slate/Public/Widgets/Text/SInlineEditableTextBlock.h b/Engine/Source/Runtime/Slate/Public/Widgets/Text/SInlineEditableTextBlock.h
--- a/Engine/Source/Runtime/Slate/Public/Widgets/Text/SInlineEditableTextBlock.h
+++ b/Engine/Source/Runtime/Slate/Public/Widgets/Text/SInlineEditableTextBlock.h
@@ -102,7 +102,7 @@ class SLATE_API SInlineEditableTextBlock: public SCompoundWidget
 		/** Callback to check if the widget is selected, should only be hooked up if parent widget is handling selection or focus. */
 		SLATE_EVENT( FIsSelected, IsSelected )
 
-		/** Called whenever the text is changed interactively by the user */
+		/** Called whenever the text is changed programmatically or interactively by the user */
 		SLATE_EVENT( FOnVerifyTextChanged, OnVerifyTextChanged )
 	SLATE_END_ARGS()
 
diff --git a/Engine/Source/Runtime/Slate/Public/Widgets/Text/SMultiLineEditableText.h b/Engine/Source/Runtime/Slate/Public/Widgets/Text/SMultiLineEditableText.h
--- a/Engine/Source/Runtime/Slate/Public/Widgets/Text/SMultiLineEditableText.h
+++ b/Engine/Source/Runtime/Slate/Public/Widgets/Text/SMultiLineEditableText.h
@@ -133,7 +133,7 @@ public:
 		 */
 		SLATE_EVENT(FOnIsTypedCharValid, OnIsTypedCharValid)
 
-		/** Called whenever the text is changed interactively by the user */
+		/** Called whenever the text is changed programmatically or interactively by the user */
 		SLATE_EVENT(FOnTextChanged, OnTextChanged)
 
 		/** Called whenever the text is committed.  This happens when the user presses enter or the text box loses focus. */
@@ -476,7 +476,7 @@ protected:
 	/** Called when a character is typed and we want to know if the text field supports typing this character. */
 	FOnIsTypedCharValid OnIsTypedCharValid;
 
-	/** Called whenever the text is changed interactively by the user */
+	/** Called whenever the text is changed programmatically or interactively by the user */
 	FOnTextChanged OnTextChangedCallback;
 
 	/** Called whenever the text is committed.  This happens when the user presses enter or the text box loses focus. */
diff --git a/Engine/Source/Runtime/UMG/Public/Components/EditableText.h b/Engine/Source/Runtime/UMG/Public/Components/EditableText.h
--- a/Engine/Source/Runtime/UMG/Public/Components/EditableText.h
+++ b/Engine/Source/Runtime/UMG/Public/Components/EditableText.h
@@ -135,7 +135,7 @@ public:
 
 public:
 
-	/** Called whenever the text is changed interactively by the user */
+	/** Called whenever the text is changed programmatically or interactively by the user */
 	UPROPERTY(BlueprintAssignable, Category="Widget Event", meta=(DisplayName="OnTextChanged (Editable Text)"))
 	FOnEditableTextChangedEvent OnTextChanged;
 
diff --git a/Engine/Source/Runtime/UMG/Public/Components/EditableTextBox.h b/Engine/Source/Runtime/UMG/Public/Components/EditableTextBox.h
--- a/Engine/Source/Runtime/UMG/Public/Components/EditableTextBox.h
+++ b/Engine/Source/Runtime/UMG/Public/Components/EditableTextBox.h
@@ -137,7 +137,7 @@ public:
 
 public:
 
-	/** Called whenever the text is changed interactively by the user */
+	/** Called whenever the text is changed programmatically or interactively by the user */
 	UPROPERTY(BlueprintAssignable, Category="TextBox|Event")
 	FOnEditableTextBoxChangedEvent OnTextChanged;
 
diff --git a/Engine/Source/Runtime/UMG/Public/Components/MultiLineEditableText.h b/Engine/Source/Runtime/UMG/Public/Components/MultiLineEditableText.h
--- a/Engine/Source/Runtime/UMG/Public/Components/MultiLineEditableText.h
+++ b/Engine/Source/Runtime/UMG/Public/Components/MultiLineEditableText.h
@@ -80,7 +80,7 @@ public:
 	UPROPERTY(EditAnywhere, Category=Behavior, AdvancedDisplay)
 	EVirtualKeyboardDismissAction VirtualKeyboardDismissAction;
 
-	/** Called whenever the text is changed interactively by the user */
+	/** Called whenever the text is changed programmatically or interactively by the user */
 	UPROPERTY(BlueprintAssignable, Category="Widget Event", meta=(DisplayName="OnTextChanged (Multi-Line Editable Text)"))
 	FOnMultiLineEditableTextChangedEvent OnTextChanged;
 
diff --git a/Engine/Source/Runtime/UMG/Public/Components/MultiLineEditableTextBox.h b/Engine/Source/Runtime/UMG/Public/Components/MultiLineEditableTextBox.h
--- a/Engine/Source/Runtime/UMG/Public/Components/MultiLineEditableTextBox.h
+++ b/Engine/Source/Runtime/UMG/Public/Components/MultiLineEditableTextBox.h
@@ -85,7 +85,7 @@ public:
 	UPROPERTY()
 	FLinearColor ReadOnlyForegroundColor_DEPRECATED;
 
-	/** Called whenever the text is changed interactively by the user */
+	/** Called whenever the text is changed programmatically or interactively by the user */
 	UPROPERTY(BlueprintAssignable, Category="Widget Event", meta=(DisplayName="OnTextChanged (Multi-Line Text Box)"))
 	FOnMultiLineEditableTextBoxChangedEvent OnTextChanged;
 
