PR #5535: Make PopUpErrorText non-focusable (Contributed by [REDACTED_PERSON])

diff --git a/Engine/Source/Runtime/Slate/Private/Widgets/Notifications/SPopUpErrorText.cpp b/Engine/Source/Runtime/Slate/Private/Widgets/Notifications/SPopUpErrorText.cpp
--- a/Engine/Source/Runtime/Slate/Private/Widgets/Notifications/SPopUpErrorText.cpp
+++ b/Engine/Source/Runtime/Slate/Private/Widgets/Notifications/SPopUpErrorText.cpp
@@ -9,6 +9,7 @@ void SPopupErrorText::Construct( const FArguments& InArgs )
 		.ComboButtonStyle( FCoreStyle::Get(), "MessageLogListingComboButton" )
 		.HasDownArrow(false)
 		.ContentPadding(0)
+		.IsFocusable(false)
 		.ButtonContent()
 		[
 			SAssignNew( HasErrorSymbol, SErrorText )
