PR #3674: Add getter for the scroll offset of the bottom of the ScrollBox. (Contributed by [REDACTED_PERSON])

diff --git a/Engine/Source/Runtime/Slate/Private/Widgets/Layout/SScrollBox.cpp b/Engine/Source/Runtime/Slate/Private/Widgets/Layout/SScrollBox.cpp
--- a/Engine/Source/Runtime/Slate/Private/Widgets/Layout/SScrollBox.cpp
+++ b/Engine/Source/Runtime/Slate/Private/Widgets/Layout/SScrollBox.cpp
@@ -323,6 +323,13 @@ float SScrollBox::GetScrollOffset() const
 	return DesiredScrollOffset;
 }
 
+float SScrollBox::GetScrollOffsetOfEnd() const
+{
+	const FGeometry ScrollPanelGeometry = FindChildGeometry(CachedGeometry, ScrollPanel.ToSharedRef());
+	const float ContentSize = GetScrollComponentFromVector(ScrollPanel->GetDesiredSize());
+	return FMath::Max(ContentSize - GetScrollComponentFromVector(ScrollPanelGeometry.Size), 0.0f);
+}
+
 float SScrollBox::GetViewFraction() const
 {
 	const FGeometry ScrollPanelGeometry = FindChildGeometry(CachedGeometry, ScrollPanel.ToSharedRef());
diff --git a/Engine/Source/Runtime/Slate/Public/Widgets/Layout/SScrollBox.h b/Engine/Source/Runtime/Slate/Public/Widgets/Layout/SScrollBox.h
--- a/Engine/Source/Runtime/Slate/Public/Widgets/Layout/SScrollBox.h
+++ b/Engine/Source/Runtime/Slate/Public/Widgets/Layout/SScrollBox.h
@@ -168,6 +168,9 @@ public:
 
 	float GetViewOffsetFraction() const;
 
+	/** Gets the scroll offset of the bottom of the ScrollBox in Slate Units. */
+	float GetScrollOffsetOfEnd() const;
+
 	void SetScrollOffset( float NewScrollOffset );
 
 	void ScrollToStart();
diff --git a/Engine/Source/Runtime/UMG/Private/Components/ScrollBox.cpp b/Engine/Source/Runtime/UMG/Private/Components/ScrollBox.cpp
--- a/Engine/Source/Runtime/UMG/Private/Components/ScrollBox.cpp
+++ b/Engine/Source/Runtime/UMG/Private/Components/ScrollBox.cpp
@@ -114,6 +114,16 @@ float UScrollBox::GetScrollOffset() const
 	return 0;
 }
 
+float UScrollBox::GetScrollOffsetOfEnd() const
+{
+	if (MyScrollBox.IsValid())
+	{
+		return MyScrollBox->GetScrollOffsetOfEnd();
+	}
+
+	return 0;
+}
+
 float UScrollBox::GetViewOffsetFraction() const
 {
 	if ( MyScrollBox.IsValid() )
diff --git a/Engine/Source/Runtime/UMG/Public/Components/ScrollBox.h b/Engine/Source/Runtime/UMG/Public/Components/ScrollBox.h
--- a/Engine/Source/Runtime/UMG/Public/Components/ScrollBox.h
+++ b/Engine/Source/Runtime/UMG/Public/Components/ScrollBox.h
@@ -116,6 +116,10 @@ public:
 	UFUNCTION(BlueprintCallable, Category="Widget")
 	float GetScrollOffset() const;
 
+	/** Gets the scroll offset of the bottom of the ScrollBox in Slate Units. */
+	UFUNCTION(BlueprintCallable, Category = "Widget")
+	float GetScrollOffsetOfEnd() const;
+
 	UFUNCTION(BlueprintCallable, Category="Widget")
 	float GetViewOffsetFraction() const;
 
