Fixed the strange color change in UI when selecting a variable in multiple editors.

diff --git a/Engine/Source/Editor/GraphEditor/Private/SGraphActionMenu.cpp b/Engine/Source/Editor/GraphEditor/Private/SGraphActionMenu.cpp
--- a/Engine/Source/Editor/GraphEditor/Private/SGraphActionMenu.cpp
+++ b/Engine/Source/Editor/GraphEditor/Private/SGraphActionMenu.cpp
@@ -90,6 +90,11 @@ public:
 		ContentBorder->SetContent(InContent);
 	}
 
+	virtual const FSlateBrush* GetBorder() const
+	{
+		return nullptr;
+	}
+
 private:
 	TSharedPtr<SBorder> ContentBorder;
 };
diff --git a/Engine/Source/Runtime/Slate/Public/Widgets/Views/STableRow.h b/Engine/Source/Runtime/Slate/Public/Widgets/Views/STableRow.h
--- a/Engine/Source/Runtime/Slate/Public/Widgets/Views/STableRow.h
+++ b/Engine/Source/Runtime/Slate/Public/Widgets/Views/STableRow.h
@@ -823,7 +823,7 @@ public:
 	}
 
 	/** @return The border to be drawn around this list item */
-	const FSlateBrush* GetBorder() const 
+	virtual const FSlateBrush* GetBorder() const 
 	{
 		TSharedPtr< ITypedTableView<ItemType> > OwnerWidget = OwnerTablePtr.Pin();
 
