fixing invaidation for menu anchor

diff --git a/Engine/Source/Runtime/Slate/Private/Widgets/Input/SMenuAnchor.cpp b/Engine/Source/Runtime/Slate/Private/Widgets/Input/SMenuAnchor.cpp
--- a/Engine/Source/Runtime/Slate/Private/Widgets/Input/SMenuAnchor.cpp
+++ b/Engine/Source/Runtime/Slate/Private/Widgets/Input/SMenuAnchor.cpp
@@ -180,7 +180,7 @@ void SMenuAnchor::Tick( const FGeometry& AllottedGeometry, const double InCurren
 
 bool SMenuAnchor::ComputeVolatility() const
 {
-	return IsOpen();
+	return SPanel::ComputeVolatility() || IsOpen();
 }
 
 void SMenuAnchor::OnArrangeChildren( const FGeometry& AllottedGeometry, FArrangedChildren& ArrangedChildren ) const
@@ -512,6 +512,8 @@ void SMenuAnchor::SetIsOpen( bool InIsOpen, const bool bFocusMenu, const int32 F
 					}
 				}
 			}
+
+			Invalidate(EInvalidateWidget::ChildOrder | EInvalidateWidget::Volatility);
 		}
 		else
 		{
@@ -520,32 +522,17 @@ void SMenuAnchor::SetIsOpen( bool InIsOpen, const bool bFocusMenu, const int32 F
 			{
 				PopupMenuPtr.Pin()->Dismiss();
 			}
-			else
-			{
-				PopupWindowPtr.Reset();
-				OwnedMenuPtr.Reset();
-				MethodInUse = FPopupMethodReply::Unhandled();
-			}
 
-			// Always clear out the menu content children slot to prevent prepass and other hierarchy queries from considering the
-			// hidden menu content as content they should be concerned with.
-			Children[1]
-			[
-				SNullWidget::NullWidget
-			];
+			ResetPopupMenuContent();
 		}
-
-		Invalidate(EInvalidateWidget::ChildOrder|EInvalidateWidget::Volatility);
 	}
 }
 
 void SMenuAnchor::OnMenuClosed(TSharedRef<IMenu> InMenu)
 {
 	bDismissedThisTick = true;
-	MethodInUse = FPopupMethodReply::Unhandled();
-	PopupMenuPtr.Reset();
-	OwnedMenuPtr.Reset();
-	PopupWindowPtr.Reset();
+
+	ResetPopupMenuContent();
 
 	if (OnMenuOpenChanged.IsBound())
 	{
@@ -553,6 +540,23 @@ void SMenuAnchor::OnMenuClosed(TSharedRef<IMenu> InMenu)
 	}
 }
 
+void SMenuAnchor::ResetPopupMenuContent()
+{
+	MethodInUse = FPopupMethodReply::Unhandled();
+	PopupMenuPtr.Reset();
+	OwnedMenuPtr.Reset();
+	PopupWindowPtr.Reset();
+
+	// Always clear out the menu content children slot to prevent prepass and other hierarchy queries from considering the
+	// hidden menu content as content they should be concerned with.
+	Children[1]
+	[
+		SNullWidget::NullWidget
+	];
+
+	Invalidate(EInvalidateWidget::ChildOrder | EInvalidateWidget::Volatility);
+}
+
 bool SMenuAnchor::IsOpen() const
 {
 	return MethodInUse.IsSet() && PopupMenuPtr.IsValid();
diff --git a/Engine/Source/Runtime/Slate/Public/Widgets/Input/SMenuAnchor.h b/Engine/Source/Runtime/Slate/Public/Widgets/Input/SMenuAnchor.h
--- a/Engine/Source/Runtime/Slate/Public/Widgets/Input/SMenuAnchor.h
+++ b/Engine/Source/Runtime/Slate/Public/Widgets/Input/SMenuAnchor.h
@@ -142,6 +142,9 @@ protected:
 	/** Handler/callback called by menus created by this anchor, when they are dismissed */
 	void OnMenuClosed(TSharedRef<IMenu> InMenu);
 
+	/** Reset the popup state to prepopup. */
+	void ResetPopupMenuContent();
+
 	/** Computes the placement geometry for menus displayed in a separately created window */
 	FGeometry ComputeNewWindowMenuPlacement(const FGeometry& AllottedGeometry, const FVector2D& PopupDesiredSize, EMenuPlacement PlacementMode) const;
 
