Fix GC crashes with slate draw elements.

diff --git a/Engine/Source/Runtime/Engine/Private/Internationalization/StringTable.cpp b/Engine/Source/Runtime/Engine/Private/Internationalization/StringTable.cpp
--- a/Engine/Source/Runtime/Engine/Private/Internationalization/StringTable.cpp
+++ b/Engine/Source/Runtime/Engine/Private/Internationalization/StringTable.cpp
@@ -370,7 +370,7 @@ void UStringTable::PostLoad()
 	if (FSlateApplicationBase::IsInitialized())
 	{
 		// Ensure all invalidation panels are updated now that the string data is loaded
-		FSlateApplicationBase::Get().InvalidateAllWidgets();
+		FSlateApplicationBase::Get().InvalidateAllWidgets(false);
 	}
 }
 
diff --git a/Engine/Source/Runtime/Slate/Private/Framework/Application/SlateApplication.cpp b/Engine/Source/Runtime/Slate/Private/Framework/Application/SlateApplication.cpp
--- a/Engine/Source/Runtime/Slate/Private/Framework/Application/SlateApplication.cpp
+++ b/Engine/Source/Runtime/Slate/Private/Framework/Application/SlateApplication.cpp
@@ -471,15 +471,15 @@ FAutoConsoleVariableRef CVarRequireFocusForGamepadInput(
 
 #if !UE_BUILD_SHIPPING
 
-static void HandleGlobalInvalidate(const TArray<FString>& Args)
+static void HandleGlobalInvalidateCVarTriggered(const TArray<FString>& Args)
 {
-	FSlateApplication::Get().InvalidateAllWidgets();
+	FSlateApplication::Get().InvalidateAllWidgets(false);
 }
 
 static FAutoConsoleCommand GlobalInvalidateCommand(
 	TEXT("Slate.TriggerInvalidate"),
 	TEXT("Triggers a global invalidate of all widgets"),
-	FConsoleCommandWithArgsDelegate::CreateStatic(&HandleGlobalInvalidate)
+	FConsoleCommandWithArgsDelegate::CreateStatic(&HandleGlobalInvalidateCVarTriggered)
 );
 #endif
 
@@ -2614,6 +2614,8 @@ void FSlateApplication::UnregisterVirtualWindow(TSharedRef<SWindow> InWindow)
 
 void FSlateApplication::FlushRenderState()
 {
+	InvalidateAllWidgets(true);
+
 	if ( Renderer.IsValid() )
 	{
 		// Release any temporary material or texture resources we may have cached and are reporting to prevent
diff --git a/Engine/Source/Runtime/SlateCore/Private/Application/SlateApplicationBase.cpp b/Engine/Source/Runtime/SlateCore/Private/Application/SlateApplicationBase.cpp
--- a/Engine/Source/Runtime/SlateCore/Private/Application/SlateApplicationBase.cpp
+++ b/Engine/Source/Runtime/SlateCore/Private/Application/SlateApplicationBase.cpp
@@ -161,9 +161,9 @@ void FSlateApplicationBase::ToggleGlobalInvalidation(bool bIsGlobalInvalidationE
 	}
 }
 
-void FSlateApplicationBase::InvalidateAllWidgets() const
+void FSlateApplicationBase::InvalidateAllWidgets(bool bClearResourcesImmediately) const
 {
 	SCOPED_NAMED_EVENT(Slate_GlobalInvalidate, FColor::Red);
-	UE_LOG(LogSlate, Log, TEXT("All widgets were invalidated"));
-	OnInvalidateAllWidgetsEvent.Broadcast();
+	UE_LOG(LogSlate, Log, TEXT("InvalidateAllWidgets triggered.  All widgets were invalidated"));
+	OnInvalidateAllWidgetsEvent.Broadcast(bClearResourcesImmediately);
 }
diff --git a/Engine/Source/Runtime/SlateCore/Private/FastUpdate/SlateInvalidationRoot.cpp b/Engine/Source/Runtime/SlateCore/Private/FastUpdate/SlateInvalidationRoot.cpp
--- a/Engine/Source/Runtime/SlateCore/Private/FastUpdate/SlateInvalidationRoot.cpp
+++ b/Engine/Source/Runtime/SlateCore/Private/FastUpdate/SlateInvalidationRoot.cpp
@@ -211,9 +211,6 @@ FSlateInvalidationResult FSlateInvalidationRoot::PaintInvalidationRoot(const FSl
 		// We should not have been supplied a different root than the one we generated a path to
 		check(RootWidget == FastWidgetPathList[0].Widget);
 
-		// hack, layout caching must be disabled
-		//TGuardValue<int32> LayoutCachingGuard(GSlateLayoutCaching, 0);
-
 		Result.bRepaintedWidgets = PaintFastPath(Context);
 	}
 
@@ -576,11 +573,15 @@ void FSlateInvalidationRoot::ClearAllFastPathData(bool bInvalidationRootBeingDes
 	FinalUpdateList.Empty();
 }
 
-void FSlateInvalidationRoot::OnInvalidateAllWidgets()
+void FSlateInvalidationRoot::OnInvalidateAllWidgets(bool bClearResourcesImmediately)
 {
 	InvalidateChildOrder();
 
 	InvalidationRootWidget->InvalidatePrepass();
 
+	if (bClearResourcesImmediately)
+	{
+		ClearAllFastPathData(false);
+	}
 	bNeedsSlowPath = true;
 }
diff --git a/Engine/Source/Runtime/SlateCore/Private/Fonts/FontCache.cpp b/Engine/Source/Runtime/SlateCore/Private/Fonts/FontCache.cpp
--- a/Engine/Source/Runtime/SlateCore/Private/Fonts/FontCache.cpp
+++ b/Engine/Source/Runtime/SlateCore/Private/Fonts/FontCache.cpp
@@ -1197,7 +1197,7 @@ bool FSlateFontCache::FlushCache()
 void FSlateFontCache::FlushData()
 {
 	// Ensure all invalidation panels are cleared of cached widgets
-	FSlateApplicationBase::Get().InvalidateAllWidgets();
+	FSlateApplicationBase::Get().InvalidateAllWidgets(false);
 
 	if (GIsEditor || UnloadFreeTypeDataOnFlush)
 	{
diff --git a/Engine/Source/Runtime/SlateCore/Private/Rendering/SlateDrawBuffer.cpp b/Engine/Source/Runtime/SlateCore/Private/Rendering/SlateDrawBuffer.cpp
--- a/Engine/Source/Runtime/SlateCore/Private/Rendering/SlateDrawBuffer.cpp
+++ b/Engine/Source/Runtime/SlateCore/Private/Rendering/SlateDrawBuffer.cpp
@@ -62,6 +62,18 @@ void FSlateDrawBuffer::Unlock()
 	FPlatformAtomics::InterlockedExchange(&Locked, 0);
 }
 
+void FSlateDrawBuffer::AddReferencedObjects(FReferenceCollector& Collector)
+{
+	// Locked buffers are the only ones that are currently referencing objects.  If unlocked, the element list is not in-use and contains to-be-cleared data
+	if(Locked != 0)
+	{
+		for (TSharedRef<FSlateWindowElementList>& ElementList : WindowElementLists)
+		{
+			ElementList->AddReferencedObjects(Collector);
+		}
+	}
+}
+
 void FSlateDrawBuffer::ClearBuffer()
 {
 	// Remove any window elements that are no longer valid.
diff --git a/Engine/Source/Runtime/SlateCore/Public/Application/SlateApplicationBase.h b/Engine/Source/Runtime/SlateCore/Public/Application/SlateApplicationBase.h
--- a/Engine/Source/Runtime/SlateCore/Public/Application/SlateApplicationBase.h
+++ b/Engine/Source/Runtime/SlateCore/Public/Application/SlateApplicationBase.h
@@ -428,7 +428,7 @@ public:
 	/**
 	 * Gets a delegate that is invoked when a global invalidate of all widgets should occur
 	 */
-	DECLARE_EVENT(FSlateApplicationBase, FOnInvalidateAllWidgets);
+	DECLARE_EVENT_OneParam(FSlateApplicationBase, FOnInvalidateAllWidgets, bool);
 	FOnInvalidateAllWidgets& OnInvalidateAllWidgets() { return OnInvalidateAllWidgetsEvent; }
 
 	DECLARE_EVENT_OneParam(FSlateApplicationBase, FOnGlobalInvalidationToggled, bool);
@@ -440,7 +440,7 @@ public:
 	 * Notifies all invalidation panels that they should invalidate their contents
 	 * Note: this is a very expensive call and should only be done in non-performance critical situations
 	 */
-	void InvalidateAllWidgets() const;
+	void InvalidateAllWidgets(bool bClearResourcesImmediately) const;
 private:
 	/**
 	 * Implementation for active timer registration. See SWidget::RegisterActiveTimer.
diff --git a/Engine/Source/Runtime/SlateCore/Public/FastUpdate/SlateInvalidationRoot.h b/Engine/Source/Runtime/SlateCore/Public/FastUpdate/SlateInvalidationRoot.h
--- a/Engine/Source/Runtime/SlateCore/Public/FastUpdate/SlateInvalidationRoot.h
+++ b/Engine/Source/Runtime/SlateCore/Public/FastUpdate/SlateInvalidationRoot.h
@@ -95,7 +95,7 @@ protected:
 
 private:
 
-	void OnInvalidateAllWidgets();
+	void OnInvalidateAllWidgets(bool bClearResourcesImmediately);
 
 	bool PaintFastPath(const FSlateInvalidationContext& Context);
 
diff --git a/Engine/Source/Runtime/SlateCore/Public/Rendering/DrawElements.h b/Engine/Source/Runtime/SlateCore/Public/Rendering/DrawElements.h
--- a/Engine/Source/Runtime/SlateCore/Public/Rendering/DrawElements.h
+++ b/Engine/Source/Runtime/SlateCore/Public/Rendering/DrawElements.h
@@ -438,7 +438,7 @@ private:
 /**
  * Represents a top level window and its draw elements.
  */
-class FSlateWindowElementList : public FGCObject, public FNoncopyable
+class FSlateWindowElementList : public FNoncopyable
 {
 	friend class FSlateElementBatcher;
 public:
@@ -614,7 +614,7 @@ public:
 
 	SLATECORE_API void SetRenderTargetWindow(SWindow* InRenderTargetWindow);
 
-	virtual void AddReferencedObjects(FReferenceCollector& Collector) override;
+	void AddReferencedObjects(FReferenceCollector& Collector);
 
 private:
 	FSlateDrawElement& AddCachedElement();
diff --git a/Engine/Source/Runtime/SlateCore/Public/Rendering/SlateDrawBuffer.h b/Engine/Source/Runtime/SlateCore/Public/Rendering/SlateDrawBuffer.h
--- a/Engine/Source/Runtime/SlateCore/Public/Rendering/SlateDrawBuffer.h
+++ b/Engine/Source/Runtime/SlateCore/Public/Rendering/SlateDrawBuffer.h
@@ -10,7 +10,7 @@ class SWindow;
 /**
  * Implements a draw buffer for Slate.
  */
-class SLATECORE_API FSlateDrawBuffer
+class SLATECORE_API FSlateDrawBuffer : public FGCObject
 {
 public:
 
@@ -62,6 +62,13 @@ public:
 	 */
 	void Unlock();
 
+	/** FGCObject Interface */
+	virtual void AddReferencedObjects(FReferenceCollector& Collector) override;
+	virtual FString GetReferencerName() const
+	{
+		return TEXT("FSlateDrawBuffer for Uncached Elements");
+	}
+
 protected:
 	// List of window element lists.
 	TArray< TSharedRef<FSlateWindowElementList> > WindowElementLists;
diff --git a/Engine/Source/Runtime/UMG/Private/Slate/SRetainerWidget.cpp b/Engine/Source/Runtime/UMG/Private/Slate/SRetainerWidget.cpp
--- a/Engine/Source/Runtime/UMG/Private/Slate/SRetainerWidget.cpp
+++ b/Engine/Source/Runtime/UMG/Private/Slate/SRetainerWidget.cpp
@@ -230,8 +230,13 @@ void SRetainerWidget::OnRetainerModeChanged()
 	Invalidate(EInvalidateWidget::ChildOrder);
 }
 
-void SRetainerWidget::OnGlobalInvalidate()
+void SRetainerWidget::OnGlobalInvalidate(bool bClearResourcesImmediately)
 {
+	if (bClearResourcesImmediately)
+	{
+		ClearAllFastPathData(false);
+	}
+
 	RequestRender();
 }
 
diff --git a/Engine/Source/Runtime/UMG/Public/Slate/SRetainerWidget.h b/Engine/Source/Runtime/UMG/Public/Slate/SRetainerWidget.h
--- a/Engine/Source/Runtime/UMG/Public/Slate/SRetainerWidget.h
+++ b/Engine/Source/Runtime/UMG/Public/Slate/SRetainerWidget.h
@@ -103,7 +103,7 @@ protected:
 	bool ShouldBeRenderingOffscreen() const;
 	bool IsAnythingVisibleToRender() const;
 	void OnRetainerModeChanged();
-	void OnGlobalInvalidate();
+	void OnGlobalInvalidate(bool bClearResourcesImmediately);
 private:
 	void OnGlobalInvalidationToggled(bool bGlobalInvalidationEnabled);
 #if !UE_BUILD_SHIPPING
