Fix widget component navigation not working in Standalone; the deepest Window in the widget path is used as the boundary widget and for the hittest grid.

diff --git a/Engine/Source/Runtime/Slate/Private/Framework/Application/SlateApplication.cpp b/Engine/Source/Runtime/Slate/Private/Framework/Application/SlateApplication.cpp
--- a/Engine/Source/Runtime/Slate/Private/Framework/Application/SlateApplication.cpp
+++ b/Engine/Source/Runtime/Slate/Private/Framework/Application/SlateApplication.cpp
@@ -3533,16 +3533,19 @@ void FSlateApplication::ProcessReply( const FWidgetPath& CurrentEventPath, const
 			}
 			else
 			{
+				TSharedRef<SWindow> NavigationWindow = NavigationSource.GetDeepestWindow();
+
 				FNavigationEvent NavigationEvent(PlatformApplication->GetModifierKeys(), UserIndex, TheReply.GetNavigationType(), TheReply.GetNavigationGenesis());
 
 				FNavigationReply NavigationReply = FNavigationReply::Escape();
+
 				for (int32 WidgetIndex = NavigationSource.Widgets.Num() - 1; WidgetIndex >= 0; --WidgetIndex)
 				{
 					FArrangedWidget& SomeWidgetGettingEvent = NavigationSource.Widgets[WidgetIndex];
 					if (SomeWidgetGettingEvent.Widget->IsEnabled())
 					{
 						NavigationReply = SomeWidgetGettingEvent.Widget->OnNavigation(SomeWidgetGettingEvent.Geometry, NavigationEvent).SetHandler(SomeWidgetGettingEvent.Widget);
-						if (NavigationReply.GetBoundaryRule() != EUINavigationRule::Escape || WidgetIndex == 0)
+						if (NavigationReply.GetBoundaryRule() != EUINavigationRule::Escape || SomeWidgetGettingEvent.Widget == NavigationWindow || WidgetIndex == 0)
 						{
 							AttemptNavigation(NavigationSource, NavigationEvent, NavigationReply, SomeWidgetGettingEvent);
 							break;
@@ -6570,7 +6573,7 @@ bool FSlateApplication::AttemptNavigation(const FWidgetPath& NavigationSource, c
 			// Switch worlds for widgets in the current path 
 			FScopedSwitchWorldHack SwitchWorld(NavigationSource);
 
-			DestinationWidget = NavigationSource.GetWindow()->GetHittestGrid().FindNextFocusableWidget(FocusedArrangedWidget, NavigationType, NavigationReply, BoundaryWidget);
+			DestinationWidget = NavigationSource.GetDeepestWindow()->GetHittestGrid().FindNextFocusableWidget(FocusedArrangedWidget, NavigationType, NavigationReply, BoundaryWidget);
 
 #if WITH_SLATE_DEBUGGING
 			NavigationMethod = ESlateDebuggingNavigationMethod::HitTestGrid;
diff --git a/Engine/Source/Runtime/SlateCore/Private/Layout/WidgetPath.cpp b/Engine/Source/Runtime/SlateCore/Private/Layout/WidgetPath.cpp
--- a/Engine/Source/Runtime/SlateCore/Private/Layout/WidgetPath.cpp
+++ b/Engine/Source/Runtime/SlateCore/Private/Layout/WidgetPath.cpp
@@ -104,8 +104,7 @@ TOptional<FWidgetAndPointer> FWidgetPath::FindArrangedWidgetAndCursor( TSharedRe
 		: FWidgetAndPointer();
 }
 
-	
-TSharedRef<SWindow> FWidgetPath::GetWindow()
+TSharedRef<SWindow> FWidgetPath::GetWindow() const
 {
 	check(IsValid());
 
@@ -113,16 +112,18 @@ TSharedRef<SWindow> FWidgetPath::GetWindow()
 	return FirstWidgetWindow;
 }
 
-
-TSharedRef<SWindow> FWidgetPath::GetWindow() const
+TSharedRef<SWindow> FWidgetPath::GetDeepestWindow() const
 {
 	check(IsValid());
-
-	TSharedRef<SWindow> FirstWidgetWindow = StaticCastSharedRef<SWindow>(Widgets[0].Widget);
+	const int32 WindowIndex = Widgets.FindLastByPredicate([](const FArrangedWidget& SomeWidget)
+	{
+		return SomeWidget.Widget->Advanced_IsWindow();
+	});
+	check(WindowIndex != INDEX_NONE);
+	TSharedRef<SWindow> FirstWidgetWindow = StaticCastSharedRef<SWindow>(Widgets[WindowIndex].Widget);
 	return FirstWidgetWindow;
 }
 
-
 bool FWidgetPath::IsValid() const
 {
 	return Widgets.Num() > 0;
diff --git a/Engine/Source/Runtime/SlateCore/Public/Layout/ArrangedChildren.h b/Engine/Source/Runtime/SlateCore/Public/Layout/ArrangedChildren.h
--- a/Engine/Source/Runtime/SlateCore/Public/Layout/ArrangedChildren.h
+++ b/Engine/Source/Runtime/SlateCore/Public/Layout/ArrangedChildren.h
@@ -165,6 +165,12 @@ public:
 		return Array.IndexOfByPredicate( Pred );
 	}
 
+	template<typename PredicateType>
+	FORCEINLINE int32 FindLastByPredicate(const PredicateType& Pred) const
+	{
+		return Array.FindLastByPredicate(Pred);
+	}
+
 	FORCEINLINE void Remove(int32 Index, int32 Count=1)
 	{
 		Array.RemoveAt(Index, Count);
diff --git a/Engine/Source/Runtime/SlateCore/Public/Layout/WidgetPath.h b/Engine/Source/Runtime/SlateCore/Public/Layout/WidgetPath.h
--- a/Engine/Source/Runtime/SlateCore/Public/Layout/WidgetPath.h
+++ b/Engine/Source/Runtime/SlateCore/Public/Layout/WidgetPath.h
@@ -57,20 +57,20 @@ public:
 	TOptional<FArrangedWidget> FindArrangedWidget( TSharedRef<const SWidget> WidgetToFind ) const;
 
 	TOptional<FWidgetAndPointer> FindArrangedWidgetAndCursor( TSharedRef<const SWidget> WidgetToFind ) const;
-	
+
 	/**
 	 * Get the first (top-most) widget in this path, which is always a window; assumes path is valid
 	 *
 	 * @return Window at the top of this path
 	 */
-	TSharedRef<SWindow> GetWindow();
+	TSharedRef<SWindow> GetWindow() const;
 
 	/**
-	 * Get the first (top-most) widget in this path, which is always a window; assumes path is valid
+	 * Get the deepest (bottom-most) window in this path; assumes path is valid.
 	 *
 	 * @return Window at the top of this path
 	 */
-	TSharedRef<SWindow> GetWindow() const;
+	TSharedRef<SWindow> GetDeepestWindow() const;
 	
 	/** A valid path has at least one widget in it */
 	bool IsValid() const;
