Slate - Updating the SConstraintCanvas slots to invalidate their widget when thier properties are changed so that the layout is invalidated.

Michel.Champoux

diff --git a/Engine/Source/Runtime/Slate/Public/Widgets/Layout/SConstraintCanvas.h b/Engine/Source/Runtime/Slate/Public/Widgets/Layout/SConstraintCanvas.h
--- a/Engine/Source/Runtime/Slate/Public/Widgets/Layout/SConstraintCanvas.h
+++ b/Engine/Source/Runtime/Slate/Public/Widgets/Layout/SConstraintCanvas.h
@@ -34,31 +34,35 @@ public:
 	public:		
 		FSlot& Offset( const TAttribute<FMargin>& InOffset )
 		{
-			OffsetAttr = InOffset;
+			SetAttribute(OffsetAttr, InOffset, EInvalidateWidgetReason::Layout);
 			return *this;
 		}
 
 		FSlot& Anchors( const TAttribute<FAnchors>& InAnchors )
 		{
-			AnchorsAttr = InAnchors;
+			SetAttribute(AnchorsAttr, InAnchors, EInvalidateWidgetReason::Layout);
 			return *this;
 		}
 
 		FSlot& Alignment(const TAttribute<FVector2D>& InAlignment)
 		{
-			AlignmentAttr = InAlignment;
+			SetAttribute(AlignmentAttr, InAlignment, EInvalidateWidgetReason::Layout);
 			return *this;
 		}
 
 		FSlot& AutoSize(const TAttribute<bool>& InAutoSize)
 		{
-			AutoSizeAttr = InAutoSize;
+			SetAttribute(AutoSizeAttr, InAutoSize, EInvalidateWidgetReason::Layout);
 			return *this;
 		}
 
 		FSlot& ZOrder(const TAttribute<float>& InZOrder)
 		{
-			ZOrderAttr = InZOrder;
+			// Layout isn't entirely correct here, but Paint wouldn't be either.
+			// We need the parent to redraw the children because the logical order didn't change
+			// but the paint order may have if ZOrder changed, so the painted elements need to
+			// be resorted.
+			SetAttribute(ZOrderAttr, InZOrder, EInvalidateWidgetReason::Layout);
 			return *this;
 		}
 
