Slate - Reworking the SScaleBox to work better after the changes of layout post SlateGI, there may be some fallout from this, but I've rewritten several bits of scaling logic to be a lot safer to prevent hysteresis.
Slate - Removing an unused path for the old defunct layout caching that's no longer used.

Matt.Kuhlenschmidt

diff --git a/Engine/Source/Runtime/Slate/Private/Widgets/Layout/SBox.cpp b/Engine/Source/Runtime/Slate/Private/Widgets/Layout/SBox.cpp
--- a/Engine/Source/Runtime/Slate/Private/Widgets/Layout/SBox.cpp
+++ b/Engine/Source/Runtime/Slate/Private/Widgets/Layout/SBox.cpp
@@ -143,19 +143,6 @@ void SBox::SetMaxAspectRatio(TAttribute<FOptionalSize> InMaxAspectRatio)
 	}
 }
 
-void SBox::ChildLayoutChanged(EInvalidateWidget InvalidateReason)
-{
-	if (WidthOverride.IsSet() && HeightOverride.IsSet() && InvalidateReason != EInvalidateWidget::Visibility)
-	{
-		// Done.  We don't need to notify anyone else that our desired size has changed.
-	}
-	else
-	{
-		SPanel::ChildLayoutChanged(InvalidateReason);
-	}
-}
-
-
 FVector2D SBox::ComputeDesiredSize( float ) const
 {
 	EVisibility ChildVisibility = ChildSlot.GetWidget()->GetVisibility();
diff --git a/Engine/Source/Runtime/Slate/Private/Widgets/Layout/SScaleBox.cpp b/Engine/Source/Runtime/Slate/Private/Widgets/Layout/SScaleBox.cpp
--- a/Engine/Source/Runtime/Slate/Private/Widgets/Layout/SScaleBox.cpp
+++ b/Engine/Source/Runtime/Slate/Private/Widgets/Layout/SScaleBox.cpp
@@ -5,23 +5,22 @@
 #include "Framework/Application/SlateApplication.h"
 #include "Widgets/SViewport.h"
 #include "Misc/CoreDelegates.h"
+#include "Widgets/SWidget.h"
 
 
 /* SScaleBox interface
  *****************************************************************************/
 
-void SScaleBox::Construct( const SScaleBox::FArguments& InArgs )
+void SScaleBox::Construct(const SScaleBox::FArguments& InArgs)
 {
-	Stretch = InArgs._Stretch;
+	bHasCustomPrepass = true;
 
+	Stretch = InArgs._Stretch;
 
 	StretchDirection = InArgs._StretchDirection;
 	UserSpecifiedScale = InArgs._UserSpecifiedScale;
 	IgnoreInheritedScale = InArgs._IgnoreInheritedScale;
-	bSingleLayoutPass = InArgs._SingleLayoutPass;
 
-	LastIncomingScale = 1.0f;
-	LastAreaSize = FVector2D(0, 0);
 	LastFinalOffset = FVector2D(0, 0);
 
 	ChildSlot
@@ -36,7 +35,6 @@ void SScaleBox::Construct( const SScaleBox::FArguments& InArgs )
 	FSlateApplication::Get().OnDebugSafeZoneChanged.AddSP(this, &SScaleBox::DebugSafeAreaUpdated);
 #endif
 
-
 	RefreshSafeZoneScale();
 	OnSafeFrameChangedHandle = FCoreDelegates::OnSafeFrameChangedEvent.AddSP(this, &SScaleBox::RefreshSafeZoneScale);
 }
@@ -46,107 +44,167 @@ SScaleBox::~SScaleBox()
 	FCoreDelegates::OnSafeFrameChangedEvent.Remove(OnSafeFrameChangedHandle);
 }
 
-/* SWidget overrides
- *****************************************************************************/
-void SScaleBox::OnArrangeChildren( const FGeometry& AllottedGeometry, FArrangedChildren& ArrangedChildren ) const
+bool SScaleBox::CustomPrepass(float LayoutScaleMultiplier)
 {
-	const EVisibility ChildVisibility = ChildSlot.GetWidget()->GetVisibility();
-	if ( ArrangedChildren.Accepts(ChildVisibility) )
+	SWidget& ChildSlotWidget = ChildSlot.GetWidget().Get();
+
+	const bool bLocalGeometryRequired = DoesScaleRequireLocalGeometry();
+	const bool bNeedsNormalizingPrepass = DoesScaleRequireNormalizingPrepass();
+
+	//bool bCanRenderThisFrame = true;
+
+	// If we need a normalizing prepass, or we've yet to give the child a chance to generate a desired
+	// size, do that now.
+	if (bNeedsNormalizingPrepass || !LastAllocatedArea.IsSet())
 	{
-		const FVector2D AreaSize = AllottedGeometry.GetLocalSize();
+		ChildSlotWidget.SlatePrepass(LayoutScaleMultiplier);
+	}
 
-		const FVector2D CurrentWidgetDesiredSize = ChildSlot.GetWidget()->GetDesiredSize();
-		FVector2D SlotWidgetDesiredSize = CurrentWidgetDesiredSize;
-	
-		float FinalScale = 1;
+	ContentDesiredSize = ChildSlotWidget.GetDesiredSize();
 
-		bool bAllowFullLayout = true;
+	if (bLocalGeometryRequired)
+	{
+		if (LastAllocatedArea.IsSet())
+		{
+			ComputedContentScale = ComputeContentScale(LastPaintGeometry.GetValue());
+		}
+	}
+	else
+	{
+		// If we don't need the area, send a false geometry.
+		static const FGeometry NullGeometry;
+		ComputedContentScale = ComputeContentScale(NullGeometry);
+	}
 
-		if ( bSingleLayoutPass && LastContentDesiredSize.IsSet() && LastFinalScale.IsSet() && LastAreaSize.Equals(AreaSize) && FMath::IsNearlyEqual(LastIncomingScale, AllottedGeometry.Scale) )
+	if (bNeedsNormalizingPrepass)
+	{
+		ChildSlotWidget.InvalidatePrepass();
+	}
+
+	// Extract the incoming scale out of the layout scale if 
+	if (ComputedContentScale.IsSet())
+	{
+		if (IgnoreInheritedScale.Get(false) && LayoutScaleMultiplier != 0)
 		{
-			if ( SlotWidgetDesiredSize.Equals(LastContentDesiredSize.GetValue()) )
-			{
-				bAllowFullLayout = false;
-				FinalScale = LastFinalScale.GetValue();
-			}
+			ComputedContentScale = ComputedContentScale.GetValue() / LayoutScaleMultiplier;
 		}
+	}
+
+	return true;
+}
+
+bool SScaleBox::DoesScaleRequireNormalizingPrepass() const
+{
+	const EStretch::Type CurrentStretch = Stretch.Get();
+	switch (CurrentStretch)
+	{
+	case EStretch::None:
+	case EStretch::Fill:
+	case EStretch::ScaleBySafeZone:
+	case EStretch::UserSpecified:
+		return false;
+	default:
+		return true;
+	}
+}
+
+bool SScaleBox::DoesScaleRequireLocalGeometry() const
+{
+	const EStretch::Type CurrentStretch = Stretch.Get();
+	switch (CurrentStretch)
+	{
+	case EStretch::None:
+	case EStretch::Fill:
+	case EStretch::ScaleBySafeZone:
+	case EStretch::UserSpecified:
+		return false;
+	default:
+		return true;
+	}
+}
 
-		if ( bAllowFullLayout )
+float SScaleBox::ComputeContentScale(const FGeometry& PaintGeometry) const
+{
+	const EStretch::Type CurrentStretch = Stretch.Get();
+	const EStretchDirection::Type CurrentStretchDirection = StretchDirection.Get();
+
+	switch (CurrentStretch)
+	{
+	case EStretch::ScaleBySafeZone:
+		return SafeZoneScale;
+	case EStretch::UserSpecified:
+		return UserSpecifiedScale.Get(1.0f);
+	}
+
+	float FinalScale = 1;
+
+	const FVector2D ChildDesiredSize = ChildSlot.GetWidget()->GetDesiredSize();
+
+	if (ChildDesiredSize.X != 0 && ChildDesiredSize.Y != 0)
+	{
+		switch (CurrentStretch)
+		{
+		case EStretch::ScaleToFit:
 		{
-			const EStretch::Type CurrentStretch = Stretch.Get();
-			const EStretchDirection::Type CurrentStretchDirection = StretchDirection.Get();
+			//FVector2D LocalSnappedArea = PaintGeometry.LocalToRoundedLocal(PaintGeometry.GetLocalSize()) - PaintGeometry.LocalToRoundedLocal(FVector2D(0, 0));
+			//FinalScale = FMath::Min(LocalSnappedArea.X / ChildDesiredSize.X, LocalSnappedArea.Y / ChildDesiredSize.Y);
+			FinalScale = FMath::Min(PaintGeometry.GetLocalSize().X / ChildDesiredSize.X, PaintGeometry.GetLocalSize().Y / ChildDesiredSize.Y);
+			break;
+		}
+		case EStretch::ScaleToFitX:
+			FinalScale = PaintGeometry.GetLocalSize().X / ChildDesiredSize.X;
+			break;
+		case EStretch::ScaleToFitY:
+			FinalScale = PaintGeometry.GetLocalSize().Y / ChildDesiredSize.Y;
+			break;
+		case EStretch::Fill:
+			break;
+		case EStretch::ScaleToFill:
+			FinalScale = FMath::Max(PaintGeometry.GetLocalSize().X / ChildDesiredSize.X, PaintGeometry.GetLocalSize().Y / ChildDesiredSize.Y);
+			break;
+		}
 
-			bool bRequiresAnotherPrepass = CurrentStretch != EStretch::UserSpecified && CurrentStretch != EStretch::ScaleBySafeZone;
+		switch (CurrentStretchDirection)
+		{
+		case EStretchDirection::DownOnly:
+			FinalScale = FMath::Min(FinalScale, 1.0f);
+			break;
+		case EStretchDirection::UpOnly:
+			FinalScale = FMath::Max(FinalScale, 1.0f);
+			break;
+		case EStretchDirection::Both:
+			break;
+		}
+	}
 
-			if ( SlotWidgetDesiredSize.X != 0 && SlotWidgetDesiredSize.Y != 0 )
-			{
-				switch ( CurrentStretch )
-				{
-				case EStretch::None:
-					bRequiresAnotherPrepass = false;
-					break;
-				case EStretch::Fill:
-					SlotWidgetDesiredSize = AreaSize;
-					bRequiresAnotherPrepass = false;
-					break;
-				case EStretch::ScaleToFit:
-					FinalScale = FMath::Min(AreaSize.X / SlotWidgetDesiredSize.X, AreaSize.Y / SlotWidgetDesiredSize.Y);
-					break;
-				case EStretch::ScaleToFitX:
-					FinalScale = AreaSize.X / SlotWidgetDesiredSize.X;
-					break;
-				case EStretch::ScaleToFitY:
-					FinalScale = AreaSize.Y / SlotWidgetDesiredSize.Y;
-					break;
-				case EStretch::ScaleToFill:
-					FinalScale = FMath::Max(AreaSize.X / SlotWidgetDesiredSize.X, AreaSize.Y / SlotWidgetDesiredSize.Y);
-					break;
-				case EStretch::ScaleBySafeZone:
-					FinalScale = SafeZoneScale;
-					bRequiresAnotherPrepass = false;
-					break;
-				case EStretch::UserSpecified:
-					FinalScale = UserSpecifiedScale.Get(1.0f);
-					bRequiresAnotherPrepass = false;
-					break;
-				}
+	return FinalScale;
+}
 
-				switch ( CurrentStretchDirection )
-				{
-				case EStretchDirection::DownOnly:
-					FinalScale = FMath::Min(FinalScale, 1.0f);
-					break;
-				case EStretchDirection::UpOnly:
-					FinalScale = FMath::Max(FinalScale, 1.0f);
-					break;
-				case EStretchDirection::Both:
-					break;
-				}
+void SScaleBox::OnArrangeChildren(const FGeometry& AllottedGeometry, FArrangedChildren& ArrangedChildren) const
+{
+	const EVisibility ChildVisibility = ChildSlot.GetWidget()->GetVisibility();
+	if (ArrangedChildren.Accepts(ChildVisibility))
+	{
+		const FVector2D AreaSize = AllottedGeometry.GetLocalSize();
 
-				// Force full layout calculations when the previously calculated final scale is zero
-				if (!FMath::IsNearlyZero(FinalScale) || !bSingleLayoutPass)
-				{
-					LastFinalScale = FinalScale;
-				}
-				else
-				{
-					LastFinalScale.Reset();
-				}
-			}
-			else
-			{
-				LastFinalScale.Reset();
-			}
+		const FVector2D CurrentWidgetDesiredSize = ChildSlot.GetWidget()->GetDesiredSize();
+		FVector2D SlotWidgetDesiredSize = CurrentWidgetDesiredSize;
 
-			if ( IgnoreInheritedScale.Get(false) && AllottedGeometry.Scale != 0 )
-			{
-				FinalScale /= AllottedGeometry.Scale;
-			}
+		const EStretch::Type CurrentStretch = Stretch.Get();
+		const EStretchDirection::Type CurrentStretchDirection = StretchDirection.Get();
 
+		if (CurrentStretch == EStretch::Fill)
+		{
+			SlotWidgetDesiredSize = AreaSize;
+		}
+
+		if (ComputedContentScale.IsSet())
+		{
 			LastFinalOffset = FVector2D(0, 0);
+			float FinalScale = ComputedContentScale.GetValue();
 
 			// If we're just filling, there's no scale applied, we're just filling the area.
-			if ( CurrentStretch != EStretch::Fill )
+			if (CurrentStretch != EStretch::Fill)
 			{
 				const FMargin SlotPadding(ChildSlot.SlotPadding.Get());
 				AlignmentArrangeResult XResult = AlignChild<Orient_Horizontal>(AreaSize.X, ChildSlot, SlotPadding, FinalScale, false);
@@ -156,66 +214,41 @@ void SScaleBox::OnArrangeChildren( const FGeometry& AllottedGeometry, FArrangedC
 
 				// If the layout horizontally is fill, then we need the desired size to be the whole size of the widget, 
 				// but scale the inverse of the scale we're applying.
-				if ( ChildSlot.HAlignment == HAlign_Fill )
+				if (ChildSlot.HAlignment == HAlign_Fill)
 				{
 					SlotWidgetDesiredSize.X = AreaSize.X / FinalScale;
 				}
 
 				// If the layout vertically is fill, then we need the desired size to be the whole size of the widget, 
 				// but scale the inverse of the scale we're applying.
-				if ( ChildSlot.VAlignment == VAlign_Fill )
+				if (ChildSlot.VAlignment == VAlign_Fill)
 				{
 					SlotWidgetDesiredSize.Y = AreaSize.Y / FinalScale;
 				}
 			}
 
-			LastAreaSize = AreaSize;
-			LastIncomingScale = AllottedGeometry.Scale;
-		
-
-			if ( bRequiresAnotherPrepass)
-			{
-				SWidget* ChildSlotWidget = &ChildSlot.GetWidget().Get();
-				// We need to run another pre-pass now that we know the final scale.
-				// This will allow things that don't scale linearly (such as text) to update their size and layout correctly.
-				//
-				// NOTE: This step is pretty expensive especially if you're nesting scale boxes.
-
-				const float NewPrepassScaleMultiplier = AllottedGeometry.GetAccumulatedLayoutTransform().GetScale() * FinalScale;
-
-				//if(ChildSlotWidget->GetPrepassLayoutScaleMultiplier() != NewPrepassScaleMultiplier)
-				{
-					ChildSlotWidget->InvalidatePrepass();
-					ChildSlotWidget->SlatePrepass(NewPrepassScaleMultiplier);
-
-					LastContentDesiredSize = ChildSlotWidget->GetDesiredSize();
-
-					if (SlotWidgetDesiredSize != LastSlotWidgetDesiredSize)
-					{
-						const_cast<SScaleBox*>(this)->Invalidate(EInvalidateWidget::Layout);
-					}
-				}
-			}
-			else
-			{
-				LastContentDesiredSize.Reset();
-				LastFinalScale.Reset();
-			}
-
-			LastSlotWidgetDesiredSize = SlotWidgetDesiredSize;
+			ArrangedChildren.AddWidget(ChildVisibility, AllottedGeometry.MakeChild(
+				ChildSlot.GetWidget(),
+				LastFinalOffset,
+				SlotWidgetDesiredSize,
+				FinalScale
+			));
 		}
-
-		ArrangedChildren.AddWidget(ChildVisibility, AllottedGeometry.MakeChild(
-			ChildSlot.GetWidget(),
-			LastFinalOffset,
-			LastSlotWidgetDesiredSize,
-			FinalScale
-		) );
 	}
 }
 
 int32 SScaleBox::OnPaint( const FPaintArgs& Args, const FGeometry& AllottedGeometry, const FSlateRect& MyCullingRect, FSlateWindowElementList& OutDrawElements, int32 LayerId, const FWidgetStyle& InWidgetStyle, bool bParentEnabled ) const
 {
+	// We need another layout pass if the incoming allocated geometry is different from last frames.
+	if (!LastAllocatedArea.IsSet() || !AllottedGeometry.GetLocalSize().Equals(LastAllocatedArea.GetValue()))
+	{
+		LastAllocatedArea = AllottedGeometry.GetLocalSize();
+		LastPaintGeometry = AllottedGeometry;
+		const_cast<SScaleBox*>(this)->Invalidate(EInvalidateWidgetReason::Layout);
+		const_cast<SScaleBox*>(this)->InvalidatePrepass();
+		const_cast<SScaleBox*>(this)->SlatePrepass(GetPrepassLayoutScaleMultiplier());
+	}
+
 	bool bClippingNeeded = false;
 
 	if (GetClipping() == EWidgetClipping::Inherit)
@@ -278,100 +311,52 @@ void SScaleBox::SetVAlign(EVerticalAlignment VAlign)
 
 void SScaleBox::SetStretchDirection(EStretchDirection::Type InStretchDirection)
 {
-	if(!StretchDirection.IdenticalTo(InStretchDirection))
-	{
-		StretchDirection = InStretchDirection;
-		Invalidate(EInvalidateWidget::Layout);
-	}
+	SetAttribute(StretchDirection, TAttribute<EStretchDirection::Type>(InStretchDirection), EInvalidateWidgetReason::Layout);
 }
 
 void SScaleBox::SetStretch(EStretch::Type InStretch)
 {
-	if(!Stretch.IdenticalTo(InStretch))
+	if (SetAttribute(Stretch, TAttribute<EStretch::Type>(InStretch), EInvalidateWidgetReason::Layout))
 	{
-		Stretch = InStretch;
 		RefreshSafeZoneScale();
-		Invalidate(EInvalidateWidget::Layout);
 	}
 }
 
 void SScaleBox::SetUserSpecifiedScale(float InUserSpecifiedScale)
 {
-	if(!UserSpecifiedScale.IdenticalTo(InUserSpecifiedScale))
-	{
-		UserSpecifiedScale = InUserSpecifiedScale;
-		Invalidate(EInvalidateWidget::Layout);
-	}
+	SetAttribute(UserSpecifiedScale, TAttribute<float>(InUserSpecifiedScale), EInvalidateWidgetReason::Layout);
 }
 
 void SScaleBox::SetIgnoreInheritedScale(bool InIgnoreInheritedScale)
 {
-	if(!IgnoreInheritedScale.IdenticalTo(InIgnoreInheritedScale))
-	{
-		IgnoreInheritedScale = InIgnoreInheritedScale;
-		Invalidate(EInvalidateWidget::Layout);
-	}
+	SetAttribute(IgnoreInheritedScale, TAttribute<bool>(InIgnoreInheritedScale), EInvalidateWidgetReason::Layout);
 }
 
 FVector2D SScaleBox::ComputeDesiredSize(float InScale) const
 {
-	float ExpectedLayoutScale = GetLayoutScale();
-	if (IgnoreInheritedScale.Get(false))
+	if (DoesScaleRequireNormalizingPrepass())
 	{
-		return ExpectedLayoutScale * SCompoundWidget::ComputeDesiredSize(InScale) / InScale;
+		// If we require a normalizing pre-pass, we can never allow the scaled content's desired size to affect
+		// the area we return that we need, otherwise, we'll be introducing hysteresis.
+		if (ContentDesiredSize.IsSet())
+		{
+			return ContentDesiredSize.GetValue();
+		}
 	}
-
-	FVector2D ComputedDesiredSize = SCompoundWidget::ComputeDesiredSize(InScale);
-
-	const EStretch::Type CurrentStretch = Stretch.Get();
-
-	switch (CurrentStretch)
+	else
 	{
-	case EStretch::ScaleToFitX:
-		ExpectedLayoutScale = ComputedDesiredSize.X == 0.0f ? 1.0f : FMath::Max(1.0f, GetTickSpaceGeometry().GetLocalSize().X / ComputedDesiredSize.X);
-		break;
-	case EStretch::ScaleToFitY:
-		ExpectedLayoutScale = ComputedDesiredSize.Y == 0.0f ? 1.0f : FMath::Max(1.0f, GetTickSpaceGeometry().GetLocalSize().Y / ComputedDesiredSize.Y);
-		break;
+		if (ContentDesiredSize.IsSet() && ComputedContentScale.IsSet())
+		{
+			return ContentDesiredSize.GetValue() * ComputedContentScale.GetValue();
+		}
 	}
 
-	return ExpectedLayoutScale * ComputedDesiredSize;
+	return SCompoundWidget::ComputeDesiredSize(InScale);
 }
 
 float SScaleBox::GetRelativeLayoutScale(const FSlotBase& Child, float LayoutScaleMultiplier) const
 {
-	if ( IgnoreInheritedScale.Get(false) )
-	{
-		return GetLayoutScale() / LayoutScaleMultiplier;
-	}
-
-	return GetLayoutScale();
-}
-
-float SScaleBox::GetLayoutScale() const
-{
-	const EStretch::Type CurrentStretch = Stretch.Get();
-	
-	switch (CurrentStretch)
-	{
-	case EStretch::ScaleBySafeZone:
-		return SafeZoneScale;
-	case EStretch::UserSpecified:
-		return UserSpecifiedScale.Get(1.0f);
-	default:
-		if ( bSingleLayoutPass )
-		{
-			if ( LastFinalScale.IsSet() )
-			{
-				return LastFinalScale.GetValue();
-			}
-		}
-
-		// Because our scale is determined by our size, we always report a scale of 1.0 here, 
-		// as reporting our actual scale can cause a feedback loop whereby the calculated size changes each frame.
-		// We workaround this by forcibly pre-passing our child content a second time once we know its final scale in OnArrangeChildren.
-		return 1.0f;
-	}
+	return ComputedContentScale.IsSet() ? ComputedContentScale.GetValue() : 1.0f;
 }
 
 void SScaleBox::RefreshSafeZoneScale()
@@ -414,6 +399,7 @@ void SScaleBox::RefreshSafeZoneScale()
 }
 
 #if WITH_EDITOR
+
 void SScaleBox::DebugSafeAreaUpdated(const FMargin& NewSafeZone, bool bShouldRecacheMetrics)
 {
 	RefreshSafeZoneScale();
diff --git a/Engine/Source/Runtime/Slate/Public/Widgets/Layout/SBox.h b/Engine/Source/Runtime/Slate/Public/Widgets/Layout/SBox.h
--- a/Engine/Source/Runtime/Slate/Public/Widgets/Layout/SBox.h
+++ b/Engine/Source/Runtime/Slate/Public/Widgets/Layout/SBox.h
@@ -129,7 +129,6 @@ public:
 
 protected:
 	// Begin SWidget overrides.
-	virtual void ChildLayoutChanged(EInvalidateWidget InvalidateReason) override;
 	virtual FVector2D ComputeDesiredSize(float) const override;
 	float ComputeDesiredWidth() const;
 	float ComputeDesiredHeight() const;
diff --git a/Engine/Source/Runtime/Slate/Public/Widgets/Layout/SScaleBox.h b/Engine/Source/Runtime/Slate/Public/Widgets/Layout/SScaleBox.h
--- a/Engine/Source/Runtime/Slate/Public/Widgets/Layout/SScaleBox.h
+++ b/Engine/Source/Runtime/Slate/Public/Widgets/Layout/SScaleBox.h
@@ -82,7 +82,6 @@ public:
 	, _Stretch(EStretch::None)
 	, _UserSpecifiedScale(1.0f)
 	, _IgnoreInheritedScale(false)
-	, _SingleLayoutPass(false)
 	{}
 		/** Slot for this designers content (optional) */
 		SLATE_DEFAULT_SLOT(FArguments, Content)
@@ -105,15 +104,6 @@ public:
 		/** Undo any inherited scale factor before applying this scale box's scale */
 		SLATE_ATTRIBUTE(bool, IgnoreInheritedScale)
 
-		/**
-		 * Only perform a single layout pass, if you do this, it can save a considerable 
-		 * amount of time, however, some things like text may not look correct. You may also
-		 * see the UI judder between frames.  This generally is caused by not explicitly
-		 * sizing the widget, and instead allowing it to layout based on desired size along
-		 * which won't work in Single Layout Pass mode.
-		 */
-		SLATE_ARGUMENT(bool, SingleLayoutPass)
-
 #if WITH_EDITOR
 		/** Force a particular screen size to be used instead of the reported device size. */
 		SLATE_ARGUMENT(TOptional<FVector2D>, OverrideScreenSize)
@@ -164,12 +154,17 @@ public:
 	
 protected:
 	// Begin SWidget overrides.
+	virtual bool CustomPrepass(float LayoutScaleMultiplier) override;
 	virtual FVector2D ComputeDesiredSize(float InScale) const override;
 	virtual float GetRelativeLayoutScale(const FSlotBase& Child, float LayoutScaleMultiplier) const override;
 	// End SWidget overrides.
 
-	float GetLayoutScale() const;
+	bool DoesScaleRequireNormalizingPrepass() const;
+	bool DoesScaleRequireLocalGeometry() const;
+	float ComputeContentScale(const FGeometry& PaintGeometry) const;
+
 	void RefreshSafeZoneScale();
+
 #if WITH_EDITOR
 	void DebugSafeAreaUpdated(const FMargin& NewSafeZone, bool bShouldRecacheMetrics);
 #endif
@@ -193,25 +188,16 @@ protected:
 	/** Delegate handle to unhook the safe frame changed. */
 	FDelegateHandle OnSafeFrameChangedHandle;
 
-	/**
-	 * Determines if this scale box should attempt to layout everything using only a 
-	 * single pass each frame.  This is a MUCH more efficient mode, since invalidating 
-	 * text during prepass is expensive, however for that very reason this mode may not 
-	 * work always, but for large pieces of UI with a restricted size you should try it.
-	 */
-	bool bSingleLayoutPass;
+	mutable TOptional<FVector2D> LastAllocatedArea;
+	mutable TOptional<FGeometry> LastPaintGeometry;
 
-	/**  */
-	mutable TOptional<FVector2D> LastContentDesiredSize;
+	mutable TOptional<FVector2D> ContentDesiredSize;
 
 	/**  */
-	mutable FVector2D LastAreaSize;
+	mutable TOptional<float> ComputedContentScale;
+
 
-	/**  */
-	mutable float LastIncomingScale;
 
-	/**  */
-	mutable TOptional<float> LastFinalScale;
 
 	/**  */
 	mutable FVector2D LastFinalOffset;
diff --git a/Engine/Source/Runtime/SlateCore/Private/Widgets/SWidget.cpp b/Engine/Source/Runtime/SlateCore/Private/Widgets/SWidget.cpp
--- a/Engine/Source/Runtime/SlateCore/Private/Widgets/SWidget.cpp
+++ b/Engine/Source/Runtime/SlateCore/Private/Widgets/SWidget.cpp
@@ -645,29 +645,6 @@ bool SWidget::ConditionallyDetatchParentWidget(SWidget* InExpectedParent)
 	return false;
 }
 
-
-void SWidget::LayoutChanged(EInvalidateWidget InvalidateReason)
-{
-	if(EnumHasAnyFlags(InvalidateReason, EInvalidateWidget::Layout))
-	{
-		bNeedsDesiredSize = true;
-
-		TSharedPtr<SWidget> ParentWidget = ParentWidgetPtr.Pin();
-		if (ParentWidget.IsValid())
-		{
-			ParentWidget->ChildLayoutChanged(InvalidateReason);
-		}
-	}
-}
-
-void SWidget::ChildLayoutChanged(EInvalidateWidget InvalidateReason)
-{
-	if (!bNeedsDesiredSize || InvalidateReason == EInvalidateWidget::Visibility )
-	{
-		LayoutChanged(InvalidateReason);
-	}
-}
-
 void SWidget::AssignIndicesToChildren(FSlateInvalidationRoot& Root, int32 ParentIndex, TArray<FWidgetProxy, TMemStackAllocator<>>& FastPathList, bool bParentVisible, bool bParentVolatile)
 {
 	FWidgetProxy MyProxy(this);
diff --git a/Engine/Source/Runtime/SlateCore/Public/Widgets/SCompoundWidget.h b/Engine/Source/Runtime/SlateCore/Public/Widgets/SCompoundWidget.h
--- a/Engine/Source/Runtime/SlateCore/Public/Widgets/SCompoundWidget.h
+++ b/Engine/Source/Runtime/SlateCore/Public/Widgets/SCompoundWidget.h
@@ -38,11 +38,7 @@ public:
 	 */
 	void SetContentScale( const TAttribute< FVector2D >& InContentScale )
 	{
-		if (!ContentScale.IdenticalTo(InContentScale))
-		{
-			ContentScale = InContentScale;
-			Invalidate(EInvalidateWidget::LayoutAndVolatility);
-		}
+		SetAttribute(ContentScale, InContentScale, EInvalidateWidgetReason::Layout);
 	}
 
 	/**
@@ -60,11 +56,7 @@ public:
 	 */
 	void SetColorAndOpacity( const TAttribute<FLinearColor>& InColorAndOpacity )
 	{
-		if (!ColorAndOpacity.IdenticalTo(InColorAndOpacity))
-		{
-			ColorAndOpacity = InColorAndOpacity;
-			Invalidate(EInvalidateWidget::PaintAndVolatility);
-		}
+		SetAttribute(ColorAndOpacity, InColorAndOpacity, EInvalidateWidgetReason::Paint);
 	}
 
 	/**
@@ -74,11 +66,7 @@ public:
 	 */
 	void SetForegroundColor( const TAttribute<FSlateColor>& InForegroundColor )
 	{
-		if (!ForegroundColor.IdenticalTo(InForegroundColor))
-		{
-			ForegroundColor = InForegroundColor;
-			Invalidate(EInvalidateWidget::PaintAndVolatility);
-		}
+		SetAttribute(ForegroundColor, InForegroundColor, EInvalidateWidgetReason::Paint);
 	}
 
 public:
diff --git a/Engine/Source/Runtime/SlateCore/Public/Widgets/SWidget.h b/Engine/Source/Runtime/SlateCore/Public/Widgets/SWidget.h
--- a/Engine/Source/Runtime/SlateCore/Public/Widgets/SWidget.h
+++ b/Engine/Source/Runtime/SlateCore/Public/Widgets/SWidget.h
@@ -26,6 +26,7 @@
 #include "SlateGlobals.h"
 #include "Types/PaintArgs.h"
 #include "FastUpdate/WidgetProxy.h"
+#include "InvalidateWidgetReason.h"
 #include "Widgets/Accessibility/SlateWidgetAccessibleTypes.h"
 
 class FActiveTimerHandle;
@@ -127,66 +128,6 @@ enum class EAccessibleType : uint8
 	Summary
 };
 
-/**
- * The different types of invalidation that are possible for a widget.
- */
-enum class EInvalidateWidgetReason : uint8
-{
-	None = 0,
-
-	/**
-	 * Use Layout invalidation if your widget needs to change desired size.  This is an expensive invalidation so do not use if all you need to do is redraw a widget
-	 */
-	Layout = 1 << 0,
-
-	/**
-	 * Use when the painting of widget has been altered, but nothing affecting sizing.
-	 */
-	Paint = 1 << 1,
-
-	/**
-	 * Use if just the volatility of the widget has been adjusted.
-	 */
-	Volatility = 1 << 2,
-
-	/**
-	 * A child was added or removed.   (this implies layout)
-	 */
-	ChildOrder = 1 << 3,
-
-	/** A Widgets render transform changed */
-	RenderTransform = 1 << 4,
-
-	/**
-	 * Changing visibility (this implies layout)
-	 */
-	Visibility = 1 << 5,
-
-	/**
-	 * Use Paint invalidation if you're changing a normal property involving painting or sizing.
-	 * Additionally if the property that was changed affects Volatility in anyway, it's important
-	 * that you invalidate volatility so that it can be recalculated and cached.
-	 */
-	PaintAndVolatility = Paint | Volatility,
-	/**
-	 * Use Layout invalidation if you're changing a normal property involving painting or sizing.
-	 * Additionally if the property that was changed affects Volatility in anyway, it's important
-	 * that you invalidate volatility so that it can be recalculated and cached.
-	 */
-	LayoutAndVolatility = Layout | Volatility,
-
-
-	/**
-	 * Do not use this ever unless you know what you are doing
-	 */
-	All UE_DEPRECATED(4.22, "EInvalidateWidget::All has been deprecated.  You probably wanted EInvalidateWidget::Layout but if you need more than that then use bitwise or to combine them") = 0xff
-};
-
-ENUM_CLASS_FLAGS(EInvalidateWidgetReason)
-
-// This typedefed because EInvalidateWidget will be deprecated soon
-typedef EInvalidateWidgetReason EInvalidateWidget;
-
 
 /**
  * An FPopupLayer hosts the pop-up content which could be anything you want to appear on top of a widget.
@@ -714,8 +655,6 @@ public:
 	void SetCanTick(bool bInCanTick) { bInCanTick ? AddUpdateFlags(EWidgetUpdateFlags::NeedsTick) : RemoveUpdateFlags(EWidgetUpdateFlags::NeedsTick); }
 	bool GetCanTick() const { return HasAnyUpdateFlags(EWidgetUpdateFlags::NeedsTick); }
 
-	virtual void ChildLayoutChanged(EInvalidateWidget InvalidateReason);
-
 	const FSlateWidgetPersistentState& GetPersistentState() const { return PersistentState; }
 	const FWidgetProxyHandle GetProxyHandle() const { return FastPathProxyHandle; }
 
@@ -784,8 +723,6 @@ private:
 		DesiredSize = InDesiredSize;
 	}
 
-	void LayoutChanged(EInvalidateWidget InvalidateReason);
-
 	void CreateStatID() const;
 
 	void AddUpdateFlags(EWidgetUpdateFlags FlagsToAdd)
diff --git a/Engine/Source/Runtime/UMG/Private/Components/ScaleBox.cpp b/Engine/Source/Runtime/UMG/Private/Components/ScaleBox.cpp
--- a/Engine/Source/Runtime/UMG/Private/Components/ScaleBox.cpp
+++ b/Engine/Source/Runtime/UMG/Private/Components/ScaleBox.cpp
@@ -34,7 +34,7 @@ TSharedRef<SWidget> UScaleBox::RebuildWidget()
 #if WITH_EDITOR
 		.OverrideScreenSize(DesignerSize)
 #endif
-	.SingleLayoutPass(bSingleLayoutPass);
+	;
 
 	if ( GetChildrenCount() > 0 )
 	{
diff --git a/Engine/Source/Runtime/UMG/Public/Components/ScaleBox.h b/Engine/Source/Runtime/UMG/Public/Components/ScaleBox.h
--- a/Engine/Source/Runtime/UMG/Public/Components/ScaleBox.h
+++ b/Engine/Source/Runtime/UMG/Public/Components/ScaleBox.h
@@ -40,16 +40,6 @@ public:
 	UPROPERTY(EditAnywhere, BlueprintReadOnly, Category = "Stretching")
 	bool IgnoreInheritedScale;
 
-	/**
-	 * Only perform a single layout pass, if you do this, it can save a considerable
-	 * amount of time, however, some things like text may not look correct.  You may also
-	 * see the UI judder between frames.  This generally is caused by not explicitly
-	 * sizing the widget, and instead allowing it to layout based on desired size along
-	 * which won't work in Single Layout Pass mode.
-	 */
-	UPROPERTY(EditAnywhere, BlueprintReadOnly, config, Category = "Performance")
-	bool bSingleLayoutPass;
-
 public:
 	UFUNCTION(BlueprintCallable, Category = "Appearance")
 	void SetStretch(EStretch::Type InStretch);
