ScaleBox - Restoring the special case for the scalebox to be able to predict the desired size of the opposite dimension when we're only scaling in a single dimension, this does involve introducing hysteresis, but it's very limited, and theoretically stable as long as the child doesn't change size in the dimesnion we're scaling based on the width in return.

diff --git a/Engine/Source/Runtime/Slate/Private/Widgets/Layout/SScaleBox.cpp b/Engine/Source/Runtime/Slate/Private/Widgets/Layout/SScaleBox.cpp
--- a/Engine/Source/Runtime/Slate/Private/Widgets/Layout/SScaleBox.cpp
+++ b/Engine/Source/Runtime/Slate/Private/Widgets/Layout/SScaleBox.cpp
@@ -61,19 +61,20 @@ bool SScaleBox::CustomPrepass(float LayoutScaleMultiplier)
 	}
 
 	ContentDesiredSize = ChildSlotWidget.GetDesiredSize();
+	TOptional<float> NewComputedContentScale;
 
 	if (bLocalGeometryRequired)
 	{
 		if (LastAllocatedArea.IsSet())
 		{
-			ComputedContentScale = ComputeContentScale(LastPaintGeometry.GetValue());
+			NewComputedContentScale = ComputeContentScale(LastPaintGeometry.GetValue());
 		}
 	}
 	else
 	{
 		// If we don't need the area, send a false geometry.
 		static const FGeometry NullGeometry;
-		ComputedContentScale = ComputeContentScale(NullGeometry);
+		NewComputedContentScale = ComputeContentScale(NullGeometry);
 	}
 
 	if (bNeedsNormalizingPrepass)
@@ -82,14 +83,16 @@ bool SScaleBox::CustomPrepass(float LayoutScaleMultiplier)
 	}
 
 	// Extract the incoming scale out of the layout scale if 
-	if (ComputedContentScale.IsSet())
+	if (NewComputedContentScale.IsSet())
 	{
 		if (IgnoreInheritedScale.Get(false) && LayoutScaleMultiplier != 0)
 		{
-			ComputedContentScale = ComputedContentScale.GetValue() / LayoutScaleMultiplier;
+			NewComputedContentScale = NewComputedContentScale.GetValue() / LayoutScaleMultiplier;
 		}
 	}
 
+	ComputedContentScale = NewComputedContentScale;
+
 	return true;
 }
 
@@ -123,6 +126,19 @@ bool SScaleBox::DoesScaleRequireLocalGeometry() const
 	}
 }
 
+bool SScaleBox::IsDesiredSizeDependentOnAreaAndScale() const
+{
+	const EStretch::Type CurrentStretch = Stretch.Get();
+	switch (CurrentStretch)
+	{
+	case EStretch::ScaleToFitX:
+	case EStretch::ScaleToFitY:
+		return true;
+	default:
+		return false;
+	}
+}
+
 float SScaleBox::ComputeContentScale(const FGeometry& PaintGeometry) const
 {
 	const EStretch::Type CurrentStretch = Stretch.Get();
@@ -246,7 +262,7 @@ int32 SScaleBox::OnPaint( const FPaintArgs& Args, const FGeometry& AllottedGeome
 		LastPaintGeometry = AllottedGeometry;
 		const_cast<SScaleBox*>(this)->Invalidate(EInvalidateWidgetReason::Layout);
 		const_cast<SScaleBox*>(this)->InvalidatePrepass();
-		const_cast<SScaleBox*>(this)->SlatePrepass(GetPrepassLayoutScaleMultiplier());
+		//const_cast<SScaleBox*>(this)->SlatePrepass(GetPrepassLayoutScaleMultiplier());
 	}
 
 	bool bClippingNeeded = false;
@@ -336,11 +352,34 @@ FVector2D SScaleBox::ComputeDesiredSize(float InScale) const
 {
 	if (DoesScaleRequireNormalizingPrepass())
 	{
-		// If we require a normalizing pre-pass, we can never allow the scaled content's desired size to affect
-		// the area we return that we need, otherwise, we'll be introducing hysteresis.
 		if (ContentDesiredSize.IsSet())
 		{
-			return ContentDesiredSize.GetValue();
+			FVector2D ContentDesiredSizeValue = ContentDesiredSize.GetValue();
+
+			if (IsDesiredSizeDependentOnAreaAndScale())
+			{
+				// SUPER SPECIAL CASE - 
+				// In the special case that we're only fitting one dimension, we can have the opposite dimension desire the growth of the
+				// expected scale, if we can get that extra space, awesome.
+				if (ComputedContentScale.IsSet() && ComputedContentScale.GetValue() != 0)
+				{
+					const EStretch::Type CurrentStretch = Stretch.Get();
+
+					switch (CurrentStretch)
+					{
+					case EStretch::ScaleToFitX:
+						ContentDesiredSizeValue.Y = ContentDesiredSizeValue.Y * ComputedContentScale.GetValue();
+						break;
+					case EStretch::ScaleToFitY:
+						ContentDesiredSizeValue.X = ContentDesiredSizeValue.X * ComputedContentScale.GetValue();
+						break;
+					}
+				}
+			}
+
+			// If we require a normalizing pre-pass, we can never allow the scaled content's desired size to affect
+			// the area we return that we need, otherwise, we'll be introducing hysteresis.
+			return ContentDesiredSizeValue;
 		}
 	}
 	else
diff --git a/Engine/Source/Runtime/Slate/Public/Widgets/Layout/SScaleBox.h b/Engine/Source/Runtime/Slate/Public/Widgets/Layout/SScaleBox.h
--- a/Engine/Source/Runtime/Slate/Public/Widgets/Layout/SScaleBox.h
+++ b/Engine/Source/Runtime/Slate/Public/Widgets/Layout/SScaleBox.h
@@ -161,6 +161,7 @@ protected:
 
 	bool DoesScaleRequireNormalizingPrepass() const;
 	bool DoesScaleRequireLocalGeometry() const;
+	bool IsDesiredSizeDependentOnAreaAndScale() const;
 	float ComputeContentScale(const FGeometry& PaintGeometry) const;
 
 	void RefreshSafeZoneScale();
