Adding some TODOs to SBox, and adding more uses for SetAttribute where it's needed to prevent overinvalidation.

diff --git a/Engine/Source/Runtime/Slate/Private/Widgets/Layout/SBox.cpp b/Engine/Source/Runtime/Slate/Private/Widgets/Layout/SBox.cpp
--- a/Engine/Source/Runtime/Slate/Private/Widgets/Layout/SBox.cpp
+++ b/Engine/Source/Runtime/Slate/Private/Widgets/Layout/SBox.cpp
@@ -41,11 +41,13 @@ void SBox::SetContent(const TSharedRef< SWidget >& InContent)
 			InContent
 		];
 
+	// TODO SlateGI - This seems no longer needed.
 	Invalidate(EInvalidateWidget::Layout);
 }
 
 void SBox::SetHAlign(EHorizontalAlignment HAlign)
 {
+	// TODO SlateGI - Fix Slots
 	if (ChildSlot.HAlignment != HAlign)
 	{
 		ChildSlot.HAlignment = HAlign;
@@ -55,6 +57,7 @@ void SBox::SetHAlign(EHorizontalAlignment HAlign)
 
 void SBox::SetVAlign(EVerticalAlignment VAlign)
 {
+	// TODO SlateGI - Fix Slots
 	if (ChildSlot.VAlignment != VAlign)
 	{
 		ChildSlot.VAlignment = VAlign;
@@ -64,6 +67,7 @@ void SBox::SetVAlign(EVerticalAlignment VAlign)
 
 void SBox::SetPadding(const TAttribute<FMargin>& InPadding)
 {
+	// TODO SlateGI - Fix Slots
 	if (!ChildSlot.SlotPadding.IdenticalTo(InPadding))
 	{
 		ChildSlot.SlotPadding = InPadding;
@@ -73,56 +77,32 @@ void SBox::SetPadding(const TAttribute<FMargin>& InPadding)
 
 void SBox::SetWidthOverride(TAttribute<FOptionalSize> InWidthOverride)
 {
-	if (!WidthOverride.IdenticalTo(InWidthOverride))
-	{
-		WidthOverride = InWidthOverride;
-		Invalidate(EInvalidateWidget::LayoutAndVolatility);
-	}
+	SetAttribute(WidthOverride, InWidthOverride, EInvalidateWidgetReason::Layout);
 }
 
 void SBox::SetHeightOverride(TAttribute<FOptionalSize> InHeightOverride)
 {
-	if (!HeightOverride.IdenticalTo(InHeightOverride))
-	{
-		HeightOverride = InHeightOverride;
-		Invalidate(EInvalidateWidget::LayoutAndVolatility);
-	}
+	SetAttribute(HeightOverride, InHeightOverride, EInvalidateWidgetReason::Layout);
 }
 
 void SBox::SetMinDesiredWidth(TAttribute<FOptionalSize> InMinDesiredWidth)
 {
-	if (!MinDesiredWidth.IdenticalTo(InMinDesiredWidth))
-	{
-		MinDesiredWidth = InMinDesiredWidth;
-		Invalidate(EInvalidateWidget::LayoutAndVolatility);
-	}
+	SetAttribute(MinDesiredWidth, InMinDesiredWidth, EInvalidateWidgetReason::Layout);
 }
 
 void SBox::SetMinDesiredHeight(TAttribute<FOptionalSize> InMinDesiredHeight)
 {
-	if (!MinDesiredHeight.IdenticalTo(InMinDesiredHeight))
-	{
-		MinDesiredHeight = InMinDesiredHeight;
-		Invalidate(EInvalidateWidget::LayoutAndVolatility);
-	}
+	SetAttribute(MinDesiredHeight, InMinDesiredHeight, EInvalidateWidgetReason::Layout);
 }
 
 void SBox::SetMaxDesiredWidth(TAttribute<FOptionalSize> InMaxDesiredWidth)
 {
-	if (!MaxDesiredWidth.IdenticalTo(InMaxDesiredWidth))
-	{
-		MaxDesiredWidth = InMaxDesiredWidth;
-		Invalidate(EInvalidateWidget::LayoutAndVolatility);
-	}
+	SetAttribute(MaxDesiredWidth, InMaxDesiredWidth, EInvalidateWidgetReason::Layout);
 }
 
 void SBox::SetMaxDesiredHeight(TAttribute<FOptionalSize> InMaxDesiredHeight)
 {
-	if (!MaxDesiredHeight.IdenticalTo(InMaxDesiredHeight))
-	{
-		MaxDesiredHeight = InMaxDesiredHeight;
-		Invalidate(EInvalidateWidget::LayoutAndVolatility);
-	}
+	SetAttribute(MaxDesiredHeight, InMaxDesiredHeight, EInvalidateWidgetReason::Layout);
 }
 
 void SBox::SetMinAspectRatio(TAttribute<FOptionalSize> InMinAspectRatio)
@@ -136,11 +116,7 @@ void SBox::SetMinAspectRatio(TAttribute<FOptionalSize> InMinAspectRatio)
 
 void SBox::SetMaxAspectRatio(TAttribute<FOptionalSize> InMaxAspectRatio)
 {
-	if (!MaxAspectRatio.IdenticalTo(InMaxAspectRatio))
-	{
-		MaxAspectRatio = InMaxAspectRatio;
-		Invalidate(EInvalidateWidget::LayoutAndVolatility);
-	}
+	SetAttribute(MaxAspectRatio, InMaxAspectRatio, EInvalidateWidgetReason::Layout);
 }
 
 FVector2D SBox::ComputeDesiredSize( float ) const
