Reduced the float percision back to the 4.21 state.
Failed experiment, users would be exposed to floating point percision in some simple cases and was not.

diff --git a/Engine/Source/Runtime/Slate/Public/Widgets/Input/NumericTypeInterface.h b/Engine/Source/Runtime/Slate/Public/Widgets/Input/NumericTypeInterface.h
--- a/Engine/Source/Runtime/Slate/Public/Widgets/Input/NumericTypeInterface.h
+++ b/Engine/Source/Runtime/Slate/Public/Widgets/Input/NumericTypeInterface.h
@@ -36,7 +36,7 @@ struct TDefaultNumericTypeInterface : INumericTypeInterface<NumericType>
 		static const FNumberFormattingOptions NumberFormattingOptions = FNumberFormattingOptions()
 			.SetUseGrouping(false)
 			.SetMinimumFractionalDigits(TIsIntegral<NumericType>::Value ? 0 : 1)
-			.SetMaximumFractionalDigits(TIsIntegral<NumericType>::Value ? 0 : 7);
+			.SetMaximumFractionalDigits(TIsIntegral<NumericType>::Value ? 0 : 6);
 		return FastDecimalFormat::NumberToString(Value, ExpressionParser::GetLocalizedNumberFormattingRules(), NumberFormattingOptions);
 	}
 	virtual TOptional<NumericType> FromString(const FString& InString, const NumericType& InExistingValue) override
