Fix delete not working on selected text when Alt was held.

diff --git a/Engine/Source/Runtime/Slate/Private/Widgets/Text/SlateEditableTextLayout.cpp b/Engine/Source/Runtime/Slate/Private/Widgets/Text/SlateEditableTextLayout.cpp
--- a/Engine/Source/Runtime/Slate/Private/Widgets/Text/SlateEditableTextLayout.cpp
+++ b/Engine/Source/Runtime/Slate/Private/Widgets/Text/SlateEditableTextLayout.cpp
@@ -984,8 +984,8 @@ FReply FSlateEditableTextLayout::HandleKeyDown(const FKeyEvent& InKeyEvent)
 	else if (Key == EKeys::Delete && !OwnerWidget->IsTextReadOnly())
 	{
 		// @Todo: Slate keybindings support more than one set of keys. 
-		// Delete to next word boundary (Ctrl+Delete)
-		if (InKeyEvent.IsControlDown() && !InKeyEvent.IsAltDown() && !InKeyEvent.IsShiftDown())
+		// Delete to next word boundary (Ctrl+Delete), only if there is no Text Selected in that case we carry on with a normal delete.
+		if (!AnyTextSelected() && InKeyEvent.IsControlDown() && !InKeyEvent.IsAltDown() && !InKeyEvent.IsShiftDown())
 		{
 			if (OwnerWidget->IsTextPassword())
 			{
