Adding scroll smoothing and wheel scroll multiplier options to ScrollBox
- Also adjusted the logic of displaying the scroll offset a wee bit to include the scroll bar as well
- This brings the SScrollBox to parity with the scrolling features of STableViewBase

diff --git a/Engine/Source/Runtime/Slate/Private/Widgets/Layout/SScrollBox.cpp b/Engine/Source/Runtime/Slate/Private/Widgets/Layout/SScrollBox.cpp
--- a/Engine/Source/Runtime/Slate/Private/Widgets/Layout/SScrollBox.cpp
+++ b/Engine/Source/Runtime/Slate/Private/Widgets/Layout/SScrollBox.cpp
@@ -127,6 +127,8 @@ void SScrollBox::Construct( const FArguments& InArgs )
 	ConsumeMouseWheel = InArgs._ConsumeMouseWheel;
 	TickScrollDelta = 0;
 	AllowOverscroll = InArgs._AllowOverscroll;
+	bAnimateWheelScrolling = InArgs._AnimateWheelScrolling;
+	WheelScrollMultiplier = InArgs._WheelScrollMultiplier;
 	NavigationScrollPadding = InArgs._NavigationScrollPadding;
 	NavigationDestination = InArgs._NavigationDestination;
 	bTouchPanningCapture = false;
@@ -583,27 +585,28 @@ void SScrollBox::Tick( const FGeometry& AllottedGeometry, const double InCurrent
 
 	// If this scroll box has no size, do not compute a view fraction because it will be wrong and causes pop in when the size is available
 	const float ViewFraction = GetViewFraction();
-	const float ViewOffset = GetViewOffsetFraction();
+	const float TargetViewOffset = GetViewOffsetFraction();
 	
+	const float CurrentViewOffset = bAnimateScroll ? FMath::FInterpTo(ScrollBar->DistanceFromTop(), TargetViewOffset, InDeltaTime, 15.f) : TargetViewOffset;
+
 	// Update the scrollbar with the clamped version of the offset
-	float TargetPhysicalOffset = GetScrollComponentFromVector(ViewOffset*ScrollPanel->GetDesiredSize());
+	float NewPhysicalOffset = GetScrollComponentFromVector(CurrentViewOffset * ScrollPanel->GetDesiredSize());
 	if ( AllowOverscroll == EAllowOverscroll::Yes )
 	{
-		TargetPhysicalOffset += Overscroll.GetOverscroll(AllottedGeometry);
+		NewPhysicalOffset += Overscroll.GetOverscroll(AllottedGeometry);
 	}
 
 	const bool bWasScrolling = bIsScrolling;
-	bIsScrolling = !FMath::IsNearlyEqual(TargetPhysicalOffset, ScrollPanel->PhysicalOffset, 0.001f);
-	ScrollPanel->PhysicalOffset = (bAnimateScroll)
-		? FMath::FInterpTo(ScrollPanel->PhysicalOffset, TargetPhysicalOffset, InDeltaTime, 15.f)
-		: TargetPhysicalOffset;
+	bIsScrolling = !FMath::IsNearlyEqual(NewPhysicalOffset, ScrollPanel->PhysicalOffset, 0.001f);
 
+	ScrollPanel->PhysicalOffset = TargetPhysicalOffset;
+	
 	if (bWasScrolling && !bIsScrolling)
 	{
 		Invalidate(EInvalidateWidget::Layout);
 	}
 	
-	ScrollBar->SetState(ViewOffset, ViewFraction);
+	ScrollBar->SetState(CurrentViewOffset, ViewFraction);
 	if (!ScrollBar->IsNeeded())
 	{
 		// We cannot scroll, so ensure that there is no offset.
@@ -822,7 +825,7 @@ FReply SScrollBox::OnMouseWheel( const FGeometry& MyGeometry, const FPointerEven
 		// Make sure scroll velocity is cleared so it doesn't fight with the mouse wheel input
 		InertialScrollManager.ClearScrollVelocity();
 
-		const bool bScrollWasHandled = ScrollBy(MyGeometry, -MouseEvent.GetWheelDelta() * GetGlobalScrollAmount(), EAllowOverscroll::No, false);
+		const bool bScrollWasHandled = ScrollBy(MyGeometry, -MouseEvent.GetWheelDelta() * GetGlobalScrollAmount() * WheelScrollMultiplier, EAllowOverscroll::No, bAnimateWheelScrolling);
 
 		if ( bScrollWasHandled && !bIsScrollingActiveTimerRegistered )
 		{
@@ -1091,6 +1094,16 @@ void SScrollBox::SetAllowOverscroll(EAllowOverscroll NewAllowOverscroll)
 	}
 }
 
+void SScrollBox::SetAnimateWheelScrolling(bool bInAnimateWheelScrolling)
+{
+	bAnimateWheelScrolling = bInAnimateWheelScrolling;
+}
+
+void SScrollBox::SetWheelScrollMultiplier(float NewWheelScrollMultiplier)
+{
+	WheelScrollMultiplier = NewWheelScrollMultiplier;
+}
+
 void SScrollBox::BeginInertialScrolling()
 {
 	if ( !UpdateInertialScrollHandle.IsValid() )
diff --git a/Engine/Source/Runtime/Slate/Public/Widgets/Layout/SScrollBox.h b/Engine/Source/Runtime/Slate/Public/Widgets/Layout/SScrollBox.h
--- a/Engine/Source/Runtime/Slate/Public/Widgets/Layout/SScrollBox.h
+++ b/Engine/Source/Runtime/Slate/Public/Widgets/Layout/SScrollBox.h
@@ -92,6 +92,8 @@ public:
 		, _ScrollBarThickness(FVector2D(9.0f, 9.0f))
 		, _ScrollBarPadding(2.0f)
 		, _AllowOverscroll(EAllowOverscroll::Yes)
+		, _AnimateWheelScrolling(false)
+		, _WheelScrollMultiplier(1.f)
 		, _NavigationDestination(EDescendantScrollDestination::IntoView)
 		, _NavigationScrollPadding(0.0f)
 		, _OnUserScrolled()
@@ -126,6 +128,10 @@ public:
 
 		SLATE_ARGUMENT(EAllowOverscroll, AllowOverscroll);
 
+		SLATE_ARGUMENT(bool, AnimateWheelScrolling);
+
+		SLATE_ARGUMENT(float, WheelScrollMultiplier);
+
 		SLATE_ARGUMENT(EDescendantScrollDestination, NavigationDestination);
 
 		/**
@@ -165,6 +171,10 @@ public:
 
 	void SetAllowOverscroll( EAllowOverscroll NewAllowOverscroll );
 
+	void SetAnimateWheelScrolling(bool bInAnimateWheelScrolling);
+
+	void SetWheelScrollMultiplier(float NewWheelScrollMultiplier);
+
 	float GetScrollOffset() const;
 
 	float GetViewFraction() const;
@@ -365,6 +375,12 @@ protected:
 
 	double LastScrollTime;
 
+	/** Multiplier applied to each click of the scroll wheel (applied alongside the global scroll amount) */
+	float WheelScrollMultiplier = 1.f;
+
+	/** Whether to animate wheel scrolling */
+	bool bAnimateWheelScrolling : 1;
+
 	/**	Whether the software cursor should be drawn in the viewport */
 	bool bShowSoftwareCursor : 1;
 
diff --git a/Engine/Source/Runtime/UMG/Private/Components/ScrollBox.cpp b/Engine/Source/Runtime/UMG/Private/Components/ScrollBox.cpp
--- a/Engine/Source/Runtime/UMG/Private/Components/ScrollBox.cpp
+++ b/Engine/Source/Runtime/UMG/Private/Components/ScrollBox.cpp
@@ -98,6 +98,8 @@ TSharedRef<SWidget> UScrollBox::RebuildWidget()
 		.ConsumeMouseWheel(ConsumeMouseWheel)
 		.NavigationDestination(NavigationDestination)
 		.NavigationScrollPadding(NavigationScrollPadding)
+		.AnimateWheelScrolling(bAnimateWheelScrolling)
+		.WheelScrollMultiplier(WheelScrollMultiplier)
 		.OnUserScrolled(BIND_UOBJECT_DELEGATE(FOnUserScrolled, SlateHandleUserScrolled));
 
 	for ( UPanelSlot* PanelSlot : Slots )
@@ -126,6 +128,8 @@ void UScrollBox::SynchronizeProperties()
 	MyScrollBox->SetAllowOverscroll(AllowOverscroll ? EAllowOverscroll::Yes : EAllowOverscroll::No);
 	MyScrollBox->SetScrollBarRightClickDragAllowed(bAllowRightClickDragScrolling);
 	MyScrollBox->SetConsumeMouseWheel(ConsumeMouseWheel);
+	MyScrollBox->SetAnimateWheelScrolling(bAnimateWheelScrolling);
+	MyScrollBox->SetWheelScrollMultiplier(WheelScrollMultiplier);
 }
 
 float UScrollBox::GetScrollOffset() const
@@ -331,6 +335,24 @@ void UScrollBox::SetAllowOverscroll(bool NewAllowOverscroll)
 	}
 }
 
+void UScrollBox::SetAnimateWheelScrolling(bool bShouldAnimateWheelScrolling)
+{
+	bAnimateWheelScrolling = bShouldAnimateWheelScrolling;
+	if (MyScrollBox)
+	{
+		MyScrollBox->SetAnimateWheelScrolling(bShouldAnimateWheelScrolling);
+	}
+}
+
+void UScrollBox::SetWheelScrollMultiplier(float NewWheelScrollMultiplier)
+{
+	WheelScrollMultiplier = NewWheelScrollMultiplier;
+	if (MyScrollBox)
+	{
+		MyScrollBox->SetWheelScrollMultiplier(NewWheelScrollMultiplier);
+	}
+}
+
 void UScrollBox::EndInertialScrolling()
 {
 	if (MyScrollBox.IsValid())
diff --git a/Engine/Source/Runtime/UMG/Public/Components/ScrollBox.h b/Engine/Source/Runtime/UMG/Public/Components/ScrollBox.h
--- a/Engine/Source/Runtime/UMG/Public/Components/ScrollBox.h
+++ b/Engine/Source/Runtime/UMG/Public/Components/ScrollBox.h
@@ -67,6 +67,10 @@ public:
 	/**  Disable to stop scrollbars from activating inertial overscrolling */
 	UPROPERTY(EditAnywhere, BlueprintReadOnly, Category = "Scroll")
 	bool AllowOverscroll;
+	
+	/** True to lerp smoothly when wheel scrolling along the scroll box */
+	UPROPERTY(EditAnywhere, BlueprintReadOnly, Category = "Scroll")
+	bool bAnimateWheelScrolling = false;
 
 	/**  */
 	UPROPERTY(EditAnywhere, BlueprintReadOnly, Category="Scroll")
@@ -83,6 +87,10 @@ public:
 	UPROPERTY(EditAnywhere, BlueprintReadOnly, Category = "Scroll")
 	bool bAllowRightClickDragScrolling;
 
+	/** The multiplier to apply when wheel scrolling */
+	UPROPERTY(EditAnywhere, BlueprintReadOnly, Category = "Scroll")
+	float WheelScrollMultiplier = 1.f;
+
 
 	UFUNCTION(BlueprintCallable, Category = "Scroll")
 	void SetConsumeMouseWheel(EConsumeMouseWheel NewConsumeMouseWheel);
@@ -105,6 +113,12 @@ public:
 	UFUNCTION(BlueprintCallable, Category = "Scroll")
 	void SetAllowOverscroll(bool NewAllowOverscroll);
 
+	UFUNCTION(BlueprintCallable, Category = "Scroll")
+	void SetAnimateWheelScrolling(bool bShouldAnimateWheelScrolling);
+
+	UFUNCTION(BlueprintCallable, Category = "Scroll")
+	void SetWheelScrollMultiplier(float NewWheelScrollMultiplier);
+
 	/** Instantly stops any inertial scrolling that is currently in progress */
 	UFUNCTION(BlueprintCallable, Category = "Scroll")
 	void EndInertialScrolling();
