Fix compile fail (I misclicked during a P4 diff edit)

diff --git a/Engine/Source/Runtime/Slate/Private/Widgets/Layout/SScrollBox.cpp b/Engine/Source/Runtime/Slate/Private/Widgets/Layout/SScrollBox.cpp
--- a/Engine/Source/Runtime/Slate/Private/Widgets/Layout/SScrollBox.cpp
+++ b/Engine/Source/Runtime/Slate/Private/Widgets/Layout/SScrollBox.cpp
@@ -599,7 +599,7 @@ void SScrollBox::Tick( const FGeometry& AllottedGeometry, const double InCurrent
 	const bool bWasScrolling = bIsScrolling;
 	bIsScrolling = !FMath::IsNearlyEqual(NewPhysicalOffset, ScrollPanel->PhysicalOffset, 0.001f);
 
-	ScrollPanel->PhysicalOffset = TargetPhysicalOffset;
+	ScrollPanel->PhysicalOffset = NewPhysicalOffset;
 	
 	if (bWasScrolling && !bIsScrolling)
 	{
