Unblocking robomerge
- Merge was a little involved, ran SlateViewer and things seem to be working as expected

ListView improvements
- Added horizontal layout option to both SListView and STileView.
- Added smooth inertial scrolling. Doing so naturally requires more refreshes, so it's not intended for use in perf-sensitive scenarios and is disabled by default.
- Added option to establish a consistent fixed offset when scrolling (to, for example, ensure the top item is always flush with the top of the list, or that the bottom 25% of the preceeding item is always visible above). Note that enabling this will override the NavigationScrollOffset property, as the two concepts fully conflict.
- Exposed all new functionality to the UMG counterparts of the slate lists
- Updated SlateViewer's ListView testing tab to include a horizontal list, a horizontal tile view, and options for enabling animated scrolling and fixed offsets. Also tinkered a bit with the default sizing and open tab, as the default size seemed universally too small.

diff --git a/Engine/Source/Developer/AutomationWindow/Private/SAutomationTestTreeView.h b/Engine/Source/Developer/AutomationWindow/Private/SAutomationTestTreeView.h
--- a/Engine/Source/Developer/AutomationWindow/Private/SAutomationTestTreeView.h
+++ b/Engine/Source/Developer/AutomationWindow/Private/SAutomationTestTreeView.h
@@ -131,7 +131,7 @@ public:
 			else
 			{
 				// Make the TableView
-				this->ConstructChildren( 0, InArgs._ItemHeight, EListItemAlignment::LeftAligned, InArgs._HeaderRow, InArgs._ExternalScrollbar, InArgs._OnTreeViewScrolled );
+				this->ConstructChildren( 0, InArgs._ItemHeight, EListItemAlignment::LeftAligned, InArgs._HeaderRow, InArgs._ExternalScrollbar, Orient_Vertical, InArgs._OnTreeViewScrolled );
 			}
 		}
 
diff --git a/Engine/Source/Runtime/AppFramework/Private/Framework/Testing/STableViewTesting.cpp b/Engine/Source/Runtime/AppFramework/Private/Framework/Testing/STableViewTesting.cpp
--- a/Engine/Source/Runtime/AppFramework/Private/Framework/Testing/STableViewTesting.cpp
+++ b/Engine/Source/Runtime/AppFramework/Private/Framework/Testing/STableViewTesting.cpp
@@ -682,6 +682,154 @@ private:
 	TSharedPtr<FTestData> ItemToEdit;
 };
 
+/**
+ * An Item Editor used by tile view testing.
+ * It visualizes a string and edits its contents.
+ */
+class SHorizontalListItemEditor : public STableRow< TSharedPtr<FTestData> >
+{
+public:
+	
+	SLATE_BEGIN_ARGS(SHorizontalListItemEditor)
+		: _ItemToEdit()
+		{}
+
+		SLATE_EVENT(FOnCanAcceptDrop, OnCanAcceptDrop)
+		SLATE_EVENT(FOnAcceptDrop, OnAcceptDrop)
+		SLATE_EVENT(FOnDragDetected, OnDragDetected)
+		SLATE_ARGUMENT(TSharedPtr<FTestData>, ItemToEdit)
+
+	SLATE_END_ARGS()
+	
+	/**
+	 * Construct the widget
+	 *
+	 * @param InArgs   A declaration from which to construct the widget
+	 */
+	void Construct( const FArguments& InArgs, const TSharedRef<STableViewBase>& InOwnerTableView )
+	{
+		STableRow<TSharedPtr<FTestData> >::ConstructInternal( STableRow::FArguments()
+			.OnCanAcceptDrop(InArgs._OnCanAcceptDrop)
+			.OnAcceptDrop(InArgs._OnAcceptDrop)
+			.OnDragDetected(InArgs._OnDragDetected)
+			.ShowSelection(true)
+		, InOwnerTableView);
+
+		ItemToEdit = InArgs._ItemToEdit;
+
+		FSlateFontInfo ItemEditorFont = FCoreStyle::Get().GetFontStyle(TEXT("NormalFont"));
+
+		this->ChildSlot
+		[
+			SNew(SBorder)
+			.BorderImage( FCoreStyle::Get().GetBrush("ToolPanel.GroupBorder") )
+			.BorderBackgroundColor(FLinearColor(1,1,1,0.45f))
+			.Padding(this, &SHorizontalListItemEditor::GetVariableWidth)
+			[
+				SNew(SVerticalBox)
+
+				+SVerticalBox::Slot()
+				.AutoHeight()
+				.Padding(4)
+				[
+					// Name
+					SNew( STextBlock )
+					.Font( ItemEditorFont )
+					.Text( this, &SHorizontalListItemEditor::OnReadText )
+				]
+
+				+SVerticalBox::Slot()
+				.AutoHeight()
+				.Padding(2)
+				[
+					// Number
+					SNew(SSpinBox<float>)
+					.MinValue(-5)
+					.MaxValue(800)
+					.Font( ItemEditorFont )
+					.Value( this, &SHorizontalListItemEditor::OnReadValue )
+					.OnValueChanged( this, &SHorizontalListItemEditor::OnWriteValue )
+				]
+
+			+SVerticalBox::Slot()
+			.AutoHeight()
+			.Padding(2)
+			[
+				//TextField
+				SNew( SEditableTextBox )
+				.Font( ItemEditorFont )
+				.Text( this, &SHorizontalListItemEditor::OnReadText )
+				.OnTextChanged( this, &SHorizontalListItemEditor::OnTextChanged )
+			]
+
+				+SVerticalBox::Slot()
+				.AutoHeight()
+				.Padding(2)
+				[
+					//TextBlock
+					SNew( STextBlock ) .Font( ItemEditorFont ) .Text( this, &SHorizontalListItemEditor::OnReadText )
+				]
+			]
+		];
+	}
+	
+	void SpawnContextMenu( const FVector2D& SpawnLocation )
+	{
+		FSlateApplication::Get().PushMenu
+		(
+			SharedThis( this ),
+			FWidgetPath(),
+			GetPopupContent(),
+			SpawnLocation,
+			FPopupTransitionEffect( FPopupTransitionEffect::ContextMenu )
+		);
+	}
+	
+private:
+	FMargin GetVariableWidth() const
+	{
+		return FMargin(0, 0, this->ItemToEdit->GetNumber(), 0);
+	}
+
+	TSharedRef<SWidget> GetPopupContent()
+	{		
+		return SNew(STestMenu);
+	}
+
+	/** Modify the text when the widget changes it */
+	void OnTextChanged( const FText& NewText )
+	{
+		if (ItemToEdit.IsValid())
+		{
+			ItemToEdit->SetName( NewText );
+		}		
+	}
+	
+	/** @return the text being edited by this widget */
+	FText OnReadText() const
+	{
+		if ( ItemToEdit.IsValid() )
+		{
+			return ItemToEdit->GetName();
+		}
+
+		return FText::GetEmpty();
+	}
+
+	float OnReadValue() const
+	{
+		return ItemToEdit->GetNumber();
+	}
+
+	void OnWriteValue( float NewValue )
+	{
+		ItemToEdit->SetNumber(NewValue);
+	}
+
+	/** A pointer to the data item that we visualize/edit */
+	TSharedPtr<FTestData> ItemToEdit;
+};
+
 /**
  * ListTesting is a test case for lists and trees.
  */
@@ -742,7 +890,7 @@ public:
 					.Text( LOCTEXT("RebuildDataButtonLabel", "Rebuild Data") )
 					.OnClicked( this, &STableViewTesting::Rebuild_OnClicked )
 				]
-				+SHorizontalBox::Slot().AutoWidth()
+				+SHorizontalBox::Slot().AutoWidth() .Padding(10, 0, 0, 0)
 				[
 					// We set how many data items we want in our test case.
 					SNew(SNumericEntryBox<int32>)
@@ -760,6 +908,44 @@ public:
 				[
 					SNew(SSpacer)
 				]
+				+SHorizontalBox::Slot() .AutoWidth()
+				[
+					SNew(SCheckBox)
+					.IsChecked(this, &STableViewTesting::IsScrollAnimChecked)
+					.OnCheckStateChanged(this, &STableViewTesting::HandleScrollAnimCheckStateChanged)
+				]
+				+SHorizontalBox::Slot().AutoWidth() .Padding(2,0,0,0) .VAlign(VAlign_Center)
+				[
+					SNew(STextBlock).Text( LOCTEXT("AnimateScrollingLabel", "Animate Scrolling") )
+				]
+				+SHorizontalBox::Slot() .AutoWidth() .Padding(10, 0, 0, 0)
+				[
+					SNew(SCheckBox)
+					.IsChecked(this, &STableViewTesting::IsFixedItemOffsetChecked)
+					.OnCheckStateChanged(this, &STableViewTesting::HandleFixedItemOffsetCheckStateChanged)
+				]
+				+SHorizontalBox::Slot().AutoWidth()
+				[
+					SNew(SHorizontalBox)
+					.IsEnabled(this, &STableViewTesting::GetIsFixedItemOffsetEnabled)
+					+SHorizontalBox::Slot().AutoWidth() .VAlign(VAlign_Center)
+					[
+						SNew(STextBlock).Text(LOCTEXT("FixedLineScrollOffsetLabel", "Fixed Item Scroll Offset: "))
+					]
+					+SHorizontalBox::Slot().AutoWidth()
+					[
+						// We set how many data items we want in our test case.
+						SNew(SNumericEntryBox<double>)
+						.MinValue(0.)
+						.MaxValue(0.5)
+						.Value(this, &STableViewTesting::GetFixedItemOffsetValue)
+						.OnValueChanged(this, &STableViewTesting::HandleFixedItemOffsetValueChanged)
+					]
+				]
+				+SHorizontalBox::Slot() .FillWidth(1)
+				[
+					SNew(SSpacer)
+				]
 				+SHorizontalBox::Slot().AutoWidth()
 				[
 					SAssignNew(SelectionModeCombo, SComboBox< TSharedPtr<ESelectionMode::Type> >)
@@ -809,49 +995,106 @@ public:
 			.FillHeight(1)
 			[
 				SNew(SBox)
-				.HeightOverride(500.0f)
+				.HeightOverride(800.0f)
 				[
-					SNew( SSplitter )
+					SNew(SSplitter)
+					.Orientation(Orient_Vertical)
 					+ SSplitter::Slot()
-					. Value(1)
+					.Value(.6f)
 					[
-						SNew(SBorder)
-						.BorderImage(FCoreStyle::Get().GetBrush("ToolPanel.GroupBorder"))
+						SNew( SSplitter )
+						+ SSplitter::Slot()
+						. Value(1)
 						[
-							SNew(SHorizontalBox)
-							+SHorizontalBox::Slot()
-							.AutoWidth()
+							SNew(SBorder)
+							.BorderImage(FCoreStyle::Get().GetBrush("ToolPanel.GroupBorder"))
 							[
-								ExternalScrollbar
+								SNew(SHorizontalBox)
+								+SHorizontalBox::Slot()
+								.AutoWidth()
+								[
+									ExternalScrollbar
+								]
+								+SHorizontalBox::Slot()
+								.FillWidth(1)
+								[
+									// The list view being tested
+									SAssignNew( ListBeingTested, SListView< TSharedPtr<FTestData> > )
+									.ExternalScrollbar( ExternalScrollbar )
+									// List view items are this tall
+									.ItemHeight(24)
+									// Tell the list view where to get its source data
+									.ListItemsSource( &Items )
+									// When the list view needs to generate a widget for some data item, use this method
+									.OnGenerateRow( this, &STableViewTesting::OnGenerateWidgetForList )
+									// What to put in the context menu
+									.OnContextMenuOpening( this, &STableViewTesting::GetListContextMenu )
+									// Single, multi or no selection
+									.SelectionMode( this, &STableViewTesting::GetSelectionMode )
+									.HeaderRow
+									(
+										SNew(SHeaderRow)
+										+ SHeaderRow::Column("Name")
+										[
+											SNew(SBorder)
+											.Padding(5)
+											[
+												SNew(STextBlock)
+												.Text(LOCTEXT("TestNameColumn", "Name"))
+											]
+										]
+										+ SHeaderRow::Column("Number") .DefaultLabel(LOCTEXT("TestNumberColumn", "Number"))
+										+ SHeaderRow::Column("TextField") .DefaultLabel(LOCTEXT("TestTextFieldColumn", "Text Field"))
+										+ SHeaderRow::Column("TextBlock") .DefaultLabel(LOCTEXT("TestTextBlockColumn", "Text Block"))
+										+ SHeaderRow::Column("AddChild") .DefaultLabel(LOCTEXT("TestAddChildColumn", "Add Child"))
+									)
+								]
 							]
-							+SHorizontalBox::Slot()
-							.FillWidth(1)
+						]
+						+ SSplitter::Slot()
+						. Value(1)
+						[
+							SNew(SBorder)
+							.BorderImage(FCoreStyle::Get().GetBrush("ToolPanel.GroupBorder"))
 							[
-								// The list view being tested
-								SAssignNew( ListBeingTested, SListView< TSharedPtr<FTestData> > )
-								.ExternalScrollbar( ExternalScrollbar )
-								// List view items are this tall
-								.ItemHeight(24)
-								// Tell the list view where to get its source data
+								// The tile view being tested
+								SAssignNew( TileViewBeingTested, STileView< TSharedPtr<FTestData> > )
+								// Tile view items are this wide
+								.ItemWidth(128)
+								// Tile view items are this tall
+								.ItemHeight(90)
+								// Tell the tile view where to get its source data
 								.ListItemsSource( &Items )
 								// When the list view needs to generate a widget for some data item, use this method
-								.OnGenerateRow( this, &STableViewTesting::OnGenerateWidgetForList )
+								.OnGenerateTile( this, &STableViewTesting::OnGenerateWidgetForTileView )
+								// What to put in the context menu
+								.OnContextMenuOpening( this, &STableViewTesting::GetTileViewContextMenu )
+								// Single, multi or no selection
+								.SelectionMode( this, &STableViewTesting::GetSelectionMode )
+							]
+						]
+						+ SSplitter::Slot()
+						. Value(1)
+						[
+							SNew(SBorder)
+							.BorderImage(FCoreStyle::Get().GetBrush("ToolPanel.GroupBorder"))
+							[
+								// The TreeView being tested; mostly identical to ListView except for OnGetChildren
+								SAssignNew( TreeBeingTested, STreeView< TSharedPtr<FTestData> > )
+								.ItemHeight(24)
+								.TreeItemsSource( &Items )
+								.OnGenerateRow( this, &STableViewTesting::OnGenerateWidgetForTree )
+								// Given some DataItem, this is how we find out if it has any children and what they are.
+								.OnGetChildren( this, &STableViewTesting::OnGetChildrenForTree )
 								// What to put in the context menu
-								.OnContextMenuOpening( this, &STableViewTesting::GetListContextMenu )
+								.OnContextMenuOpening( this, &STableViewTesting::GetTreeContextMenu )
 								// Single, multi or no selection
 								.SelectionMode( this, &STableViewTesting::GetSelectionMode )
+								.WheelScrollMultiplier(60.f)
 								.HeaderRow
 								(
 									SNew(SHeaderRow)
-									+ SHeaderRow::Column("Name")
-									[
-										SNew(SBorder)
-										.Padding(5)
-										[
-											SNew(STextBlock)
-											.Text(LOCTEXT("TestNameColumn", "Name"))
-										]
-									]
+									+ SHeaderRow::Column("Name") .DefaultLabel(LOCTEXT("TestNameColumn","Name"))
 									+ SHeaderRow::Column("Number") .DefaultLabel(LOCTEXT("TestNumberColumn", "Number"))
 									+ SHeaderRow::Column("TextField") .DefaultLabel(LOCTEXT("TestTextFieldColumn", "Text Field"))
 									+ SHeaderRow::Column("TextBlock") .DefaultLabel(LOCTEXT("TestTextBlockColumn", "Text Block"))
@@ -861,53 +1104,49 @@ public:
 						]
 					]
 					+ SSplitter::Slot()
-					. Value(1)
+					.Value(.15f)
 					[
 						SNew(SBorder)
 						.BorderImage(FCoreStyle::Get().GetBrush("ToolPanel.GroupBorder"))
 						[
-							// The tile view being tested
-							SAssignNew( TileViewBeingTested, STileView< TSharedPtr<FTestData> > )
-							// Tile view items are this wide
-							.ItemWidth(128)
-							// Tile view items are this tall
-							.ItemHeight(75)
-							// Tell the tile view where to get its source data
+							// The horizontally-oriented list view being tested
+							SAssignNew(HorizontalListBeingTested, SListView< TSharedPtr<FTestData> > )
+							// This is a horizontally scrolling list view
+							.Orientation(Orient_Horizontal)
+							// List view items are this tall
+							//.ItemHeight(24)
+							// Tell the list view where to get its source data
 							.ListItemsSource( &Items )
 							// When the list view needs to generate a widget for some data item, use this method
-							.OnGenerateTile( this, &STableViewTesting::OnGenerateWidgetForTileView )
+							.OnGenerateRow( this, &STableViewTesting::OnGenerateWidgetForHorizontalList )
 							// What to put in the context menu
-							.OnContextMenuOpening( this, &STableViewTesting::GetTileViewContextMenu )
+							.OnContextMenuOpening( this, &STableViewTesting::GetListContextMenu )
 							// Single, multi or no selection
 							.SelectionMode( this, &STableViewTesting::GetSelectionMode )
 						]
 					]
 					+ SSplitter::Slot()
-					. Value(1)
+					.Value(.25f)
 					[
 						SNew(SBorder)
 						.BorderImage(FCoreStyle::Get().GetBrush("ToolPanel.GroupBorder"))
 						[
-							// The TreeView being tested; mostly identical to ListView except for OnGetChildren
-							SAssignNew( TreeBeingTested, STreeView< TSharedPtr<FTestData> > )
-							.ItemHeight(24)
-							.TreeItemsSource( &Items )
-							.OnGenerateRow( this, &STableViewTesting::OnGenerateWidgetForTree )
-							// Given some DataItem, this is how we find out if it has any children and what they are.
-							.OnGetChildren( this, &STableViewTesting::OnGetChildrenForTree )
+							// The horizontally-oriented tile view being tested
+							SAssignNew( HorizontalTileViewBeingTested, STileView< TSharedPtr<FTestData> > )
+							// This is a horizontally scrolling tile view
+							.Orientation(Orient_Horizontal)
+							// Tile view items are this wide
+							.ItemWidth(128)
+							// Tile view items are this tall
+							.ItemHeight(90)
+							// Tell the tile view where to get its source data
+							.ListItemsSource( &Items )
+							// When the list view needs to generate a widget for some data item, use this method
+							.OnGenerateTile( this, &STableViewTesting::OnGenerateWidgetForTileView )
 							// What to put in the context menu
-							.OnContextMenuOpening( this, &STableViewTesting::GetTreeContextMenu )
+							.OnContextMenuOpening( this, &STableViewTesting::GetTileViewContextMenu )
 							// Single, multi or no selection
 							.SelectionMode( this, &STableViewTesting::GetSelectionMode )
-							.HeaderRow
-							(
-								SNew(SHeaderRow)
-								+ SHeaderRow::Column("Name") .DefaultLabel(LOCTEXT("TestNameColumn","Name"))
-								+ SHeaderRow::Column("Number") .DefaultLabel(LOCTEXT("TestNumberColumn", "Number"))
-								+ SHeaderRow::Column("TextField") .DefaultLabel(LOCTEXT("TestTextFieldColumn", "Text Field"))
-								+ SHeaderRow::Column("TextBlock") .DefaultLabel(LOCTEXT("TestTextBlockColumn", "Text Block"))
-								+ SHeaderRow::Column("AddChild") .DefaultLabel(LOCTEXT("TestAddChildColumn", "Add Child"))
-							)
 						]
 					]
 				]
@@ -965,9 +1204,12 @@ protected:
 			const TSharedRef<FTestData>& ItemBeingDragged = DragDropOperation->GetDraggedItem();
 			FTestData::RemoveRecursive(Items, ItemBeingDragged);
 			FTestData::InsertRecursive(Items, ItemBeingDragged, TargetItem.ToSharedRef(), DropZone);
-			this->ListBeingTested->RequestListRefresh();
-			this->TreeBeingTested->RequestTreeRefresh();
-			this->TileViewBeingTested->RequestListRefresh();
+			ListBeingTested->RequestListRefresh();
+			HorizontalListBeingTested->RequestListRefresh();
+			TreeBeingTested->RequestTreeRefresh();
+			TileViewBeingTested->RequestListRefresh();
+			HorizontalTileViewBeingTested->RequestListRefresh();
+
 			return FReply::Handled();
 		}
 		else
@@ -1103,20 +1345,20 @@ protected:
 	{
 		if (Items.IsValidIndex(ScrollToIndex))
 		{
-			if (TreeBeingTested.IsValid())
-			{
-				TreeBeingTested->RequestScrollIntoView( Items[ ScrollToIndex ] );
-				TreeBeingTested->SetSelection( Items[ ScrollToIndex ] );
-			}
+			ListBeingTested->RequestScrollIntoView(Items[ScrollToIndex]);
+			ListBeingTested->SetSelection(Items[ScrollToIndex]);
 
-			if ( TileViewBeingTested.IsValid() )
-			{
-				TileViewBeingTested->RequestScrollIntoView( Items[ ScrollToIndex ] );
-				TileViewBeingTested->SetSelection( Items[ ScrollToIndex ] );
-			}
+			HorizontalListBeingTested->RequestScrollIntoView(Items[ScrollToIndex]);
+			HorizontalListBeingTested->SetSelection(Items[ScrollToIndex]);
 
-			ListBeingTested->RequestScrollIntoView( Items[ ScrollToIndex ] );
-			ListBeingTested->SetSelection( Items[ ScrollToIndex ] );
+			TreeBeingTested->RequestScrollIntoView( Items[ ScrollToIndex ] );
+			TreeBeingTested->SetSelection( Items[ ScrollToIndex ] );
+			
+			TileViewBeingTested->RequestScrollIntoView( Items[ ScrollToIndex ] );
+			TileViewBeingTested->SetSelection( Items[ ScrollToIndex ] );
+			
+			HorizontalTileViewBeingTested->RequestScrollIntoView(Items[ScrollToIndex]);
+			HorizontalTileViewBeingTested->SetSelection(Items[ScrollToIndex]);
 		}
 
 		return FReply::Handled();
@@ -1125,15 +1367,12 @@ protected:
 	/** Request that both Tree and List refresh themselves on the next tick */
 	FReply RequestRefresh()
 	{
-		if (TreeBeingTested.IsValid())
-		{
-			TreeBeingTested->RequestTreeRefresh();
-		}
-		if (TileViewBeingTested.IsValid())
-		{
-			TileViewBeingTested->RequestListRefresh();
-		}
 		ListBeingTested->RequestListRefresh();
+		HorizontalListBeingTested->RequestListRefresh();
+		TileViewBeingTested->RequestListRefresh();
+		HorizontalTileViewBeingTested->RequestListRefresh();
+		TreeBeingTested->RequestTreeRefresh();
+
 		return FReply::Handled();
 	}
 
@@ -1179,6 +1418,16 @@ protected:
 			.ItemToEdit( InItem );
 	}
 
+	/** @return Given a data item return a new widget to represent it in the ListView */
+	TSharedRef<ITableRow> OnGenerateWidgetForHorizontalList(TSharedPtr<FTestData> InItem, const TSharedRef<STableViewBase>& OwnerTable)
+	{
+		return
+			SNew(SHorizontalListItemEditor, OwnerTable)
+			.OnCanAcceptDrop(this, &STableViewTesting::OnCanAcceptDrop_Handler)
+			.OnAcceptDrop(this, &STableViewTesting::OnAcceptDrop_Handler)
+			.OnDragDetected(this, &STableViewTesting::OnDragDetected_Handler, TWeakPtr<FTestData>(InItem))
+			.ItemToEdit(InItem);
+	}
 
 	// Tile view test
 
@@ -1221,6 +1470,96 @@ protected:
 			: FText::GetEmpty();
 	}
 
+	/** 
+	+SHorizontalBox::Slot() .AutoWidth()
+				[
+					SNew(SCheckBox)
+					.IsChecked(this, &STableViewTesting::IsScrollAnimChecked)
+					.OnCheckStateChanged(this, &STableViewTesting::HandleScrollAnimCheckStateChanged)
+				]
+				+SHorizontalBox::Slot().AutoWidth() .Padding(2,0,0,0) .VAlign(VAlign_Center)
+				[
+					SNew(STextBlock).Text( LOCTEXT("AnimateScrollingLabel", "Animate Scrolling") )
+				]
+				+SHorizontalBox::Slot() .AutoWidth() .Padding(10, 0, 0, 0)
+				[
+					SNew(SCheckBox)
+					.IsChecked(this, &STableViewTesting::IsFixedItemOffsetChecked)
+					.OnCheckStateChanged(this, &STableViewTesting::HandleFixedItemOffsetCheckStateChanged)
+				]
+				+SHorizontalBox::Slot().AutoWidth()
+				[
+					SNew(SHorizontalBox)
+					.IsEnabled(this, &STableViewTesting::GetIsFixedItemOffsetEnabled)
+					+SHorizontalBox::Slot().AutoWidth()
+					[
+						SNew(STextBlock).Text(LOCTEXT("FixedLineScrollOffsetLabel", "Fixed Item Scroll Offset: "))
+					]
+					+SHorizontalBox::Slot().AutoWidth()
+					[
+						// We set how many data items we want in our test case.
+						SNew(SNumericEntryBox<float>)
+						.MinValue(0.f)
+						.MaxValue(1.f)
+						.Value(this, &STableViewTesting::GetFixedLineScrollOffsetValue)
+						.OnValueChanged(this, &STableViewTesting::HandleFixedLineScrollOffsetValueChanged)
+					]
+	*/
+
+	ECheckBoxState IsScrollAnimChecked() const
+	{
+		return bIsScrollAnimEnabled ? ECheckBoxState::Checked : ECheckBoxState::Unchecked;
+	}
+
+	void HandleScrollAnimCheckStateChanged(ECheckBoxState NewState)
+	{
+		bIsScrollAnimEnabled = NewState == ECheckBoxState::Checked;
+		
+		ListBeingTested->SetIsScrollAnimationEnabled(bIsScrollAnimEnabled);
+		HorizontalListBeingTested->SetIsScrollAnimationEnabled(bIsScrollAnimEnabled);
+		TileViewBeingTested->SetIsScrollAnimationEnabled(bIsScrollAnimEnabled);
+		HorizontalTileViewBeingTested->SetIsScrollAnimationEnabled(bIsScrollAnimEnabled);
+		TreeBeingTested->SetIsScrollAnimationEnabled(bIsScrollAnimEnabled);
+	}
+
+	ECheckBoxState IsFixedItemOffsetChecked() const
+	{
+		return bIsFixedLineScrollOffsetEnabled ? ECheckBoxState::Checked : ECheckBoxState::Unchecked;
+	}
+
+	void HandleFixedItemOffsetCheckStateChanged(ECheckBoxState NewState)
+	{
+		bIsFixedLineScrollOffsetEnabled = NewState == ECheckBoxState::Checked;
+		UpdateFixedItemOffset();
+	}
+
+	bool GetIsFixedItemOffsetEnabled() const
+	{
+		return bIsFixedLineScrollOffsetEnabled;
+	}
+
+	TOptional<double> GetFixedItemOffsetValue() const
+	{
+		return FixedLineScrollOffset;
+	}
+
+	void HandleFixedItemOffsetValueChanged(double NewValue)
+	{
+		FixedLineScrollOffset = NewValue;
+		UpdateFixedItemOffset();
+	}
+
+	void UpdateFixedItemOffset()
+	{
+		TOptional<double> FixedOffset = bIsFixedLineScrollOffsetEnabled ? FixedLineScrollOffset : TOptional<double>();
+
+		ListBeingTested->SetFixedLineScrollOffset(FixedOffset);
+		HorizontalListBeingTested->SetFixedLineScrollOffset(FixedOffset);
+		TileViewBeingTested->SetFixedLineScrollOffset(FixedOffset);
+		HorizontalTileViewBeingTested->SetFixedLineScrollOffset(FixedOffset);
+		TreeBeingTested->SetFixedLineScrollOffset(FixedOffset);
+	}
+
 	void OnSelectionModeChanged( ESelectionModePtr InMode, ESelectInfo::Type )
 	{
 		if (InMode.IsValid())
@@ -1273,8 +1612,10 @@ protected:
 	
 	/** A pointer to the ListView being tested */
 	TSharedPtr< SListView< TSharedPtr<FTestData> > > ListBeingTested;
+	TSharedPtr< SListView< TSharedPtr<FTestData> > > HorizontalListBeingTested;
 	/** A pointer to the TileView being tested */
 	TSharedPtr< STileView< TSharedPtr<FTestData> > > TileViewBeingTested;
+	TSharedPtr< STileView< TSharedPtr<FTestData> > > HorizontalTileViewBeingTested;
 	/** A pointer to the TreeView being tested */
 	TSharedPtr< STreeView< TSharedPtr<FTestData> > > TreeBeingTested;
 	/** The data items being tested */
@@ -1292,6 +1633,9 @@ protected:
 	/** All available selection modes */
 	TArray< ESelectionModePtr > SelectionModes;
 
+	bool bIsScrollAnimEnabled = false;
+	bool bIsFixedLineScrollOffsetEnabled = false;
+	double FixedLineScrollOffset = 0.f;
 };
 
 
diff --git a/Engine/Source/Runtime/AppFramework/Private/Widgets/Testing/STestSuite.cpp b/Engine/Source/Runtime/AppFramework/Private/Widgets/Testing/STestSuite.cpp
--- a/Engine/Source/Runtime/AppFramework/Private/Widgets/Testing/STestSuite.cpp
+++ b/Engine/Source/Runtime/AppFramework/Private/Widgets/Testing/STestSuite.cpp
@@ -6025,7 +6025,7 @@ TSharedRef<SDockTab> SpawnTestSuite1( const FSpawnTabArgs& Args )
 			//The first cell in the primary area will be occupied by a stack of tabs.
 			// They are all opened.
 			FTabManager::NewStack()
-			->SetSizeCoefficient(0.35f)
+			->SetSizeCoefficient(0.2f)
 			->AddTab("LayoutExampleTab", ETabState::OpenedTab)
 			->AddTab("DocumentsTestTab", ETabState::OpenedTab)
 		)
@@ -6034,7 +6034,7 @@ TSharedRef<SDockTab> SpawnTestSuite1( const FSpawnTabArgs& Args )
 			// We can subdivide a cell further by using an additional splitter
 			FTabManager::NewSplitter()
 			->SetOrientation( Orient_Vertical )
-			->SetSizeCoefficient(0.65f)
+			->SetSizeCoefficient(0.8f)
 			->Split
 			(
 				FTabManager::NewStack()
@@ -6050,6 +6050,7 @@ TSharedRef<SDockTab> SpawnTestSuite1( const FSpawnTabArgs& Args )
 				->SetSizeCoefficient(0.75f)
 				#if PLATFORM_SUPPORTS_MULTIPLE_NATIVE_WINDOWS
 				->AddTab("TableViewTestTab", ETabState::OpenedTab)
+				->SetForegroundTab(FName("TableViewTestTab"))
 				#endif
 				->AddTab("LayoutRoundingTab", ETabState::OpenedTab)
 				->AddTab("EditableTextTab", ETabState::OpenedTab)
@@ -6291,7 +6292,7 @@ void RestoreSlateTestSuite()
 	TSharedRef<FTabManager::FLayout> Layout = FTabManager::NewLayout( "SlateTestSuite_Layout" )
 	->AddArea
 	(
-		FTabManager::NewArea(720, 600)
+		FTabManager::NewArea(1600, 1200)
 #if PLATFORM_MAC
 		->SetWindow( FVector2D(420, 32), false )
 #else
@@ -6304,6 +6305,7 @@ void RestoreSlateTestSuite()
 			->AddTab( "TestSuite1", ETabState::OpenedTab )
 			->AddTab("RenderTransformManipulator", ETabState::OpenedTab)
 			->AddTab("WidgetGalleryTab", ETabState::OpenedTab)
+			->SetForegroundTab(FName("TestSuite1"))
 		)		
 	)
 	#if PLATFORM_SUPPORTS_MULTIPLE_NATIVE_WINDOWS
diff --git a/Engine/Source/Runtime/Slate/Private/Widgets/Views/SListPanel.cpp b/Engine/Source/Runtime/Slate/Private/Widgets/Views/SListPanel.cpp
--- a/Engine/Source/Runtime/Slate/Private/Widgets/Views/SListPanel.cpp
+++ b/Engine/Source/Runtime/Slate/Private/Widgets/Views/SListPanel.cpp
@@ -6,35 +6,28 @@
 
 FNoChildren SListPanel::NoChildren = FNoChildren();
 
+// Used to subtract a tiny amount from the available dimension to avoid floating point precision problems when arranging children
+static const float FloatingPointPrecisionOffset = 0.001f;
+
 SListPanel::SListPanel()
 	: Children(this)
 {
 }
 
-/**
- * Construct the widget
- *
- * @param InArgs   A declaration from which to construct the widget
- */
 void SListPanel::Construct( const FArguments& InArgs )
 {
-	PreferredRowNum = 0;
-	SmoothScrollOffsetInItems = 0;
-	OverscrollAmount = 0;
 	ItemWidth = InArgs._ItemWidth;
 	ItemHeight = InArgs._ItemHeight;
 	NumDesiredItems = InArgs._NumDesiredItems;
 	ItemAlignment = InArgs._ItemAlignment;
-	bIsRefreshPending = false;
+	Orientation = InArgs._ListOrientation;
 }
 
-/** Make a new ListPanel::Slot  */
 SListPanel::FSlot& SListPanel::Slot()
 {
 	return *(new FSlot());
 }
 	
-/** Add a slot to the ListPanel */
 SListPanel::FSlot& SListPanel::AddSlot(int32 InsertAtIndex)
 {
 	FSlot& NewSlot = SListPanel::Slot();
@@ -49,80 +42,80 @@ SListPanel::FSlot& SListPanel::AddSlot(int32 InsertAtIndex)
 	
 	return NewSlot;
 }
-	
-/**
- * Arrange the children top-to-bottom with not additional layout info.
- */
+
 void SListPanel::OnArrangeChildren( const FGeometry& AllottedGeometry, FArrangedChildren& ArrangedChildren ) const
 {
-	if ( ShouldArrangeHorizontally() )
+	if (Children.Num() > 0)
 	{
-		const EListItemAlignment ListItemAlignment = ItemAlignment.Get();
+		const FTableViewDimensions AllottedDimensions(Orientation, AllottedGeometry.GetLocalSize());
+		FTableViewDimensions DimensionsSoFar(Orientation);
 
-		// This is a tile view list, arrange items horizontally until there is no more room then create a new row.
-		const float AllottedWidth = AllottedGeometry.GetLocalSize().X;
-		const float ItemPadding = GetItemPadding(AllottedGeometry, ListItemAlignment);
-		const float HalfItemPadding = ItemPadding * 0.5;
-
-		const FVector2D LocalItemSize = GetItemSize(AllottedGeometry, ListItemAlignment);
+		if (ShouldArrangeAsTiles())
+		{
+			// This is a tile view list - arrange items side by side along the line axis until there is no more room, then create a new line along the scroll axis.
+			const EListItemAlignment ListItemAlignment = ItemAlignment.Get();
+			const float ItemPadding = GetItemPadding(AllottedGeometry, ListItemAlignment);
+			const float HalfItemPadding = ItemPadding * 0.5;
 
-		float WidthSoFar = 0;
-		float HeightSoFar = -FMath::FloorToInt(SmoothScrollOffsetInItems * LocalItemSize.Y) - OverscrollAmount;
+			const FTableViewDimensions LocalItemSize = GetItemSize(AllottedGeometry, ListItemAlignment);
+			DimensionsSoFar.ScrollAxis = -FMath::FloorToInt(FirstLineScrollOffset * LocalItemSize.ScrollAxis) - OverscrollAmount;
 
-		bool bIsNewLine = true;
-		for( int32 ItemIndex = 0; ItemIndex < Children.Num(); ++ItemIndex )
-		{
-			if ( bIsNewLine )
+			bool bIsNewLine = true;
+			for (int32 ItemIndex = 0; ItemIndex < Children.Num(); ++ItemIndex)
 			{
-				if ( ListItemAlignment == EListItemAlignment::RightAligned || ListItemAlignment == EListItemAlignment::CenterAligned )
+				if (bIsNewLine)
 				{
-					const float LinePadding = GetLinePadding(AllottedGeometry, ItemIndex);
-					const bool IsVisible = Children[ItemIndex].GetWidget()->GetVisibility().IsVisible();
-					if ( ListItemAlignment == EListItemAlignment::RightAligned )
-					{
-						WidthSoFar += IsVisible ? LinePadding : 0;
-					}
-					else
+					if (ListItemAlignment == EListItemAlignment::RightAligned || ListItemAlignment == EListItemAlignment::CenterAligned)
 					{
-						const float HalfLinePadding = LinePadding * 0.5;
-						WidthSoFar += IsVisible ? HalfLinePadding : 0;
+						const float LinePadding = GetLinePadding(AllottedGeometry, ItemIndex);
+						const bool IsVisible = Children[ItemIndex].GetWidget()->GetVisibility().IsVisible();
+						if (ListItemAlignment == EListItemAlignment::RightAligned)
+						{
+							DimensionsSoFar.LineAxis += IsVisible ? LinePadding : 0;
+						}
+						else
+						{
+							const float HalfLinePadding = LinePadding * 0.5;
+							DimensionsSoFar.LineAxis += IsVisible ? HalfLinePadding : 0;
+						}
 					}
-				}
 
-				bIsNewLine = false;
-			}
+					bIsNewLine = false;
+				}
 
-			ArrangedChildren.AddWidget(
-				AllottedGeometry.MakeChild(Children[ItemIndex].GetWidget(), FVector2D(WidthSoFar + HalfItemPadding, HeightSoFar), LocalItemSize)
-				);
-		
-			WidthSoFar += LocalItemSize.X + ItemPadding;
+				FTableViewDimensions ChildOffset = DimensionsSoFar;
+				ChildOffset.LineAxis += HalfItemPadding;
+				ArrangedChildren.AddWidget(AllottedGeometry.MakeChild(Children[ItemIndex].GetWidget(), ChildOffset.ToVector2D(), LocalItemSize.ToVector2D()));
 
-			if ( WidthSoFar + LocalItemSize.X + ItemPadding > AllottedWidth )
-			{
-				WidthSoFar = 0;
-				HeightSoFar += LocalItemSize.Y;
-				bIsNewLine = true;
+				DimensionsSoFar.LineAxis += LocalItemSize.LineAxis + ItemPadding;
+				if (DimensionsSoFar.LineAxis + LocalItemSize.LineAxis + ItemPadding > AllottedDimensions.LineAxis)
+				{
+					DimensionsSoFar.LineAxis = 0;
+					DimensionsSoFar.ScrollAxis += LocalItemSize.ScrollAxis;
+					bIsNewLine = true;
+				}
 			}
 		}
-	}
-	else
-	{
-		if (Children.Num() > 0)
+		else
 		{
-			// This is a normal list, arrange items vertically
-			float HeightSoFar = -FMath::FloorToInt(SmoothScrollOffsetInItems * Children[0].GetWidget()->GetDesiredSize().Y) - OverscrollAmount;
-			for( int32 ItemIndex = 0; ItemIndex < Children.Num(); ++ItemIndex )
+			// This is a normal list, arrange items in a line along the scroll axis
+			const FTableViewDimensions FirstChildDimensions(Orientation, Children[0].GetWidget()->GetDesiredSize());
+			DimensionsSoFar.ScrollAxis = -FMath::FloorToInt(FirstLineScrollOffset * FirstChildDimensions.ScrollAxis) - OverscrollAmount;
+			
+			for (int32 ItemIndex = 0; ItemIndex < Children.Num(); ++ItemIndex)
 			{
-				const FVector2D ItemDesiredSize = Children[ItemIndex].GetWidget()->GetDesiredSize();
-				const bool IsVisible = Children[ItemIndex].GetWidget()->GetVisibility().IsVisible();
-				const float LocalItemHeight = IsVisible ? ItemDesiredSize.Y : 0.f;
+				const FTableViewDimensions WidgetDesiredDimensions(Orientation, Children[ItemIndex].GetWidget()->GetDesiredSize());
+				const bool bIsVisible = Children[ItemIndex].GetWidget()->GetVisibility().IsVisible();
+				
+				FTableViewDimensions FinalWidgetDimensions(Orientation);
+				FinalWidgetDimensions.ScrollAxis = bIsVisible ? WidgetDesiredDimensions.ScrollAxis : 0.f;
+				FinalWidgetDimensions.LineAxis = AllottedDimensions.LineAxis;
 
 				ArrangedChildren.AddWidget(
-					AllottedGeometry.MakeChild( Children[ItemIndex].GetWidget(), FVector2D(0, HeightSoFar), FVector2D(AllottedGeometry.GetLocalSize().X, LocalItemHeight) )
-					);
+					AllottedGeometry.MakeChild(Children[ItemIndex].GetWidget(), DimensionsSoFar.ToVector2D(), FinalWidgetDimensions.ToVector2D())
+				);
 
-				HeightSoFar += LocalItemHeight;
+				DimensionsSoFar.ScrollAxis += FinalWidgetDimensions.ScrollAxis;
 			}
 		}
 	}
@@ -130,57 +123,60 @@ void SListPanel::OnArrangeChildren( const FGeometry& AllottedGeometry, FArranged
 	
 void SListPanel::Tick( const FGeometry& AllottedGeometry, const double InCurrentTime, const float InDeltaTime )
 {
-	if (ShouldArrangeHorizontally())
+	if (ShouldArrangeAsTiles())
 	{
+		const FTableViewDimensions AllottedDimensions(Orientation, AllottedGeometry.GetLocalSize());
+
 		const EListItemAlignment ListItemAlignment = ItemAlignment.Get();
-		const float AllottedWidth = AllottedGeometry.GetLocalSize().X;
 		const float ItemPadding = GetItemPadding(AllottedGeometry, ListItemAlignment);
-		const FVector2D LocalItemSize = GetItemSize(AllottedGeometry, ListItemAlignment);
-		const float TotalItemWidth = LocalItemSize.X + ItemPadding;
+		FTableViewDimensions ItemDimensions = GetItemSize(AllottedGeometry, ListItemAlignment);
+		ItemDimensions.LineAxis += ItemPadding;
+
 		const int32 NumChildren = NumDesiredItems.Get();
 
-		if ( TotalItemWidth > 0.0f && NumChildren > 0)
+		if (ItemDimensions.LineAxis > 0.0f && NumChildren > 0)
 		{
-			int32 NumColumns = FMath::Clamp(FMath::CeilToInt(AllottedWidth / TotalItemWidth) - 1, 1, NumChildren);
-			PreferredRowNum = FMath::CeilToInt(NumChildren / (float)NumColumns);
+			const int32 NumItemsPerLine = FMath::Clamp(FMath::CeilToInt(AllottedDimensions.LineAxis / ItemDimensions.LineAxis) - 1, 1, NumChildren);
+			PreferredNumLines = FMath::CeilToInt(NumChildren / (float)NumItemsPerLine);
 		}
 		else
 		{
-			PreferredRowNum = 1;
+			PreferredNumLines = 1;
 		}
 	}
 	else
 	{
-		PreferredRowNum = NumDesiredItems.Get();
+		PreferredNumLines = NumDesiredItems.Get();
 	}
 }
-	
-/**
- * Simply the sum of all the children (vertically), and the largest width (horizontally).
- */
+
 FVector2D SListPanel::ComputeDesiredSize( float ) const
 {
-	float MaxWidth = 0;
-	float TotalHeight = 0;
-	for( int32 ItemIndex = 0; ItemIndex < Children.Num(); ++ItemIndex )
-	{
-		// Notice that we do not respect child Visibility.
-		// It is simply not useful for ListPanels.
-		FVector2D ChildDesiredSize = Children[ItemIndex].GetWidget()->GetDesiredSize();
-		MaxWidth = FMath::Max( ChildDesiredSize.X, MaxWidth );
-		TotalHeight += ChildDesiredSize.Y;
-	}
+	FTableViewDimensions DesiredListPanelDimensions(Orientation);
 
-	if (ShouldArrangeHorizontally())
+	if (ShouldArrangeAsTiles())
 	{
-		return FVector2D( MaxWidth, ItemHeight.Get() * PreferredRowNum );
+		FTableViewDimensions ItemDimensions = GetDesiredItemDimensions();
+
+		DesiredListPanelDimensions.LineAxis = ItemDimensions.LineAxis;
+		DesiredListPanelDimensions.ScrollAxis = ItemDimensions.ScrollAxis;
 	}
-	else
+	else if (Children.Num() > 0)
 	{
-		return (Children.Num() > 0)
-			? FVector2D( MaxWidth, TotalHeight/Children.Num()*PreferredRowNum )
-			: FVector2D::ZeroVector;
+		// Simply the sum of all the children along the scroll axis and the largest width along the line axis.
+		for (int32 ItemIndex = 0; ItemIndex < Children.Num(); ++ItemIndex)
+		{
+			// Notice that we do not respect child Visibility - it's not useful for ListPanels.
+			FTableViewDimensions ChildDimensions(Orientation, Children[ItemIndex].GetWidget()->GetDesiredSize());
+			DesiredListPanelDimensions.ScrollAxis += ChildDimensions.ScrollAxis;
+			DesiredListPanelDimensions.LineAxis = FMath::Max(DesiredListPanelDimensions.LineAxis, ChildDimensions.LineAxis);
+		}
+
+		DesiredListPanelDimensions.ScrollAxis /= Children.Num();
 	}
+
+	DesiredListPanelDimensions.ScrollAxis *= PreferredNumLines;
+	return DesiredListPanelDimensions.ToVector2D();
 }
 
 FChildren* SListPanel::GetAllChildren()
@@ -188,9 +184,6 @@ FChildren* SListPanel::GetAllChildren()
 	return &Children;
 }
 
-/**
- * @return A slot-agnostic representation of this panel's children
- */
 FChildren* SListPanel::GetChildren()
 {
 	if (bIsRefreshPending)
@@ -207,10 +200,9 @@ FChildren* SListPanel::GetChildren()
 	
 }
 
-/** Set the offset of the view area from the top of the list. */
-void SListPanel::SmoothScrollOffset(float InOffsetInItems)
+void SListPanel::SetFirstLineScrollOffset(float InFirstLineScrollOffset)
 {
-	SmoothScrollOffsetInItems = InOffsetInItems;
+	FirstLineScrollOffset = InFirstLineScrollOffset;
 }
 
 void SListPanel::SetOverscrollAmount( float InOverscrollAmount )
@@ -218,20 +210,14 @@ void SListPanel::SetOverscrollAmount( float InOverscrollAmount )
 	OverscrollAmount = InOverscrollAmount;
 }
 
-/** Remove all the children from this panel */
 void SListPanel::ClearItems()
 {
 	Children.Empty();
 }
 
-float SListPanel::GetDesiredItemWidth() const
-{
-	return ItemWidth.Get();
-}
-
-float SListPanel::GetDesiredItemHeight() const
+FTableViewDimensions SListPanel::GetDesiredItemDimensions() const
 {
-	return ItemHeight.Get();
+	return FTableViewDimensions(Orientation, ItemWidth.Get(), ItemHeight.Get());
 }
 
 float SListPanel::GetItemPadding(const FGeometry& AllottedGeometry) const
@@ -241,78 +227,54 @@ float SListPanel::GetItemPadding(const FGeometry& AllottedGeometry) const
 
 float SListPanel::GetItemPadding(const FGeometry& AllottedGeometry, const EListItemAlignment ListItemAlignment) const
 {
-	// Subtract a tiny amount from the available width to avoid floating point precision problems when arranging children
-	static const float FloatingPointPrecisionOffset = 0.001f;
-
-	const float DesiredWidth = GetDesiredItemWidth();
-
-	// Only add padding between items if we have more total items that we can fit on a single row.  Otherwise,
-	// the padding around items would continue to change proportionately with the (ample) free horizontal space
-	float Padding = 0.0f;
-	switch ( ListItemAlignment )
+	if (ListItemAlignment == EListItemAlignment::EvenlyDistributed)
 	{
-		case EListItemAlignment::EvenlyDistributed:
+		FTableViewDimensions AllottedDimensions(Orientation, AllottedGeometry.GetLocalSize());
+		FTableViewDimensions DesiredDimensions = GetDesiredItemDimensions();
+
+		// Only add padding between items if we have more total items that we can fit on a single line.  Otherwise,
+		// the padding around items would continue to change proportionately with the (ample) free space along the line axis
+		const int32 NumItemsPerLine = DesiredDimensions.LineAxis > 0.f ? FMath::FloorToInt(AllottedDimensions.LineAxis / DesiredDimensions.LineAxis) : 0;
+		if (NumItemsPerLine > 0 && Children.Num() > NumItemsPerLine)
 		{
-			const int32 NumItemsWide = DesiredWidth > 0 ? FMath::FloorToInt(AllottedGeometry.Size.X / DesiredWidth) : 0;
-			if ( NumItemsWide > 0 && Children.Num() > NumItemsWide )
-			{
-				Padding = ( AllottedGeometry.Size.X - FloatingPointPrecisionOffset - NumItemsWide * DesiredWidth ) / NumItemsWide;
-			}
-			break;
+			return (AllottedDimensions.LineAxis - FloatingPointPrecisionOffset - (NumItemsPerLine * DesiredDimensions.LineAxis)) / NumItemsPerLine;
 		}
 	}
 
-	return Padding;
+	return 0.f;
 }
 
-FVector2D SListPanel::GetItemSize(const FGeometry& AllottedGeometry) const
+FTableViewDimensions SListPanel::GetItemSize(const FGeometry& AllottedGeometry) const
 {
 	return GetItemSize(AllottedGeometry, ItemAlignment.Get());
 }
 
-FVector2D SListPanel::GetItemSize(const FGeometry& AllottedGeometry, const EListItemAlignment ListItemAlignment) const
+FTableViewDimensions SListPanel::GetItemSize(const FGeometry& AllottedGeometry, const EListItemAlignment ListItemAlignment) const
 {
-	// Subtract a tiny amount from the available width to avoid floating point precision problems when arranging children
-	static const float FloatingPointPrecisionOffset = 0.001f;
+	const FTableViewDimensions AllottedDimensions(Orientation, AllottedGeometry.GetLocalSize());
+	const FTableViewDimensions ItemDimensions = GetDesiredItemDimensions();
 
-	const float DesiredWidth = GetDesiredItemWidth();
-	const float DesiredHeight = GetDesiredItemHeight();
+	// Only add padding between items if we have more total items that we can fit on a single line.  Otherwise,
+	// the padding around items would continue to change proportionately with the (ample) free space in our minor axis
+	const int32 NumItemsPerLine = ItemDimensions.LineAxis > 0 ? FMath::FloorToInt(AllottedDimensions.LineAxis / ItemDimensions.LineAxis) : 0;
 
-	float ExtraWidth = 0.0f;
-	float ExtraHeight = 0.0f;
-	switch ( ListItemAlignment )
+	FTableViewDimensions ExtraDimensions(Orientation);
+	if (NumItemsPerLine > 0)
 	{
-		case EListItemAlignment::Fill:
-		{
-			const int32 NumItemsWide = DesiredWidth > 0 ? FMath::Min(Children.Num(), FMath::FloorToInt(AllottedGeometry.Size.X / DesiredWidth)) : 0;
-			if ( NumItemsWide > 0 )
-			{
-				ExtraWidth = ( AllottedGeometry.Size.X - FloatingPointPrecisionOffset - NumItemsWide * DesiredWidth ) / NumItemsWide;
-			}
-			break;
-		}
-		case EListItemAlignment::EvenlySize:
+		if (ListItemAlignment == EListItemAlignment::Fill || 
+			ListItemAlignment == EListItemAlignment::EvenlyWide || 
+			ListItemAlignment == EListItemAlignment::EvenlySize)
 		{
-			const int32 NumItemsWide = DesiredWidth > 0 ? FMath::FloorToInt(AllottedGeometry.Size.X / DesiredWidth) : 0;
-			if ( NumItemsWide > 0 )
+			ExtraDimensions.LineAxis = (AllottedDimensions.LineAxis - FloatingPointPrecisionOffset - NumItemsPerLine * ItemDimensions.LineAxis) / NumItemsPerLine;
+			
+			if (ListItemAlignment == EListItemAlignment::EvenlySize)
 			{
-				ExtraWidth = ( AllottedGeometry.Size.X - FloatingPointPrecisionOffset - NumItemsWide * DesiredWidth ) / NumItemsWide;
-				ExtraHeight = DesiredHeight * ( ExtraWidth / ( DesiredWidth + ExtraWidth ) );
+				ExtraDimensions.ScrollAxis = ItemDimensions.ScrollAxis * (ExtraDimensions.LineAxis / (ItemDimensions.LineAxis + ExtraDimensions.LineAxis));
 			}
-			break;
-		}
-		case EListItemAlignment::EvenlyWide:
-		{
-			const int32 NumItemsWide = DesiredWidth > 0 ? FMath::FloorToInt(AllottedGeometry.Size.X / DesiredWidth) : 0;
-			if ( NumItemsWide > 0 )
-			{
-				ExtraWidth = ( AllottedGeometry.Size.X - FloatingPointPrecisionOffset - NumItemsWide * DesiredWidth ) / NumItemsWide;
-			}
-			break;
 		}
 	}
-
-	return FVector2D(DesiredWidth + ExtraWidth, DesiredHeight + ExtraHeight);
+	
+	return ItemDimensions + ExtraDimensions;
 }
 
 float SListPanel::GetLinePadding(const FGeometry& AllottedGeometry, const int32 LineStartIndex) const
@@ -323,14 +285,13 @@ float SListPanel::GetLinePadding(const FGeometry& AllottedGeometry, const int32
 		return 0.0f;
 	}
 
-	const FVector2D LocalItemSize = GetItemSize(AllottedGeometry);
-	const int32 NumItemsWide = LocalItemSize.X > 0 ? FMath::FloorToInt(AllottedGeometry.Size.X / LocalItemSize.X) : 0;
-	const int32 NumItemsOnLine = FMath::Min(NumItemsLeft, NumItemsWide);
+	const FTableViewDimensions AllottedDimensions(Orientation, AllottedGeometry.GetLocalSize());
+	const FTableViewDimensions ItemSize = GetItemSize(AllottedGeometry);
 
-	// Subtract a tiny amount from the available width to avoid floating point precision problems when arranging children
-	static const float FloatingPointPrecisionOffset = 0.001f;
+	const int32 MaxNumItemsOnLine = ItemSize.LineAxis > 0 ? FMath::FloorToInt(AllottedDimensions.LineAxis / ItemSize.LineAxis) : 0;
+	const int32 NumItemsOnLine = FMath::Min(NumItemsLeft, MaxNumItemsOnLine);
 
-	return (AllottedGeometry.Size.X - FloatingPointPrecisionOffset - NumItemsOnLine * LocalItemSize.X);
+	return AllottedDimensions.LineAxis - FloatingPointPrecisionOffset - (NumItemsOnLine * ItemSize.LineAxis);
 }
 
 void SListPanel::SetRefreshPending( bool IsPendingRefresh )
@@ -343,9 +304,10 @@ bool SListPanel::IsRefreshPending() const
 	return bIsRefreshPending;
 }
 
-bool SListPanel::ShouldArrangeHorizontally() const
+bool SListPanel::ShouldArrangeAsTiles() const
 {
-	return GetDesiredItemWidth() > 0;
+	FTableViewDimensions DesiredItemDimensions = GetDesiredItemDimensions();
+	return DesiredItemDimensions.LineAxis > 0 && DesiredItemDimensions.ScrollAxis > 0;
 }
 
 void SListPanel::SetItemHeight(TAttribute<float> Height)
diff --git a/Engine/Source/Runtime/Slate/Private/Widgets/Views/SListPanel.h b/Engine/Source/Runtime/Slate/Private/Widgets/Views/SListPanel.h
--- a/Engine/Source/Runtime/Slate/Private/Widgets/Views/SListPanel.h
+++ b/Engine/Source/Runtime/Slate/Private/Widgets/Views/SListPanel.h
@@ -41,6 +41,7 @@ public:
 		, _ItemHeight(16)
 		, _NumDesiredItems(0)
 		, _ItemAlignment(EListItemAlignment::EvenlyDistributed)
+		, _ListOrientation(Orient_Vertical)
 		{
 			_Visibility = EVisibility::SelfHitTestInvisible;
 			_Clipping = EWidgetClipping::ClipToBounds;
@@ -50,6 +51,7 @@ public:
 		SLATE_ATTRIBUTE( float, ItemHeight )
 		SLATE_ATTRIBUTE( int32, NumDesiredItems )
 		SLATE_ATTRIBUTE( EListItemAlignment, ItemAlignment )
+		SLATE_ARGUMENT( EOrientation, ListOrientation )
 
 	SLATE_END_ARGS()
 	
@@ -72,8 +74,8 @@ public:
 	virtual void Tick( const FGeometry& AllottedGeometry, const double InCurrentTime, const float InDeltaTime ) override;
 	// End of SWidget interface
 
-	/** Fraction of the first item that we should offset by to simulate smooth scrolling. */
-	void SmoothScrollOffset(float OffsetInItems);
+	/** Fraction of the first line that we should offset by to account for the current scroll amount. */
+	void SetFirstLineScrollOffset(float InFirstLineScrollOffset);
 
 	/** Set how much we should appear to have scrolled past the beginning/end of the list. */
 	void SetOverscrollAmount( float InOverscrollAmount );
@@ -81,17 +83,14 @@ public:
 	/** Remove all the children from this panel */
 	void ClearItems();
 
-	/** @return the uniform desired item width used when arranging children. */
-	float GetDesiredItemWidth() const;
-
-	/** @return the uniform item height used when arranging children. */
-	float GetDesiredItemHeight() const;
+	/** @return the uniform desired item dimensions used when arranging children. */
+	FTableViewDimensions GetDesiredItemDimensions() const;
 
 	/** @return the uniform item width used when arranging children. */
-	FVector2D GetItemSize(const FGeometry& AllottedGeometry) const;
+	FTableViewDimensions GetItemSize(const FGeometry& AllottedGeometry) const;
 
 	/** @return the uniform item width used when arranging children. */
-	FVector2D GetItemSize(const FGeometry& AllottedGeometry, const EListItemAlignment ListItemAlignment) const;
+	FTableViewDimensions GetItemSize(const FGeometry& AllottedGeometry, const EListItemAlignment ListItemAlignment) const;
 
 	/** @return the horizontal padding applied to each tile item */
 	float GetItemPadding(const FGeometry& AllottedGeometry) const;
@@ -113,11 +112,11 @@ public:
 
 	/** See ItemWidth attribute */
 	void SetItemWidth(TAttribute<float> Width);
-	
+
 protected:
 
-	/** @return true if this panel should arrange items horizontally until it runs out of room, then create new rows */
-	bool ShouldArrangeHorizontally() const;
+	/** @return true if this panel should arrange items as tiles placed alongside one another in each line */
+	bool ShouldArrangeAsTiles() const;
 	
 protected:
 
@@ -133,26 +132,30 @@ protected:
 	/** Total number of items that the tree wants to visualize */
 	TAttribute<int32> NumDesiredItems;
 	
-	/** How should be horizontally aligned? Only relevant for tile views. */
-	TAttribute<EListItemAlignment> ItemAlignment;
-
 	/**
 	 * The offset of the view area from the top of the list in item heights.
-	 * Translate to physical units based on first item in list.
+	 * Translate to physical units based on first line in list.
 	 */
-	float SmoothScrollOffsetInItems;
+	float FirstLineScrollOffset = 0.f;
 
 	/** Amount scrolled past beginning/end of list in Slate Units. */
-	float OverscrollAmount;
-
-	/** The preferred number of rows that this widget should have */
-	int32 PreferredRowNum;
+	float OverscrollAmount = 0.f;
 
 	/**
 	 * When true, a refresh of the table view control that is using this panel is pending.
 	 * Some of the widgets in this panel are associated with items that may no longer be sound data.
 	 */
-	bool bIsRefreshPending;
+	bool bIsRefreshPending = false;
+
+	/** How should be horizontally aligned? Only relevant for tile views. */
+	TAttribute<EListItemAlignment> ItemAlignment;
+
+	/** Overall orientation of the list for layout and scrolling. Only relevant for tile views. */
+	EOrientation Orientation;
+
+	/** The preferred number of lines that this widget should have orthogonal to the scroll axis. Only relevant for tile views. */
+	int32 PreferredNumLines = 1;
+
 
 private:
 	
diff --git a/Engine/Source/Runtime/Slate/Private/Widgets/Views/STableViewBase.cpp b/Engine/Source/Runtime/Slate/Private/Widgets/Views/STableViewBase.cpp
--- a/Engine/Source/Runtime/Slate/Private/Widgets/Views/STableViewBase.cpp
+++ b/Engine/Source/Runtime/Slate/Private/Widgets/Views/STableViewBase.cpp
@@ -23,7 +23,34 @@ namespace ListConstants
 	static const float OvershootBounceRate = 250.0f;
 }
 
-void STableViewBase::ConstructChildren( const TAttribute<float>& InItemWidth, const TAttribute<float>& InItemHeight, const TAttribute<EListItemAlignment>& InItemAlignment, const TSharedPtr<SHeaderRow>& InHeaderRow, const TSharedPtr<SScrollBar>& InScrollBar, const FOnTableViewScrolled& InOnTableViewScrolled )
+
+FTableViewDimensions::FTableViewDimensions(EOrientation InOrientation)
+	: Orientation(InOrientation)
+{
+}
+
+
+FTableViewDimensions::FTableViewDimensions(EOrientation InOrientation, float X, float Y)
+	: FTableViewDimensions(InOrientation, FVector2D(X, Y))
+{
+}
+
+FTableViewDimensions::FTableViewDimensions(EOrientation InOrientation, const FVector2D& Size)
+	: FTableViewDimensions(InOrientation)
+{
+	if (InOrientation == Orient_Vertical)
+	{
+		LineAxis = Size.X;
+		ScrollAxis = Size.Y;
+	}
+	else
+	{
+		ScrollAxis = Size.X;
+		LineAxis = Size.Y;
+	}
+}
+
+void STableViewBase::ConstructChildren( const TAttribute<float>& InItemWidth, const TAttribute<float>& InItemHeight, const TAttribute<EListItemAlignment>& InItemAlignment, const TSharedPtr<SHeaderRow>& InHeaderRow, const TSharedPtr<SScrollBar>& InScrollBar, EOrientation InScrollOrientation, const FOnTableViewScrolled& InOnTableViewScrolled )
 {
 	bItemsNeedRefresh = true;
 	
@@ -31,55 +58,78 @@ void STableViewBase::ConstructChildren( const TAttribute<float>& InItemWidth, co
 
 	OnTableViewScrolled = InOnTableViewScrolled;
 
-	// If the user provided a scrollbar, we do not need to make one of our own.
-	if (InScrollBar.IsValid())
+	Orientation = InHeaderRow ? Orient_Vertical : InScrollOrientation;
+	UE_CLOG(InScrollOrientation != Orientation, LogSlate, Error, TEXT("STableViewBase does not support horizontal scrolling when displaying a header row"));
+
+	ItemsPanel = SNew(SListPanel)
+		.ItemWidth(InItemWidth)
+		.ItemHeight(InItemHeight)
+		.NumDesiredItems(this, &STableViewBase::GetNumItemsBeingObserved)
+		.ItemAlignment(InItemAlignment)
+		.ListOrientation(Orientation);
+
+	TSharedPtr<SWidget> ListAndScrollbar;
+	if (InScrollBar)
 	{
+		// The user provided us with a scrollbar; we will rely on it.
 		ScrollBar = InScrollBar;
-		ScrollBar->SetOnUserScrolled( FOnUserScrolled::CreateSP(this, &STableViewBase::ScrollBar_OnUserScrolled) );
+		ScrollBar->SetOnUserScrolled(FOnUserScrolled::CreateSP(this, &STableViewBase::ScrollBar_OnUserScrolled));
+		
+		ListAndScrollbar = ItemsPanel;
 	}
-	
-	TSharedRef<SWidget> ListAndScrollbar = (!ScrollBar.IsValid())
-		// We need to make our own scrollbar
-		? StaticCastSharedRef<SWidget>
-		(
-			SNew(SHorizontalBox)
-			+SHorizontalBox::Slot()
-			.FillWidth(1)
-			[
-				SAssignNew( ItemsPanel, SListPanel )
-				.ItemWidth( InItemWidth )
-				.ItemHeight( InItemHeight )
-				.NumDesiredItems( this, &STableViewBase::GetNumItemsBeingObserved )
-				.ItemAlignment( InItemAlignment )
-			]
-			+SHorizontalBox::Slot()
-			.AutoWidth()
-			[
-				SNew( SBox )
-				.WidthOverride( FOptionalSize( 16 ) )
+	else
+	{
+		ScrollBar = SNew(SScrollBar)
+			.OnUserScrolled(this, &STableViewBase::ScrollBar_OnUserScrolled)
+			.Orientation(Orientation);
+
+		const FOptionalSize ScrollBarSize(16.f);
+
+		if (Orientation == Orient_Vertical)
+		{
+			ListAndScrollbar = SNew(SHorizontalBox)
+				+SHorizontalBox::Slot()
+				.FillWidth(1)
 				[
-					SAssignNew(ScrollBar, SScrollBar)
-					.OnUserScrolled( this, &STableViewBase::ScrollBar_OnUserScrolled )
+					ItemsPanel.ToSharedRef()
 				]
-			]
-		)
-		// The user provided us with a scrollbar; we will rely on it.
-		: StaticCastSharedRef<SWidget>
-		(
-			SAssignNew( ItemsPanel, SListPanel )
-			.ItemWidth( InItemWidth )
-			.ItemHeight( InItemHeight )
-			.NumDesiredItems( this, &STableViewBase::GetNumItemsBeingObserved )
-			.ItemAlignment( InItemAlignment )
-		);
+				+SHorizontalBox::Slot()
+				.AutoWidth()
+				[
+					SNew(SBox)
+					.WidthOverride( FOptionalSize( 16 ) )
+					[
+						ScrollBar.ToSharedRef()
+					]
+				];
+		}
+		else
+		{
+			ListAndScrollbar = SNew(SVerticalBox)
+				+SVerticalBox::Slot()
+				.FillHeight(1)
+				[
+					ItemsPanel.ToSharedRef()
+				]
+				+SVerticalBox::Slot()
+				.AutoHeight()
+				[
+					SNew(SBox)
+					.HeightOverride(ScrollBarSize)
+					[
+						ScrollBar.ToSharedRef()
+					]
+				];
+		}
+	}
 
-	if (InHeaderRow.IsValid())
+	if (InHeaderRow)
 	{
 		// Only associate the scrollbar if we created it.
 		// If the scrollbar was passed in from outside then it won't appear under our header row so doesn't need compensating for.
-		if (!InScrollBar.IsValid())
+		if (!InScrollBar)
 		{
-			InHeaderRow->SetAssociatedVerticalScrollBar( ScrollBar.ToSharedRef(), 16 );
+			InHeaderRow->SetAssociatedVerticalScrollBar(ScrollBar.ToSharedRef(), 16.f);
 		}
 
 		this->ChildSlot
@@ -93,7 +143,7 @@ void STableViewBase::ConstructChildren( const TAttribute<float>& InItemWidth, co
 			+SVerticalBox::Slot()
 			.FillHeight(1)
 			[
-				ListAndScrollbar
+				ListAndScrollbar.ToSharedRef()
 			]
 		];
 	}
@@ -101,7 +151,7 @@ void STableViewBase::ConstructChildren( const TAttribute<float>& InItemWidth, co
 	{
 		this->ChildSlot
 		[
-			ListAndScrollbar
+			ListAndScrollbar.ToSharedRef()
 		];
 	}
 	
@@ -127,49 +177,6 @@ void STableViewBase::OnMouseCaptureLost(const FCaptureLostEvent& CaptureLostEven
 	bShowSoftwareCursor = false;
 }
 
-struct FEndOfListResult
-{
-	FEndOfListResult( float InOffset, float InItemsAboveView )
-	: OffsetFromEndOfList(InOffset)
-	, ItemsAboveView(InItemsAboveView)
-	{
-
-	}
-
-	float OffsetFromEndOfList;
-	float ItemsAboveView;
-};
-
-static FEndOfListResult ComputeOffsetForEndOfList( const FGeometry& ListPanelGeometry, const FChildren& ListPanelChildren )
-{
-	float OffsetFromEndOfList = 0.0f;
-	float AvailableSpace = ListPanelGeometry.GetLocalSize().Y;
-	float ItemsAboveView = 0.0f;
-	for ( int ChildIndex=ListPanelChildren.Num()-1; ChildIndex >= 0; --ChildIndex )
-	{
-		const float CurChildHeight = ListPanelChildren.GetChildAt(ChildIndex)->GetDesiredSize().Y;
-		if (AvailableSpace == 0)
-		{
-			ItemsAboveView ++;
-		}
-
-		if ( CurChildHeight < AvailableSpace )
-		{
-			// This whole child fits
-			OffsetFromEndOfList += 1;
-			AvailableSpace -= CurChildHeight;
-		}
-		else
-		{
-			OffsetFromEndOfList += AvailableSpace / CurChildHeight;
-			ItemsAboveView += (CurChildHeight - AvailableSpace)/CurChildHeight;
-			AvailableSpace = 0;
-		}
-	}
-
-	return FEndOfListResult( OffsetFromEndOfList, ItemsAboveView );
-}
-
 EActiveTimerReturnType STableViewBase::UpdateInertialScroll(double InCurrentTime, float InDeltaTime)
 {
 	bool bKeepTicking = false;
@@ -241,39 +248,49 @@ void STableViewBase::Tick( const FGeometry& AllottedGeometry, const double InCur
 		if ( bItemsNeedRefresh || PanelGeometryLastTick.GetLocalSize() != PanelGeometry.GetLocalSize())
 		{
 			PanelGeometryLastTick = PanelGeometry;
+			
+			const int32 NumItemsPerLine = GetNumItemsPerLine();
+			const EScrollIntoViewResult ScrollIntoViewResult = ScrollIntoView(PanelGeometry);
+
+			double TargetScrollOffset = GetTargetScrollOffset();
 
-			// We never create the ItemsPanel if the user did not specify all the parameters required to successfully make a list.
-			const EScrollIntoViewResult ScrollIntoViewResult = ScrollIntoView( PanelGeometry );
+			if (bEnableAnimatedScrolling)
+			{
+				CurrentScrollOffset = FMath::FInterpTo(CurrentScrollOffset, TargetScrollOffset, InDeltaTime, 12.f);
+				if (FMath::IsNearlyEqual(CurrentScrollOffset, TargetScrollOffset, 0.01))
+				{
+					CurrentScrollOffset = TargetScrollOffset;
+				}
+			}
+			else
+			{
+				CurrentScrollOffset = TargetScrollOffset;
+			}
 
 			const FReGenerateResults ReGenerateResults = ReGenerateItems( PanelGeometry );
 			LastGenerateResults = ReGenerateResults;
 
 			const int32 NumItemsBeingObserved = GetNumItemsBeingObserved();
+			const int32 NumItemLines = NumItemsBeingObserved / NumItemsPerLine;
 
-			const int32 NumItemsWide = GetNumItemsWide();
-			const int32 NumItemRows = NumItemsBeingObserved / NumItemsWide;
-
-			const bool bEnoughRoomForAllItems = ReGenerateResults.ExactNumRowsOnScreen >= NumItemRows;
+			const double InitialDesiredOffset = DesiredScrollOffset;
+			const bool bEnoughRoomForAllItems = ReGenerateResults.ExactNumLinesOnScreen >= NumItemLines;
 			if (bEnoughRoomForAllItems)
 			{
 				// We can show all the items, so make sure there is no scrolling.
-				ScrollOffset = 0;
+				SetScrollOffset(0.0);
+				CurrentScrollOffset = TargetScrollOffset = DesiredScrollOffset;
 			}
-			else if ( ReGenerateResults.bGeneratedPastLastItem )
+			else if (ReGenerateResults.bGeneratedPastLastItem)
 			{
-				ScrollOffset = ReGenerateResults.NewScrollOffset;
+				SetScrollOffset(FMath::Max(0.0, ReGenerateResults.NewScrollOffset));
+				CurrentScrollOffset = TargetScrollOffset = DesiredScrollOffset;
 			}
-
-			// Compute SmoothScrollOffset before SetScrollOffset(). This is becasue SetScrollOffset() modifies the ScrollOffset (double downcast to float).
-			// We need original value in computation of SmoothScrollOffset.
-			double SmoothScrollOffset = ScrollOffset / GetNumItemsWide();
-			SmoothScrollOffset = SmoothScrollOffset - (int64)SmoothScrollOffset; // get fractional part
-
-			SetScrollOffset(FMath::Max(0.0, ScrollOffset));
-
-			// FMath::Fractional() is buggy as it casts to int32 (too small for the integer part of float).
-			//ItemsPanel->SmoothScrollOffset(FMath::Fractional(ScrollOffset / GetNumItemsWide()));
-			ItemsPanel->SmoothScrollOffset(SmoothScrollOffset);
+			
+			// FMath::Fractional() is insufficient here as it casts to int32 (too small for the integer part of a float when the scroll offset is enormous), so we do a double/int64 version here.
+			double FirstLineScrollOffset = CurrentScrollOffset / NumItemsPerLine;
+			FirstLineScrollOffset = FirstLineScrollOffset - (int64)FirstLineScrollOffset;
+			ItemsPanel->SetFirstLineScrollOffset(FirstLineScrollOffset);
 
 			if (AllowOverscroll == EAllowOverscroll::Yes)
 			{
@@ -286,20 +303,20 @@ void STableViewBase::Tick( const FGeometry& AllottedGeometry, const double InCur
 			// Update scrollbar
 			if (NumItemsBeingObserved > 0)
 			{
-				if (ReGenerateResults.ExactNumRowsOnScreen < 1.0f)
+				if (ReGenerateResults.ExactNumLinesOnScreen < 1.0f)
 				{
 					// We are be observing a single row which is larger than the available visible area, so we should calculate thumb size based on that
-					const double VisibleSizeFraction = AllottedGeometry.GetLocalSize().Y / ReGenerateResults.HeightOfGeneratedItems;
+					const double VisibleSizeFraction = AllottedGeometry.GetLocalSize().Y / ReGenerateResults.LengthOfGeneratedItems;
 					const double ThumbSizeFraction = FMath::Min(VisibleSizeFraction, 1.0);
-					const double OffsetFraction = ScrollOffset / NumItemsBeingObserved;
+					const double OffsetFraction = CurrentScrollOffset / NumItemsBeingObserved;
 					ScrollBar->SetState( OffsetFraction, ThumbSizeFraction );
 				}
 				else
 				{
 					// The thumb size is whatever fraction of the items we are currently seeing (including partially seen items).
 					// e.g. if we are seeing 0.5 of the first generated widget and 0.75 of the last widget, that's 1.25 widgets.
-					const double ThumbSizeFraction = ReGenerateResults.ExactNumRowsOnScreen / NumItemRows;
-					const double OffsetFraction = ScrollOffset / NumItemsBeingObserved;
+					const double ThumbSizeFraction = ReGenerateResults.ExactNumLinesOnScreen / NumItemLines;
+					const double OffsetFraction = CurrentScrollOffset / NumItemsBeingObserved;
 					ScrollBar->SetState( OffsetFraction, ThumbSizeFraction );
 				}
 			}
@@ -317,7 +334,7 @@ void STableViewBase::Tick( const FGeometry& AllottedGeometry, const double InCur
 
 			Invalidate(EInvalidateWidget::ChildOrder);
 
-			if (ScrollIntoViewResult == EScrollIntoViewResult::Deferred)
+			if (ScrollIntoViewResult == EScrollIntoViewResult::Deferred || CurrentScrollOffset != TargetScrollOffset)
 			{
 				// We call this rather than just leave bItemsNeedRefresh as true to ensure that EnsureTickToRefresh is registered
 				RequestLayoutRefresh();
@@ -435,7 +452,12 @@ FReply STableViewBase::OnMouseMove( const FGeometry& MyGeometry, const FPointerE
 {	
 	if( MouseEvent.IsMouseButtonDown( EKeys::RightMouseButton ) && !MouseEvent.IsTouchEvent())
 	{
-		const float ScrollByAmount = MouseEvent.GetCursorDelta().Y / MyGeometry.Scale;
+		// We only care about deltas along the scroll axis
+		FTableViewDimensions CursorDeltaDimensions(Orientation, MouseEvent.GetCursorDelta());
+		CursorDeltaDimensions.LineAxis = 0.f;
+		
+		const float ScrollByAmount = CursorDeltaDimensions.ScrollAxis / MyGeometry.Scale;
+
 		// If scrolling with the right mouse button, we need to remember how much we scrolled.
 		// If we did not scroll at all, we will bring up the context menu when the mouse is released.
 		AmountScrolledWhileRightMouseDown += FMath::Abs( ScrollByAmount );
@@ -469,7 +491,7 @@ FReply STableViewBase::OnMouseMove( const FGeometry& MyGeometry, const FPointerE
 			// Check if the mouse has moved.
 			if( AmountScrolled != 0 )
 			{
-				SoftwareCursorPosition.Y += ScrollByAmount;
+				SoftwareCursorPosition += CursorDeltaDimensions.ToVector2D();
 			}
 
 			return Reply;
@@ -517,20 +539,24 @@ FReply STableViewBase::OnMouseWheel( const FGeometry& MyGeometry, const FPointer
 		// Make sure scroll velocity is cleared so it doesn't fight with the mouse wheel input
 		this->InertialScrollManager.ClearScrollVelocity();
 
-		const float AmountScrolledInItems = this->ScrollBy( MyGeometry, -MouseEvent.GetWheelDelta() * WheelScrollMultiplier, EAllowOverscroll::No );
-
-		switch ( ConsumeMouseWheel )
+		float AmountScrolledInItems = 0.f;
+		if (FixedLineScrollOffset.IsSet())
 		{
-		case EConsumeMouseWheel::Always:
-			return FReply::Handled();
-		case EConsumeMouseWheel::WhenScrollingPossible: //default behavior
-		default:
-			if ( FMath::Abs( AmountScrolledInItems ) > 0.0f )
-			{
-				return FReply::Handled();
-			}
+			// When we need to maintain a fixed offset, we scroll by items. This prevents the list not moving or jumping unexpectedly far on an individual scroll wheel motion.
+			const double AdditionalOffset = (MouseEvent.GetWheelDelta() >= 0.f ? -1.f : 1.f) * GetNumItemsPerLine();
+			const double NewScrollOffset = FMath::Max(0., DesiredScrollOffset + AdditionalOffset);
+			AmountScrolledInItems = this->ScrollTo(NewScrollOffset);
+		}
+		else
+		{
+			// No required offset to maintain, so we scroll by units
+			AmountScrolledInItems = this->ScrollBy(MyGeometry, -MouseEvent.GetWheelDelta() * WheelScrollMultiplier, EAllowOverscroll::No);
 		}
 		
+		if (ConsumeMouseWheel == EConsumeMouseWheel::Always || (FMath::Abs(AmountScrolledInItems) > 0.0f && ConsumeMouseWheel != EConsumeMouseWheel::Never))
+		{
+			return FReply::Handled();
+		}
 	}
 	return FReply::Unhandled();
 }
@@ -543,8 +569,6 @@ FReply STableViewBase::OnKeyDown( const FGeometry& MyGeometry, const FKeyEvent&
 		ScrollToBottom();
 		return FReply::Handled();
 	}
-
-
 	return SCompoundWidget::OnKeyDown(MyGeometry, InKeyEvent); 
 }
 
@@ -678,7 +702,6 @@ int32 STableViewBase::OnPaint( const FPaintArgs& Args, const FGeometry& Allotted
 
 STableViewBase::STableViewBase( ETableViewMode::Type InTableViewMode )
 	: TableViewMode( InTableViewMode )
-	, ScrollOffset( 0 )
 	, bStartedTouchInteraction( false )
 	, AmountScrolledWhileRightMouseDown( 0 )
 	, TickScrollDelta( 0 )
@@ -694,12 +717,26 @@ STableViewBase::STableViewBase( ETableViewMode::Type InTableViewMode )
 	, ConsumeMouseWheel(EConsumeMouseWheel::WhenScrollingPossible)
 	, bItemsNeedRefresh( false )	
 {
+	FixedLineScrollOffset = 0.25f;
+}
+
+double STableViewBase::GetTargetScrollOffset() const
+{
+	if (FixedLineScrollOffset.IsSet() && !IsRightClickScrolling() && InertialScrollManager.GetScrollVelocity() == 0.f)
+	{
+		const int32 NumItemsPerLine = GetNumItemsPerLine();
+		const double DesiredLineOffset = FMath::FloorToDouble(DesiredScrollOffset / NumItemsPerLine) - FixedLineScrollOffset.GetValue();
+		
+		//return FMath::Max(0.0, FMath::CeilToDouble(DesiredScrollOffset) - AdditionalOffset);
+		return FMath::Max(0.0, DesiredLineOffset * NumItemsPerLine);
+	}
+	return DesiredScrollOffset;
 }
 
 float STableViewBase::ScrollBy(const FGeometry& MyGeometry, float ScrollByAmountInSlateUnits, EAllowOverscroll InAllowOverscroll)
 {
 	const int32 NumItemsBeingObserved = GetNumItemsBeingObserved();
-	const float FractionalScrollOffsetInItems = (ScrollOffset + GetScrollRateInItems() * ScrollByAmountInSlateUnits) / NumItemsBeingObserved;
+	const float FractionalScrollOffsetInItems = (DesiredScrollOffset + GetScrollRateInItems() * ScrollByAmountInSlateUnits) / NumItemsBeingObserved;
 	const double ClampedScrollOffsetInItems = FMath::Clamp<double>( FractionalScrollOffsetInItems*NumItemsBeingObserved, -10.0f, NumItemsBeingObserved+10.0f ) * NumItemsBeingObserved;
 	if (InAllowOverscroll == EAllowOverscroll::Yes)
 	{
@@ -711,11 +748,11 @@ float STableViewBase::ScrollBy(const FGeometry& MyGeometry, float ScrollByAmount
 float STableViewBase::ScrollTo( float InScrollOffset)
 {
 	const float NewScrollOffset = FMath::Clamp( InScrollOffset, -10.0f, GetNumItemsBeingObserved()+10.0f );
-	float AmountScrolled = FMath::Abs( ScrollOffset - NewScrollOffset );
+	float AmountScrolled = FMath::Abs( DesiredScrollOffset - NewScrollOffset );
 
 	SetScrollOffset( NewScrollOffset );
 	
-	if ( bWasAtEndOfList && NewScrollOffset >= ScrollOffset )
+	if ( bWasAtEndOfList && NewScrollOffset >= DesiredScrollOffset )
 	{
 		AmountScrolled = 0;
 	}
@@ -725,15 +762,16 @@ float STableViewBase::ScrollTo( float InScrollOffset)
 
 float STableViewBase::GetScrollOffset() const
 {
-	return ScrollOffset;
+	return DesiredScrollOffset;
 }
 
 void STableViewBase::SetScrollOffset( const float InScrollOffset )
 {
-	if ( ScrollOffset != InScrollOffset )
+	const float InValidatedOffset = FMath::Max(0.0f, InScrollOffset);
+	if (DesiredScrollOffset != InValidatedOffset)
 	{
-		ScrollOffset = InScrollOffset;
-		OnTableViewScrolled.ExecuteIfBound( ScrollOffset );
+		DesiredScrollOffset = InValidatedOffset;
+		OnTableViewScrolled.ExecuteIfBound(DesiredScrollOffset);
 		RequestLayoutRefresh();
 	}
 }
@@ -747,10 +785,10 @@ void STableViewBase::AddScrollOffset(const float InScrollOffsetDelta, bool Refre
 {
 	if (FMath::IsNearlyEqual(InScrollOffsetDelta, 0.0f) == false)
 	{
-		ScrollOffset += InScrollOffsetDelta;
+		DesiredScrollOffset += InScrollOffsetDelta;
 		if (RefreshList)
 		{
-			OnTableViewScrolled.ExecuteIfBound(ScrollOffset);
+			OnTableViewScrolled.ExecuteIfBound(DesiredScrollOffset);
 			RequestLayoutRefresh();
 		}
 	}
@@ -764,6 +802,25 @@ void STableViewBase::SetScrollbarVisibility(const EVisibility InVisibility)
 	}
 }
 
+void STableViewBase::SetFixedLineScrollOffset(TOptional<double> InFixedLineScrollOffset)
+{
+	if (FixedLineScrollOffset != InFixedLineScrollOffset)
+	{
+		FixedLineScrollOffset = InFixedLineScrollOffset;
+		RequestLayoutRefresh();
+	}
+}
+
+void STableViewBase::SetIsScrollAnimationEnabled(bool bInEnableScrollAnimation)
+{
+	bEnableAnimatedScrolling = bInEnableScrollAnimation;
+}
+
+void STableViewBase::SetWheelScrollMultiplier(float NewWheelScrollMultiplier)
+{
+	WheelScrollMultiplier = NewWheelScrollMultiplier;
+}
+
 void STableViewBase::InsertWidget( const TSharedRef<ITableRow> & WidgetToInset )
 {
 	ItemsPanel->AddSlot(0)
@@ -772,11 +829,6 @@ void STableViewBase::InsertWidget( const TSharedRef<ITableRow> & WidgetToInset )
 	];
 }
 
-/**
- * Add a widget to the view.
- *
- * @param WidgetToAppend   Widget to append to the view.
- */
 void STableViewBase::AppendWidget( const TSharedRef<ITableRow>& WidgetToAppend )
 {
 	ItemsPanel->AddSlot()
@@ -785,9 +837,6 @@ void STableViewBase::AppendWidget( const TSharedRef<ITableRow>& WidgetToAppend )
 	];
 }
 
-/**
- * Remove all the widgets from the view.
- */
 void STableViewBase::ClearWidgets()
 {
 	ItemsPanel->ClearItems();
@@ -805,7 +854,9 @@ float STableViewBase::GetItemHeight() const
 
 FVector2D STableViewBase::GetItemSize() const
 {
-	return ItemsPanel->GetItemSize(PanelGeometryLastTick) + FVector2D(ItemsPanel->GetItemPadding(PanelGeometryLastTick), 0.0f);
+	FTableViewDimensions ItemDimensions = ItemsPanel->GetItemSize(PanelGeometryLastTick);
+	ItemDimensions.LineAxis += ItemsPanel->GetItemPadding(PanelGeometryLastTick);
+	return ItemDimensions.ToVector2D();
 }
 
 void STableViewBase::SetItemHeight(TAttribute<float> Height)
@@ -823,7 +874,7 @@ float STableViewBase::GetNumLiveWidgets() const
 	return ItemsPanel->GetChildren()->Num();
 }
 
-int32 STableViewBase::GetNumItemsWide() const
+int32 STableViewBase::GetNumItemsPerLine() const
 {
 	return 1;
 }
@@ -858,9 +909,9 @@ void STableViewBase::OnRightMouseButtonUp(const FPointerEvent& MouseEvent)
 
 float STableViewBase::GetScrollRateInItems() const
 {
-	return (LastGenerateResults.HeightOfGeneratedItems != 0 && LastGenerateResults.ExactNumRowsOnScreen != 0)
+	return (LastGenerateResults.LengthOfGeneratedItems != 0 && LastGenerateResults.ExactNumLinesOnScreen != 0)
 		// Approximate a consistent scrolling rate based on the average item height.
-		? LastGenerateResults.ExactNumRowsOnScreen / LastGenerateResults.HeightOfGeneratedItems
+		? LastGenerateResults.ExactNumLinesOnScreen / LastGenerateResults.LengthOfGeneratedItems
 		// Scroll 1/2 an item at a time as a default.
 		: 0.5f;
 }
diff --git a/Engine/Source/Runtime/Slate/Public/Widgets/Views/SListView.h b/Engine/Source/Runtime/Slate/Public/Widgets/Views/SListView.h
--- a/Engine/Source/Runtime/Slate/Public/Widgets/Views/SListView.h
+++ b/Engine/Source/Runtime/Slate/Public/Widgets/Views/SListView.h
@@ -81,6 +81,8 @@ public:
 		, _SelectionMode(ESelectionMode::Multi)
 		, _ClearSelectionOnClick(true)
 		, _ExternalScrollbar()
+		, _Orientation(Orient_Vertical)
+		, _EnableAnimatedScrolling(false)
 		, _ScrollbarDragFocusCause(EFocusCause::Mouse)
 		, _AllowOverscroll(EAllowOverscroll::Yes)
 		, _ConsumeMouseWheel(EConsumeMouseWheel::WhenScrollingPossible)
@@ -125,6 +127,12 @@ public:
 
 		SLATE_ARGUMENT( TSharedPtr<SScrollBar>, ExternalScrollbar )
 
+		SLATE_ARGUMENT( EOrientation, Orientation )
+
+		SLATE_ARGUMENT( bool, EnableAnimatedScrolling)
+
+		SLATE_ARGUMENT( TOptional<double>, FixedLineScrollOffset )
+
 		SLATE_ATTRIBUTE( EVisibility, ScrollbarVisibility)
 		
 		SLATE_ARGUMENT( EFocusCause, ScrollbarDragFocusCause )
@@ -184,6 +192,9 @@ public:
 
 		this->bReturnFocusToSelection = InArgs._ReturnFocusToSelection;
 
+		this->bEnableAnimatedScrolling = InArgs._EnableAnimatedScrolling;
+		this->FixedLineScrollOffset = InArgs._FixedLineScrollOffset;
+
 		this->OnItemToString_Debug =
 			InArgs._OnItemToString_Debug.IsBound()
 			? InArgs._OnItemToString_Debug
@@ -218,7 +229,7 @@ public:
 		else
 		{
 			// Make the TableView
-			ConstructChildren( 0, InArgs._ItemHeight, EListItemAlignment::LeftAligned, InArgs._HeaderRow, InArgs._ExternalScrollbar, InArgs._OnListViewScrolled );
+			ConstructChildren( 0, InArgs._ItemHeight, EListItemAlignment::LeftAligned, InArgs._HeaderRow, InArgs._ExternalScrollbar, InArgs._Orientation, InArgs._OnListViewScrolled );
 			if(ScrollBar.IsValid())
 			{
 				ScrollBar->SetDragFocusCause(InArgs._ScrollbarDragFocusCause);
@@ -281,7 +292,7 @@ public:
 				}
 
 				int32 NumItemsInAPage = GetNumLiveWidgets();
-				int32 Remainder = NumItemsInAPage % GetNumItemsWide();
+				int32 Remainder = NumItemsInAPage % GetNumItemsPerLine();
 				NumItemsInAPage -= Remainder;
 
 				if ( SelectionIndex >= NumItemsInAPage )
@@ -305,7 +316,7 @@ public:
 				}
 
 				int32 NumItemsInAPage = GetNumLiveWidgets();
-				int32 Remainder = NumItemsInAPage % GetNumItemsWide();
+				int32 Remainder = NumItemsInAPage % GetNumItemsPerLine();
 				NumItemsInAPage -= Remainder;
 
 				if ( SelectionIndex < ItemsSourceRef.Num() - NumItemsInAPage )
@@ -396,19 +407,21 @@ public:
 		{
 			const TArray<ItemType>& ItemsSourceRef = (*this->ItemsSource);
 
-			const int32 NumItemsWide = GetNumItemsWide();
+			const int32 NumItemsPerLine = GetNumItemsPerLine();
 			const int32 CurSelectionIndex = (!TListTypeTraits<ItemType>::IsPtrValid(SelectorItem)) ? 0 : ItemsSourceRef.Find(TListTypeTraits<ItemType>::NullableItemTypeConvertToItemType(SelectorItem));
 			int32 AttemptSelectIndex = -1;
 
 			const EUINavigation NavType = InNavigationEvent.GetNavigationType();
-			switch (NavType)
+			if ((Orientation == Orient_Vertical && NavType == EUINavigation::Up) ||
+				(Orientation == Orient_Horizontal && NavType == EUINavigation::Left))
 			{
-			case EUINavigation::Up:
-				AttemptSelectIndex = CurSelectionIndex - NumItemsWide;
-				break;
-
-			case EUINavigation::Down:
-				AttemptSelectIndex = CurSelectionIndex + NumItemsWide;
+				// Nav backward by a line
+				AttemptSelectIndex = CurSelectionIndex - NumItemsPerLine;
+			}
+			else if ((Orientation == Orient_Vertical && NavType == EUINavigation::Down) ||
+				(Orientation == Orient_Horizontal && NavType == EUINavigation::Right))
+			{
+				AttemptSelectIndex = CurSelectionIndex + NumItemsPerLine;
 
 				// The list might be jagged so attempt to determine if there's a partially filled row we can move to
 				if (!ItemsSourceRef.IsValidIndex(AttemptSelectIndex))
@@ -417,23 +430,19 @@ public:
 					if (NumItems > 0)
 					{
 						// NumItemsWide should never be 0, ensuring for sanity
-						ensure(NumItemsWide > 0);
+						ensure(NumItemsPerLine > 0);
 
 						// calculate total number of rows and row of current index (1 index)
-						int32 NumRows = NumItems / NumItemsWide + (NumItems % NumItemsWide != 0 ? 1 : 0);
-						int32 CurRow = CurSelectionIndex / NumItems + 1;
+						int32 NumLines = NumItems / NumItemsPerLine + (NumItems % NumItemsPerLine != 0 ? 1 : 0);
+						int32 CurLine = CurSelectionIndex / NumItems + 1;
 
 						// if not on final row, assume a jagged list and select the final item
-						if (CurRow < NumRows)
+						if (CurLine < NumLines)
 						{
 							AttemptSelectIndex = NumItems - 1;
 						}
 					}
 				}
-				break;
-
-			default:
-				break;
 			}
 
 			// If it's valid we'll scroll it into view and return an explicit widget in the FNavigationReply
@@ -982,124 +991,123 @@ public:
 		this->ClearWidgets();
 
 		// Ensure that we always begin and clean up a generation pass.
-		
 		FGenerationPassGuard GenerationPassGuard(WidgetGenerator);
 
-		const TArray<ItemType>* SourceItems = ItemsSource;
-		if ( SourceItems != nullptr && SourceItems->Num() > 0 )
+		if (ItemsSource && ItemsSource->Num() > 0)
 		{
 			// Items in view, including fractional items
 			float ItemsInView = 0.0f;
 
-			// Height of generated widgets that is landing in the bounds of the view.
-			float ViewHeightUsedSoFar = 0.0f;
+			// Total length of widgets generated so far (height for vertical lists, width for horizontal)
+			float LengthGeneratedSoFar = 0.0f;
 
-			// Total height of widgets generated so far.
-			float HeightGeneratedSoFar = 0.0f;
+			// Length of generated widgets that in the bounds of the view.
+			float ViewLengthUsedSoFar = 0.0f;
 
 			// Index of the item at which we start generating based on how far scrolled down we are
 			// Note that we must generate at LEAST one item.
-			int32 StartIndex = FMath::Clamp( FMath::FloorToInt(ScrollOffset), 0, SourceItems->Num()-1 );
+			int32 StartIndex = FMath::Clamp( FMath::FloorToInt(CurrentScrollOffset), 0, ItemsSource->Num() - 1 );
 
-			// Height of the first item that is generated. This item is at the location where the user requested we scroll
-			float FirstItemHeight = 0.0f;
+			// Length of the first item that is generated. This item is at the location where the user requested we scroll
+			float FirstItemLength = 0.0f;
 
 			// Generate widgets assuming scenario a.
-			bool bGeneratedEnoughForSmoothScrolling = false;
+			bool bHasFilledAvailableArea = false;
 			bool bAtEndOfList = false;
 
 			const float LayoutScaleMultiplier = MyGeometry.GetAccumulatedLayoutTransform().GetScale();
-			for( int32 ItemIndex = StartIndex; !bGeneratedEnoughForSmoothScrolling && ItemIndex < SourceItems->Num(); ++ItemIndex )
+			FTableViewDimensions MyDimensions(this->Orientation, MyGeometry.GetLocalSize());
+			
+			for( int32 ItemIndex = StartIndex; !bHasFilledAvailableArea && ItemIndex < ItemsSource->Num(); ++ItemIndex )
 			{
-				const ItemType& CurItem = (*SourceItems)[ItemIndex];
+				const ItemType& CurItem = (*ItemsSource)[ItemIndex];
 
-				// We do not generatie a new widget if the CurItem is an invalid object.
 				if (!TListTypeTraits<ItemType>::IsPtrValid(CurItem))
 				{
+					// Don't bother generating widgets for invalid items
 					continue;
 				}
 
-				const float ItemHeight = GenerateWidgetForItem(CurItem, ItemIndex, StartIndex, LayoutScaleMultiplier);
+				const float ItemLength = GenerateWidgetForItem(CurItem, ItemIndex, StartIndex, LayoutScaleMultiplier);
 
 				const bool bIsFirstItem = ItemIndex == StartIndex;
 
 				if (bIsFirstItem)
 				{
-					FirstItemHeight = ItemHeight;
+					FirstItemLength = ItemLength;
 				}
 
 				// Track the number of items in the view, including fractions.
 				if (bIsFirstItem)
 				{
 					// The first item may not be fully visible (but cannot exceed 1)
-					// FirstItemFractionScrolledIntoView is the fraction of the item that is visible after taking into account anything that may be scrolled off the top of the list view
-					// FirstItemHeightScrolledIntoView is the height of the item, ignoring anything that is scrolled off the top of the list view
-					// FirstItemVisibleFraction is either: The visible item height as a fraction of the available list view height (if the item size is larger than the available size, otherwise this will be >1), or just FirstItemFractionScrolledIntoView (which can never be >1)
-					const float FirstItemFractionScrolledIntoView = 1.0f - FMath::Max(FMath::Fractional(ScrollOffset), 0.0f);
-					const float FirstItemHeightScrolledIntoView = ItemHeight * FirstItemFractionScrolledIntoView;
-					const float FirstItemVisibleFraction = FMath::Min(MyGeometry.GetLocalSize().Y / FirstItemHeightScrolledIntoView, FirstItemFractionScrolledIntoView);
+					// FirstItemFractionScrolledIntoView is the fraction of the item that is visible after taking into account anything that may be scrolled off the top/left of the list view
+					const float FirstItemFractionScrolledIntoView = 1.0f - FMath::Max(FMath::Fractional(CurrentScrollOffset), 0.0f);
+					
+					// FirstItemLengthScrolledIntoView is the length of the item, ignoring anything that is scrolled off the top/left of the list view
+					const float FirstItemLengthScrolledIntoView = ItemLength * FirstItemFractionScrolledIntoView;
+					
+					// FirstItemVisibleFraction is either: The visible item length as a fraction of the available list view length (if the item size is larger than the available size, otherwise this will be >1), or just FirstItemFractionScrolledIntoView (which can never be >1)
+					const float FirstItemVisibleFraction = FMath::Min(MyDimensions.ScrollAxis / FirstItemLengthScrolledIntoView, FirstItemFractionScrolledIntoView);
+					
 					ItemsInView += FirstItemVisibleFraction;
 				}
-				else if (ViewHeightUsedSoFar + ItemHeight > MyGeometry.GetLocalSize().Y)
+				else if (ViewLengthUsedSoFar + ItemLength > MyDimensions.ScrollAxis)
 				{
 					// The last item may not be fully visible either
-					ItemsInView += (MyGeometry.GetLocalSize().Y - ViewHeightUsedSoFar) / ItemHeight;
+					ItemsInView += (MyDimensions.ScrollAxis - ViewLengthUsedSoFar) / ItemLength;
 				}
 				else
 				{
 					ItemsInView += 1;
 				}
 
-				HeightGeneratedSoFar += ItemHeight;
+				LengthGeneratedSoFar += ItemLength;
 
-				ViewHeightUsedSoFar += (bIsFirstItem)
-					? ItemHeight * ItemsInView	// For the first item, ItemsInView <= 1.0f
-					: ItemHeight;
+				ViewLengthUsedSoFar += (bIsFirstItem)
+					? ItemLength * ItemsInView	// For the first item, ItemsInView <= 1.0f
+					: ItemLength;
 
-				if (ItemIndex >= SourceItems->Num()-1)
+				if (ItemIndex >= ItemsSource->Num() - 1)
 				{
 					bAtEndOfList = true;
 				}
 
-				if (ViewHeightUsedSoFar > MyGeometry.GetLocalSize().Y )
+				if (ViewLengthUsedSoFar >= MyDimensions.ScrollAxis)
 				{
-					bGeneratedEnoughForSmoothScrolling = true;
+					bHasFilledAvailableArea = true;
 				}
 			}
 
 			// Handle scenario b.
-			// We may have stopped because we got to the end of the items.
-			// But we may still have space to fill!
-			if (bAtEndOfList && ViewHeightUsedSoFar < MyGeometry.GetLocalSize().Y)
+			// We may have stopped because we got to the end of the items, but we may still have space to fill!
+			if (bAtEndOfList && !bHasFilledAvailableArea)
 			{
-				double NewScrollOffsetForBackfill = static_cast<double>(StartIndex) + (HeightGeneratedSoFar - MyGeometry.GetLocalSize().Y) / FirstItemHeight;
+				double NewScrollOffsetForBackfill = static_cast<double>(StartIndex) + (LengthGeneratedSoFar - MyDimensions.ScrollAxis) / FirstItemLength;
 
-				for( int32 ItemIndex = StartIndex-1; HeightGeneratedSoFar < MyGeometry.GetLocalSize().Y && ItemIndex >= 0; --ItemIndex )
+				for (int32 ItemIndex = StartIndex - 1; LengthGeneratedSoFar < MyDimensions.ScrollAxis && ItemIndex >= 0; --ItemIndex)
 				{
-					const ItemType& CurItem = (*SourceItems)[ItemIndex];
-					// When the item is not valid, we do not generate a widget for it.
-					if (!TListTypeTraits<ItemType>::IsPtrValid(CurItem))
+					const ItemType& CurItem = (*ItemsSource)[ItemIndex];
+					if (TListTypeTraits<ItemType>::IsPtrValid(CurItem))
 					{
-						continue;
-					}
+						const float ItemLength = GenerateWidgetForItem(CurItem, ItemIndex, StartIndex, LayoutScaleMultiplier);
 
-					const float ItemHeight = GenerateWidgetForItem(CurItem, ItemIndex, StartIndex, LayoutScaleMultiplier);
+						if (LengthGeneratedSoFar + ItemLength > MyDimensions.ScrollAxis && ItemLength > 0.f)
+						{
+							// Generated the item that puts us over the top.
+							// Count the fraction of this item that will stick out beyond the list
+							NewScrollOffsetForBackfill = static_cast<double>(ItemIndex) + (LengthGeneratedSoFar + ItemLength - MyDimensions.ScrollAxis) / ItemLength;
+						}
 
-					if (HeightGeneratedSoFar + ItemHeight > MyGeometry.GetLocalSize().Y && ItemHeight > 0.f)
-					{
-						// Generated the item that puts us over the top.
-						// Count the fraction of this item that will stick out above the list
-						NewScrollOffsetForBackfill = static_cast<double>(ItemIndex) + (HeightGeneratedSoFar + ItemHeight - MyGeometry.GetLocalSize().Y) / ItemHeight;
+						// The widget used up some of the available vertical space.
+						LengthGeneratedSoFar += ItemLength;
 					}
-
-					// The widget used up some of the available vertical space.
-					HeightGeneratedSoFar += ItemHeight;
 				}
 
-				return FReGenerateResults(NewScrollOffsetForBackfill, HeightGeneratedSoFar, ItemsSource->Num() - NewScrollOffsetForBackfill, bAtEndOfList);
+				return FReGenerateResults(NewScrollOffsetForBackfill, LengthGeneratedSoFar, ItemsSource->Num() - NewScrollOffsetForBackfill, true);
 			}
 
-			return FReGenerateResults(ScrollOffset, HeightGeneratedSoFar, ItemsInView, bAtEndOfList);
+			return FReGenerateResults(CurrentScrollOffset, LengthGeneratedSoFar, ItemsInView, false);
 		}
 
 		return FReGenerateResults(0.0f, 0.0f, 0.0f, false);
@@ -1128,10 +1136,6 @@ public:
 		const TSharedRef<SWidget> NewlyGeneratedWidget = WidgetForItem->AsWidget();
 		NewlyGeneratedWidget->SlatePrepass(LayoutScaleMultiplier);
 
-		const bool IsFirstWidgetOnScreen = (ItemIndex == StartIndex);
-		const bool IsVisible = NewlyGeneratedWidget->GetVisibility().IsVisible();
-		const float ItemHeight = IsVisible ? NewlyGeneratedWidget->GetDesiredSize().Y : 0;
-
 		// We have a widget for this item; add it to the panel so that it is part of the UI.
 		if (ItemIndex >= StartIndex)
 		{
@@ -1144,7 +1148,9 @@ public:
 			this->InsertWidget( WidgetForItem.ToSharedRef() );
 		}
 
-		return ItemHeight;
+		const bool bIsVisible = NewlyGeneratedWidget->GetVisibility().IsVisible();
+		FTableViewDimensions GeneratedWidgetDimensions(this->Orientation, bIsVisible ? NewlyGeneratedWidget->GetDesiredSize() : FVector2D::ZeroVector);
+		return GeneratedWidgetDimensions.ScrollAxis;
 	}
 
 	/** @return how many items there are in the TArray being observed */
@@ -1480,11 +1486,11 @@ public:
 	* Will determine the max row size for the specified column id
 	*
 	* @param ColumnId  Column Id
-	* @param Orientation  Orientation that is main axis you want to query
+	* @param ColumnOrientation  Orientation that is main axis you want to query
 	*
 	* @return The max size for a column Id.
 	*/
-	FVector2D GetMaxRowSizeForColumn(const FName& ColumnId, EOrientation Orientation)
+	FVector2D GetMaxRowSizeForColumn(const FName& ColumnId, EOrientation ColumnOrientation)
 	{
 		FVector2D MaxSize = FVector2D::ZeroVector;
 
@@ -1494,7 +1500,7 @@ public:
 			FVector2D NewMaxSize = TableRow->GetRowSizeForColumn(ColumnId);
 
 			// We'll return the full size, but we only take into consideration the asked axis for the calculation of the size
-			if (NewMaxSize.Component(Orientation) > MaxSize.Component(Orientation))
+			if (NewMaxSize.Component(ColumnOrientation) > MaxSize.Component(ColumnOrientation))
 			{
 				MaxSize = NewMaxSize;
 			}
@@ -1517,16 +1523,16 @@ protected:
 	 */
 	virtual EScrollIntoViewResult ScrollIntoView( const FGeometry& ListViewGeometry ) override
 	{
-		if ( TListTypeTraits<ItemType>::IsPtrValid(ItemToScrollIntoView) && ItemsSource != nullptr )
+		if (ItemsSource && TListTypeTraits<ItemType>::IsPtrValid(ItemToScrollIntoView))
 		{
 			const int32 IndexOfItem = ItemsSource->Find( TListTypeTraits<ItemType>::NullableItemTypeConvertToItemType( ItemToScrollIntoView ) );
 			if (IndexOfItem != INDEX_NONE)
 			{
 				double NumLiveWidgets = GetNumLiveWidgets();
-				if (NumLiveWidgets == 0 && IsPendingRefresh())
+				if (NumLiveWidgets == 0. && IsPendingRefresh())
 				{
 					// Use the last number of widgets on screen to estimate if we actually need to scroll.
-					NumLiveWidgets = LastGenerateResults.ExactNumRowsOnScreen;
+					NumLiveWidgets = LastGenerateResults.ExactNumLinesOnScreen;
 
 					// If we still don't have any widgets, we're not in a situation where we can scroll an item into view
 					// (probably as nothing has been generated yet), so we'll defer this again until the next frame
@@ -1538,11 +1544,13 @@ protected:
 
 				EndInertialScrolling();
 				
+				const int32 NumFullEntriesInView = FMath::FloorToInt(CurrentScrollOffset + NumLiveWidgets) - FMath::CeilToInt(CurrentScrollOffset);
+
 				// Only scroll the item into view if it's not already in the visible range
-				// When navigating, we don't want to scroll partially visible existing rows all the way to the center
-				const double MinDisplayedIndex = bNavigateOnScrollIntoView ? ScrollOffset - 1 : ScrollOffset;
-				const double MaxDisplayedIndex = FMath::Max(ScrollOffset + NumLiveWidgets + (bNavigateOnScrollIntoView ? 0 : -1), MinDisplayedIndex);
-				if (IndexOfItem <= MinDisplayedIndex || IndexOfItem > MaxDisplayedIndex)
+				// When navigating, we don't want to scroll partially visible existing rows all the way to the center, so we count partially displayed indices in the displayed range
+				const double MinDisplayedIndex = bNavigateOnScrollIntoView ? FMath::FloorToDouble(CurrentScrollOffset) : FMath::CeilToDouble(CurrentScrollOffset);
+				const double MaxDisplayedIndex = bNavigateOnScrollIntoView ? FMath::CeilToDouble(CurrentScrollOffset + NumFullEntriesInView) : FMath::FloorToDouble(CurrentScrollOffset + NumFullEntriesInView);
+				if (IndexOfItem < MinDisplayedIndex || IndexOfItem > MaxDisplayedIndex)
 				{
 					// Scroll the top of the listview to the item in question
 					double NewScrollOffset = IndexOfItem;
@@ -1560,47 +1568,51 @@ protected:
 				{
 					if (TSharedPtr<ITableRow> TableRow = WidgetFromItem((*ItemsSource)[IndexOfItem]))
 					{
-						// Make sure the existing row for this item is fully in view
-						const FGeometry& RowGeometry = TableRow->AsWidget()->GetCachedGeometry();
-						if (RowGeometry.GetAbsolutePositionAtCoordinates(FVector2D::ZeroVector).Y < ListViewGeometry.GetAbsolutePositionAtCoordinates(FVector2D::ZeroVector).Y)
+						const FGeometry& WidgetGeometry = TableRow->AsWidget()->GetCachedGeometry();
+						const FTableViewDimensions WidgetTopLeft(this->Orientation, WidgetGeometry.GetAbsolutePositionAtCoordinates(FVector2D::ZeroVector));
+						const FTableViewDimensions ListViewTopLeft(this->Orientation, ListViewGeometry.GetAbsolutePositionAtCoordinates(FVector2D::ZeroVector));
+
+						double NewScrollOffset = DesiredScrollOffset;
+						// Make sure the existing entry for this item is fully in view
+						if (WidgetTopLeft.ScrollAxis < ListViewTopLeft.ScrollAxis)
 						{
-							// This row is clipped on the top, so simply set it as the new scroll offset target to bump it down a bit
-							const double MaxScrollOffset = FMath::Max(0.0, static_cast<double>(ItemsSource->Num()) - NumLiveWidgets);
-							double NewScrollOffset = FMath::Clamp<double>(static_cast<double>(IndexOfItem - NavigationScrollOffset), 0.0, MaxScrollOffset);
-							SetScrollOffset(NewScrollOffset);
+							// This entry is clipped at the top/left, so simply set it as the new scroll offset target to bump it down into view
+							NewScrollOffset = static_cast<double>(IndexOfItem - NavigationScrollOffset);
 						}
 						else
 						{
 							const FVector2D BottomRight(1.f, 1.f);
-							const float RowBottomY = RowGeometry.GetAbsolutePositionAtCoordinates(BottomRight).Y;
-							const float PanelBottomY = ListViewGeometry.GetAbsolutePositionAtCoordinates(BottomRight).Y;
-							if (RowBottomY > PanelBottomY)
+							const FTableViewDimensions WidgetBottomRight(this->Orientation, WidgetGeometry.GetAbsolutePositionAtCoordinates(BottomRight));
+							const FTableViewDimensions ListViewBottomRight(this->Orientation, ListViewGeometry.GetAbsolutePositionAtCoordinates(BottomRight));
+							
+							if (WidgetBottomRight.ScrollAxis > ListViewBottomRight.ScrollAxis)
 							{
-								// This row is clipped on the bottom, so we need to determine the exact item offset required to get the bottom of this entry into view
-								// To do so, we need to push the current offset down by the clipped amount translated into number of rows
-								float DistanceRemaining = RowBottomY - PanelBottomY;
+								// This entry is clipped at the end, so we need to determine the exact item offset required to get it fully into view
+								// To do so, we need to push the current offset down by the clipped amount translated into number of items
+								float DistanceRemaining = WidgetBottomRight.ScrollAxis - ListViewBottomRight.ScrollAxis;
 								float AdditionalOffset = 0.f;
 								for (const ItemType& ItemWithWidget : WidgetGenerator.ItemsWithGeneratedWidgets)
 								{
-									const float ExistingRowSizeY = WidgetGenerator.GetWidgetForItem(ItemWithWidget)->AsWidget()->GetCachedGeometry().GetAbsoluteSize().Y;
-									if (ExistingRowSizeY < DistanceRemaining)
+									FTableViewDimensions WidgetAbsoluteDimensions(this->Orientation, WidgetGenerator.GetWidgetForItem(ItemWithWidget)->AsWidget()->GetCachedGeometry().GetAbsoluteSize());
+									if (WidgetAbsoluteDimensions.ScrollAxis < DistanceRemaining)
 									{
-										DistanceRemaining -= ExistingRowSizeY;
+										DistanceRemaining -= WidgetAbsoluteDimensions.ScrollAxis;
 										AdditionalOffset += 1.f;
 									}
 									else
 									{
-										AdditionalOffset += DistanceRemaining / ExistingRowSizeY;
+										AdditionalOffset += DistanceRemaining / WidgetAbsoluteDimensions.ScrollAxis;
 										DistanceRemaining = 0.f;
 										break;
 									}
 								}
 
-								const double MaxScrollOffset = FMath::Max(0.0, static_cast<double>(ItemsSource->Num()) - NumLiveWidgets);
-								double NewScrollOffset = FMath::Clamp<double>(ScrollOffset + AdditionalOffset + NavigationScrollOffset, 0.0, MaxScrollOffset);
-								SetScrollOffset(NewScrollOffset);
+								NewScrollOffset = DesiredScrollOffset + AdditionalOffset + (FixedLineScrollOffset.IsSet() ? 0.f : NavigationScrollOffset);
 							}
 						}
+
+						const double MaxScrollOffset = FMath::Max(0.0, static_cast<double>(ItemsSource->Num()) - NumLiveWidgets);
+						SetScrollOffset(FMath::Min(NewScrollOffset, MaxScrollOffset));
 					}
 				}
 
@@ -1641,7 +1653,7 @@ protected:
 		if (InAllowOverscroll == EAllowOverscroll::No)
 		{
 			//check if we are on the top of the list and want to scroll up
-			if (ScrollOffset < KINDA_SMALL_NUMBER && ScrollByAmountInSlateUnits < 0)
+			if (DesiredScrollOffset < KINDA_SMALL_NUMBER && ScrollByAmountInSlateUnits < 0)
 			{
 				return 0.0f;
 			}
@@ -1654,11 +1666,11 @@ protected:
 		}
 
 		float AbsScrollByAmount = FMath::Abs( ScrollByAmountInSlateUnits );
-		int32 StartingItemIndex = (int32)ScrollOffset;
-		double NewScrollOffset = ScrollOffset;
+		int32 StartingItemIndex = (int32)CurrentScrollOffset;
+		double NewScrollOffset = DesiredScrollOffset;
 
-		const bool bWholeListVisible = ScrollOffset == 0 && bWasAtEndOfList;
-		if ( InAllowOverscroll == EAllowOverscroll::Yes && Overscroll.ShouldApplyOverscroll( ScrollOffset == 0, bWasAtEndOfList, ScrollByAmountInSlateUnits ) )
+		const bool bWholeListVisible = DesiredScrollOffset == 0 && bWasAtEndOfList;
+		if ( InAllowOverscroll == EAllowOverscroll::Yes && Overscroll.ShouldApplyOverscroll(DesiredScrollOffset == 0, bWasAtEndOfList, ScrollByAmountInSlateUnits ) )
 		{
 			const float UnclampedScrollDelta = FMath::Sign(ScrollByAmountInSlateUnits) * AbsScrollByAmount;				
 			const float ActuallyScrolledBy = Overscroll.ScrollBy(MyGeometry, UnclampedScrollDelta);
@@ -1671,28 +1683,27 @@ protected:
 		else if (!bWholeListVisible)
 		{
 			// We know how far we want to scroll in SlateUnits, but we store scroll offset in "number of widgets".
-			// Challenge: each widget can be a different height.
+			// Challenge: each widget can be a different height/width.
 			// Strategy:
-			//           Scroll "one widget's height" at a time until we've scrolled as far as the user asked us to.
-			//           Generate widgets on demand so we can figure out how tall they are.
+			//           Scroll "one widget's length" at a time until we've scrolled as far as the user asked us to.
+			//           Generate widgets on demand so we can figure out how big they are.
 
-			const TArray<ItemType>* SourceItems = ItemsSource;
-			if ( SourceItems != nullptr && SourceItems->Num() > 0 )
+			if (ItemsSource && ItemsSource->Num() > 0)
 			{
-				int ItemIndex = StartingItemIndex;
+				int32 ItemIndex = StartingItemIndex;
 				const float LayoutScaleMultiplier = MyGeometry.GetAccumulatedLayoutTransform().GetScale();
-				while( AbsScrollByAmount != 0 && ItemIndex < SourceItems->Num() && ItemIndex >= 0 )
+				while( AbsScrollByAmount != 0 && ItemIndex < ItemsSource->Num() && ItemIndex >= 0 )
 				{
-					const ItemType CurItem = (*SourceItems)[ ItemIndex ];
-					// If the CurItem is not valid, we do not generate a new widget for it, we skip it.
+					const ItemType& CurItem = (*ItemsSource)[ItemIndex];
 					if (!TListTypeTraits<ItemType>::IsPtrValid(CurItem))
 					{
+						// If the CurItem is not valid, we do not generate a new widget for it, we skip it.
 						++ItemIndex;
 						continue;
 					}
 
 					TSharedPtr<ITableRow> RowWidget = WidgetGenerator.GetWidgetForItem( CurItem );
-					if ( !RowWidget.IsValid() )
+					if (!RowWidget.IsValid())
 					{
 						// We couldn't find an existing widgets, meaning that this data item was not visible before.
 						// Make a new widget for it.
@@ -1708,40 +1719,38 @@ protected:
 						RowWidget->AsWidget()->SlatePrepass(LayoutScaleMultiplier);
 					}
 
-					if ( ScrollByAmountInSlateUnits > 0 )
+					const FTableViewDimensions WidgetDimensions(this->Orientation, RowWidget->AsWidget()->GetDesiredSize());
+					if (ScrollByAmountInSlateUnits > 0)
 					{
-						FVector2D WidgetDesiredSize = RowWidget->AsWidget()->GetDesiredSize();
-						const float RemainingHeight = WidgetDesiredSize.Y * ( 1.0 - FMath::Fractional( NewScrollOffset ) );
+						const float RemainingDistance = WidgetDimensions.ScrollAxis * (1.0 - FMath::Fractional(NewScrollOffset));
 
-						if ( AbsScrollByAmount > RemainingHeight )
+						if (AbsScrollByAmount > RemainingDistance)
 						{
-							if ( ItemIndex != SourceItems->Num() )
+							if (ItemIndex != ItemsSource->Num())
 							{
-								AbsScrollByAmount -= RemainingHeight;
+								AbsScrollByAmount -= RemainingDistance;
 								NewScrollOffset = 1.0f + (int32)NewScrollOffset;
 								++ItemIndex;
 							}
 							else
 							{
-								NewScrollOffset = SourceItems->Num();
+								NewScrollOffset = ItemsSource->Num();
 								break;
 							}
 						} 
-						else if ( AbsScrollByAmount == RemainingHeight )
+						else if ( AbsScrollByAmount == RemainingDistance)
 						{
 							NewScrollOffset = 1.0f + (int32)NewScrollOffset;
 							break;
 						}
 						else
 						{
-							NewScrollOffset = (int32)NewScrollOffset + ( 1.0f - ( ( RemainingHeight - AbsScrollByAmount ) / WidgetDesiredSize.Y ) );
+							NewScrollOffset = (int32)NewScrollOffset + (1.0f - ((RemainingDistance - AbsScrollByAmount) / WidgetDimensions.ScrollAxis));
 							break;
 						}
 					}
 					else
 					{
-						FVector2D WidgetDesiredSize = RowWidget->AsWidget()->GetDesiredSize();
-
 						float Fractional = FMath::Fractional( NewScrollOffset );
 						if ( Fractional == 0 )
 						{
@@ -1749,13 +1758,13 @@ protected:
 							--NewScrollOffset;
 						}
 
-						const float PrecedingHeight = WidgetDesiredSize.Y * Fractional;
+						const float PrecedingDistance = WidgetDimensions.ScrollAxis * Fractional;
 
-						if ( AbsScrollByAmount > PrecedingHeight )
+						if ( AbsScrollByAmount > PrecedingDistance)
 						{
 							if ( ItemIndex != 0 )
 							{
-								AbsScrollByAmount -= PrecedingHeight;
+								AbsScrollByAmount -= PrecedingDistance;
 								NewScrollOffset -= FMath::Fractional( NewScrollOffset );
 								--ItemIndex;
 							}
@@ -1765,14 +1774,14 @@ protected:
 								break;
 							}
 						} 
-						else if ( AbsScrollByAmount == PrecedingHeight )
+						else if ( AbsScrollByAmount == PrecedingDistance)
 						{
 							NewScrollOffset -= FMath::Fractional( NewScrollOffset );
 							break;
 						}
 						else
 						{
-							NewScrollOffset = (int32)NewScrollOffset + ( ( PrecedingHeight - AbsScrollByAmount ) / WidgetDesiredSize.Y );
+							NewScrollOffset = (int32)NewScrollOffset + ((PrecedingDistance - AbsScrollByAmount) / WidgetDimensions.ScrollAxis);
 							break;
 						}
 					}
@@ -1903,10 +1912,10 @@ protected:
 	/** If true, the focus will be returned to the last selected object in a list when navigated to. */
 	bool bReturnFocusToSelection;
 
-private:
-
+	/** If true, the item currently slated to be scrolled into view will also be navigated to after being scrolled in */
 	bool bNavigateOnScrollIntoView = false;
 
+private:
 	struct FGenerationPassGuard
 	{
 		FWidgetGenerator& Generator;
diff --git a/Engine/Source/Runtime/Slate/Public/Widgets/Views/STableViewBase.h b/Engine/Source/Runtime/Slate/Public/Widgets/Views/STableViewBase.h
--- a/Engine/Source/Runtime/Slate/Public/Widgets/Views/STableViewBase.h
+++ b/Engine/Source/Runtime/Slate/Public/Widgets/Views/STableViewBase.h
@@ -25,37 +25,61 @@ class SScrollBar;
 enum class EConsumeMouseWheel : uint8;
 enum class ESlateVisibility : uint8;
 
-/** If the list panel is arranging items horizontally, this enum dictates how the items should be aligned (basically, where any extra space is placed) */
+/** If the list panel is arranging items as tiles, this enum dictates how the items should be aligned (basically, where any extra space is placed) */
 UENUM(BlueprintType)
 enum class EListItemAlignment : uint8
 {
-	/** Items are distributed evenly along the row (any extra space is added as padding between the items) */
+	/** Items are distributed evenly along the line (any extra space is added as padding between the items) */
 	EvenlyDistributed UMETA(DisplayName = "Evenly (Padding)"),
 
-	/** Items are distributed evenly along the row (any extra space is used to scale up the size of the item proportionally.) */
+	/** Items are distributed evenly along the line (any extra space is used to scale up the size of the item proportionally.) */
 	EvenlySize UMETA(DisplayName = "Evenly (Size)"),
 
-	/** Items are distributed evenly along the row, any extra space is used to scale up width of the items proportionally.) */
+	/** Items are distributed evenly along the line, any extra space is used to scale up width of the items proportionally.) */
 	EvenlyWide UMETA(DisplayName = "Evenly (Wide)"),
 
-	/** Items are left aligned on the row (any extra space is added to the right of the items) */
+	/** Items are left aligned on the line (any extra space is added to the right of the items) */
 	LeftAligned,
 
-	/** Items are right aligned on the row (any extra space is added to the left of the items) */
+	/** Items are right aligned on the line (any extra space is added to the left of the items) */
 	RightAligned,
 
-	/** Items are center aligned on the row (any extra space is halved and added to the left of the items) */
+	/** Items are center aligned on the line (any extra space is halved and added to the left of the items) */
 	CenterAligned,
 
-	/** Items are evenly horizontally stretched to distribute any extra space */
+	/** Items are evenly stretched to distribute any extra space on the line */
 	Fill,
 };
 
-
 DECLARE_DELEGATE_OneParam(
 	FOnTableViewScrolled,
 	double );	/** Scroll offset from the beginning of the list in items */
 
+/** Abstracts away the need to distinguish between X or Y when calculating table layout elements */
+struct SLATE_API FTableViewDimensions
+{
+	FTableViewDimensions(EOrientation InOrientation);
+	FTableViewDimensions(EOrientation InOrientation, float X, float Y);
+	FTableViewDimensions(EOrientation InOrientation, const FVector2D& Size);
+
+	FVector2D ToVector2D() const
+	{
+		return Orientation == Orient_Vertical ? FVector2D(LineAxis, ScrollAxis) : FVector2D(ScrollAxis, LineAxis);
+	}
+
+	FTableViewDimensions operator+(const FTableViewDimensions& Other) const
+	{
+		return FTableViewDimensions(Orientation, ToVector2D() + Other.ToVector2D());
+	}
+
+	EOrientation Orientation = Orient_Vertical;
+
+	/** The dimension along the scrolling axis of the table view (Y when oriented vertically, X when horizontal) */
+	float ScrollAxis = 0.f;
+	
+	/** The dimension orthogonal to the scroll axis, along which lines of items are created. Only really relevant for tile views. */
+	float LineAxis = 0.f;
+};
 
 /**
  * Contains ListView functionality that does not depend on the type of data being observed by the ListView.
@@ -67,7 +91,7 @@ class SLATE_API STableViewBase
 public:
 
 	/** Create the child widgets that comprise the list */
-	void ConstructChildren( const TAttribute<float>& InItemWidth, const TAttribute<float>& InItemHeight, const TAttribute<EListItemAlignment>& InItemAlignment, const TSharedPtr<SHeaderRow>& InColumnHeaders, const TSharedPtr<SScrollBar>& InScrollBar, const FOnTableViewScrolled& InOnTableViewScrolled );
+	void ConstructChildren( const TAttribute<float>& InItemWidth, const TAttribute<float>& InItemHeight, const TAttribute<EListItemAlignment>& InItemAlignment, const TSharedPtr<SHeaderRow>& InHeaderRow, const TSharedPtr<SScrollBar>& InScrollBar, EOrientation InScrollOrientation, const FOnTableViewScrolled& InOnTableViewScrolled );
 
 	/** Sets the item height */
 	void SetItemHeight(TAttribute<float> Height);
@@ -130,6 +154,15 @@ public:
 
 	void SetScrollbarVisibility(const EVisibility InVisibility);
 
+	/** Sets the fixed offset in items to always apply to the top/left (depending on orientation) of the list. */
+	void SetFixedLineScrollOffset(TOptional<double> InFixedLineScrollOffset);
+
+	/** Sets whether the list should lerp between scroll offsets or jump instantly between them. */
+	void SetIsScrollAnimationEnabled(bool bInEnableScrollAnimation);
+
+	/** Sets the multiplier applied when wheel scrolling. Higher numbers will cover more distance per click of the wheel. */
+	void SetWheelScrollMultiplier(float NewWheelScrollMultiplier);
+
 public:
 
 	// SWidget interface
@@ -165,6 +198,9 @@ protected:
 
 	STableViewBase( ETableViewMode::Type InTableViewMode );
 
+	/** Returns the "true" scroll offset where the list will ultimately settle (and may already be). */
+	double GetTargetScrollOffset() const;
+
 	/**
 	 * Scroll the list view by some number of screen units.
 	 *
@@ -215,10 +251,10 @@ protected:
 	virtual float GetNumLiveWidgets() const;
 
 	/**
-	 * Get the number of items that can fit in the view horizontally before creating a new row.
+	 * Get the number of items that can fit in the view along the line axis (orthogonal to the scroll axis) before creating a new line.
 	 * Default is 1, but may be more in subclasses (like STileView)
 	 */
-	virtual int32 GetNumItemsWide() const;
+	virtual int32 GetNumItemsPerLine() const;
 
 	/*
 	 * Right click down
@@ -246,26 +282,24 @@ protected:
 	/** Information about the outcome of the WidgetRegeneratePass */
 	struct SLATE_API FReGenerateResults
 	{
-		FReGenerateResults( double InNewScrollOffset, double InHeightGenerated, double InItemsOnScreen, bool AtEndOfList )
-		: NewScrollOffset( InNewScrollOffset )
-		, HeightOfGeneratedItems( InHeightGenerated )
-		, ExactNumRowsOnScreen( InItemsOnScreen )
-		, bGeneratedPastLastItem( AtEndOfList )
-		{
-
-		}
+		FReGenerateResults(double InNewScrollOffset, double InLengthGenerated, double InItemsOnScreen, bool AtEndOfList)
+			: NewScrollOffset(InNewScrollOffset)
+			, LengthOfGeneratedItems(InLengthGenerated)
+			, ExactNumLinesOnScreen(InItemsOnScreen)
+			, bGeneratedPastLastItem(AtEndOfList)
+		{}
 
 		/** The scroll offset that we actually use might not be what the user asked for */
-		double NewScrollOffset;
+		double NewScrollOffset = 0.;
 
-		/** The total height of the widgets that we have generated to represent the visible subset of the items*/
-		double HeightOfGeneratedItems;
+		/** The total length along the scroll axis of the widgets that we have generated to represent the visible subset of the items*/
+		double LengthOfGeneratedItems = 0.;
 
-		/** How many rows are fitting on the screen, including fractions */
-		double ExactNumRowsOnScreen;
+		/** How many lines are fitting on the screen, including fractions */
+		double ExactNumLinesOnScreen = 0.;
 
 		/** True when we have generated  */
-		bool bGeneratedPastLastItem;
+		bool bGeneratedPastLastItem = false;
 	};
 	/**
 	 * Update generate Widgets for Items as needed and clean up any Widgets that are no longer needed.
@@ -308,8 +342,25 @@ protected:
 	/** Delegate to call when the table view is scrolled */
 	FOnTableViewScrolled OnTableViewScrolled;
 
-	/** Scroll offset from the beginning of the list in items */
-	double ScrollOffset;
+	/** 
+	 * The fixed offset in lines to always apply to the top/left (depending on orientation) of the list.
+	 * If provided, all non-inertial means of scrolling will settle with exactly this offset of the topmost entry.
+	 * Ex: A value of 0.25 would cause the topmost full entry to be offset down by a quarter length of the preceeding entry.
+	 */
+	TOptional<double> FixedLineScrollOffset;
+
+	/** True to lerp smoothly between offsets when the desired scroll offset changes. */
+	bool bEnableAnimatedScrolling = false;
+
+	/** The currently displayed scroll offset from the beginning of the list in items. */
+	double CurrentScrollOffset = 0.;
+
+	/** 
+	 * The raw desired scroll offset from the beginning of the list in items.
+	 * Does not incorporate the FixedLineScrollOffset. Use GetTargetScrollOffset() to know the final target offset to display.
+	 * Note: If scroll animation is disabled and there is no FixedLineScrollOffset, this is identical to both the CurrentScrollOffset and the target offset.
+	 */
+	double DesiredScrollOffset = 0.;
 
 	/** Did the user start an interaction in this list? */
 	bool bStartedTouchInteraction;
@@ -353,6 +404,9 @@ protected:
 	/** How much to scroll when using mouse wheel */
 	float WheelScrollMultiplier;
 
+	/** The layout and scroll orientation of the list */
+	EOrientation Orientation = Orient_Vertical;
+
 protected:
 
 	/** Check whether the current state of the table warrants inertial scroll by the specified amount */
diff --git a/Engine/Source/Runtime/Slate/Public/Widgets/Views/STileView.h b/Engine/Source/Runtime/Slate/Public/Widgets/Views/STileView.h
--- a/Engine/Source/Runtime/Slate/Public/Widgets/Views/STileView.h
+++ b/Engine/Source/Runtime/Slate/Public/Widgets/Views/STileView.h
@@ -51,6 +51,8 @@ public:
 		, _SelectionMode(ESelectionMode::Multi)
 		, _ClearSelectionOnClick(true)
 		, _ExternalScrollbar()
+		, _Orientation(Orient_Vertical)
+		, _EnableAnimatedScrolling(false)
 		, _ScrollbarVisibility(EVisibility::Visible)
 		, _ScrollbarDragFocusCause(EFocusCause::Mouse)
 		, _AllowOverscroll(EAllowOverscroll::Yes)
@@ -95,6 +97,12 @@ public:
 
 		SLATE_ARGUMENT( TSharedPtr<SScrollBar>, ExternalScrollbar )
 
+		SLATE_ARGUMENT(EOrientation, Orientation)
+
+		SLATE_ARGUMENT( bool, EnableAnimatedScrolling)
+
+		SLATE_ARGUMENT( TOptional<double>, FixedLineScrollOffset )
+
 		SLATE_ATTRIBUTE(EVisibility, ScrollbarVisibility)
 
 		SLATE_ARGUMENT(EFocusCause, ScrollbarDragFocusCause)
@@ -149,6 +157,9 @@ public:
 		this->bHandleDirectionalNavigation = InArgs._HandleDirectionalNavigation;
 		this->IsFocusable = InArgs._IsFocusable;
 
+		this->bEnableAnimatedScrolling = InArgs._EnableAnimatedScrolling;
+		this->FixedLineScrollOffset = InArgs._FixedLineScrollOffset;
+
 		this->OnItemToString_Debug = InArgs._OnItemToString_Debug.IsBound()
 			? InArgs._OnItemToString_Debug
 			: SListView< ItemType >::GetDefaultDebugDelegate();
@@ -184,7 +195,7 @@ public:
 		else
 		{
 			// Make the TableView
-			this->ConstructChildren(InArgs._ItemWidth, InArgs._ItemHeight, InArgs._ItemAlignment, TSharedPtr<SHeaderRow>(), InArgs._ExternalScrollbar, InArgs._OnTileViewScrolled);
+			this->ConstructChildren(InArgs._ItemWidth, InArgs._ItemHeight, InArgs._ItemAlignment, TSharedPtr<SHeaderRow>(), InArgs._ExternalScrollbar, InArgs._Orientation, InArgs._OnTileViewScrolled);
 			if (this->ScrollBar.IsValid())
 			{
 				this->ScrollBar->SetDragFocusCause(InArgs._ScrollbarDragFocusCause);
@@ -208,29 +219,26 @@ public:
 		{
 			const TArray<ItemType>& ItemsSourceRef = (*this->ItemsSource);
 
-			const int32 NumItemsWide = GetNumItemsWide();
+			const int32 NumItemsPerLine = GetNumItemsPerLine();
 			const int32 CurSelectionIndex = (!TListTypeTraits<ItemType>::IsPtrValid(this->SelectorItem)) ? 0 : ItemsSourceRef.Find(TListTypeTraits<ItemType>::NullableItemTypeConvertToItemType(this->SelectorItem));
 			int32 AttemptSelectIndex = -1;
 
 			const EUINavigation NavType = InNavigationEvent.GetNavigationType();
-			switch (NavType)
+			if ((this->Orientation == Orient_Vertical && NavType == EUINavigation::Left) ||
+				(this->Orientation == Orient_Horizontal && NavType == EUINavigation::Up))
 			{
-			case EUINavigation::Left:
-				if (bWrapHorizontalNavigation || (CurSelectionIndex % NumItemsWide) > 0)
+				if (bWrapHorizontalNavigation || (CurSelectionIndex % NumItemsPerLine) > 0)
 				{
 					AttemptSelectIndex = CurSelectionIndex - 1;
 				}
-				break;
-
-			case EUINavigation::Right:
-				if (bWrapHorizontalNavigation || (CurSelectionIndex % NumItemsWide) < (NumItemsWide - 1))
+			}
+			else if ((this->Orientation == Orient_Vertical && NavType == EUINavigation::Right) ||
+				(this->Orientation == Orient_Horizontal && NavType == EUINavigation::Down))
+			{
+				if (bWrapHorizontalNavigation || (CurSelectionIndex % NumItemsPerLine) < (NumItemsPerLine - 1))
 				{
 					AttemptSelectIndex = CurSelectionIndex + 1;
 				}
-				break;
-
-			default:
-				break;
 			}
 
 			// If it's valid we'll scroll it into view and return an explicit widget in the FNavigationReply
@@ -252,90 +260,95 @@ public:
 		this->ClearWidgets();
 		
 		const TArray<ItemType>* SourceItems = this->ItemsSource;
-		if ( SourceItems != nullptr )
+		if (SourceItems && SourceItems->Num() > 0)
 		{
 			// Item width and height is constant by design.
-			const float AllottedWidth = MyGeometry.Size.X;
-			const float ItemWidth = this->GetItemWidth();
-			const float ItemHeight = this->GetItemHeight();
-
-			const int32 NumItemsWide = GetNumItemsWide();
-			const int32 NumItemsPaddedToFillLastRow = (SourceItems->Num() % NumItemsWide != 0)
-				? SourceItems->Num() + NumItemsWide - SourceItems->Num() % NumItemsWide
-				: SourceItems->Num();
-			const double RowsPerScreen = MyGeometry.Size.Y / ItemHeight;
-			const double EndOfListOffset = NumItemsPaddedToFillLastRow - NumItemsWide * RowsPerScreen;
-			const double ClampedScrollOffset = FMath::Clamp(STableViewBase::ScrollOffset, 0.0, EndOfListOffset);
-			const bool bAtEndOfList = (STableViewBase::ScrollOffset >= ClampedScrollOffset);
+			FTableViewDimensions TileDimensions = GetTileDimensions();
+			FTableViewDimensions AllottedDimensions(this->Orientation, MyGeometry.GetLocalSize());
+
+			const int32 NumItems = SourceItems->Num();
+			const int32 NumItemsPerLine = GetNumItemsPerLine();
+			const int32 NumItemsPaddedToFillLastLine = (NumItems % NumItemsPerLine != 0)
+				? NumItems + NumItemsPerLine - NumItems % NumItemsPerLine
+				: NumItems;
+
+			const double LinesPerScreen = AllottedDimensions.ScrollAxis / TileDimensions.ScrollAxis;
+			const double EndOfListOffset = NumItemsPaddedToFillLastLine - NumItemsPerLine * LinesPerScreen;
+			const double ClampedScrollOffset = FMath::Clamp(this->CurrentScrollOffset, 0.0, EndOfListOffset);
 			const float LayoutScaleMultiplier = MyGeometry.GetAccumulatedLayoutTransform().GetScale();
 			
 			// Once we run out of vertical and horizontal space, we stop generating widgets.
-			float WidthUsedSoFar = 0.0f;
-			float HeightUsedSoFar = 0.0f;
-			float WidgetHeightSoFar = 0.0f;
+			FTableViewDimensions DimensionsUsedSoFar(this->Orientation);
+			
 			// Index of the item at which we start generating based on how far scrolled down we are
-			int32 StartIndex = FMath::Max( 0, FMath::FloorToInt(ClampedScrollOffset / NumItemsWide) * NumItemsWide );
+			int32 StartIndex = FMath::Max( 0, FMath::FloorToInt(ClampedScrollOffset / NumItemsPerLine) * NumItemsPerLine);
 
 			// Let the WidgetGenerator know that we are starting a pass so that it can keep track of data items and widgets.
 			this->WidgetGenerator.OnBeginGenerationPass();
 
 			// Actually generate the widgets.
-			bool bKeepGenerating = true;
-			bool bNewRow = true;
-			bool bFirstRow = true;
-			double NumRowsShownOnScreen = 0;
-			for( int32 ItemIndex = StartIndex; bKeepGenerating && ItemIndex < SourceItems->Num(); ++ItemIndex )
+			bool bIsAtEndOfList = false;
+			bool bHasFilledAvailableArea = false;
+			bool bNewLine = true;
+			bool bFirstLine = true;
+			double NumLinesShownOnScreen = 0;
+			for( int32 ItemIndex = StartIndex; !bHasFilledAvailableArea && ItemIndex < NumItems; ++ItemIndex )
 			{
 				const ItemType& CurItem = (*SourceItems)[ItemIndex];
 
-				if (bNewRow)
+				if (bNewLine)
 				{
-					bNewRow = false;
-					WidgetHeightSoFar += ItemHeight;
-
-					float RowFraction = 1.0f;
-					if (bFirstRow)
+					bNewLine = false;
+					
+					float LineFraction = 1.0f;
+					if (bFirstLine)
 					{
-						bFirstRow = false;
-						RowFraction = 1.0f - FMath::Max(FMath::Fractional(ClampedScrollOffset / NumItemsWide), 0.0f);
+						bFirstLine = false;
+						LineFraction = 1.0f - FMath::Max(FMath::Fractional(ClampedScrollOffset / NumItemsPerLine), 0.0f);
 					}
 
-					HeightUsedSoFar += ItemHeight * RowFraction;
-
-					if (HeightUsedSoFar > MyGeometry.Size.Y)
+					DimensionsUsedSoFar.ScrollAxis += TileDimensions.ScrollAxis * LineFraction;
+					
+					if (DimensionsUsedSoFar.ScrollAxis > AllottedDimensions.ScrollAxis)
 					{
-						NumRowsShownOnScreen += FMath::Max(1.0f - ((HeightUsedSoFar - MyGeometry.Size.Y) / ItemHeight), 0.0f);
+						NumLinesShownOnScreen += FMath::Max(1.0f - ((DimensionsUsedSoFar.ScrollAxis - AllottedDimensions.ScrollAxis) / TileDimensions.ScrollAxis), 0.0f);
 					}
 					else
 					{
-						NumRowsShownOnScreen += RowFraction;
+						NumLinesShownOnScreen += LineFraction;
 					}
 				}
 
-				const float GeneratedItemHeight = SListView<ItemType>::GenerateWidgetForItem(CurItem, ItemIndex, StartIndex, LayoutScaleMultiplier);
+				SListView<ItemType>::GenerateWidgetForItem(CurItem, ItemIndex, StartIndex, LayoutScaleMultiplier);
 
-				// The widget used up some of the available horizontal space.
-				WidthUsedSoFar += ItemWidth;
+				// The widget used up some of the available space for the current line
+				DimensionsUsedSoFar.LineAxis += TileDimensions.LineAxis;
 
-				// Check to see if we have completed a row of widgets
-				if ( WidthUsedSoFar + ItemWidth > AllottedWidth )
+				// Check to see if we have completed a line of widgets
+				if (DimensionsUsedSoFar.LineAxis + TileDimensions.LineAxis > AllottedDimensions.LineAxis)
 				{
-					// A new row of widgets was added.
-					WidthUsedSoFar = 0;
-					bNewRow = true;
+					// A new line of widgets was completed - time to start another one
+					DimensionsUsedSoFar.LineAxis = 0;
+					bNewLine = true;
 
-					// Stop when we've generated a widget that's partially clipped by the bottom of the list.
-					if ( HeightUsedSoFar >= MyGeometry.Size.Y )
+					// Stop when we've finished generating a line of widgets partially clipped by the end of the view
+					if (DimensionsUsedSoFar.ScrollAxis >= AllottedDimensions.ScrollAxis)
 					{
-						bKeepGenerating = false;
+						bHasFilledAvailableArea = true;
 					}
 				}
+
+				if (ItemIndex >= NumItems - 1)
+				{
+					bIsAtEndOfList = true;
+				}
 			}
 
 			// We have completed the generation pass. The WidgetGenerator will clean up unused Widgets.
 			this->WidgetGenerator.OnEndGenerationPass();
 
-			return STableViewBase::FReGenerateResults(ClampedScrollOffset, WidgetHeightSoFar, NumRowsShownOnScreen, bAtEndOfList);
+			const float TotalGeneratedLineAxisSize = FMath::CeilToFloat(NumLinesShownOnScreen) * TileDimensions.ScrollAxis;
+			return STableViewBase::FReGenerateResults(ClampedScrollOffset, TotalGeneratedLineAxisSize, NumLinesShownOnScreen, bIsAtEndOfList);
 		}
 
 		return STableViewBase::FReGenerateResults(0, 0, 0, false);
@@ -345,13 +358,13 @@ public:
 	virtual int32 GetNumItemsBeingObserved() const override
 	{
 		const int32 NumItemsBeingObserved = this->ItemsSource == nullptr ? 0 : this->ItemsSource->Num();
-		const int32 NumItemsWide = GetNumItemsWide();
+		const int32 NumItemsPerLine = GetNumItemsPerLine();
 		
 		int32 NumEmptySpacesAtEnd = 0;
-		if ( NumItemsWide > 0 )
+		if ( NumItemsPerLine > 0 )
 		{
-			NumEmptySpacesAtEnd = NumItemsWide - (NumItemsBeingObserved % NumItemsWide);
-			if ( NumEmptySpacesAtEnd >= NumItemsWide )
+			NumEmptySpacesAtEnd = NumItemsPerLine - (NumItemsBeingObserved % NumItemsPerLine);
+			if ( NumEmptySpacesAtEnd >= NumItemsPerLine )
 			{
 				NumEmptySpacesAtEnd = 0;
 			}
@@ -362,17 +375,18 @@ public:
 
 protected:
 
-	virtual float ScrollBy( const FGeometry& MyGeometry, float ScrollByAmountInSlateUnits, EAllowOverscroll InAllowOverscroll ) override
+	FTableViewDimensions GetTileDimensions() const
 	{
-		// Working around a CLANG bug, where all the base class
-		// members require an explicit namespace resolution.
-		typedef STableViewBase S;
+		return FTableViewDimensions(this->Orientation, this->GetItemWidth(), this->GetItemHeight());
+	}
 
-		const bool bWholeListVisible = S::ScrollOffset == 0 && S::bWasAtEndOfList;
-		if ( InAllowOverscroll == EAllowOverscroll::Yes && S::Overscroll.ShouldApplyOverscroll( S::ScrollOffset == 0, S::bWasAtEndOfList, ScrollByAmountInSlateUnits ) )
+	virtual float ScrollBy( const FGeometry& MyGeometry, float ScrollByAmountInSlateUnits, EAllowOverscroll InAllowOverscroll ) override
+	{
+		const bool bWholeListVisible = this->DesiredScrollOffset == 0 && this->bWasAtEndOfList;
+		if (InAllowOverscroll == EAllowOverscroll::Yes && this->Overscroll.ShouldApplyOverscroll(this->DesiredScrollOffset == 0, this->bWasAtEndOfList, ScrollByAmountInSlateUnits))
 		{
-			const float UnclampedScrollDelta = ScrollByAmountInSlateUnits / GetNumItemsWide();
-			const float ActuallyScrolledBy = S::Overscroll.ScrollBy(MyGeometry, UnclampedScrollDelta);
+			const float UnclampedScrollDelta = ScrollByAmountInSlateUnits / GetNumItemsPerLine();
+			const float ActuallyScrolledBy = this->Overscroll.ScrollBy(MyGeometry, UnclampedScrollDelta);
 			if (ActuallyScrolledBy != 0.0f)
 			{
 				this->RequestListRefresh();
@@ -381,27 +395,28 @@ protected:
 		}
 		else if (!bWholeListVisible)
 		{
-			const float ItemHeight = this->GetItemHeight();
-			const double NewScrollOffset = this->ScrollOffset + ( ( ScrollByAmountInSlateUnits * GetNumItemsWide() ) / ItemHeight );
+			const double NewScrollOffset = this->DesiredScrollOffset + ((ScrollByAmountInSlateUnits * GetNumItemsPerLine()) / GetTileDimensions().ScrollAxis);
 
 			return this->ScrollTo( NewScrollOffset );
 		}
 
-		return 0;
+		return 0.f;
 	}
 
-	virtual int32 GetNumItemsWide() const override
+	virtual int32 GetNumItemsPerLine() const override
 	{
-		const float ItemWidth = this->GetItemWidth();
-		const int32 NumItemsWide = ItemWidth > 0 ? FMath::FloorToInt(this->PanelGeometryLastTick.GetLocalSize().X / ItemWidth) : 1;
-		return FMath::Max(1, NumItemsWide);
+		FTableViewDimensions PanelDimensions(this->Orientation, this->PanelGeometryLastTick.GetLocalSize());
+		FTableViewDimensions TileDimensions = GetTileDimensions();
+
+		const int32 NumItemsPerLine = TileDimensions.LineAxis > 0 ? FMath::FloorToInt(PanelDimensions.LineAxis / TileDimensions.LineAxis) : 1;
+		return FMath::Max(1, NumItemsPerLine);
 	}
 
 	/**
-	* If there is a pending request to scroll an item into view, do so.
-	*
-	* @param ListViewGeometry  The geometry of the listView; can be useful for centering the item.
-	*/
+	 * If there is a pending request to scroll an item into view, do so.
+	 *
+	 * @param ListViewGeometry  The geometry of the listView; can be useful for centering the item.
+	 */
 	virtual typename SListView<ItemType>::EScrollIntoViewResult ScrollIntoView(const FGeometry& ListViewGeometry) override
 	{
 		if (TListTypeTraits<ItemType>::IsPtrValid(this->ItemToScrollIntoView) && this->ItemsSource != nullptr)
@@ -409,12 +424,13 @@ protected:
 			const int32 IndexOfItem = this->ItemsSource->Find(TListTypeTraits<ItemType>::NullableItemTypeConvertToItemType(this->ItemToScrollIntoView));
 			if (IndexOfItem != INDEX_NONE)
 			{
-				const float NumItemsHigh = ListViewGeometry.GetLocalSize().Y / this->GetItemHeight();
+				const float NumLinesInView = FTableViewDimensions(this->Orientation, ListViewGeometry.GetLocalSize()).ScrollAxis / GetTileDimensions().ScrollAxis;
+
 				float NumLiveWidgets = this->GetNumLiveWidgets();
 				if (NumLiveWidgets == 0 && this->IsPendingRefresh())
 				{
 					// Use the last number of widgets on screen to estimate if we actually need to scroll.
-					NumLiveWidgets = this->LastGenerateResults.ExactNumRowsOnScreen;
+					NumLiveWidgets = this->LastGenerateResults.ExactNumLinesOnScreen;
 
 					// If we still don't have any widgets, we're not in a situation where we can scroll an item into view
 					// (probably as nothing has been generated yet), so we'll defer this again until the next frame
@@ -427,20 +443,43 @@ protected:
 				this->EndInertialScrolling();
 
 				// Only scroll the item into view if it's not already in the visible range
-				const int32 NumItemsWide = GetNumItemsWide();
-				const float Index = (IndexOfItem / NumItemsWide) * NumItemsWide;
-				const float IndexPlusOne = ((IndexOfItem / NumItemsWide) + 1) * NumItemsWide;
-				if (Index < this->ScrollOffset || IndexPlusOne > (this->ScrollOffset + NumItemsHigh * NumItemsWide))
+				const int32 NumItemsPerLine = GetNumItemsPerLine();
+				const double ScrollLineOffset = this->GetTargetScrollOffset() / NumItemsPerLine;
+				const int32 LineOfItem = FMath::FloorToInt(IndexOfItem / NumItemsPerLine);
+				const int32 NumFullLinesInView = FMath::FloorToInt(ScrollLineOffset + NumLinesInView) - FMath::CeilToInt(ScrollLineOffset);
+				
+				const double MinDisplayedLine = this->bNavigateOnScrollIntoView ? FMath::FloorToDouble(ScrollLineOffset) : FMath::CeilToDouble(ScrollLineOffset);
+				const double MaxDisplayedLine = this->bNavigateOnScrollIntoView ? FMath::CeilToDouble(ScrollLineOffset + NumFullLinesInView) : FMath::FloorToDouble(ScrollLineOffset + NumFullLinesInView);
+
+				if (LineOfItem < MinDisplayedLine || LineOfItem > MaxDisplayedLine)
 				{
-					// Scroll the top of the listview to the item in question
-					float NewScrollOffset = Index;
-					// Center the list view on the item in question.
-					NewScrollOffset -= ((NumItemsHigh - 1.0f) * NumItemsWide * 0.5f);
+					// Set the line with the item at the beginning of the view area
+					float NewLineOffset = LineOfItem;
+					// Center the line in the view area
+					NewLineOffset -= NumLinesInView * 0.5f;
+					// Convert the line offset into an item offset
+					double NewScrollOffset = NewLineOffset * NumItemsPerLine;
 					// And clamp the scroll offset within the allowed limits
-					NewScrollOffset = FMath::Clamp(NewScrollOffset, 0.0f, GetNumItemsBeingObserved() - NumItemsWide * NumItemsHigh);
+					NewScrollOffset = FMath::Clamp(NewScrollOffset, 0., (double)(GetNumItemsBeingObserved() - NumItemsPerLine * NumLinesInView));
 
 					this->SetScrollOffset(NewScrollOffset);
 				}
+				else if (this->bNavigateOnScrollIntoView)
+				{
+					// Make sure the line containing the existing entry for this item is fully in view
+					if (LineOfItem == MinDisplayedLine)
+					{
+						// This line is clipped at the top/left, so set it as the new offset
+						this->SetScrollOffset(LineOfItem * NumItemsPerLine - (this->FixedLineScrollOffset.IsSet() && LineOfItem > 0 ? 0.f : this->NavigationScrollOffset));
+					}
+					else if (LineOfItem == MaxDisplayedLine)
+					{
+						// This line is clipped at the end, so we need to advance just enough to bring it fully into view
+						// Since all tiles are required to be of the same size, this is straightforward
+						const float NewLineOffset = LineOfItem - NumLinesInView + 1.f + (this->FixedLineScrollOffset.IsSet() ? 0.f : this->NavigationScrollOffset);
+						this->SetScrollOffset(NewLineOffset * NumItemsPerLine);
+					}
+				}
 
 				this->RequestListRefresh();
 
diff --git a/Engine/Source/Runtime/Slate/Public/Widgets/Views/STreeView.h b/Engine/Source/Runtime/Slate/Public/Widgets/Views/STreeView.h
--- a/Engine/Source/Runtime/Slate/Public/Widgets/Views/STreeView.h
+++ b/Engine/Source/Runtime/Slate/Public/Widgets/Views/STreeView.h
@@ -124,6 +124,7 @@ public:
 		, _SelectionMode(ESelectionMode::Multi)
 		, _ClearSelectionOnClick(true)
 		, _ExternalScrollbar()
+		, _EnableAnimatedScrolling(false)
 		, _ScrollbarDragFocusCause(EFocusCause::Mouse)
 		, _ConsumeMouseWheel( EConsumeMouseWheel::WhenScrollingPossible )
 		, _AllowOverscroll(EAllowOverscroll::Yes)
@@ -172,6 +173,10 @@ public:
 
 		SLATE_ARGUMENT( TSharedPtr<SScrollBar>, ExternalScrollbar )
 
+		SLATE_ARGUMENT( bool, EnableAnimatedScrolling)
+
+		SLATE_ARGUMENT( TOptional<double>, FixedLineScrollOffset )
+
 		SLATE_ARGUMENT( EFocusCause, ScrollbarDragFocusCause )
 
 		SLATE_ARGUMENT( EConsumeMouseWheel, ConsumeMouseWheel );
@@ -223,6 +228,9 @@ public:
 
 		this->WheelScrollMultiplier = InArgs._WheelScrollMultiplier;
 
+		this->bEnableAnimatedScrolling = InArgs._EnableAnimatedScrolling;
+		this->FixedLineScrollOffset = InArgs._FixedLineScrollOffset;
+
 		this->OnItemToString_Debug = InArgs._OnItemToString_Debug.IsBound()
 			? InArgs._OnItemToString_Debug
 			: SListView< ItemType >::GetDefaultDebugDelegate();
@@ -266,7 +274,7 @@ public:
 		else
 		{
 			// Make the TableView
-			this->ConstructChildren( 0, InArgs._ItemHeight, EListItemAlignment::LeftAligned, InArgs._HeaderRow, InArgs._ExternalScrollbar, InArgs._OnTreeViewScrolled );
+			this->ConstructChildren( 0, InArgs._ItemHeight, EListItemAlignment::LeftAligned, InArgs._HeaderRow, InArgs._ExternalScrollbar, Orient_Vertical, InArgs._OnTreeViewScrolled );
 			if (this->ScrollBar.IsValid())
 			{
 				this->ScrollBar->SetDragFocusCause(InArgs._ScrollbarDragFocusCause);
diff --git a/Engine/Source/Runtime/UMG/Private/Components/ListView.cpp b/Engine/Source/Runtime/UMG/Private/Components/ListView.cpp
--- a/Engine/Source/Runtime/UMG/Private/Components/ListView.cpp
+++ b/Engine/Source/Runtime/UMG/Private/Components/ListView.cpp
@@ -108,8 +108,7 @@ UObject* UListView::BP_GetSelectedItem() const
 
 void UListView::HandleOnEntryInitializedInternal(UObject* Item, const TSharedRef<ITableRow>& TableRow)
 {
-	UUserWidget* const RowWidget = GetEntryWidgetFromItem(Item);
-	BP_OnEntryInitialized.Broadcast(Item, RowWidget);
+	BP_OnEntryInitialized.Broadcast(Item, GetEntryWidgetFromItem(Item));
 }
 
 bool UListView::BP_GetSelectedItems(TArray<UObject*>& Items) const
diff --git a/Engine/Source/Runtime/UMG/Private/Components/ListViewBase.cpp b/Engine/Source/Runtime/UMG/Private/Components/ListViewBase.cpp
--- a/Engine/Source/Runtime/UMG/Private/Components/ListViewBase.cpp
+++ b/Engine/Source/Runtime/UMG/Private/Components/ListViewBase.cpp
@@ -62,6 +62,15 @@ void UListViewBase::ScrollToBottom()
 	}
 }
 
+void UListViewBase::SetWheelScrollMultiplier(float NewWheelScrollMultiplier)
+{
+	WheelScrollMultiplier = NewWheelScrollMultiplier;
+	if (MyTableViewBase)
+	{
+		MyTableViewBase->SetWheelScrollMultiplier(GetGlobalScrollAmount() * NewWheelScrollMultiplier);
+	}
+}
+
 void UListViewBase::SetScrollbarVisibility(ESlateVisibility InVisibility)
 {
 	if (MyTableViewBase)
@@ -104,6 +113,10 @@ TSharedRef<SWidget> UListViewBase::RebuildWidget()
 	}
 
 	MyTableViewBase = RebuildListWidget();
+	MyTableViewBase->SetIsScrollAnimationEnabled(bEnableScrollAnimation);
+	MyTableViewBase->SetFixedLineScrollOffset(bEnableFixedLineOffset ? TOptional<double>(FixedLineScrollOffset) : TOptional<double>());
+	MyTableViewBase->SetWheelScrollMultiplier(GetGlobalScrollAmount() * WheelScrollMultiplier);
+
 	return MyTableViewBase.ToSharedRef();
 }
 
@@ -120,6 +133,13 @@ void UListViewBase::SynchronizeProperties()
 {
 	Super::SynchronizeProperties();
 
+	if (MyTableViewBase)
+	{
+		MyTableViewBase->SetIsScrollAnimationEnabled(bEnableScrollAnimation);
+		MyTableViewBase->SetFixedLineScrollOffset(bEnableFixedLineOffset ? TOptional<double>(FixedLineScrollOffset) : TOptional<double>());
+		MyTableViewBase->SetWheelScrollMultiplier(GetGlobalScrollAmount() * WheelScrollMultiplier);
+	}
+
 #if WITH_EDITORONLY_DATA
 	if (IsDesignTime() && MyTableViewBase.IsValid())
 	{
diff --git a/Engine/Source/Runtime/UMG/Private/Components/TileView.cpp b/Engine/Source/Runtime/UMG/Private/Components/TileView.cpp
--- a/Engine/Source/Runtime/UMG/Private/Components/TileView.cpp
+++ b/Engine/Source/Runtime/UMG/Private/Components/TileView.cpp
@@ -9,7 +9,6 @@ UTileView::UTileView(const FObjectInitializer& ObjectInitializer)
 	: Super(ObjectInitializer)
 {
 	bIsVariable = true;
-	EntryHeight = 128.f;
 }
 
 TSharedRef<STableViewBase> UTileView::RebuildListWidget()
diff --git a/Engine/Source/Runtime/UMG/Public/Components/ListView.h b/Engine/Source/Runtime/UMG/Public/Components/ListView.h
--- a/Engine/Source/Runtime/UMG/Public/Components/ListView.h
+++ b/Engine/Source/Runtime/UMG/Public/Components/ListView.h
@@ -132,19 +132,35 @@ protected:
 	virtual void OnSelectionChangedInternal(UObject* FirstSelectedItem) override;
 	virtual void OnItemScrolledIntoViewInternal(UObject* Item, UUserWidget& EntryWidget) override;
 
+	void HandleOnEntryInitializedInternal(UObject* Item, const TSharedRef<ITableRow>& TableRow);
+
 	/** SListView construction helper - useful if using a custom STreeView subclass */
 	template <template<typename> class ListViewT = SListView>
 	TSharedRef<ListViewT<UObject*>> ConstructListView()
 	{
-		MyListView = ITypedUMGListView<UObject*>::ConstructListView<ListViewT>(this, ListItems, bIsFocusable, SelectionMode, bClearSelectionOnClick, ConsumeMouseWheel, bReturnFocusToSelection);
-		if (MyListView)
-		{
-			MyListView->SetOnEntryInitialized(SListView<UObject*>::FOnEntryInitialized::CreateUObject(this, &ThisClass::HandleOnEntryInitializedInternal));
-		}
+		FListViewConstructArgs Args;
+		Args.bAllowFocus = bIsFocusable;
+		Args.SelectionMode = SelectionMode;
+		Args.bClearSelectionOnClick = bClearSelectionOnClick;
+		Args.ConsumeMouseWheel = ConsumeMouseWheel;
+		Args.bReturnFocusToSelection = bReturnFocusToSelection;
+		Args.Orientation = Orientation;
+		MyListView = ITypedUMGListView<UObject*>::ConstructListView<ListViewT>(this, ListItems, Args);
+		
+		MyListView->SetOnEntryInitialized(SListView<UObject*>::FOnEntryInitialized::CreateUObject(this, &UListView::HandleOnEntryInitializedInternal));
+
 		return StaticCastSharedRef<ListViewT<UObject*>>(MyListView.ToSharedRef());
 	}
 
 protected:
+	/** 
+	 * The scroll & layout orientation of the list. ListView and TileView only. 
+	 * Vertical will scroll vertically and arrange tiles into rows.
+	 * Horizontal will scroll horizontally and arrange tiles into columns.
+	 */
+	UPROPERTY(EditAnywhere, Category = ListView)
+	TEnumAsByte<EOrientation> Orientation;
+
 	UPROPERTY(EditAnywhere, BlueprintReadOnly, Category = ListView)
 	TEnumAsByte<ESelectionMode::Type> SelectionMode = ESelectionMode::Single;
 
@@ -215,8 +231,6 @@ private:
 	UFUNCTION(BlueprintCallable, Category = ListView, meta = (DisplayName = "Get Selected Item", AllowPrivateAccess = true))
 	UObject* BP_GetSelectedItem() const;
 
-	void HandleOnEntryInitializedInternal(UObject* Item, const TSharedRef<ITableRow>& TableRow);
-
 private:
 	/** Called when a row widget is generated for a list item */
 	UPROPERTY(BlueprintAssignable, Category = Events, meta = (DisplayName = "On Entry Initialized"))
diff --git a/Engine/Source/Runtime/UMG/Public/Components/ListViewBase.h b/Engine/Source/Runtime/UMG/Public/Components/ListViewBase.h
--- a/Engine/Source/Runtime/UMG/Public/Components/ListViewBase.h
+++ b/Engine/Source/Runtime/UMG/Public/Components/ListViewBase.h
@@ -209,24 +209,31 @@ protected:
 	 * Use these instead of SNew-ing your owned ListView directly to get exposed events for free
 	 */
 
+	struct FListViewConstructArgs
+	{
+		bool bAllowFocus = true;
+		ESelectionMode::Type SelectionMode = ESelectionMode::Single;
+		bool bClearSelectionOnClick = false;
+		EConsumeMouseWheel ConsumeMouseWheel = EConsumeMouseWheel::WhenScrollingPossible;
+		bool bReturnFocusToSelection = false;
+		EOrientation Orientation = Orient_Vertical;
+	};
+
 	template <template<typename> class ListViewT = SListView, typename UListViewBaseT>
-	static TSharedRef<ListViewT<ItemType>> ConstructListView(UListViewBaseT* Implementer,
+	static TSharedRef<ListViewT<ItemType>> ConstructListView(UListViewBaseT* Implementer, 
 		const TArray<ItemType>& ListItems,
-		bool bAllowFocus = true,
-		ESelectionMode::Type SelectionMode = ESelectionMode::Single,
-		bool bClearSelectionOnClick = false,
-		EConsumeMouseWheel ConsumeMouseWheel = EConsumeMouseWheel::WhenScrollingPossible,
-		bool bReturnFocusToSelection = false)
+		const FListViewConstructArgs& Args = FListViewConstructArgs())
 	{
 		static_assert(TIsDerivedFrom<ListViewT<ItemType>, SListView<ItemType>>::IsDerived, "ConstructListView can only construct instances of SListView classes");
 		return SNew(ListViewT<ItemType>)
 			.HandleGamepadEvents(true)
 			.ListItemsSource(&ListItems)
-			.IsFocusable(bAllowFocus)
-			.ClearSelectionOnClick(bClearSelectionOnClick)
-			.ConsumeMouseWheel(ConsumeMouseWheel)
-			.SelectionMode(SelectionMode)
-			.ReturnFocusToSelection(bReturnFocusToSelection)
+			.IsFocusable(Args.bAllowFocus)
+			.ClearSelectionOnClick(Args.bClearSelectionOnClick)
+			.ConsumeMouseWheel(Args.ConsumeMouseWheel)
+			.SelectionMode(Args.SelectionMode)
+			.ReturnFocusToSelection(Args.bReturnFocusToSelection)
+			.Orientation(Args.Orientation)
 			.OnGenerateRow_UObject(Implementer, &UListViewBaseT::HandleGenerateRow)
 			.OnSelectionChanged_UObject(Implementer, &UListViewBaseT::HandleSelectionChanged)
 			.OnRowReleased_UObject(Implementer, &UListViewBaseT::HandleRowReleased)
@@ -236,28 +243,32 @@ protected:
 			.OnMouseButtonDoubleClick_UObject(Implementer, &UListViewBaseT::HandleItemDoubleClicked);
 	}
 
+	struct FTileViewConstructArgs : public FListViewConstructArgs
+	{
+		EListItemAlignment TileAlignment = EListItemAlignment::EvenlyDistributed;
+		TAttribute<float> EntryHeight;
+		TAttribute<float> EntryWidth;
+		bool bWrapDirectionalNavigation = false;
+	};
+
 	template <template<typename> class TileViewT = STileView, typename UListViewBaseT>
 	static TSharedRef<TileViewT<ItemType>> ConstructTileView(UListViewBaseT* Implementer,
 		const TArray<ItemType>& ListItems,
-		EListItemAlignment TileAlignment,
-		TAttribute<float> TileHeight,
-		TAttribute<float> TileWidth,
-		ESelectionMode::Type SelectionMode = ESelectionMode::Single,
-		bool bClearSelectionOnClick = false,
-		bool bWrapHorizontalNavigation = false,
-		EConsumeMouseWheel ConsumeMouseWheel = EConsumeMouseWheel::WhenScrollingPossible)
+		const FTileViewConstructArgs& Args = FTileViewConstructArgs())
 	{
 		static_assert(TIsDerivedFrom<TileViewT<ItemType>, STileView<ItemType>>::IsDerived, "ConstructTileView can only construct instances of STileView classes");
 		return SNew(TileViewT<ItemType>)
 			.HandleGamepadEvents(true)
 			.ListItemsSource(&ListItems)
-			.ClearSelectionOnClick(bClearSelectionOnClick)
-			.WrapHorizontalNavigation(bWrapHorizontalNavigation)
-			.ConsumeMouseWheel(ConsumeMouseWheel)
-			.SelectionMode(SelectionMode)
-			.ItemHeight(TileHeight)
-			.ItemWidth(TileWidth)
-			.ItemAlignment(TileAlignment)
+			.IsFocusable(Args.bAllowFocus)
+			.ClearSelectionOnClick(Args.bClearSelectionOnClick)
+			.WrapHorizontalNavigation(Args.bWrapDirectionalNavigation)
+			.ConsumeMouseWheel(Args.ConsumeMouseWheel)
+			.SelectionMode(Args.SelectionMode)
+			.ItemHeight(Args.EntryHeight)
+			.ItemWidth(Args.EntryWidth)
+			.ItemAlignment(Args.TileAlignment)
+			.Orientation(Args.Orientation)
 			.OnGenerateTile_UObject(Implementer, &UListViewBaseT::HandleGenerateRow)
 			.OnTileReleased_UObject(Implementer, &UListViewBaseT::HandleRowReleased)
 			.OnSelectionChanged_UObject(Implementer, &UListViewBaseT::HandleSelectionChanged)
@@ -267,20 +278,25 @@ protected:
 			.OnMouseButtonDoubleClick_UObject(Implementer, &UListViewBaseT::HandleItemDoubleClicked);
 	}
 
+	struct FTreeViewConstructArgs
+	{
+		ESelectionMode::Type SelectionMode = ESelectionMode::Single;
+		bool bClearSelectionOnClick = false;
+		EConsumeMouseWheel ConsumeMouseWheel = EConsumeMouseWheel::WhenScrollingPossible;
+	};
+
 	template <template<typename> class TreeViewT = STreeView, typename UListViewBaseT>
 	static TSharedRef<TreeViewT<ItemType>> ConstructTreeView(UListViewBaseT* Implementer,
 		const TArray<ItemType>& ListItems,
-		ESelectionMode::Type SelectionMode = ESelectionMode::Single,
-		bool bClearSelectionOnClick = false,
-		EConsumeMouseWheel ConsumeMouseWheel = EConsumeMouseWheel::WhenScrollingPossible)
+		const FTreeViewConstructArgs& Args = FTreeViewConstructArgs())
 	{
 		static_assert(TIsDerivedFrom<TreeViewT<ItemType>, STreeView<ItemType>>::IsDerived, "ConstructTreeView can only construct instances of STreeView classes");
 		return SNew(TreeViewT<ItemType>)
 			.HandleGamepadEvents(true)
 			.TreeItemsSource(&ListItems)
-			.ClearSelectionOnClick(bClearSelectionOnClick)
-			.ConsumeMouseWheel(ConsumeMouseWheel)
-			.SelectionMode(SelectionMode)
+			.ClearSelectionOnClick(Args.bClearSelectionOnClick)
+			.ConsumeMouseWheel(Args.ConsumeMouseWheel)
+			.SelectionMode(Args.SelectionMode)
 			.OnGenerateRow_UObject(Implementer, &UListViewBaseT::HandleGenerateRow)
 			.OnSelectionChanged_UObject(Implementer, &UListViewBaseT::HandleSelectionChanged)
 			.OnRowReleased_UObject(Implementer, &UListViewBaseT::HandleRowReleased)
@@ -292,6 +308,71 @@ protected:
 			.OnExpansionChanged_UObject(Implementer, &UListViewBaseT::HandleExpansionChanged);
 	}
 
+	template <template<typename> class ListViewT = SListView, typename UListViewBaseT>
+	UE_DEPRECATED(4.23, "Use ConstructListView with FListViewConstructArgs instead of individual arguments.")
+	static TSharedRef<ListViewT<ItemType>> ConstructListView(UListViewBaseT* Implementer,
+		const TArray<ItemType>& ListItems,
+		bool bAllowFocus,
+		ESelectionMode::Type SelectionMode = ESelectionMode::Single,
+		bool bClearSelectionOnClick = false,
+		EConsumeMouseWheel ConsumeMouseWheel = EConsumeMouseWheel::WhenScrollingPossible,
+		bool bReturnFocusToSelection = false,
+		EOrientation Orientation = Orient_Vertical)
+	{
+		FListViewConstructArgs Args;
+		Args.bAllowFocus = bAllowFocus;
+		Args.SelectionMode = SelectionMode;
+		Args.bClearSelectionOnClick = bClearSelectionOnClick;
+		Args.ConsumeMouseWheel = ConsumeMouseWheel;
+		Args.bReturnFocusToSelection = bReturnFocusToSelection;
+		Args.Orientation = Orientation;
+
+		return ConstructListView<ListViewT>(Implementer, ListItems, Args);
+	}
+
+	template <template<typename> class TileViewT = STileView, typename UListViewBaseT>
+	UE_DEPRECATED(4.23, "Use ConstructTileView with FTileViewConstructArgs instead of individual arguments.")
+	static TSharedRef<TileViewT<ItemType>> ConstructTileView(UListViewBaseT* Implementer,
+		const TArray<ItemType>& ListItems,
+		EListItemAlignment TileAlignment,
+		TAttribute<float> TileHeight,
+		TAttribute<float> TileWidth,
+		ESelectionMode::Type SelectionMode = ESelectionMode::Single,
+		bool bClearSelectionOnClick = false,
+		bool bWrapDirectionalNavigation = false,
+		EConsumeMouseWheel ConsumeMouseWheel = EConsumeMouseWheel::WhenScrollingPossible,
+		EOrientation Orientation = Orient_Vertical)
+	{
+		FTileViewConstructArgs Args;
+		Args.SelectionMode = SelectionMode;
+		Args.bClearSelectionOnClick = bClearSelectionOnClick;
+		Args.ConsumeMouseWheel = ConsumeMouseWheel;
+		Args.Orientation = Orientation;
+		Args.bWrapDirectionalNavigation = bWrapDirectionalNavigation;
+		Args.TileHeight = TileHeight;
+		Args.TileWidth = TileWidth;
+		Args.TileAlignment = TileAlignment;
+
+		return ConstructTileView<TileViewT>(Implementer, ListItems, Args);
+	}
+
+	template <template<typename> class TreeViewT = STreeView, typename UListViewBaseT>
+	UE_DEPRECATED(4.23, "Use ConstructTreeView with FTreeViewConstructArgs instead of individual arguments.")
+	static TSharedRef<TreeViewT<ItemType>> ConstructTreeView(UListViewBaseT* Implementer,
+		const TArray<ItemType>& ListItems,
+		ESelectionMode::Type SelectionMode,
+		bool bClearSelectionOnClick = false,
+		EConsumeMouseWheel ConsumeMouseWheel = EConsumeMouseWheel::WhenScrollingPossible)
+	{
+		FTreeViewConstructArgs Args;
+		Args.SelectionMode = SelectionMode;
+		Args.bClearSelectionOnClick = bClearSelectionOnClick;
+		Args.ConsumeMouseWheel = ConsumeMouseWheel;
+
+		return ConstructTreeView<TreeViewT>(Implementer, ListItems, Args);
+	}
+
+protected:
 	/** Gets the SObjectTableRow underlying the UMG EntryWidget that represents the given item (if one exists) */
 	template <template<typename> class ObjectRowT = SObjectTableRow>
 	TSharedPtr<ObjectRowT<ItemType>> GetObjectRowFromItem(const ItemType& Item) const
@@ -453,6 +534,12 @@ public:
 	UFUNCTION(BlueprintCallable, Category = ListViewBase)
 	void ScrollToBottom();
 
+	UFUNCTION(BlueprintCallable, Category = ListViewBase)
+	void SetWheelScrollMultiplier(float NewWheelScrollMultiplier);
+
+	UFUNCTION(BlueprintCallable, Category = ListViewBase)
+	void SetScrollbarVisibility(ESlateVisibility InVisibility);
+
 	/**
 	 * Sets the list to refresh on the next tick.
 	 *
@@ -471,8 +558,6 @@ public:
 	DECLARE_EVENT_OneParam(UListView, FOnEntryWidgetReleased, UUserWidget&);
 	FOnEntryWidgetReleased& OnEntryWidgetReleased() { return OnEntryWidgetReleasedEvent; }
 
-	void SetScrollbarVisibility(ESlateVisibility InVisibility);
-
 protected:
 	virtual TSharedRef<SWidget> RebuildWidget() override final;
 	virtual void ReleaseSlateResources(bool bReleaseChildren) override;
@@ -565,6 +650,25 @@ protected:
 	UPROPERTY(EditAnywhere, BlueprintReadOnly, Category = ListEntries, meta = (DesignerRebuild, AllowPrivateAccess = true, MustImplement = UserListEntry))
 	TSubclassOf<UUserWidget> EntryWidgetClass;
 
+	/** The multiplier to apply when wheel scrolling */
+	UPROPERTY(EditAnywhere, BlueprintReadOnly, Category = Scrolling)
+	float WheelScrollMultiplier = 1.f;
+
+	/** True to enable lerped animation when scrolling through the list */
+	UPROPERTY(EditAnywhere, BlueprintReadOnly, Category = Scrolling)
+	bool bEnableScrollAnimation = false;
+
+	UPROPERTY()
+	bool bEnableFixedLineOffset = false;
+
+	/** 
+	 * Optional fixed offset (in lines) to always apply to the top/left (depending on orientation) of the list.
+	 * If provided, all non-inertial means of scrolling will settle with exactly this offset of the topmost entry.
+	 * Ex: A value of 0.25 would cause the topmost full entry to be offset down by a quarter length of the preceeding entry.
+	 */
+	UPROPERTY(EditAnywhere, Category = Scrolling, meta = (EditCondition = bEnableFixedLineOffset, ClampMin = 0.0f, ClampMax = 0.5f))
+	float FixedLineScrollOffset = 0.f;
+
 private:
 	void FinishGeneratingEntry(UUserWidget& GeneratedEntry);
 	void HandleAnnounceGeneratedEntries();
diff --git a/Engine/Source/Runtime/UMG/Public/Components/TileView.h b/Engine/Source/Runtime/UMG/Public/Components/TileView.h
--- a/Engine/Source/Runtime/UMG/Public/Components/TileView.h
+++ b/Engine/Source/Runtime/UMG/Public/Components/TileView.h
@@ -38,7 +38,20 @@ protected:
 	template <template<typename> class TileViewT = STileView>
 	TSharedRef<TileViewT<UObject*>> ConstructTileView()
 	{
-		MyListView = MyTileView = ITypedUMGListView<UObject*>::ConstructTileView<TileViewT>(this, ListItems, TileAlignment, EntryHeight, EntryWidth, SelectionMode, bClearSelectionOnClick, bWrapHorizontalNavigation, ConsumeMouseWheel);
+		FTileViewConstructArgs Args;
+		Args.bAllowFocus = bIsFocusable;
+		Args.SelectionMode = SelectionMode;
+		Args.bClearSelectionOnClick = bClearSelectionOnClick;
+		Args.ConsumeMouseWheel = ConsumeMouseWheel;
+		Args.bReturnFocusToSelection = bReturnFocusToSelection;
+		Args.TileAlignment = TileAlignment;
+		Args.EntryHeight = EntryHeight;
+		Args.EntryWidth = EntryWidth;
+		Args.bWrapDirectionalNavigation = bWrapHorizontalNavigation;
+		Args.Orientation = Orientation;
+
+		MyListView = MyTileView = ITypedUMGListView<UObject*>::ConstructTileView<TileViewT>(this, ListItems, Args);
+		MyTileView->SetOnEntryInitialized(SListView<UObject*>::FOnEntryInitialized::CreateUObject(this, &UTileView::HandleOnEntryInitializedInternal));
 		return StaticCastSharedRef<TileViewT<UObject*>>(MyTileView.ToSharedRef());
 	}
 
diff --git a/Engine/Source/Runtime/UMG/Public/Components/TreeView.h b/Engine/Source/Runtime/UMG/Public/Components/TreeView.h
--- a/Engine/Source/Runtime/UMG/Public/Components/TreeView.h
+++ b/Engine/Source/Runtime/UMG/Public/Components/TreeView.h
@@ -58,7 +58,13 @@ protected:
 	template <template<typename> class TreeViewT = STreeView>
 	TSharedRef<TreeViewT<UObject*>> ConstructTreeView()
 	{
-		MyListView = MyTreeView = ITypedUMGListView<UObject*>::ConstructTreeView<TreeViewT>(this, ListItems, SelectionMode, bClearSelectionOnClick, ConsumeMouseWheel);
+		FTreeViewConstructArgs Args;
+		Args.bClearSelectionOnClick = bClearSelectionOnClick;
+		Args.SelectionMode = SelectionMode;
+		Args.ConsumeMouseWheel = ConsumeMouseWheel;
+
+		MyListView = MyTreeView = ITypedUMGListView<UObject*>::ConstructTreeView<TreeViewT>(this, ListItems, Args);
+		MyTreeView->SetOnEntryInitialized(SListView<UObject*>::FOnEntryInitialized::CreateUObject(this, &UTreeView::HandleOnEntryInitializedInternal));
 		return StaticCastSharedRef<TreeViewT<UObject*>>(MyTreeView.ToSharedRef());
 	}
 
