Fix user scaled scalebox widgets not having a prepass when they finally get geometry

diff --git a/Engine/Source/Runtime/Slate/Private/Widgets/Layout/SScaleBox.cpp b/Engine/Source/Runtime/Slate/Private/Widgets/Layout/SScaleBox.cpp
--- a/Engine/Source/Runtime/Slate/Private/Widgets/Layout/SScaleBox.cpp
+++ b/Engine/Source/Runtime/Slate/Private/Widgets/Layout/SScaleBox.cpp
@@ -262,7 +262,10 @@ int32 SScaleBox::OnPaint( const FPaintArgs& Args, const FGeometry& AllottedGeome
 		LastPaintGeometry = AllottedGeometry;
 		const_cast<SScaleBox*>(this)->Invalidate(EInvalidateWidgetReason::Layout);
 		const_cast<SScaleBox*>(this)->InvalidatePrepass();
-		//const_cast<SScaleBox*>(this)->SlatePrepass(GetPrepassLayoutScaleMultiplier());
+		if (Stretch.Get() == EStretch::UserSpecified)
+		{
+			const_cast<SScaleBox*>(this)->SlatePrepass(GetPrepassLayoutScaleMultiplier());
+		}
 	}
 
 	bool bClippingNeeded = false;
