When buttons change style they must invalidate.  Fixes changing the look of a button not updating on the screen

diff --git a/Engine/Source/Runtime/Slate/Private/Widgets/Input/SButton.cpp b/Engine/Source/Runtime/Slate/Private/Widgets/Input/SButton.cpp
--- a/Engine/Source/Runtime/Slate/Private/Widgets/Input/SButton.cpp
+++ b/Engine/Source/Runtime/Slate/Private/Widgets/Input/SButton.cpp
@@ -510,6 +510,8 @@ void SButton::SetButtonStyle(const FButtonStyle* ButtonStyle)
 
 	HoveredSound = Style->HoveredSlateSound;
 	PressedSound = Style->PressedSlateSound;
+
+	Invalidate(EInvalidateWidget::Layout);
 }
 
 void SButton::SetClickMethod(EButtonClickMethod::Type InClickMethod)
