Fixing an issue with SScaleboxes not invalidating prepass when their scale is changed directly.

diff --git a/Engine/Source/Runtime/Slate/Private/Widgets/Layout/SScaleBox.cpp b/Engine/Source/Runtime/Slate/Private/Widgets/Layout/SScaleBox.cpp
--- a/Engine/Source/Runtime/Slate/Private/Widgets/Layout/SScaleBox.cpp
+++ b/Engine/Source/Runtime/Slate/Private/Widgets/Layout/SScaleBox.cpp
@@ -330,7 +330,10 @@ void SScaleBox::SetVAlign(EVerticalAlignment VAlign)
 
 void SScaleBox::SetStretchDirection(EStretchDirection::Type InStretchDirection)
 {
-	SetAttribute(StretchDirection, TAttribute<EStretchDirection::Type>(InStretchDirection), EInvalidateWidgetReason::Layout);
+	if (SetAttribute(StretchDirection, TAttribute<EStretchDirection::Type>(InStretchDirection), EInvalidateWidgetReason::Layout))
+	{
+		const_cast<SScaleBox*>(this)->InvalidatePrepass();
+	}
 }
 
 void SScaleBox::SetStretch(EStretch::Type InStretch)
@@ -338,17 +341,24 @@ void SScaleBox::SetStretch(EStretch::Type InStretch)
 	if (SetAttribute(Stretch, TAttribute<EStretch::Type>(InStretch), EInvalidateWidgetReason::Layout))
 	{
 		RefreshSafeZoneScale();
+		const_cast<SScaleBox*>(this)->InvalidatePrepass();
 	}
 }
 
 void SScaleBox::SetUserSpecifiedScale(float InUserSpecifiedScale)
 {
-	SetAttribute(UserSpecifiedScale, TAttribute<float>(InUserSpecifiedScale), EInvalidateWidgetReason::Layout);
+	if (SetAttribute(UserSpecifiedScale, TAttribute<float>(InUserSpecifiedScale), EInvalidateWidgetReason::Layout))
+	{
+		const_cast<SScaleBox*>(this)->InvalidatePrepass();
+	}
 }
 
 void SScaleBox::SetIgnoreInheritedScale(bool InIgnoreInheritedScale)
 {
-	SetAttribute(IgnoreInheritedScale, TAttribute<bool>(InIgnoreInheritedScale), EInvalidateWidgetReason::Layout);
+	if (SetAttribute(IgnoreInheritedScale, TAttribute<bool>(InIgnoreInheritedScale), EInvalidateWidgetReason::Layout))
+	{
+		const_cast<SScaleBox*>(this)->InvalidatePrepass();
+	}
 }
 
 FVector2D SScaleBox::ComputeDesiredSize(float InScale) const
