Scalebox - Fixing an early out case where we skipped drawing a child, we no longer skip that path if we're on the fast path and draw the child even in the failure case.  Also fixing a case where we had to do another prepass even with a constant scale.  The SScaleBox no longer requires an extra prepass in cases where we don't need a normalizing scale prepass.

Matt.Kuhlenschmidt, Mic.Rooney

diff --git a/Engine/Source/Runtime/Slate/Private/Widgets/Layout/SScaleBox.cpp b/Engine/Source/Runtime/Slate/Private/Widgets/Layout/SScaleBox.cpp
--- a/Engine/Source/Runtime/Slate/Private/Widgets/Layout/SScaleBox.cpp
+++ b/Engine/Source/Runtime/Slate/Private/Widgets/Layout/SScaleBox.cpp
@@ -48,22 +48,26 @@ bool SScaleBox::CustomPrepass(float LayoutScaleMultiplier)
 {
 	SWidget& ChildSlotWidget = ChildSlot.GetWidget().Get();
 
-	const bool bLocalGeometryRequired = DoesScaleRequireLocalGeometry();
-	const bool bNeedsNormalizingPrepass = DoesScaleRequireNormalizingPrepass();
+	const bool bNeedsNormalizingPrepassOrLocalGeometry = DoesScaleRequireNormalizingPrepassOrLocalGeometry();
 
 	//bool bCanRenderThisFrame = true;
 
 	// If we need a normalizing prepass, or we've yet to give the child a chance to generate a desired
 	// size, do that now.
-	if (bNeedsNormalizingPrepass || !LastAllocatedArea.IsSet())
+	if (bNeedsNormalizingPrepassOrLocalGeometry || !LastAllocatedArea.IsSet())
 	{
 		ChildSlotWidget.SlatePrepass(LayoutScaleMultiplier);
+
+		NormalizedContentDesiredSize = ChildSlotWidget.GetDesiredSize();
+	}
+	else
+	{
+		NormalizedContentDesiredSize.Reset();
 	}
 
-	ContentDesiredSize = ChildSlotWidget.GetDesiredSize();
 	TOptional<float> NewComputedContentScale;
 
-	if (bLocalGeometryRequired)
+	if (bNeedsNormalizingPrepassOrLocalGeometry)
 	{
 		if (LastAllocatedArea.IsSet())
 		{
@@ -77,7 +81,7 @@ bool SScaleBox::CustomPrepass(float LayoutScaleMultiplier)
 		NewComputedContentScale = ComputeContentScale(NullGeometry);
 	}
 
-	if (bNeedsNormalizingPrepass)
+	if (bNeedsNormalizingPrepassOrLocalGeometry)
 	{
 		ChildSlotWidget.InvalidatePrepass();
 	}
@@ -96,22 +100,7 @@ bool SScaleBox::CustomPrepass(float LayoutScaleMultiplier)
 	return true;
 }
 
-bool SScaleBox::DoesScaleRequireNormalizingPrepass() const
-{
-	const EStretch::Type CurrentStretch = Stretch.Get();
-	switch (CurrentStretch)
-	{
-	case EStretch::None:
-	case EStretch::Fill:
-	case EStretch::ScaleBySafeZone:
-	case EStretch::UserSpecified:
-		return false;
-	default:
-		return true;
-	}
-}
-
-bool SScaleBox::DoesScaleRequireLocalGeometry() const
+bool SScaleBox::DoesScaleRequireNormalizingPrepassOrLocalGeometry() const
 {
 	const EStretch::Type CurrentStretch = Stretch.Get();
 	switch (CurrentStretch)
@@ -250,6 +239,18 @@ void SScaleBox::OnArrangeChildren(const FGeometry& AllottedGeometry, FArrangedCh
 				FinalScale
 			));
 		}
+		else if (GSlateIsOnFastUpdatePath)
+		{
+			// If we're unable to draw it at a scale yet, we need to process the draw request anyway
+			// since the object is visible, we can't just ignore the update for paint since we're
+			// on the fast path.
+			ArrangedChildren.AddWidget(ChildVisibility, AllottedGeometry.MakeChild(
+				ChildSlot.GetWidget(),
+				FVector2D(0, 0),
+				FVector2D(0, 0),
+				0
+			));
+		}
 	}
 }
 
@@ -260,11 +261,11 @@ int32 SScaleBox::OnPaint( const FPaintArgs& Args, const FGeometry& AllottedGeome
 	{
 		LastAllocatedArea = AllottedGeometry.GetLocalSize();
 		LastPaintGeometry = AllottedGeometry;
-		const_cast<SScaleBox*>(this)->Invalidate(EInvalidateWidgetReason::Layout);
-		const_cast<SScaleBox*>(this)->InvalidatePrepass();
-		if (Stretch.Get() == EStretch::UserSpecified)
+
+		if (DoesScaleRequireNormalizingPrepassOrLocalGeometry())
 		{
-			const_cast<SScaleBox*>(this)->SlatePrepass(GetPrepassLayoutScaleMultiplier());
+			const_cast<SScaleBox*>(this)->Invalidate(EInvalidateWidgetReason::Layout);
+			const_cast<SScaleBox*>(this)->InvalidatePrepass();
 		}
 	}
 
@@ -315,7 +316,7 @@ void SScaleBox::SetHAlign(EHorizontalAlignment HAlign)
 	if(ChildSlot.HAlignment != HAlign)
 	{
 		ChildSlot.HAlignment = HAlign;
-		Invalidate(EInvalidateWidget::Layout);
+		Invalidate(EInvalidateWidgetReason::Layout);
 	}
 }
 
@@ -324,7 +325,7 @@ void SScaleBox::SetVAlign(EVerticalAlignment VAlign)
 	if(ChildSlot.VAlignment != VAlign)
 	{
 		ChildSlot.VAlignment = VAlign;
-		Invalidate(EInvalidateWidget::Layout);
+		Invalidate(EInvalidateWidgetReason::Layout);
 	}
 }
 
@@ -332,7 +333,7 @@ void SScaleBox::SetStretchDirection(EStretchDirection::Type InStretchDirection)
 {
 	if (SetAttribute(StretchDirection, TAttribute<EStretchDirection::Type>(InStretchDirection), EInvalidateWidgetReason::Layout))
 	{
-		const_cast<SScaleBox*>(this)->InvalidatePrepass();
+		InvalidatePrepass();
 	}
 }
 
@@ -340,8 +341,9 @@ void SScaleBox::SetStretch(EStretch::Type InStretch)
 {
 	if (SetAttribute(Stretch, TAttribute<EStretch::Type>(InStretch), EInvalidateWidgetReason::Layout))
 	{
+		// This function invalidates the prepass.
 		RefreshSafeZoneScale();
-		const_cast<SScaleBox*>(this)->InvalidatePrepass();
+		check(NeedsPrepass());
 	}
 }
 
@@ -349,7 +351,7 @@ void SScaleBox::SetUserSpecifiedScale(float InUserSpecifiedScale)
 {
 	if (SetAttribute(UserSpecifiedScale, TAttribute<float>(InUserSpecifiedScale), EInvalidateWidgetReason::Layout))
 	{
-		const_cast<SScaleBox*>(this)->InvalidatePrepass();
+		InvalidatePrepass();
 	}
 }
 
@@ -357,17 +359,17 @@ void SScaleBox::SetIgnoreInheritedScale(bool InIgnoreInheritedScale)
 {
 	if (SetAttribute(IgnoreInheritedScale, TAttribute<bool>(InIgnoreInheritedScale), EInvalidateWidgetReason::Layout))
 	{
-		const_cast<SScaleBox*>(this)->InvalidatePrepass();
+		InvalidatePrepass();
 	}
 }
 
 FVector2D SScaleBox::ComputeDesiredSize(float InScale) const
 {
-	if (DoesScaleRequireNormalizingPrepass())
+	if (DoesScaleRequireNormalizingPrepassOrLocalGeometry())
 	{
-		if (ContentDesiredSize.IsSet())
+		if (NormalizedContentDesiredSize.IsSet())
 		{
-			FVector2D ContentDesiredSizeValue = ContentDesiredSize.GetValue();
+			FVector2D ContentDesiredSizeValue = NormalizedContentDesiredSize.GetValue();
 
 			if (IsDesiredSizeDependentOnAreaAndScale())
 			{
@@ -395,14 +397,14 @@ FVector2D SScaleBox::ComputeDesiredSize(float InScale) const
 			return ContentDesiredSizeValue;
 		}
 	}
-	else
+	// If we don't need a normalizing prepass, then we can safely just multiply
+	// the desired size of the children by the computed content scale, so that we request the now larger or smaller
+	// area that we need - this area is a constant scale, either by safezone or user scale.
+	else if (ComputedContentScale.IsSet())
 	{
-		if (ContentDesiredSize.IsSet() && ComputedContentScale.IsSet())
-		{
-			return ContentDesiredSize.GetValue() * ComputedContentScale.GetValue();
-		}
+		SCompoundWidget::ComputeDesiredSize(InScale) * ComputedContentScale.GetValue();
 	}
-
+	
 	return SCompoundWidget::ComputeDesiredSize(InScale);
 }
 
@@ -414,8 +416,9 @@ float SScaleBox::GetRelativeLayoutScale(const FSlotBase& Child, float LayoutScal
 void SScaleBox::RefreshSafeZoneScale()
 {
 	float ScaleDownBy = 0.f;
-	FMargin SafeMargin;
-	FVector2D ScaleBy;
+	FMargin SafeMargin(0, 0, 0, 0);
+	FVector2D ScaleBy(1, 1);
+
 #if WITH_EDITOR
 	if (OverrideScreenSize.IsSet() && !OverrideScreenSize.GetValue().IsZero())
 	{
@@ -441,6 +444,7 @@ void SScaleBox::RefreshSafeZoneScale()
 			}
 		}
 	}
+
 	const float SafeZoneScaleX = FMath::Max(SafeMargin.Left, SafeMargin.Right)/ (float)ScaleBy.X;
 	const float SafeZoneScaleY = FMath::Max(SafeMargin.Top, SafeMargin.Bottom) / (float)ScaleBy.Y;
 
@@ -448,6 +452,9 @@ void SScaleBox::RefreshSafeZoneScale()
 	ScaleDownBy = FMath::Max(SafeZoneScaleX, SafeZoneScaleY);
 
 	SafeZoneScale = 1.f - ScaleDownBy;
+
+	Invalidate(EInvalidateWidgetReason::Layout);
+	InvalidatePrepass();
 }
 
 #if WITH_EDITOR
diff --git a/Engine/Source/Runtime/Slate/Public/Widgets/Layout/SScaleBox.h b/Engine/Source/Runtime/Slate/Public/Widgets/Layout/SScaleBox.h
--- a/Engine/Source/Runtime/Slate/Public/Widgets/Layout/SScaleBox.h
+++ b/Engine/Source/Runtime/Slate/Public/Widgets/Layout/SScaleBox.h
@@ -159,8 +159,7 @@ protected:
 	virtual float GetRelativeLayoutScale(const FSlotBase& Child, float LayoutScaleMultiplier) const override;
 	// End SWidget overrides.
 
-	bool DoesScaleRequireNormalizingPrepass() const;
-	bool DoesScaleRequireLocalGeometry() const;
+	bool DoesScaleRequireNormalizingPrepassOrLocalGeometry() const;
 	bool IsDesiredSizeDependentOnAreaAndScale() const;
 	float ComputeContentScale(const FGeometry& PaintGeometry) const;
 
@@ -192,7 +191,7 @@ protected:
 	mutable TOptional<FVector2D> LastAllocatedArea;
 	mutable TOptional<FGeometry> LastPaintGeometry;
 
-	mutable TOptional<FVector2D> ContentDesiredSize;
+	mutable TOptional<FVector2D> NormalizedContentDesiredSize;
 
 	/**  */
 	mutable TOptional<float> ComputedContentScale;
