Add FindInputPreprocessor to SlateApplication.

diff --git a/Engine/Source/Runtime/Slate/Private/Framework/Application/SlateApplication.cpp b/Engine/Source/Runtime/Slate/Private/Framework/Application/SlateApplication.cpp
--- a/Engine/Source/Runtime/Slate/Private/Framework/Application/SlateApplication.cpp
+++ b/Engine/Source/Runtime/Slate/Private/Framework/Application/SlateApplication.cpp
@@ -4382,6 +4382,11 @@ void FSlateApplication::UnregisterInputPreProcessor(TSharedPtr<IInputProcessor>
 	InputPreProcessors.Remove(InputProcessor);
 }
 
+int32 FSlateApplication::FindInputPreProcessor(TSharedPtr<class IInputProcessor> InputProcessor) const
+{
+	return InputPreProcessors.Find(InputProcessor);
+}
+
 void FSlateApplication::SetCursorRadius(float NewRadius)
 {
 	CursorRadius = FMath::Max<float>(0.0f, NewRadius);
@@ -7696,6 +7701,11 @@ void FSlateApplication::InputPreProcessorsHelper::RemoveAll()
 }
 
 
+int32 FSlateApplication::InputPreProcessorsHelper::Find(TSharedPtr<IInputProcessor> InputProcessor) const
+{
+	return InputPreProcessorList.Find(InputProcessor);
+}
+
 bool FSlateApplication::InputPreProcessorsHelper::PreProcessInput(TFunctionRef<bool(TSharedPtr<IInputProcessor>)> ToRun)
 {
 	TGuardValue<bool> IteratingGuard(bIsIteratingPreProcessors, true);
diff --git a/Engine/Source/Runtime/Slate/Public/Framework/Application/SlateApplication.h b/Engine/Source/Runtime/Slate/Public/Framework/Application/SlateApplication.h
--- a/Engine/Source/Runtime/Slate/Public/Framework/Application/SlateApplication.h
+++ b/Engine/Source/Runtime/Slate/Public/Framework/Application/SlateApplication.h
@@ -1314,6 +1314,13 @@ public:
 	 */
 	void UnregisterInputPreProcessor(TSharedPtr<class IInputProcessor> InputProcessor);
 
+	/**
+	 * Get the index of a registered pre-processor.
+	 * @param InputProcessor	The input pre-processor to find.
+	 * @return The index of the pre-processor, or INDEX_NONE if not registered.
+	 */
+	int32 FindInputPreProcessor(TSharedPtr<class IInputProcessor> InputProcessor) const;
+
 	/** Sets the hit detection radius of the cursor */
 	void SetCursorRadius(float NewRadius);
 
@@ -2187,6 +2194,14 @@ private:
 		 */
 		void RemoveAll();
 
+		/**
+		 * Get the index of an input pre-processor.
+		 * @param InputProcessor	The InputProcessor to find.
+		 * @return The index of the pre-processor, or INDEX_NONE if not registered.
+		 */
+		int32 Find(TSharedPtr<IInputProcessor> InputProcessor) const;
+
+
 	private:
 
 		bool PreProcessInput(TFunctionRef<bool(TSharedPtr<IInputProcessor>)> ToRun);
