Scalebox - Need to actually return the desired size scaled to what we calculated.

Matt.Kuhlenschmidt

diff --git a/Engine/Source/Runtime/Slate/Private/Widgets/Layout/SScaleBox.cpp b/Engine/Source/Runtime/Slate/Private/Widgets/Layout/SScaleBox.cpp
--- a/Engine/Source/Runtime/Slate/Private/Widgets/Layout/SScaleBox.cpp
+++ b/Engine/Source/Runtime/Slate/Private/Widgets/Layout/SScaleBox.cpp
@@ -402,7 +402,7 @@ FVector2D SScaleBox::ComputeDesiredSize(float InScale) const
 	// area that we need - this area is a constant scale, either by safezone or user scale.
 	else if (ComputedContentScale.IsSet())
 	{
-		SCompoundWidget::ComputeDesiredSize(InScale) * ComputedContentScale.GetValue();
+		return SCompoundWidget::ComputeDesiredSize(InScale) * ComputedContentScale.GetValue();
 	}
 	
 	return SCompoundWidget::ComputeDesiredSize(InScale);
