Slate - Removing the hack that was placed in the TextLayout to overcome a vectorization compiler bug that has been fixed.

Matt.Kuhlenschmidt

diff --git a/Engine/Source/Runtime/Slate/Private/Framework/Text/TextLayout.cpp b/Engine/Source/Runtime/Slate/Private/Framework/Text/TextLayout.cpp
--- a/Engine/Source/Runtime/Slate/Private/Framework/Text/TextLayout.cpp
+++ b/Engine/Source/Runtime/Slate/Private/Framework/Text/TextLayout.cpp
@@ -2781,7 +2781,7 @@ FVector2D FTextLayout::FRunModel::Measure(int32 BeginIndex, int32 EndIndex, floa
 	FVector2D Size = Run->Measure(BeginIndex, EndIndex, InScale, InTextContext);
 
 	MeasuredRanges.Add( FTextRange( BeginIndex, EndIndex ) );
-	MeasuredRangeSizes.Add(FVector4(Size, FVector2D::ZeroVector));
+	MeasuredRangeSizes.Add(Size);
 
 	return Size;
 }
diff --git a/Engine/Source/Runtime/Slate/Public/Framework/Text/TextLayout.h b/Engine/Source/Runtime/Slate/Public/Framework/Text/TextLayout.h
--- a/Engine/Source/Runtime/Slate/Public/Framework/Text/TextLayout.h
+++ b/Engine/Source/Runtime/Slate/Public/Framework/Text/TextLayout.h
@@ -261,8 +261,7 @@ public:
 
 		TSharedRef< IRun > Run;
 		TArray< FTextRange > MeasuredRanges;
-		// HACK: This should be a Vector2D, but changing to a FVector4 to overcome a compiler issue with vectorization running off the end of the array.
-		TArray< FVector4 > MeasuredRangeSizes;
+		TArray< FVector2D > MeasuredRangeSizes;
 	};
 
 	struct ELineModelDirtyState
diff --git a/Engine/Source/Runtime/SlateCore/Private/SlateCoreClasses.cpp b/Engine/Source/Runtime/SlateCore/Private/SlateCoreClasses.cpp
--- a/Engine/Source/Runtime/SlateCore/Private/SlateCoreClasses.cpp
+++ b/Engine/Source/Runtime/SlateCore/Private/SlateCoreClasses.cpp
@@ -17,7 +17,7 @@ float GSlateContrast = 1;
 FAutoConsoleVariableRef CVarSlateContrast(
 	TEXT("Slate.Contrast"),
 	GSlateContrast,
-	TEXT("The amount of contrast to apply to the UI.")
+	TEXT("The amount of contrast to apply to the UI (default 1).")
 );
 
 
