Fix double unhovered events when mouse leaves a pressed button.

diff --git a/Engine/Source/Runtime/Slate/Private/Widgets/Input/SButton.cpp b/Engine/Source/Runtime/Slate/Private/Widgets/Input/SButton.cpp
--- a/Engine/Source/Runtime/Slate/Private/Widgets/Input/SButton.cpp
+++ b/Engine/Source/Runtime/Slate/Private/Widgets/Input/SButton.cpp
@@ -352,6 +352,8 @@ void SButton::OnMouseEnter( const FGeometry& MyGeometry, const FPointerEvent& Mo
 
 void SButton::OnMouseLeave( const FPointerEvent& MouseEvent )
 {
+	const bool bWasHovered = IsHovered();
+	
 	// Call parent implementation
 	SWidget::OnMouseLeave( MouseEvent );
 
@@ -362,7 +364,10 @@ void SButton::OnMouseLeave( const FPointerEvent& MouseEvent )
 		Release();
 	}
 
-	OnUnhovered.ExecuteIfBound();
+	if (bWasHovered)
+	{
+		OnUnhovered.ExecuteIfBound();
+	}
 
 	Invalidate(EInvalidateWidget::Layout);
 }
