Slate - While ticking input preprocessors we now use index in case we add to the collection, and adding a guard to buffer any removes.

diff --git a/Engine/Source/Runtime/Slate/Private/Framework/Application/SlateApplication.cpp b/Engine/Source/Runtime/Slate/Private/Framework/Application/SlateApplication.cpp
--- a/Engine/Source/Runtime/Slate/Private/Framework/Application/SlateApplication.cpp
+++ b/Engine/Source/Runtime/Slate/Private/Framework/Application/SlateApplication.cpp
@@ -7607,8 +7607,11 @@ void FSlateApplication::NavigateFromWidgetUnderCursor(const uint32 InUserIndex,
 
 void FSlateApplication::InputPreProcessorsHelper::Tick(const float DeltaTime, FSlateApplication& SlateApp, TSharedRef<ICursor> Cursor)
 {
-	for (TSharedPtr<IInputProcessor> InputPreProcessor : InputPreProcessorList)
+	TGuardValue<bool> IteratingGuard(bIsIteratingPreProcessors, true);
+
+	for (int32 ProcessorIndex = 0; ProcessorIndex < InputPreProcessorList.Num(); ProcessorIndex++)
 	{
+		TSharedPtr<IInputProcessor> InputPreProcessor = InputPreProcessorList[ProcessorIndex];
 		InputPreProcessor->Tick(DeltaTime, SlateApp, Cursor);
 	}
 }
