fixing issue where an added input processor would still be in the ProcessorsPendingRemoval array and would get removed the next time a button was pressed even though it was added back to the InputPreProcessorList. So now any time we add to the InputPreProcessorList we remove that processor from the ProcessorsPendingRemoval array

diff --git a/Engine/Source/Runtime/Slate/Private/Framework/Application/SlateApplication.cpp b/Engine/Source/Runtime/Slate/Private/Framework/Application/SlateApplication.cpp
--- a/Engine/Source/Runtime/Slate/Private/Framework/Application/SlateApplication.cpp
+++ b/Engine/Source/Runtime/Slate/Private/Framework/Application/SlateApplication.cpp
@@ -7676,6 +7676,11 @@ bool FSlateApplication::InputPreProcessorsHelper::Add(TSharedPtr<IInputProcessor
 		bResult = true;
 	}
 
+	if (bResult)
+	{
+		ProcessorsPendingRemoval.Remove(InputProcessor);
+	}
+
 	return bResult;
 }
 
@@ -7723,7 +7728,7 @@ bool FSlateApplication::InputPreProcessorsHelper::PreProcessInput(TFunctionRef<b
 		}
 	}
 
-	InputPreProcessorList.RemoveAll([this](const TSharedPtr<IInputProcessor> Processor){ return ProcessorsPendingRemoval.Contains(Processor); });
+	InputPreProcessorList.RemoveAll([this](const TSharedPtr<IInputProcessor> Processor) { return ProcessorsPendingRemoval.Contains(Processor); });
 	ProcessorsPendingRemoval.Reset();
 
 	return bShouldExit;
