Fixed a crash when running the -AutomatedMapBuild command line with the editor due to windows not being properly closed.

While running the -AutomatedMapBuild command, several windows will be created (such as the notification window panel) which then rely on the method FSlateApplication::CloseAllWindowsImmediately to close them properly. If the windows are not closed properly we ended up receiving messages from these windows after FWindowsApplication had been deleted, which meant that the message proc could access the now deleted members of FWindowsApplication and crash.

The reason why ::CloseAllWindowsImmediately could fail was the loop that iterated over all of the top level windows was not taking into account that the call to destroy the window (::RequestDestroyWindow) would also remove that window from the array of top level windows, meaning we would skip every other window.

To fix it we now iterate over the array in reverse order so that ever window gets hit (although this relies on ::RequestDestroyWindow destroying a single window each time it is called)

diff --git a/Engine/Source/Runtime/Slate/Private/Framework/Application/SlateApplication.cpp b/Engine/Source/Runtime/Slate/Private/Framework/Application/SlateApplication.cpp
--- a/Engine/Source/Runtime/Slate/Private/Framework/Application/SlateApplication.cpp
+++ b/Engine/Source/Runtime/Slate/Private/Framework/Application/SlateApplication.cpp
@@ -3138,9 +3138,12 @@ void FSlateApplication::CloseAllWindowsImmediately()
 		ToolTipWindow.Reset();
 	}
 
-	for (int32 WindowIndex = 0; WindowIndex < SlateWindows.Num(); ++WindowIndex)
+	// Destroy all top level windows.  
+	// ::RequestDestroyWindow will remove the window from the TArray SlateWindows so we
+	// should iterate over it backwards to make sure that the WindowIndex is still correct.
+	for (int32 WindowIndex = SlateWindows.Num() -1; WindowIndex >= 0; --WindowIndex)
 	{
-		// Destroy all top level windows.  This will also request that all children of each window be destroyed
+		// This will also request that all children of each window be destroyed
 		RequestDestroyWindow(SlateWindows[WindowIndex]);
 	}
 
