SEditableText::OnCursorMoved now peforms a layout invalidation

diff --git a/Engine/Source/Runtime/Slate/Private/Widgets/Input/SEditableText.cpp b/Engine/Source/Runtime/Slate/Private/Widgets/Input/SEditableText.cpp
--- a/Engine/Source/Runtime/Slate/Private/Widgets/Input/SEditableText.cpp
+++ b/Engine/Source/Runtime/Slate/Private/Widgets/Input/SEditableText.cpp
@@ -593,6 +593,7 @@ void SEditableText::OnTextCommitted(const FText& InText, const ETextCommit::Type
 
 void SEditableText::OnCursorMoved(const FTextLocation& InLocation)
 {
+	Invalidate(EInvalidateWidgetReason::Layout);
 }
 
 float SEditableText::UpdateAndClampHorizontalScrollBar(const float InViewOffset, const float InViewFraction, const EVisibility InVisiblityOverride)
