Slate - Updating the default navigation config's horizontal deadzone to match the vertical one, we've also recieved feedback that the default at 40 isn't quite high enough anyway for some well loved controllers.

diff --git a/Engine/Source/Runtime/Slate/Private/Framework/Application/NavigationConfig.cpp b/Engine/Source/Runtime/Slate/Private/Framework/Application/NavigationConfig.cpp
--- a/Engine/Source/Runtime/Slate/Private/Framework/Application/NavigationConfig.cpp
+++ b/Engine/Source/Runtime/Slate/Private/Framework/Application/NavigationConfig.cpp
@@ -9,7 +9,7 @@ FNavigationConfig::FNavigationConfig()
 	: bTabNavigation(true)
 	, bKeyNavigation(true)
 	, bAnalogNavigation(true)
-	, AnalogNavigationHorizontalThreshold(0.40f)
+	, AnalogNavigationHorizontalThreshold(0.50f)
 	, AnalogNavigationVerticalThreshold(0.50f)
 {
 	AnalogHorizontalKey = EKeys::Gamepad_LeftX;
