When updating menu bar items on Mac, reset GCachedMenuState on the game thread instead of the main thread, as some SWidget destructors may require that

diff --git a/Engine/Source/Runtime/Slate/Private/Framework/MultiBox/Mac/MacMenu.cpp b/Engine/Source/Runtime/Slate/Private/Framework/MultiBox/Mac/MacMenu.cpp
--- a/Engine/Source/Runtime/Slate/Private/Framework/MultiBox/Mac/MacMenu.cpp
+++ b/Engine/Source/Runtime/Slate/Private/Framework/MultiBox/Mac/MacMenu.cpp
@@ -500,17 +500,17 @@ void FSlateMacMenu::UpdateWithMultiBox(const TSharedPtr< FMultiBox > MultiBox)
 	FSafeMultiBoxPass* SafeMultiBoxPtr = new FSafeMultiBoxPass;
 	SafeMultiBoxPtr->MultiBox = MultiBox;
 
-	MainThreadCall(^{
-		FScopeLock Lock(&GCachedMenuStateCS);
+	FScopeLock Lock(&GCachedMenuStateCS);
+	GCachedMenuState.Reset();
 
+	MainThreadCall(^{
 		int32 NumItems = [[NSApp mainMenu] numberOfItems];
 		FText WindowLabel = NSLOCTEXT("MainMenu", "WindowMenu", "Window");
 		for (int32 Index = NumItems - 1; Index > 0; Index--)
 		{
 			[[NSApp mainMenu] removeItemAtIndex:Index];
 		}
-		GCachedMenuState.Reset();
-		
+
 		if( SafeMultiBoxPtr->MultiBox.IsValid() )
 		{
 			const TArray<TSharedRef<const FMultiBlock> >& MenuBlocks = SafeMultiBoxPtr->MultiBox->GetBlocks();
