Fix tree view to release all references to items when RebuildList is called
#rnx

diff --git a/Engine/Source/Runtime/Slate/Public/Widgets/Views/STreeView.h b/Engine/Source/Runtime/Slate/Public/Widgets/Views/STreeView.h
--- a/Engine/Source/Runtime/Slate/Public/Widgets/Views/STreeView.h
+++ b/Engine/Source/Runtime/Slate/Public/Widgets/Views/STreeView.h
@@ -540,8 +540,8 @@ public:
 					// series of calls, Private_SignalSelectionChanged() could end up in a child
 					// that indexes into either of these arrays (the index wouldn't be updated yet,
 					// and could be invalid)
-					SparseItemInfos = TempSparseItemInfo;
-					DenseItemInfos  = TempDenseItemInfos;
+					SparseItemInfos = MoveTemp(TempSparseItemInfo);
+					DenseItemInfos  = MoveTemp(TempDenseItemInfos);
 
 					// Once the selection changed events have gone through we can update the parent highlight statuses, which are based on your current selection.
 					if (bHighlightParentNodesForSelection)
@@ -684,7 +684,13 @@ public:
 	{
 		RequestListRefresh();
 	}
-		
+
+	virtual void RebuildList() override
+	{
+		LinearizedItems.Empty();
+		SListView<ItemType>::RebuildList();
+	}
+
 	/**
 	 * Set whether some data item is expanded or not.
 	 * 
