Fix newly created SEditableTextBox caret blink speed when focus is forced via code.

diff --git a/Engine/Source/Runtime/Slate/Private/Widgets/Text/SlateEditableTextLayout.cpp b/Engine/Source/Runtime/Slate/Private/Widgets/Text/SlateEditableTextLayout.cpp
--- a/Engine/Source/Runtime/Slate/Private/Widgets/Text/SlateEditableTextLayout.cpp
+++ b/Engine/Source/Runtime/Slate/Private/Widgets/Text/SlateEditableTextLayout.cpp
@@ -749,6 +749,9 @@ bool FSlateEditableTextLayout::HandleFocusReceived(const FFocusEvent& InFocusEve
 	// Focus change affects volatility, so update that too
 	OwnerWidget->GetSlateWidget()->Invalidate(EInvalidateWidget::LayoutAndVolatility);
 
+	// Force update of last cursor interaction time to ensure the caret blinks at the correct frequency regardless of how focus is set.
+	CursorInfo.UpdateLastCursorInteractionTime();
+
 	return true;
 }
 
diff --git a/Engine/Source/Runtime/Slate/Private/Widgets/Text/SlateEditableTextTypes.cpp b/Engine/Source/Runtime/Slate/Private/Widgets/Text/SlateEditableTextTypes.cpp
--- a/Engine/Source/Runtime/Slate/Private/Widgets/Text/SlateEditableTextTypes.cpp
+++ b/Engine/Source/Runtime/Slate/Private/Widgets/Text/SlateEditableTextTypes.cpp
@@ -11,6 +11,11 @@
 namespace SlateEditableTextTypes
 {
 
+void FCursorInfo::UpdateLastCursorInteractionTime()
+{
+	LastCursorInteractionTime = FSlateApplication::Get().GetCurrentTime();
+}
+
 void FCursorInfo::SetCursorLocationAndCalculateAlignment(const FTextLayout& InTextLayout, const FTextLocation& InCursorPosition)
 {
 	FTextLocation NewCursorPosition = InCursorPosition;
@@ -65,7 +70,7 @@ void FCursorInfo::SetCursorLocationAndAlignment(const FTextLayout& InTextLayout,
 		}
 	}
 
-	LastCursorInteractionTime = FSlateApplication::Get().GetCurrentTime();
+	UpdateLastCursorInteractionTime();
 }
 
 FCursorInfo FCursorInfo::CreateUndo() const
@@ -82,7 +87,8 @@ void FCursorInfo::RestoreFromUndo(const FCursorInfo& UndoData)
 	CursorPosition = UndoData.CursorPosition;
 	CursorAlignment = UndoData.CursorAlignment;
 	CursorTextDirection = UndoData.CursorTextDirection;
-	LastCursorInteractionTime = FSlateApplication::Get().GetCurrentTime();
+
+	UpdateLastCursorInteractionTime();
 }
 
 FCursorLineHighlighter::FCursorLineHighlighter(const FCursorInfo* InCursorInfo)
diff --git a/Engine/Source/Runtime/Slate/Public/Widgets/Text/SlateEditableTextTypes.h b/Engine/Source/Runtime/Slate/Public/Widgets/Text/SlateEditableTextTypes.h
--- a/Engine/Source/Runtime/Slate/Public/Widgets/Text/SlateEditableTextTypes.h
+++ b/Engine/Source/Runtime/Slate/Public/Widgets/Text/SlateEditableTextTypes.h
@@ -69,6 +69,9 @@ public:
 		return LastCursorInteractionTime;
 	}
 
+	/** Update the last cursor interaction time to now */
+	void UpdateLastCursorInteractionTime();
+
 	/** Set the position of the cursor, and then work out the correct alignment based on the current text layout */
 	void SetCursorLocationAndCalculateAlignment(const FTextLayout& InTextLayout, const FTextLocation& InCursorPosition);
 
