Slate - The scalebox now attempts to resolve the scale in a single frame like it use to before, even if the children in question didn't get a chance to be cached knowing this new scale was coming.

diff --git a/Engine/Source/Runtime/Slate/Private/Widgets/Layout/SScaleBox.cpp b/Engine/Source/Runtime/Slate/Private/Widgets/Layout/SScaleBox.cpp
--- a/Engine/Source/Runtime/Slate/Private/Widgets/Layout/SScaleBox.cpp
+++ b/Engine/Source/Runtime/Slate/Private/Widgets/Layout/SScaleBox.cpp
@@ -50,8 +50,6 @@ bool SScaleBox::CustomPrepass(float LayoutScaleMultiplier)
 
 	const bool bNeedsNormalizingPrepassOrLocalGeometry = DoesScaleRequireNormalizingPrepassOrLocalGeometry();
 
-	//bool bCanRenderThisFrame = true;
-
 	// If we need a normalizing prepass, or we've yet to give the child a chance to generate a desired
 	// size, do that now.
 	if (bNeedsNormalizingPrepassOrLocalGeometry || !LastAllocatedArea.IsSet())
@@ -203,54 +201,45 @@ void SScaleBox::OnArrangeChildren(const FGeometry& AllottedGeometry, FArrangedCh
 			SlotWidgetDesiredSize = AreaSize;
 		}
 
-		if (ComputedContentScale.IsSet())
-		{
-			LastFinalOffset = FVector2D(0, 0);
-			float FinalScale = ComputedContentScale.GetValue();
+		// This scale may not look right, the item being 
+		// shown may need 2 frames to truly be drawn correctly,
+		// but rather than have a blank frame, it's better for us to try
+		// and fit the contents to our known geometry.
+		const float TempComputedContentScale = ComputedContentScale.IsSet() ? ComputedContentScale.GetValue() : ComputeContentScale(AllottedGeometry);
 
-			// If we're just filling, there's no scale applied, we're just filling the area.
-			if (CurrentStretch != EStretch::Fill)
-			{
-				const FMargin SlotPadding(ChildSlot.SlotPadding.Get());
-				AlignmentArrangeResult XResult = AlignChild<Orient_Horizontal>(AreaSize.X, ChildSlot, SlotPadding, FinalScale, false);
-				AlignmentArrangeResult YResult = AlignChild<Orient_Vertical>(AreaSize.Y, ChildSlot, SlotPadding, FinalScale, false);
+		LastFinalOffset = FVector2D(0, 0);
+		float FinalScale = TempComputedContentScale;
 
-				LastFinalOffset = FVector2D(XResult.Offset, YResult.Offset) / FinalScale;
+		// If we're just filling, there's no scale applied, we're just filling the area.
+		if (CurrentStretch != EStretch::Fill)
+		{
+			const FMargin SlotPadding(ChildSlot.SlotPadding.Get());
+			AlignmentArrangeResult XResult = AlignChild<Orient_Horizontal>(AreaSize.X, ChildSlot, SlotPadding, FinalScale, false);
+			AlignmentArrangeResult YResult = AlignChild<Orient_Vertical>(AreaSize.Y, ChildSlot, SlotPadding, FinalScale, false);
 
-				// If the layout horizontally is fill, then we need the desired size to be the whole size of the widget, 
-				// but scale the inverse of the scale we're applying.
-				if (ChildSlot.HAlignment == HAlign_Fill)
-				{
-					SlotWidgetDesiredSize.X = AreaSize.X / FinalScale;
-				}
+			LastFinalOffset = FVector2D(XResult.Offset, YResult.Offset) / FinalScale;
 
-				// If the layout vertically is fill, then we need the desired size to be the whole size of the widget, 
-				// but scale the inverse of the scale we're applying.
-				if (ChildSlot.VAlignment == VAlign_Fill)
-				{
-					SlotWidgetDesiredSize.Y = AreaSize.Y / FinalScale;
-				}
+			// If the layout horizontally is fill, then we need the desired size to be the whole size of the widget, 
+			// but scale the inverse of the scale we're applying.
+			if (ChildSlot.HAlignment == HAlign_Fill)
+			{
+				SlotWidgetDesiredSize.X = AreaSize.X / FinalScale;
 			}
 
-			ArrangedChildren.AddWidget(ChildVisibility, AllottedGeometry.MakeChild(
-				ChildSlot.GetWidget(),
-				LastFinalOffset,
-				SlotWidgetDesiredSize,
-				FinalScale
-			));
-		}
-		else if (GSlateIsOnFastUpdatePath)
-		{
-			// If we're unable to draw it at a scale yet, we need to process the draw request anyway
-			// since the object is visible, we can't just ignore the update for paint since we're
-			// on the fast path.
-			ArrangedChildren.AddWidget(ChildVisibility, AllottedGeometry.MakeChild(
-				ChildSlot.GetWidget(),
-				FVector2D(0, 0),
-				FVector2D(0, 0),
-				0
-			));
+			// If the layout vertically is fill, then we need the desired size to be the whole size of the widget, 
+			// but scale the inverse of the scale we're applying.
+			if (ChildSlot.VAlignment == VAlign_Fill)
+			{
+				SlotWidgetDesiredSize.Y = AreaSize.Y / FinalScale;
+			}
 		}
+
+		ArrangedChildren.AddWidget(ChildVisibility, AllottedGeometry.MakeChild(
+			ChildSlot.GetWidget(),
+			LastFinalOffset,
+			SlotWidgetDesiredSize,
+			FinalScale
+		));
 	}
 }
 
