Fix multiline edit boxes not invalidating when the cursor moves

diff --git a/Engine/Source/Runtime/Slate/Private/Widgets/Text/SMultiLineEditableText.cpp b/Engine/Source/Runtime/Slate/Private/Widgets/Text/SMultiLineEditableText.cpp
--- a/Engine/Source/Runtime/Slate/Private/Widgets/Text/SMultiLineEditableText.cpp
+++ b/Engine/Source/Runtime/Slate/Private/Widgets/Text/SMultiLineEditableText.cpp
@@ -373,6 +373,7 @@ void SMultiLineEditableText::OnTextCommitted(const FText& InText, const ETextCom
 void SMultiLineEditableText::OnCursorMoved(const FTextLocation& InLocation)
 {
 	OnCursorMovedCallback.ExecuteIfBound(InLocation);
+	Invalidate(EInvalidateWidget::Layout);
 }
 
 float SMultiLineEditableText::UpdateAndClampHorizontalScrollBar(const float InViewOffset, const float InViewFraction, const EVisibility InVisiblityOverride)
