Make sure newly added widgets to list views actually get prepassed.  many widgets (like text) require an initial prepass to function

diff --git a/Engine/Source/Runtime/Slate/Public/Widgets/Views/SListView.h b/Engine/Source/Runtime/Slate/Public/Widgets/Views/SListView.h
--- a/Engine/Source/Runtime/Slate/Public/Widgets/Views/SListView.h
+++ b/Engine/Source/Runtime/Slate/Public/Widgets/Views/SListView.h
@@ -1134,6 +1134,7 @@ public:
 
 		// We rely on the widgets desired size in order to determine how many will fit on screen.
 		const TSharedRef<SWidget> NewlyGeneratedWidget = WidgetForItem->AsWidget();
+		NewlyGeneratedWidget->InvalidatePrepass();
 		NewlyGeneratedWidget->SlatePrepass(LayoutScaleMultiplier);
 
 		// We have a widget for this item; add it to the panel so that it is part of the UI.
