Slate - Update the Scalebox to update layout when prepass changes since layout needs to be redone when the scalebox changes prepass.

diff --git a/Engine/Source/Runtime/Slate/Private/Widgets/Layout/SScaleBox.cpp b/Engine/Source/Runtime/Slate/Private/Widgets/Layout/SScaleBox.cpp
--- a/Engine/Source/Runtime/Slate/Private/Widgets/Layout/SScaleBox.cpp
+++ b/Engine/Source/Runtime/Slate/Private/Widgets/Layout/SScaleBox.cpp
@@ -81,6 +81,7 @@ bool SScaleBox::CustomPrepass(float LayoutScaleMultiplier)
 
 	if (bNeedsNormalizingPrepassOrLocalGeometry)
 	{
+		ChildSlotWidget.Invalidate(EInvalidateWidgetReason::Layout);
 		ChildSlotWidget.InvalidatePrepass();
 	}
 
@@ -322,6 +323,7 @@ void SScaleBox::SetStretchDirection(EStretchDirection::Type InStretchDirection)
 {
 	if (SetAttribute(StretchDirection, TAttribute<EStretchDirection::Type>(InStretchDirection), EInvalidateWidgetReason::Layout))
 	{
+		Invalidate(EInvalidateWidgetReason::Layout);
 		InvalidatePrepass();
 	}
 }
@@ -340,6 +342,7 @@ void SScaleBox::SetUserSpecifiedScale(float InUserSpecifiedScale)
 {
 	if (SetAttribute(UserSpecifiedScale, TAttribute<float>(InUserSpecifiedScale), EInvalidateWidgetReason::Layout))
 	{
+		Invalidate(EInvalidateWidgetReason::Layout);
 		InvalidatePrepass();
 	}
 }
