Copy the child array when drawing windows to ensure they can't be deleted while keep in the recursion.

diff --git a/Engine/Source/Runtime/Slate/Private/Framework/Application/SlateApplication.cpp b/Engine/Source/Runtime/Slate/Private/Framework/Application/SlateApplication.cpp
--- a/Engine/Source/Runtime/Slate/Private/Framework/Application/SlateApplication.cpp
+++ b/Engine/Source/Runtime/Slate/Private/Framework/Application/SlateApplication.cpp
@@ -1444,7 +1444,7 @@ void FSlateApplication::DrawWindowAndChildren( const TSharedRef<SWindow>& Window
 	if (bDrawChildWindows)
 	{
 		// Draw the child windows
-		const TArray< TSharedRef<SWindow> >& WindowChildren = WindowToDraw->GetChildWindows();
+		const TArray< TSharedRef<SWindow> > WindowChildren = WindowToDraw->GetChildWindows();
 		for (int32 ChildIndex=0; ChildIndex < WindowChildren.Num(); ++ChildIndex)
 		{
 			DrawWindowAndChildren( WindowChildren[ChildIndex], DrawWindowArgs );
