Text Layout Improvements:
 - Fixed some places that were treating text as a sequence of characters rather than a sequence of grapheme clusters.
   - Cursor highlighting and the delete/backspace key were both handling this incorrectly.
 - Converted the text layout code to use IBreakIterator::SetStringRef.

diff --git a/Engine/Source/Runtime/Slate/Private/Framework/Text/TextLayout.cpp b/Engine/Source/Runtime/Slate/Private/Framework/Text/TextLayout.cpp
--- a/Engine/Source/Runtime/Slate/Private/Framework/Text/TextLayout.cpp
+++ b/Engine/Source/Runtime/Slate/Private/Framework/Text/TextLayout.cpp
@@ -602,7 +602,7 @@ void FTextLayout::FlowLineLayout(const int32 LineModelIndex, const float Wrappin
 				// This is a single word that's too long to fit onto a line, so we'll fallback to wrapping this word at grapheme cluster boundaries - this may require more than a single split
 				const int32 NonBreakingStringIndexOffset = PreviousBlockEnd;
 				const FString NonBreakingString = FString(Break.TrimmedRange.EndIndex - NonBreakingStringIndexOffset, **LineModel.Text + NonBreakingStringIndexOffset);
-				GraphemeBreakIterator->SetString(NonBreakingString);
+				GraphemeBreakIterator->SetStringRef(&NonBreakingString);
 
 				CurrentWidth = 0.0f;
 				for (int32 PreviousBreak = 0, CurrentBreak = GraphemeBreakIterator->MoveToNext(); CurrentBreak != INDEX_NONE;)
@@ -1073,7 +1073,7 @@ void FTextLayout::CreateLineWrappingCache(FLineModel& LineModel)
 		LineBreakIterator = FBreakIterator::CreateLineBreakIterator();
 	}
 
-	LineBreakIterator->SetString( **LineModel.Text );
+	LineBreakIterator->SetStringRef(&LineModel.Text.Get());
 
 	int32 PreviousBreak = 0;
 	int32 CurrentBreak = 0;
@@ -2349,6 +2349,33 @@ void FTextLayout::GetSelectionAsText(FString& DisplayText, const FTextSelection&
 	}
 }
 
+FTextSelection FTextLayout::GetGraphemeAt(const FTextLocation& Location) const
+{
+	const int32 LineIndex = Location.GetLineIndex();
+	const int32 Offset = Location.GetOffset();
+
+	if (!LineModels.IsValidIndex(LineIndex))
+	{
+		return FTextSelection();
+	}
+
+	const FLineModel& LineModel = LineModels[LineIndex];
+
+	GraphemeBreakIterator->SetStringRef(&LineModel.Text.Get());
+
+	const int32 PreviousBreak = (Offset < LineModel.Text->Len()) ? GraphemeBreakIterator->MoveToCandidateAfter(Offset) : GraphemeBreakIterator->ResetToEnd();
+	const int32 CurrentBreak = GraphemeBreakIterator->MoveToPrevious();
+
+	GraphemeBreakIterator->ClearString();
+
+	if (PreviousBreak == CurrentBreak || CurrentBreak == INDEX_NONE)
+	{
+		return FTextSelection();
+	}
+
+	return FTextSelection(FTextLocation(LineIndex, CurrentBreak), FTextLocation(LineIndex, PreviousBreak));
+}
+
 FTextSelection FTextLayout::GetWordAt(const FTextLocation& Location) const
 {
 	const int32 LineIndex = Location.GetLineIndex();
@@ -2361,7 +2388,7 @@ FTextSelection FTextLayout::GetWordAt(const FTextLocation& Location) const
 
 	const FLineModel& LineModel = LineModels[LineIndex];
 
-	WordBreakIterator->SetString(**LineModel.Text);
+	WordBreakIterator->SetStringRef(&LineModel.Text.Get());
 
 	int32 PreviousBreak = (Offset < LineModel.Text->Len()) ? WordBreakIterator->MoveToCandidateAfter(Offset) : WordBreakIterator->ResetToEnd();
 	int32 CurrentBreak = 0;
diff --git a/Engine/Source/Runtime/Slate/Private/Widgets/Text/SlateEditableTextLayout.cpp b/Engine/Source/Runtime/Slate/Private/Widgets/Text/SlateEditableTextLayout.cpp
--- a/Engine/Source/Runtime/Slate/Private/Widgets/Text/SlateEditableTextLayout.cpp
+++ b/Engine/Source/Runtime/Slate/Private/Widgets/Text/SlateEditableTextLayout.cpp
@@ -1375,11 +1375,13 @@ bool FSlateEditableTextLayout::HandleBackspace()
 		}
 		else
 		{
-			// Delete character to the left of the caret
-			if (TextLayout->RemoveAt(FTextLocation(CursorInteractionPosition, -1)))
+			// Delete grapheme to the left of the caret
+			const FTextSelection DeleteSelection = TextLayout->GetGraphemeAt(FTextLocation(CursorInteractionPosition, -1));
+			const int32 GraphemeSize = DeleteSelection.GetEnd().GetOffset() - DeleteSelection.GetBeginning().GetOffset();
+			if (TextLayout->RemoveAt(DeleteSelection.GetBeginning(), GraphemeSize))
 			{
-				// Adjust caret position one to the left
-				FinalCursorLocation = FTextLocation(CursorInteractionPosition, -1);
+				// Adjust caret to the left
+				FinalCursorLocation = FTextLocation(CursorInteractionPosition, -GraphemeSize);
 			}
 		}
 
@@ -1436,8 +1438,10 @@ bool FSlateEditableTextLayout::HandleDelete()
 		}
 		else
 		{
-			// Delete character to the right of the caret
-			TextLayout->RemoveAt(CursorInteractionPosition);
+			// Delete grapheme to the right of the caret
+			const FTextSelection DeleteSelection = TextLayout->GetGraphemeAt(CursorInteractionPosition);
+			const int32 GraphemeSize = DeleteSelection.GetEnd().GetOffset() - DeleteSelection.GetBeginning().GetOffset();
+			TextLayout->RemoveAt(DeleteSelection.GetBeginning(), GraphemeSize);
 			// do nothing to the cursor as the FinalCursorLocation is already correct
 		}
 
@@ -2599,20 +2603,10 @@ void FSlateEditableTextLayout::UpdateCursorHighlight()
 		const TArray< FTextLayout::FLineModel >& Lines = TextLayout->GetLineModels();
 		if (Lines.IsValidIndex(CursorPosition.GetLineIndex()))
 		{
-			const int32 LineTextLength = Lines[CursorPosition.GetLineIndex()].Text->Len();
-
-			if (LineTextLength == 0)
-			{
-				ActiveLineHighlights.Add(FTextLineHighlight(CursorPosition.GetLineIndex(), FTextRange(0, 0), CursorZOrder, CursorLineHighlighter.ToSharedRef()));
-			}
-			else if (CursorPosition.GetOffset() == LineTextLength)
-			{
-				ActiveLineHighlights.Add(FTextLineHighlight(CursorPosition.GetLineIndex(), FTextRange(LineTextLength - 1, LineTextLength), CursorZOrder, CursorLineHighlighter.ToSharedRef()));
-			}
-			else
-			{
-				ActiveLineHighlights.Add(FTextLineHighlight(CursorPosition.GetLineIndex(), FTextRange(CursorPosition.GetOffset(), CursorPosition.GetOffset() + 1), CursorZOrder, CursorLineHighlighter.ToSharedRef()));
-			}
+			// Ensure the cursor is sitting on a valid grapheme
+			// Right-aligned cursors naively subtract 1 from their real position, which may leave us in the middle of a codepoint or grapheme and break measuring
+			const FTextSelection CursorSelection = TextLayout->GetGraphemeAt(CursorPosition);
+			ActiveLineHighlights.Add(FTextLineHighlight(CursorPosition.GetLineIndex(), FTextRange(CursorSelection.GetBeginning().GetOffset(), CursorSelection.GetEnd().GetOffset()), CursorZOrder, CursorLineHighlighter.ToSharedRef()));
 		}
 	}
 
@@ -2831,7 +2825,7 @@ FTextLocation FSlateEditableTextLayout::TranslatedLocation(const FTextLocation&
 	const TArray< FTextLayout::FLineModel >& Lines = TextLayout->GetLineModels();
 
 	// Move to the previous or next grapheme based upon the requested direction and current position
-	GraphemeBreakIterator->SetString(*Lines[Location.GetLineIndex()].Text);
+	GraphemeBreakIterator->SetStringRef(&Lines[Location.GetLineIndex()].Text.Get());
 	const int32 NewOffsetInLine = (Direction > 0) ? GraphemeBreakIterator->MoveToCandidateAfter(Location.GetOffset()) : GraphemeBreakIterator->MoveToCandidateBefore(Location.GetOffset());
 	GraphemeBreakIterator->ClearString();
 
diff --git a/Engine/Source/Runtime/Slate/Public/Framework/Text/TextLayout.h b/Engine/Source/Runtime/Slate/Public/Framework/Text/TextLayout.h
--- a/Engine/Source/Runtime/Slate/Public/Framework/Text/TextLayout.h
+++ b/Engine/Source/Runtime/Slate/Public/Framework/Text/TextLayout.h
@@ -529,6 +529,8 @@ public:
 
 	void GetSelectionAsText(FString& DisplayText, const FTextSelection& Selection) const;
 
+	FTextSelection GetGraphemeAt(const FTextLocation& Location) const;
+
 	FTextSelection GetWordAt(const FTextLocation& Location) const;
 
 protected:
