- Changed the SSafeZone scale multiplier so it uses a TOptional instead of PLATFORM_XBOXONE.

diff --git a/Engine/Source/Runtime/Slate/Private/Widgets/Layout/SSafeZone.cpp b/Engine/Source/Runtime/Slate/Private/Widgets/Layout/SSafeZone.cpp
--- a/Engine/Source/Runtime/Slate/Private/Widgets/Layout/SSafeZone.cpp
+++ b/Engine/Source/Runtime/Slate/Private/Widgets/Layout/SSafeZone.cpp
@@ -7,23 +7,37 @@
 #include "Widgets/SViewport.h"
 
 float GSafeZoneScale = 1.0f;
-static FAutoConsoleVariableRef CVarDumpVMIR(
+static FAutoConsoleVariableRef CVarSafeZoneScale(
 	TEXT("SafeZone.Scale"),
 	GSafeZoneScale,
 	TEXT("The safezone scale."),
 	ECVF_Default
 );
 
-void SSafeZone::SetGlobalSafeZoneScale(float InScale)
+bool bEnableSafeZoneScale = false;
+static FAutoConsoleVariableRef CVarEnableSafeZoneScale(
+	TEXT("SafeZone.EnableScale"),
+	bEnableSafeZoneScale,
+	TEXT("IS the safe zone scale enabled?"),
+	ECVF_Default
+);
+
+
+void SSafeZone::SetGlobalSafeZoneScale(TOptional<float> InScale)
 {
-	GSafeZoneScale = InScale;
+	bEnableSafeZoneScale = InScale.IsSet();
+	GSafeZoneScale = InScale.Get(0.f);
 
 	FCoreDelegates::OnSafeFrameChangedEvent.Broadcast();
 }
 
-float SSafeZone::GetGlobalSafeZoneScale()
+TOptional<float> SSafeZone::GetGlobalSafeZoneScale()
 {
-	return GSafeZoneScale;
+	if (bEnableSafeZoneScale)
+	{
+		return GSafeZoneScale;
+	}
+	return TOptional<float>();
 }
 
 void SSafeZone::Construct( const FArguments& InArgs )
@@ -100,9 +114,11 @@ void SSafeZone::UpdateSafeMargin() const
 		}
 	}
 
-#if PLATFORM_XBOXONE
-	SafeMargin = SafeMargin * GSafeZoneScale;
-#endif
+	TOptional<float> SafeZoneScale = GetGlobalSafeZoneScale();
+	if (SafeZoneScale.IsSet())
+	{
+		SafeMargin = SafeMargin * GSafeZoneScale;
+	}
 
 	SafeMargin = FMargin(bPadLeft ? SafeMargin.Left : 0.0f, bPadTop ? SafeMargin.Top : 0.0f, bPadRight ? SafeMargin.Right : 0.0f, bPadBottom ? SafeMargin.Bottom : 0.0f);
 
diff --git a/Engine/Source/Runtime/Slate/Public/Widgets/Layout/SSafeZone.h b/Engine/Source/Runtime/Slate/Public/Widgets/Layout/SSafeZone.h
--- a/Engine/Source/Runtime/Slate/Public/Widgets/Layout/SSafeZone.h
+++ b/Engine/Source/Runtime/Slate/Public/Widgets/Layout/SSafeZone.h
@@ -114,8 +114,8 @@ public:
 	virtual void OnArrangeChildren( const FGeometry& AllottedGeometry, FArrangedChildren& ArrangedChildren ) const override;
 	virtual FVector2D ComputeDesiredSize(float LayoutScale) const override;
 
-	static void SetGlobalSafeZoneScale(float InScale);
-	static float GetGlobalSafeZoneScale();
+	static void SetGlobalSafeZoneScale(TOptional<float> InScale);
+	static TOptional<float> GetGlobalSafeZoneScale();
 
 private:
 
