ListViews utilizing animated scrolling trigger OnItemScrolledIntoView as soon as the widget for the target item is available, even if we still have some scrolling to do.
- Previously we waited until the inertial scroll was completely finished, incorrectly introducing a small delay between the widget being in view and it being acknowledged
- SListView & STileView both now only return a ScrollIntoViewResult of Success if a widget for the item that we're meant to notify is in view actually exists when we're animating scrolling. For instant scrolling we still can and do return success with the confidence that changing the offset is a sufficient guarantee that the widget will be generated by the time the STableViewBase is done with its Tick

diff --git a/Engine/Source/Runtime/Slate/Private/Widgets/Views/STableViewBase.cpp b/Engine/Source/Runtime/Slate/Private/Widgets/Views/STableViewBase.cpp
--- a/Engine/Source/Runtime/Slate/Private/Widgets/Views/STableViewBase.cpp
+++ b/Engine/Source/Runtime/Slate/Private/Widgets/Views/STableViewBase.cpp
@@ -334,14 +334,18 @@ void STableViewBase::Tick( const FGeometry& AllottedGeometry, const double InCur
 
 			Invalidate(EInvalidateWidget::ChildOrder);
 
-			if (ScrollIntoViewResult == EScrollIntoViewResult::Deferred || CurrentScrollOffset != TargetScrollOffset)
+			const bool bScrollIntoViewSuccess = ScrollIntoViewResult == EScrollIntoViewResult::Success;
+			if (bScrollIntoViewSuccess)
 			{
-				// We call this rather than just leave bItemsNeedRefresh as true to ensure that EnsureTickToRefresh is registered
-				RequestLayoutRefresh();
+				// Notify as soon as we've made a widget for the item, even if we still have scrolling to do
+				NotifyItemScrolledIntoView();
 			}
-			else
+
+			if (!bScrollIntoViewSuccess || CurrentScrollOffset != TargetScrollOffset)
 			{
-				NotifyItemScrolledIntoView();
+				// Either we haven't made the item yet or we still have scrolling to do, so we'll need another refresh next frame
+				// We call this rather than just leave bItemsNeedRefresh as true to ensure that EnsureTickToRefresh is registered
+				RequestLayoutRefresh();
 			}
 		}
 	}
diff --git a/Engine/Source/Runtime/Slate/Public/Widgets/Views/SListView.h b/Engine/Source/Runtime/Slate/Public/Widgets/Views/SListView.h
--- a/Engine/Source/Runtime/Slate/Public/Widgets/Views/SListView.h
+++ b/Engine/Source/Runtime/Slate/Public/Widgets/Views/SListView.h
@@ -1625,6 +1625,12 @@ protected:
 			TListTypeTraits<ItemType>::ResetPtr(ItemToScrollIntoView);
 		}
 
+		if (this->bEnableAnimatedScrolling && TListTypeTraits<ItemType>::IsPtrValid(ItemToNotifyWhenInView))
+		{
+			// When we have a target item we're shooting for, we haven't succeeded with the scroll until a widget for it exists
+			const bool bHasWidgetForItem = WidgetFromItem(TListTypeTraits<ItemType>::NullableItemTypeConvertToItemType(ItemToNotifyWhenInView)).IsValid();
+			return bHasWidgetForItem ? EScrollIntoViewResult::Success : EScrollIntoViewResult::Deferred;
+		}
 		return EScrollIntoViewResult::Success;
 	}
 
diff --git a/Engine/Source/Runtime/Slate/Public/Widgets/Views/STileView.h b/Engine/Source/Runtime/Slate/Public/Widgets/Views/STileView.h
--- a/Engine/Source/Runtime/Slate/Public/Widgets/Views/STileView.h
+++ b/Engine/Source/Runtime/Slate/Public/Widgets/Views/STileView.h
@@ -489,7 +489,13 @@ protected:
 			TListTypeTraits<ItemType>::ResetPtr(this->ItemToScrollIntoView);
 		}
 
-		return SListView<ItemType>::EScrollIntoViewResult::Success;
+		if (this->bEnableAnimatedScrolling && TListTypeTraits<ItemType>::IsPtrValid(this->ItemToNotifyWhenInView))
+		{
+			// When we have a target item we're shooting for, we haven't succeeded with the scroll until a widget for it exists
+			const bool bHasWidgetForItem = WidgetFromItem(TListTypeTraits<ItemType>::NullableItemTypeConvertToItemType(this->ItemToNotifyWhenInView)).IsValid();
+			return bHasWidgetForItem ? EScrollIntoViewResult::Success : EScrollIntoViewResult::Deferred;
+		}
+		return EScrollIntoViewResult::Success;
 	}
 
 	/** Should the left and right navigations be handled as a wrap when hitting the bounds. (you'll move to the previous / next row when appropriate) */
