Appeasing clang

diff --git a/Engine/Source/Runtime/Slate/Public/Widgets/Views/STileView.h b/Engine/Source/Runtime/Slate/Public/Widgets/Views/STileView.h
--- a/Engine/Source/Runtime/Slate/Public/Widgets/Views/STileView.h
+++ b/Engine/Source/Runtime/Slate/Public/Widgets/Views/STileView.h
@@ -493,9 +493,9 @@ protected:
 		{
 			// When we have a target item we're shooting for, we haven't succeeded with the scroll until a widget for it exists
 			const bool bHasWidgetForItem = WidgetFromItem(TListTypeTraits<ItemType>::NullableItemTypeConvertToItemType(this->ItemToNotifyWhenInView)).IsValid();
-			return bHasWidgetForItem ? EScrollIntoViewResult::Success : EScrollIntoViewResult::Deferred;
+			return bHasWidgetForItem ? SListView<ItemType>::EScrollIntoViewResult::Success : SListView<ItemType>::EScrollIntoViewResult::Deferred;
 		}
-		return EScrollIntoViewResult::Success;
+		return SListView<ItemType>::EScrollIntoViewResult::Success;
 	}
 
 	/** Should the left and right navigations be handled as a wrap when hitting the bounds. (you'll move to the previous / next row when appropriate) */
