More clang appeasement in STileView

diff --git a/Engine/Source/Runtime/Slate/Public/Widgets/Views/STileView.h b/Engine/Source/Runtime/Slate/Public/Widgets/Views/STileView.h
--- a/Engine/Source/Runtime/Slate/Public/Widgets/Views/STileView.h
+++ b/Engine/Source/Runtime/Slate/Public/Widgets/Views/STileView.h
@@ -492,7 +492,7 @@ protected:
 		if (this->bEnableAnimatedScrolling && TListTypeTraits<ItemType>::IsPtrValid(this->ItemToNotifyWhenInView))
 		{
 			// When we have a target item we're shooting for, we haven't succeeded with the scroll until a widget for it exists
-			const bool bHasWidgetForItem = WidgetFromItem(TListTypeTraits<ItemType>::NullableItemTypeConvertToItemType(this->ItemToNotifyWhenInView)).IsValid();
+			const bool bHasWidgetForItem = this->WidgetFromItem(TListTypeTraits<ItemType>::NullableItemTypeConvertToItemType(this->ItemToNotifyWhenInView)).IsValid();
 			return bHasWidgetForItem ? SListView<ItemType>::EScrollIntoViewResult::Success : SListView<ItemType>::EScrollIntoViewResult::Deferred;
 		}
 		return SListView<ItemType>::EScrollIntoViewResult::Success;
