Animated scrolling on tile views no longer siezes up into non-animated scrolling when the last line of widgets is in view

diff --git a/Engine/Source/Runtime/Slate/Public/Widgets/Views/SListView.h b/Engine/Source/Runtime/Slate/Public/Widgets/Views/SListView.h
--- a/Engine/Source/Runtime/Slate/Public/Widgets/Views/SListView.h
+++ b/Engine/Source/Runtime/Slate/Public/Widgets/Views/SListView.h
@@ -1068,15 +1068,12 @@ public:
 					? ItemLength * ItemsInView	// For the first item, ItemsInView <= 1.0f
 					: ItemLength;
 
-				if (ItemIndex >= ItemsSource->Num() - 1)
-				{
-					bAtEndOfList = true;
-				}
+				bAtEndOfList = ItemIndex >= ItemsSource->Num() - 1;
 
-				if (ViewLengthUsedSoFar >= MyDimensions.ScrollAxis)
-				{
-					bHasFilledAvailableArea = true;
-				}
+				// Note: To account for accrued error from floating point truncation and addition in our sum of dimensions used, 
+				//	we pad the allotted axis just a little to be sure we have filled the available space.
+				static const float FloatingPointPrecisionOffset = 0.001f;
+				bHasFilledAvailableArea = ViewLengthUsedSoFar >= MyDimensions.ScrollAxis + FloatingPointPrecisionOffset;
 			}
 
 			// Handle scenario b.
diff --git a/Engine/Source/Runtime/Slate/Public/Widgets/Views/STileView.h b/Engine/Source/Runtime/Slate/Public/Widgets/Views/STileView.h
--- a/Engine/Source/Runtime/Slate/Public/Widgets/Views/STileView.h
+++ b/Engine/Source/Runtime/Slate/Public/Widgets/Views/STileView.h
@@ -300,11 +300,11 @@ public:
 				{
 					bNewLine = false;
 					
-					float LineFraction = 1.0f;
+					float LineFraction = 1.f;
 					if (bFirstLine)
 					{
 						bFirstLine = false;
-						LineFraction = 1.0f - FMath::Max(FMath::Fractional(ClampedScrollOffset / NumItemsPerLine), 0.0f);
+						LineFraction -= FMath::Fractional(ClampedScrollOffset / NumItemsPerLine);
 					}
 
 					DimensionsUsedSoFar.ScrollAxis += TileDimensions.ScrollAxis * LineFraction;
@@ -324,23 +324,20 @@ public:
 				// The widget used up some of the available space for the current line
 				DimensionsUsedSoFar.LineAxis += TileDimensions.LineAxis;
 
-				// Check to see if we have completed a line of widgets
+				bIsAtEndOfList = ItemIndex >= NumItems - 1;
+
 				if (DimensionsUsedSoFar.LineAxis + TileDimensions.LineAxis > AllottedDimensions.LineAxis)
 				{
 					// A new line of widgets was completed - time to start another one
 					DimensionsUsedSoFar.LineAxis = 0;
 					bNewLine = true;
-
-					// Stop when we've finished generating a line of widgets partially clipped by the end of the view
-					if (DimensionsUsedSoFar.ScrollAxis >= AllottedDimensions.ScrollAxis)
-					{
-						bHasFilledAvailableArea = true;
-					}
 				}
 
-				if (ItemIndex >= NumItems - 1)
+				if (bIsAtEndOfList || bNewLine)
 				{
-					bIsAtEndOfList = true;
+					// We've filled all the available area when we've finished a line that's partially clipped by the end of the view
+					static const float FloatingPointPrecisionOffset = 0.001f;
+					bHasFilledAvailableArea = DimensionsUsedSoFar.ScrollAxis > AllottedDimensions.ScrollAxis + FloatingPointPrecisionOffset;
 				}
 			}
 
@@ -348,7 +345,7 @@ public:
 			this->WidgetGenerator.OnEndGenerationPass();
 
 			const float TotalGeneratedLineAxisSize = FMath::CeilToFloat(NumLinesShownOnScreen) * TileDimensions.ScrollAxis;
-			return STableViewBase::FReGenerateResults(ClampedScrollOffset, TotalGeneratedLineAxisSize, NumLinesShownOnScreen, bIsAtEndOfList);
+			return STableViewBase::FReGenerateResults(ClampedScrollOffset, TotalGeneratedLineAxisSize, NumLinesShownOnScreen, bIsAtEndOfList && !bHasFilledAvailableArea);
 		}
 
 		return STableViewBase::FReGenerateResults(0, 0, 0, false);
