[UE-79379] - Resolving the silly compile error in StressTestUnity about shadow declarations on a static property declared in three separate scopes. One is global, but static should be restraining it to the cpp and disallowing any extern declarations of it. Somehow unity is throwing this off I suppose.

diff --git a/Engine/Source/Runtime/Slate/Public/Widgets/Views/SListView.h b/Engine/Source/Runtime/Slate/Public/Widgets/Views/SListView.h
--- a/Engine/Source/Runtime/Slate/Public/Widgets/Views/SListView.h
+++ b/Engine/Source/Runtime/Slate/Public/Widgets/Views/SListView.h
@@ -1072,7 +1072,7 @@ public:
 
 				// Note: To account for accrued error from floating point truncation and addition in our sum of dimensions used, 
 				//	we pad the allotted axis just a little to be sure we have filled the available space.
-				static const float FloatingPointPrecisionOffset = 0.001f;
+				const float FloatingPointPrecisionOffset = 0.001f;
 				bHasFilledAvailableArea = ViewLengthUsedSoFar >= MyDimensions.ScrollAxis + FloatingPointPrecisionOffset;
 			}
 
diff --git a/Engine/Source/Runtime/Slate/Public/Widgets/Views/STileView.h b/Engine/Source/Runtime/Slate/Public/Widgets/Views/STileView.h
--- a/Engine/Source/Runtime/Slate/Public/Widgets/Views/STileView.h
+++ b/Engine/Source/Runtime/Slate/Public/Widgets/Views/STileView.h
@@ -336,7 +336,7 @@ public:
 				if (bIsAtEndOfList || bNewLine)
 				{
 					// We've filled all the available area when we've finished a line that's partially clipped by the end of the view
-					static const float FloatingPointPrecisionOffset = 0.001f;
+					const float FloatingPointPrecisionOffset = 0.001f;
 					bHasFilledAvailableArea = DimensionsUsedSoFar.ScrollAxis > AllottedDimensions.ScrollAxis + FloatingPointPrecisionOffset;
 				}
 			}
