Slate - Adding the ability to directly control which list view items are navigable or selectable before we actually select or navigate to them.  With this control we can now correctly skip things like headers or other just visual elements in a list that are intended to never get focus and just have the gamepad navigate past them.  Before this was done with selection hacks, where after selection we'd attempt to select the item beyond it, however this tended to trigger side effects that were unintended since anyone listening would still get those selections.

Matt.Kuhlenschmidt, Dan.Hertzka

diff --git a/Engine/Source/Runtime/Slate/Public/Framework/SlateDelegates.h b/Engine/Source/Runtime/Slate/Public/Framework/SlateDelegates.h
--- a/Engine/Source/Runtime/Slate/Public/Framework/SlateDelegates.h
+++ b/Engine/Source/Runtime/Slate/Public/Framework/SlateDelegates.h
@@ -188,6 +188,14 @@ public:
 		ESelectInfo::Type
 	)
 
+	/** A delegate type to control if a item is selectable or navigable. */
+	DECLARE_DELEGATE_RetVal_OneParam(
+		bool,
+		FIsSelectableOrNavigable,
+		/** param: The newly selected value */
+		ArgumentType
+	)
+
 	DECLARE_DELEGATE_TwoParams(
 		FOnExpansionChanged,
 		ArgumentType,
diff --git a/Engine/Source/Runtime/Slate/Public/Widgets/Views/SListView.h b/Engine/Source/Runtime/Slate/Public/Widgets/Views/SListView.h
--- a/Engine/Source/Runtime/Slate/Public/Widgets/Views/SListView.h
+++ b/Engine/Source/Runtime/Slate/Public/Widgets/Views/SListView.h
@@ -58,8 +58,9 @@ public:
 	using FOnGenerateRow            = typename TSlateDelegates< ItemType >::FOnGenerateRow;
 	using FOnItemScrolledIntoView   = typename TSlateDelegates< ItemType >::FOnItemScrolledIntoView;
 	using FOnSelectionChanged       = typename TSlateDelegates< NullableItemType >::FOnSelectionChanged;
-	using FOnMouseButtonClick       = typename TSlateDelegates< ItemType >::FOnMouseButtonClick ;
-	using FOnMouseButtonDoubleClick = typename TSlateDelegates< ItemType >::FOnMouseButtonDoubleClick ;
+	using FIsSelectableOrNavigable	= typename TSlateDelegates< ItemType >::FIsSelectableOrNavigable;
+	using FOnMouseButtonClick       = typename TSlateDelegates< ItemType >::FOnMouseButtonClick;
+	using FOnMouseButtonDoubleClick = typename TSlateDelegates< ItemType >::FOnMouseButtonDoubleClick;
 
 	typedef typename TSlateDelegates< ItemType >::FOnItemToString_Debug FOnItemToString_Debug;
 
@@ -78,6 +79,7 @@ public:
 		, _OnMouseButtonClick()
 		, _OnMouseButtonDoubleClick()
 		, _OnSelectionChanged()
+		, _OnIsSelectableOrNavigable()
 		, _SelectionMode(ESelectionMode::Multi)
 		, _ClearSelectionOnClick(true)
 		, _ExternalScrollbar()
@@ -119,6 +121,8 @@ public:
 
 		SLATE_EVENT( FOnSelectionChanged, OnSelectionChanged )
 
+		SLATE_EVENT( FIsSelectableOrNavigable, OnIsSelectableOrNavigable)
+
 		SLATE_ATTRIBUTE( ESelectionMode::Type, SelectionMode )
 
 		SLATE_ARGUMENT( TSharedPtr<SHeaderRow>, HeaderRow )
@@ -177,6 +181,7 @@ public:
 		this->OnClick = InArgs._OnMouseButtonClick;
 		this->OnDoubleClick = InArgs._OnMouseButtonDoubleClick;
 		this->OnSelectionChanged = InArgs._OnSelectionChanged;
+		this->OnIsSelectableOrNavigable = InArgs._OnIsSelectableOrNavigable;
 		this->SelectionMode = InArgs._SelectionMode;
 
 		this->bClearSelectionOnClick = InArgs._ClearSelectionOnClick;
@@ -1244,7 +1249,7 @@ public:
 			return;
 		}
 
-		Private_SetItemSelection( InItem, bSelected, SelectInfo != ESelectInfo::Direct);
+		Private_SetItemSelection(InItem, bSelected, SelectInfo != ESelectInfo::Direct);
 		Private_SignalSelectionChanged(SelectInfo);
 	}
 
@@ -1397,11 +1402,20 @@ public:
 	 * @param Item The item to navigate to on the next tick.
 	 */
 	void RequestNavigateToItem(ItemType Item, const uint32 UserIndex = 0)
+	{
+		Item = Private_FindNextSelectableOrNavigable(Item);
+		Private_RequestNavigateToItem(Item, UserIndex);
+	}
+
+private:
+	void Private_RequestNavigateToItem(ItemType Item, const uint32 UserIndex = 0)
 	{
 		bNavigateOnScrollIntoView = true;
 		RequestScrollIntoView(Item, UserIndex);
 	}
 
+public:
+
 	/**
 	 * Cancels a previous request to scroll an item into view (cancels navigation requests as well).
 	 */
@@ -1418,12 +1432,21 @@ public:
 	 * @param SoleSelectedItem   Sole item that should be selected.
 	 * @param SelectInfo Provides context on how the selection changed
 	 */
-	void SetSelection( ItemType SoleSelectedItem, ESelectInfo::Type SelectInfo = ESelectInfo::Direct  )
+	void SetSelection(ItemType SoleSelectedItem, ESelectInfo::Type SelectInfo = ESelectInfo::Direct)
+	{
+		SoleSelectedItem = Private_FindNextSelectableOrNavigable(SoleSelectedItem);
+		Private_SetSelection(SoleSelectedItem);
+	}
+
+private:
+	void Private_SetSelection(ItemType SoleSelectedItem, ESelectInfo::Type SelectInfo = ESelectInfo::Direct)
 	{
 		SelectedItems.Empty();
 		SetItemSelection( SoleSelectedItem, true, SelectInfo );
 	}
 
+public:
+
 	/** 
 	 * Set the current selection mode of the list.
 	 * If going from multi-select to a type of single-select and one item is selected, it will be maintained (otherwise all will be cleared).
@@ -1801,14 +1824,58 @@ protected:
 
 protected:
 
+	ItemType Private_FindNextSelectableOrNavigable(const ItemType& InItemToSelect)
+	{
+		ItemType ItemToSelect = InItemToSelect;
+
+		if (OnIsSelectableOrNavigable.IsBound())
+		{
+			const int32 PendingItemIndex = ItemsSource->Find(ItemToSelect);
+
+			NullableItemType LastSelectedItem = nullptr;
+			if (SelectedItems.Num() == 1)
+			{
+				LastSelectedItem = *SelectedItems.CreateIterator();
+			}
+
+			bool bSelectNextItem = true;
+			while (!OnIsSelectableOrNavigable.Execute(ItemToSelect))
+			{
+				if (TListTypeTraits<ItemType>::IsPtrValid(LastSelectedItem) && PendingItemIndex > 0)
+				{
+					ItemType NonNullLastSelectedItem = TListTypeTraits<ItemType>::NullableItemTypeConvertToItemType(LastSelectedItem);
+					const int32 LastSelectedItemIdx = ItemsSource->Find(NonNullLastSelectedItem);
+
+					// If the previously selected item was before the header, assume we're navigating down and want to select the next item
+					// Otherwise, assume the opposite and navigate to the previous item
+					bSelectNextItem = LastSelectedItemIdx < PendingItemIndex;
+				}
+
+				const int32 NewSelectionIdx = PendingItemIndex + (bSelectNextItem ? 1 : -1);
+				if (ItemsSource->IsValidIndex(NewSelectionIdx))
+				{
+					ItemToSelect = (*ItemsSource)[NewSelectionIdx];
+				}
+				else
+				{
+					break;
+				}
+			}
+		}
+
+		return ItemToSelect;
+	}
+
 	/**
 	 * Selects the specified item and scrolls it into view. If shift is held, it will be a range select.
 	 * 
 	 * @param ItemToSelect		The item that was selected by a keystroke
 	 * @param InInputEvent	The key event that caused this selection
 	 */
-	virtual void NavigationSelect(const ItemType& ItemToSelect, const FInputEvent& InInputEvent)
+	virtual void NavigationSelect(const ItemType& InItemToSelect, const FInputEvent& InInputEvent)
 	{
+		ItemType ItemToSelect = Private_FindNextSelectableOrNavigable(InItemToSelect);
+
 		const ESelectionMode::Type CurrentSelectionMode = SelectionMode.Get();
 
 		if (CurrentSelectionMode != ESelectionMode::None)
@@ -1817,7 +1884,7 @@ protected:
 			SelectorItem = ItemToSelect;
 
 			// Always request scroll into view, otherwise partially visible items will be selected - also do this before signaling selection for similar stomp-allowing reasons
-			RequestNavigateToItem(ItemToSelect, InInputEvent.GetUserIndex());
+			Private_RequestNavigateToItem(ItemToSelect, InInputEvent.GetUserIndex());
 
 			if (CurrentSelectionMode == ESelectionMode::Multi && (InInputEvent.IsShiftDown() || InInputEvent.IsControlDown()))
 			{
@@ -1838,7 +1905,7 @@ protected:
 			else
 			{
 				// Single select.
-				this->SetSelection(ItemToSelect, ESelectInfo::OnNavigation);;
+				this->Private_SetSelection(ItemToSelect, ESelectInfo::OnNavigation);
 			}
 		}
 	}
@@ -1892,6 +1959,9 @@ protected:
 	/** Delegate to invoke when selection changes. */
 	FOnSelectionChanged OnSelectionChanged;
 
+	/** Delegate to invoke to see if we can navigate or select item. */
+	FIsSelectableOrNavigable OnIsSelectableOrNavigable;
+
 	/** Called when the user clicks on an element int he list view with the left mouse button */
 	FOnMouseButtonClick OnClick;
 
diff --git a/Engine/Source/Runtime/Slate/Public/Widgets/Views/STileView.h b/Engine/Source/Runtime/Slate/Public/Widgets/Views/STileView.h
--- a/Engine/Source/Runtime/Slate/Public/Widgets/Views/STileView.h
+++ b/Engine/Source/Runtime/Slate/Public/Widgets/Views/STileView.h
@@ -30,6 +30,7 @@ public:
 	typedef typename TSlateDelegates< ItemType >::FOnMouseButtonClick FOnMouseButtonClick;
 	typedef typename TSlateDelegates< ItemType >::FOnMouseButtonDoubleClick FOnMouseButtonDoubleClick;
 	typedef typename TSlateDelegates< NullableItemType >::FOnSelectionChanged FOnSelectionChanged;
+	typedef typename TSlateDelegates< ItemType >::FIsSelectableOrNavigable FIsSelectableOrNavigable;
 
 	typedef typename TSlateDelegates< ItemType >::FOnItemToString_Debug FOnItemToString_Debug; 
 
@@ -48,6 +49,7 @@ public:
 		, _OnMouseButtonClick()
 		, _OnMouseButtonDoubleClick()
 		, _OnSelectionChanged()
+		, _OnIsSelectableOrNavigable()
 		, _SelectionMode(ESelectionMode::Multi)
 		, _ClearSelectionOnClick(true)
 		, _ExternalScrollbar()
@@ -91,6 +93,8 @@ public:
 
 		SLATE_EVENT( FOnSelectionChanged, OnSelectionChanged )
 
+		SLATE_EVENT( FIsSelectableOrNavigable, OnIsSelectableOrNavigable)
+
 		SLATE_ATTRIBUTE( ESelectionMode::Type, SelectionMode )
 
 		SLATE_ARGUMENT ( bool, ClearSelectionOnClick )
@@ -144,6 +148,7 @@ public:
 		this->OnClick = InArgs._OnMouseButtonClick;
 		this->OnDoubleClick = InArgs._OnMouseButtonDoubleClick;
 		this->OnSelectionChanged = InArgs._OnSelectionChanged;
+		this->OnIsSelectableOrNavigable = InArgs._OnIsSelectableOrNavigable;
 		this->SelectionMode = InArgs._SelectionMode;
 
 		this->bClearSelectionOnClick = InArgs._ClearSelectionOnClick;
diff --git a/Engine/Source/Runtime/Slate/Public/Widgets/Views/STreeView.h b/Engine/Source/Runtime/Slate/Public/Widgets/Views/STreeView.h
--- a/Engine/Source/Runtime/Slate/Public/Widgets/Views/STreeView.h
+++ b/Engine/Source/Runtime/Slate/Public/Widgets/Views/STreeView.h
@@ -100,6 +100,7 @@ public:
 	using FOnSetExpansionRecursive  = typename TSlateDelegates< ItemType >::FOnSetExpansionRecursive;
 	using FOnItemScrolledIntoView   = typename TSlateDelegates< ItemType >::FOnItemScrolledIntoView;
 	using FOnSelectionChanged       = typename TSlateDelegates< NullableItemType >::FOnSelectionChanged;
+	using FIsSelectableOrNavigable	= typename TSlateDelegates< ItemType >::FIsSelectableOrNavigable;
 	using FOnMouseButtonClick       = typename TSlateDelegates< ItemType >::FOnMouseButtonClick;
 	using FOnMouseButtonDoubleClick = typename TSlateDelegates< ItemType >::FOnMouseButtonDoubleClick;
 	using FOnExpansionChanged       = typename TSlateDelegates< ItemType >::FOnExpansionChanged;
@@ -121,6 +122,7 @@ public:
 		, _OnMouseButtonDoubleClick()
 		, _OnSelectionChanged()
 		, _OnExpansionChanged()
+		, _OnIsSelectableOrNavigable()
 		, _SelectionMode(ESelectionMode::Multi)
 		, _ClearSelectionOnClick(true)
 		, _ExternalScrollbar()
@@ -165,6 +167,8 @@ public:
 
 		SLATE_EVENT( FOnExpansionChanged, OnExpansionChanged )
 
+		SLATE_EVENT(FIsSelectableOrNavigable, OnIsSelectableOrNavigable)
+
 		SLATE_ATTRIBUTE( ESelectionMode::Type, SelectionMode )
 
 		SLATE_ARGUMENT( TSharedPtr<SHeaderRow>, HeaderRow )
@@ -220,6 +224,7 @@ public:
 		this->OnDoubleClick = InArgs._OnMouseButtonDoubleClick;
 		this->OnSelectionChanged = InArgs._OnSelectionChanged;
 		this->OnExpansionChanged = InArgs._OnExpansionChanged;
+		this->OnIsSelectableOrNavigable = InArgs._OnIsSelectableOrNavigable;
 		this->SelectionMode = InArgs._SelectionMode;
 
 		this->bClearSelectionOnClick = InArgs._ClearSelectionOnClick;
diff --git a/Engine/Source/Runtime/UMG/Public/Components/ListViewBase.h b/Engine/Source/Runtime/UMG/Public/Components/ListViewBase.h
--- a/Engine/Source/Runtime/UMG/Public/Components/ListViewBase.h
+++ b/Engine/Source/Runtime/UMG/Public/Components/ListViewBase.h
@@ -236,6 +236,7 @@ protected:
 			.Orientation(Args.Orientation)
 			.OnGenerateRow_UObject(Implementer, &UListViewBaseT::HandleGenerateRow)
 			.OnSelectionChanged_UObject(Implementer, &UListViewBaseT::HandleSelectionChanged)
+			.OnIsSelectableOrNavigable_UObject(Implementer, &UListViewBaseT::HandleIsSelectableOrNavigable)
 			.OnRowReleased_UObject(Implementer, &UListViewBaseT::HandleRowReleased)
 			.OnItemScrolledIntoView_UObject(Implementer, &UListViewBaseT::HandleItemScrolledIntoView)
 			.OnListViewScrolled_UObject(Implementer, &UListViewBaseT::HandleListViewScrolled)
@@ -272,6 +273,7 @@ protected:
 			.OnGenerateTile_UObject(Implementer, &UListViewBaseT::HandleGenerateRow)
 			.OnTileReleased_UObject(Implementer, &UListViewBaseT::HandleRowReleased)
 			.OnSelectionChanged_UObject(Implementer, &UListViewBaseT::HandleSelectionChanged)
+			.OnIsSelectableOrNavigable_UObject(Implementer, &UListViewBaseT::HandleIsSelectableOrNavigable)
 			.OnItemScrolledIntoView_UObject(Implementer, &UListViewBaseT::HandleItemScrolledIntoView)
 			.OnTileViewScrolled_UObject(Implementer, &UListViewBaseT::HandleListViewScrolled)
 			.OnMouseButtonClick_UObject(Implementer, &UListViewBaseT::HandleItemClicked)
@@ -299,6 +301,7 @@ protected:
 			.SelectionMode(Args.SelectionMode)
 			.OnGenerateRow_UObject(Implementer, &UListViewBaseT::HandleGenerateRow)
 			.OnSelectionChanged_UObject(Implementer, &UListViewBaseT::HandleSelectionChanged)
+			.OnIsSelectableOrNavigable_UObject(Implementer, &UListViewBaseT::HandleIsSelectableOrNavigable)
 			.OnRowReleased_UObject(Implementer, &UListViewBaseT::HandleRowReleased)
 			.OnItemScrolledIntoView_UObject(Implementer, &UListViewBaseT::HandleItemScrolledIntoView)
 			.OnTreeViewScrolled_UObject(Implementer, &UListViewBaseT::HandleListViewScrolled)
@@ -403,6 +406,7 @@ protected:
 	virtual void OnItemClickedInternal(ItemType Item) {}
 	virtual void OnItemDoubleClickedInternal(ItemType Item) {}
 	virtual void OnSelectionChangedInternal(NullableItemType FirstSelectedItem) {}
+	virtual bool OnIsSelectableOrNavigableInternal(ItemType FirstSelectedItem) { return true; }
 	virtual void OnItemScrolledIntoViewInternal(ItemType Item, UUserWidget& EntryWidget) {}
 	virtual void OnListViewScrolledInternal(float ItemOffset, float DistanceRemaining) {}
 	virtual void OnItemExpansionChangedInternal(ItemType Item, bool bIsExpanded) {}
@@ -443,6 +447,11 @@ private:
 		OnItemSelectionChanged().Broadcast(Item);
 	}
 
+	bool HandleIsSelectableOrNavigable(ItemType Item)
+	{
+		return OnIsSelectableOrNavigableInternal(Item);
+	}
+
 	void HandleListViewScrolled(double OffsetInItems)
 	{
 		if (SListView<ItemType>* MyListView = GetMyListView())
