Slate - Fixing losing the selection type when calling the private version of SetSelection.

diff --git a/Engine/Source/Runtime/Slate/Public/Widgets/Views/SListView.h b/Engine/Source/Runtime/Slate/Public/Widgets/Views/SListView.h
--- a/Engine/Source/Runtime/Slate/Public/Widgets/Views/SListView.h
+++ b/Engine/Source/Runtime/Slate/Public/Widgets/Views/SListView.h
@@ -1408,7 +1408,7 @@ public:
 	}
 
 private:
-	void Private_RequestNavigateToItem(ItemType Item, const uint32 UserIndex = 0)
+	void Private_RequestNavigateToItem(ItemType Item, const uint32 UserIndex)
 	{
 		bNavigateOnScrollIntoView = true;
 		RequestScrollIntoView(Item, UserIndex);
@@ -1435,11 +1435,11 @@ public:
 	void SetSelection(ItemType SoleSelectedItem, ESelectInfo::Type SelectInfo = ESelectInfo::Direct)
 	{
 		SoleSelectedItem = Private_FindNextSelectableOrNavigable(SoleSelectedItem);
-		Private_SetSelection(SoleSelectedItem);
+		Private_SetSelection(SoleSelectedItem, SelectInfo);
 	}
 
 private:
-	void Private_SetSelection(ItemType SoleSelectedItem, ESelectInfo::Type SelectInfo = ESelectInfo::Direct)
+	void Private_SetSelection(ItemType SoleSelectedItem, ESelectInfo::Type SelectInfo)
 	{
 		SelectedItems.Empty();
 		SetItemSelection( SoleSelectedItem, true, SelectInfo );
