Renaming local const named properties to appease the most nonsensical shadow variable warning I've ever seen

diff --git a/Engine/Source/Runtime/Slate/Public/Widgets/Views/SListView.h b/Engine/Source/Runtime/Slate/Public/Widgets/Views/SListView.h
--- a/Engine/Source/Runtime/Slate/Public/Widgets/Views/SListView.h
+++ b/Engine/Source/Runtime/Slate/Public/Widgets/Views/SListView.h
@@ -1077,7 +1077,7 @@ public:
 
 				// Note: To account for accrued error from floating point truncation and addition in our sum of dimensions used, 
 				//	we pad the allotted axis just a little to be sure we have filled the available space.
-				const float FloatingPointPrecisionOffset = 0.001f;
+				const float FloatPrecisionOffset = 0.001f;
 				bHasFilledAvailableArea = ViewLengthUsedSoFar >= MyDimensions.ScrollAxis + FloatingPointPrecisionOffset;
 			}
 
diff --git a/Engine/Source/Runtime/Slate/Public/Widgets/Views/STileView.h b/Engine/Source/Runtime/Slate/Public/Widgets/Views/STileView.h
--- a/Engine/Source/Runtime/Slate/Public/Widgets/Views/STileView.h
+++ b/Engine/Source/Runtime/Slate/Public/Widgets/Views/STileView.h
@@ -341,7 +341,7 @@ public:
 				if (bIsAtEndOfList || bNewLine)
 				{
 					// We've filled all the available area when we've finished a line that's partially clipped by the end of the view
-					const float FloatingPointPrecisionOffset = 0.001f;
+					const float FloatPrecisionOffset = 0.001f;
 					bHasFilledAvailableArea = DimensionsUsedSoFar.ScrollAxis > AllottedDimensions.ScrollAxis + FloatingPointPrecisionOffset;
 				}
 			}
