Fix compile

diff --git a/Engine/Source/Runtime/Slate/Public/Widgets/Views/SListView.h b/Engine/Source/Runtime/Slate/Public/Widgets/Views/SListView.h
--- a/Engine/Source/Runtime/Slate/Public/Widgets/Views/SListView.h
+++ b/Engine/Source/Runtime/Slate/Public/Widgets/Views/SListView.h
@@ -1078,7 +1078,7 @@ public:
 				// Note: To account for accrued error from floating point truncation and addition in our sum of dimensions used, 
 				//	we pad the allotted axis just a little to be sure we have filled the available space.
 				const float FloatPrecisionOffset = 0.001f;
-				bHasFilledAvailableArea = ViewLengthUsedSoFar >= MyDimensions.ScrollAxis + FloatingPointPrecisionOffset;
+				bHasFilledAvailableArea = ViewLengthUsedSoFar >= MyDimensions.ScrollAxis + FloatPrecisionOffset;
 			}
 
 			// Handle scenario b.
diff --git a/Engine/Source/Runtime/Slate/Public/Widgets/Views/STileView.h b/Engine/Source/Runtime/Slate/Public/Widgets/Views/STileView.h
--- a/Engine/Source/Runtime/Slate/Public/Widgets/Views/STileView.h
+++ b/Engine/Source/Runtime/Slate/Public/Widgets/Views/STileView.h
@@ -342,7 +342,7 @@ public:
 				{
 					// We've filled all the available area when we've finished a line that's partially clipped by the end of the view
 					const float FloatPrecisionOffset = 0.001f;
-					bHasFilledAvailableArea = DimensionsUsedSoFar.ScrollAxis > AllottedDimensions.ScrollAxis + FloatingPointPrecisionOffset;
+					bHasFilledAvailableArea = DimensionsUsedSoFar.ScrollAxis > AllottedDimensions.ScrollAxis + FloatPrecisionOffset;
 				}
 			}
 
