Fixing the sliders so that they feel better to move and properly invalidate when changed.

diff --git a/Engine/Source/Runtime/Slate/Private/Widgets/Input/SSlider.cpp b/Engine/Source/Runtime/Slate/Private/Widgets/Input/SSlider.cpp
--- a/Engine/Source/Runtime/Slate/Private/Widgets/Input/SSlider.cpp
+++ b/Engine/Source/Runtime/Slate/Private/Widgets/Input/SSlider.cpp
@@ -368,6 +368,8 @@ void SSlider::CommitValue(float NewValue)
 		ValueAttribute.Set(NewValue);
 	}
 
+	Invalidate(EInvalidateWidgetReason::Paint);
+
 	OnValueChanged.ExecuteIfBound(NewValue);
 }
 
