Fixing the logic for navigating away from jagged tile views.  They now correctly calculcate the current row and the number of rows to determine if it's fine to allow navigation to continue off the bottom of the tileview.

diff --git a/Engine/Source/Runtime/Slate/Public/Widgets/Views/SListView.h b/Engine/Source/Runtime/Slate/Public/Widgets/Views/SListView.h
--- a/Engine/Source/Runtime/Slate/Public/Widgets/Views/SListView.h
+++ b/Engine/Source/Runtime/Slate/Public/Widgets/Views/SListView.h
@@ -424,22 +424,22 @@ public:
 				AttemptSelectIndex = CurSelectionIndex - NumItemsPerLine;
 			}
 			else if ((Orientation == Orient_Vertical && NavType == EUINavigation::Down) ||
-				(Orientation == Orient_Horizontal && NavType == EUINavigation::Right))
+					 (Orientation == Orient_Horizontal && NavType == EUINavigation::Right))
 			{
 				AttemptSelectIndex = CurSelectionIndex + NumItemsPerLine;
 
 				// The list might be jagged so attempt to determine if there's a partially filled row we can move to
 				if (!ItemsSourceRef.IsValidIndex(AttemptSelectIndex))
 				{
-					int32 NumItems = ItemsSourceRef.Num();
+					const int32 NumItems = ItemsSourceRef.Num();
 					if (NumItems > 0)
 					{
 						// NumItemsWide should never be 0, ensuring for sanity
 						ensure(NumItemsPerLine > 0);
 
 						// calculate total number of rows and row of current index (1 index)
-						int32 NumLines = NumItems / NumItemsPerLine + (NumItems % NumItemsPerLine != 0 ? 1 : 0);
-						int32 CurLine = CurSelectionIndex / NumItems + 1;
+						const int32 NumLines = FMath::CeilToInt(NumItems / (float)NumItemsPerLine);
+						const int32 CurLine = FMath::CeilToInt((CurSelectionIndex + 1) / (float)NumItemsPerLine);
 
 						// if not on final row, assume a jagged list and select the final item
 						if (CurLine < NumLines)
