Moving everything based on FUserAndPointer in FSlateApplication over into FSlateUser directly
- Secondary users now have a faux cursor to control
- Slate now supports per-user tooltips and drag-drop operations
- FSlateUsers are now accessible from outside of FSlateApplication
- Internally enforcing validity of the UserIndex associated with input events within GetOrCreateUser/RegisterNewUser
- ForEachUser provides an FSlateUser& now
- Added the SLATE_SCOPE macro to allow public access to class elements from within the slate module only (external modules see them as protected)

diff --git a/Engine/Plugins/Experimental/ControlRig/Source/ControlRigEditor/Private/Graph/SGraphPinBoneNameValueWidget.cpp b/Engine/Plugins/Experimental/ControlRig/Source/ControlRigEditor/Private/Graph/SGraphPinBoneNameValueWidget.cpp
--- a/Engine/Plugins/Experimental/ControlRig/Source/ControlRigEditor/Private/Graph/SGraphPinBoneNameValueWidget.cpp
+++ b/Engine/Plugins/Experimental/ControlRig/Source/ControlRigEditor/Private/Graph/SGraphPinBoneNameValueWidget.cpp
@@ -2,6 +2,7 @@
 
 #include "SGraphPinBoneNameValueWidget.h"
 #include "DetailLayoutBuilder.h"
+#include "Framework/Application/SlateUser.h"
 
 #define LOCTEXT_NAMESPACE "GraphPinBoneNameValueWidget"
 
@@ -147,10 +148,11 @@ void SGraphPinBoneNameValueWidget::OnMenuOpenChanged(bool bOpen)
 		}
 
 		// Set focus back to ComboBox for users focusing the ListView that just closed
-		FSlateApplication::Get().ForEachUser([&](FSlateUser* User) {
-			if (FSlateApplication::Get().HasUserFocusedDescendants(AsShared(), User->GetUserIndex()))
+		TSharedRef<SWidget> ThisRef = AsShared();
+		FSlateApplication::Get().ForEachUser([&ThisRef](FSlateUser& User) {
+			if (User.HasFocusedDescendants(ThisRef))
 			{
-				FSlateApplication::Get().SetUserFocus(User->GetUserIndex(), AsShared(), EFocusCause::SetDirectly);
+				User.SetFocus(ThisRef);
 			}
 		});
 
@@ -166,11 +168,8 @@ EActiveTimerReturnType SGraphPinBoneNameValueWidget::SetFocusPostConstruct(doubl
 	if (SearchField.IsValid())
 	{
 		bool bSucceeded = false;
-		FSlateApplication::Get().ForEachUser([&](FSlateUser* User) {
-			if (FSlateApplication::Get().SetUserFocus(User->GetUserIndex(), SearchField->AsShared(), EFocusCause::SetDirectly))
-			{
-				bSucceeded = true;
-			}
+		FSlateApplication::Get().ForEachUser([&](FSlateUser& User) {
+			bSucceeded |= User.SetFocus(SearchField.ToSharedRef());
 		});
 
 		if (bSucceeded)
diff --git a/Engine/Source/Editor/LevelEditor/Private/LevelViewportLayout.cpp b/Engine/Source/Editor/LevelEditor/Private/LevelViewportLayout.cpp
--- a/Engine/Source/Editor/LevelEditor/Private/LevelViewportLayout.cpp
+++ b/Engine/Source/Editor/LevelEditor/Private/LevelViewportLayout.cpp
@@ -675,7 +675,7 @@ void FLevelViewportLayout::FinishMaximizeTransition()
 				OwnerWindow->EndFullWindowOverlayTransition();
 			}
 			// Release overlay mouse capture to prevent situations where user is unable to get the mouse cursor back if they were holding one of the buttons down and exited immersive mode.
-			FSlateApplication::Get().ReleaseMouseCapture();
+			FSlateApplication::Get().ReleaseAllPointerCapture();
 
 			if( bIsMaximized )
 			{
diff --git a/Engine/Source/Editor/Persona/Private/Customization/SSearchableComboBox.cpp b/Engine/Source/Editor/Persona/Private/Customization/SSearchableComboBox.cpp
--- a/Engine/Source/Editor/Persona/Private/Customization/SSearchableComboBox.cpp
+++ b/Engine/Source/Editor/Persona/Private/Customization/SSearchableComboBox.cpp
@@ -2,6 +2,7 @@
 
 #include "SSearchableComboBox.h"
 #include "DetailLayoutBuilder.h"
+#include "Framework/Application/SlateUser.h"
 
 #define LOCTEXT_NAMESPACE "SearchableComboBox"
 
@@ -161,10 +162,11 @@ void SSearchableComboBox::OnMenuOpenChanged(bool bOpen)
 		}
 
 		// Set focus back to ComboBox for users focusing the ListView that just closed
-		FSlateApplication::Get().ForEachUser([&](FSlateUser* User) {
-			if (FSlateApplication::Get().HasUserFocusedDescendants(AsShared(), User->GetUserIndex()))
+		TSharedRef<SWidget> ThisRef = AsShared();
+		FSlateApplication::Get().ForEachUser([&ThisRef](FSlateUser& User) {
+			if (User.HasFocusedDescendants(ThisRef))
 			{
-				FSlateApplication::Get().SetUserFocus(User->GetUserIndex(), AsShared(), EFocusCause::SetDirectly);
+				User.SetFocus(ThisRef);
 			}
 		});
 
diff --git a/Engine/Source/Editor/Persona/Private/SAnimNotifyPanel.cpp b/Engine/Source/Editor/Persona/Private/SAnimNotifyPanel.cpp
--- a/Engine/Source/Editor/Persona/Private/SAnimNotifyPanel.cpp
+++ b/Engine/Source/Editor/Persona/Private/SAnimNotifyPanel.cpp
@@ -1832,7 +1832,7 @@ FReply SAnimNotifyNode::OnMouseMove( const FGeometry& MyGeometry, const FPointer
 	if(CurrentDragHandle == ENotifyStateHandleHit::None)
 	{
 		// We've had focus taken away - realease the mouse
-		FSlateApplication::Get().ReleaseMouseCapture();
+		FSlateApplication::Get().ReleaseAllPointerCapture();
 		return FReply::Unhandled();
 	}
 	
diff --git a/Engine/Source/Editor/UnrealEd/Private/DataTableEditorUtils.cpp b/Engine/Source/Editor/UnrealEd/Private/DataTableEditorUtils.cpp
--- a/Engine/Source/Editor/UnrealEd/Private/DataTableEditorUtils.cpp
+++ b/Engine/Source/Editor/UnrealEd/Private/DataTableEditorUtils.cpp
@@ -6,6 +6,7 @@
 #include "Styling/SlateTypes.h"
 #include "Fonts/FontMeasure.h"
 #include "Framework/Application/SlateApplication.h"
+#include "Framework/Application/SlateUser.h"
 #include "EditorStyleSet.h"
 #include "Engine/UserDefinedStruct.h"
 #include "ScopedTransaction.h"
@@ -351,10 +352,11 @@ void SDataTableStructComboBox::OnMenuOpenChanged(bool bOpen)
 		}
 
 		// Set focus back to ComboBox for users focusing the ListView that just closed
-		FSlateApplication::Get().ForEachUser([&](FSlateUser* User) {
-			if (FSlateApplication::Get().HasUserFocusedDescendants(AsShared(), User->GetUserIndex()))
+		TSharedRef<SWidget> ThisRef = AsShared();
+		FSlateApplication::Get().ForEachUser([&ThisRef](FSlateUser& User) {
+			if (User.HasFocusedDescendants(ThisRef))
 			{
-				FSlateApplication::Get().SetUserFocus(User->GetUserIndex(), AsShared(), EFocusCause::SetDirectly);
+				User.SetFocus(ThisRef, EFocusCause::SetDirectly);
 			}
 		});
 	}
diff --git a/Engine/Source/Runtime/Engine/Private/Slate/SceneViewport.cpp b/Engine/Source/Runtime/Engine/Private/Slate/SceneViewport.cpp
--- a/Engine/Source/Runtime/Engine/Private/Slate/SceneViewport.cpp
+++ b/Engine/Source/Runtime/Engine/Private/Slate/SceneViewport.cpp
@@ -14,6 +14,7 @@
 #include "Layout/WidgetPath.h"
 #include "UnrealEngine.h"
 #include "Framework/Application/SlateApplication.h"
+#include "Framework/Application/SlateUser.h"
 #include "Slate/SlateTextures.h"
 #include "Slate/DebugCanvas.h"
 
@@ -1120,10 +1121,10 @@ void FSceneViewport::OnFocusLost( const FFocusEvent& InFocusEvent )
 		TSharedPtr<SWidget> ViewportWidgetPin = ViewportWidget.Pin();
 		if ( ViewportWidgetPin.IsValid() )
 		{
-			FSlateApplication::Get().ForEachUser([&] (FSlateUser* User) {
-				if ( User->GetFocusedWidget() == ViewportWidgetPin )
+			FSlateApplication::Get().ForEachUser([&ViewportWidgetPin] (FSlateUser& User) {
+				if (User.GetFocusedWidget() == ViewportWidgetPin )
 				{
-					FSlateApplication::Get().ClearUserFocus(User->GetUserIndex());
+					User.ClearFocus();
 				}
 			});
 		}
diff --git a/Engine/Source/Runtime/MoviePlayer/Private/DefaultGameMoviePlayer.cpp b/Engine/Source/Runtime/MoviePlayer/Private/DefaultGameMoviePlayer.cpp
--- a/Engine/Source/Runtime/MoviePlayer/Private/DefaultGameMoviePlayer.cpp
+++ b/Engine/Source/Runtime/MoviePlayer/Private/DefaultGameMoviePlayer.cpp
@@ -453,9 +453,7 @@ void FDefaultGameMoviePlayer::WaitForMovieToFinish(bool bAllowEngineTick)
 		// Make sure the movie player widget has user focus to accept keypresses
 		if (LoadingScreenContents.IsValid())
 		{
-			SlateApp.ForEachUser([&](FSlateUser* User) {
-				SlateApp.SetUserFocus(User->GetUserIndex(), LoadingScreenContents);
-			});
+			SlateApp.SetAllUserFocus(LoadingScreenContents);
 		}
 
 		// Continue to wait until the user calls finish (if enabled) or when loading completes or the minimum enforced time (if any) has been reached.
diff --git a/Engine/Source/Runtime/Slate/Private/Framework/Application/MenuStack.cpp b/Engine/Source/Runtime/Slate/Private/Framework/Application/MenuStack.cpp
--- a/Engine/Source/Runtime/Slate/Private/Framework/Application/MenuStack.cpp
+++ b/Engine/Source/Runtime/Slate/Private/Framework/Application/MenuStack.cpp
@@ -377,7 +377,7 @@ FMenuStack::FPrePushResults FMenuStack::PrePush(const FPrePushArgs& InArgs)
 	// Release the mouse so that context can be properly restored upon closing menus.  See CL 1411833 before changing this.
 	if (InArgs.bFocusImmediately)
 	{
-		FSlateApplication::Get().ReleaseMouseCapture();
+		FSlateApplication::Get().ReleaseAllPointerCapture();
 	}
 
 	return OutResults;
diff --git a/Engine/Source/Runtime/Slate/Private/Framework/Application/SlateApplication.cpp b/Engine/Source/Runtime/Slate/Private/Framework/Application/SlateApplication.cpp
--- a/Engine/Source/Runtime/Slate/Private/Framework/Application/SlateApplication.cpp
+++ b/Engine/Source/Runtime/Slate/Private/Framework/Application/SlateApplication.cpp
@@ -2,6 +2,7 @@
 
 
 #include "Framework/Application/SlateApplication.h"
+#include "Framework/Application/SlateUser.h"
 #include "Rendering/SlateDrawBuffer.h"
 #include "Misc/CommandLine.h"
 #include "Misc/ScopeLock.h"
@@ -52,6 +53,104 @@
 
 CSV_DECLARE_CATEGORY_MODULE_EXTERN(CORE_API, Basic);
 
+//////////////////////////////////////////////////////////////////////////
+
+/** 
+ * The cursor given to any additional input-providing SlateUsers beyond the first one (since the first one owns the platform cursor)
+ * Or, if we're on a platform without a cursor, User 0 gets one of these too
+ */
+class FFauxSlateCursor : public ICursor
+{
+public:
+	FFauxSlateCursor() {}
+	virtual ~FFauxSlateCursor() {}
+	virtual void* CreateCursorFromFile(const FString& InPathToCursorWithoutExtension, FVector2D HotSpot) { return nullptr; }
+	virtual void* CreateCursorFromRGBABuffer(const FColor* Pixels, int32 Width, int32 Height, FVector2D InHotSpot) { return nullptr; }
+	virtual void SetTypeShape(EMouseCursor::Type InCursorType, void* CursorHandle) override {}
+	
+	virtual FVector2D GetPosition() const override { return CurrentPosition; }
+	virtual EMouseCursor::Type GetType() const override { return CurrentType; }
+	
+	virtual void SetType(const EMouseCursor::Type InNewCursor) override { CurrentType = InNewCursor; }
+	virtual void Show(bool bShow) override { bIsVisible = bShow; }
+
+	virtual void SetPosition(const int32 X, const int32 Y) override
+	{
+		FVector2D NewPosition(X, Y);
+		UpdateCursorClipping(NewPosition);
+
+		CurrentPosition = NewPosition;
+	}
+
+	virtual void GetSize(int32& Width, int32& Height) const override
+	{
+		Width = Height = 32;
+	}
+
+	
+	virtual void Lock(const RECT* const Bounds) override
+	{
+		if (Bounds)
+		{
+			CursorClipRect.Min.X = FMath::TruncToInt(Bounds->left);
+			CursorClipRect.Min.Y = FMath::TruncToInt(Bounds->top);
+			CursorClipRect.Max.X = FMath::TruncToInt(Bounds->right) - 1;
+			CursorClipRect.Max.Y = FMath::TruncToInt(Bounds->bottom) - 1;
+		}
+		else
+		{
+			//@todo DanH: Do we want to enforce a limit to the bounds of the screen like the console cursors?
+			CursorClipRect = FIntRect();
+		}
+
+		FVector2D Position = GetPosition();
+		if (UpdateCursorClipping(Position))
+		{
+			SetPosition(Position.X, Position.Y);
+		}
+	}
+
+private:
+	bool UpdateCursorClipping(FVector2D& CursorPosition)
+	{
+		bool bAdjusted = false;
+
+		if (CursorClipRect.Area() > 0)
+		{
+			if (CursorPosition.X < CursorClipRect.Min.X)
+			{
+				CursorPosition.X = CursorClipRect.Min.X;
+				bAdjusted = true;
+			}
+			else if (CursorPosition.X > CursorClipRect.Max.X)
+			{
+				CursorPosition.X = CursorClipRect.Max.X;
+				bAdjusted = true;
+			}
+
+			if (CursorPosition.Y < CursorClipRect.Min.Y)
+			{
+				CursorPosition.Y = CursorClipRect.Min.Y;
+				bAdjusted = true;
+			}
+			else if (CursorPosition.Y > CursorClipRect.Max.Y)
+			{
+				CursorPosition.Y = CursorClipRect.Max.Y;
+				bAdjusted = true;
+			}
+		}
+
+		return bAdjusted;
+	}
+
+	bool bIsVisible = false;
+	EMouseCursor::Type CurrentType = EMouseCursor::None;
+	FVector2D CurrentPosition = FVector2D::ZeroVector;
+	FIntRect CursorClipRect;
+};
+
+//////////////////////////////////////////////////////////////////////////
+
 class FEventRouter
 {
 
@@ -327,76 +426,11 @@ struct FEventRouter::Translate<FPointerEvent>
 	}
 };
 
-
-FSlateUser::FSlateUser(int32 InUserIndex, bool InVirtualUser)
-	: UserIndex(InUserIndex)
-	, bVirtualUser(InVirtualUser)
-	, FocusVersion(0)
-{
-	UE_LOG(LogSlate, Log, TEXT("New Slate User Created.  User Index %d, Is Virtual User: %d"), UserIndex, bVirtualUser);
-	FocusWidgetPathWeak = FWidgetPath();
-	FocusCause = EFocusCause::Cleared;
-	ShowFocus = false;
-}
-
-FSlateUser::~FSlateUser()
-{
-	UE_LOG(LogSlate, Log, TEXT("Slate User Destroyed.  User Index %d, Is Virtual User: %d"), UserIndex, bVirtualUser);
-}
-
-TSharedPtr<SWidget> FSlateUser::GetFocusedWidget() const
-{
-	if ( FocusWidgetPathWeak.IsValid() )
-	{
-		return FocusWidgetPathWeak.GetLastWidget().Pin();
-	}
-
-	return TSharedPtr<SWidget>();
-}
-
-void FSlateUser::SetFocusPath(const FWidgetPath& InWidgetPath, EFocusCause InCause, bool InShowFocus)
-{
-	FocusWidgetPathStrong.Reset();
-	FocusWidgetPathWeak = InWidgetPath;
-	FocusCause = InCause;
-	ShowFocus = InShowFocus;
-}
-
-void FSlateUser::FinishFrame()
-{
-	FocusWidgetPathStrong.Reset();
-}
-
-void FSlateUser::NotifyWindowDestroyed(TSharedRef<SWindow> DestroyedWindow)
-{
-	if (FocusWidgetPathStrong.IsValid() && FocusWidgetPathStrong->IsValid() && DestroyedWindow == FocusWidgetPathStrong->GetWindow())
-	{
-		FocusWidgetPathStrong.Reset();
-	}
-}
-
-FSlateVirtualUser::FSlateVirtualUser(int32 InUserIndex, int32 InVirtualUserIndex)
-	: UserIndex(InUserIndex)
-	, VirtualUserIndex(InVirtualUserIndex)
-{
-}
-
-FSlateVirtualUser::~FSlateVirtualUser()
-{
-	if ( FSlateApplication::IsInitialized() )
-	{
-		FSlateApplication::Get().UnregisterUser(UserIndex);
-	}
-}
-
-
 DECLARE_CYCLE_STAT( TEXT("Message Tick Time"), STAT_SlateMessageTick, STATGROUP_Slate );
-DECLARE_CYCLE_STAT( TEXT("Update Tooltip Time"), STAT_SlateUpdateTooltip, STATGROUP_Slate );
 DECLARE_CYCLE_STAT( TEXT("Total Slate Tick Time"), STAT_SlateTickTime, STATGROUP_Slate );
 DECLARE_CYCLE_STAT( TEXT("Draw Window And Children Time"), STAT_SlateDrawWindowTime, STATGROUP_Slate );
 DECLARE_CYCLE_STAT( TEXT("TickRegisteredWidgets"), STAT_SlateTickRegisteredWidgets, STATGROUP_Slate );
 DECLARE_CYCLE_STAT( TEXT("PreTickEvent"), STAT_SlatePreTickEvent, STATGROUP_Slate );
-DECLARE_CYCLE_STAT( TEXT("QueryCursor"), STAT_SlateQueryCursor, STATGROUP_Slate);
 
 DECLARE_CYCLE_STAT(TEXT("ShowVirtualKeyboard"), STAT_ShowVirtualKeyboard, STATGROUP_Slate);
 
@@ -448,27 +482,23 @@ TAutoConsoleVariable<float> SleepBufferPostInput(
 	0.0f,
 	TEXT("The amount of time that must pass without any user action before Slate is put to sleep (provided that there are no active timers)."));
 
-int32 GEnableSyntheticMouseMoves = 1;
-FAutoConsoleVariableRef CVarEnableSyntheticMouseMoves(
-	TEXT("Slate.EnableSyntheticMouseMoves"),
-	GEnableSyntheticMouseMoves,
-	TEXT("")
-);
-
-int32 GEnableCursorQueries = 1;
-FAutoConsoleVariableRef CVarAllowCursorQueries(
-	TEXT("Slate.EnableCursorQueries"),
-	GEnableCursorQueries,
-	TEXT("")
-);
-
-int32 bRequireFocusForGamepadInput = 0;
+static bool bRequireFocusForGamepadInput = false;
 FAutoConsoleVariableRef CVarRequireFocusForGamepadInput(
 	TEXT("Slate.RequireFocusForGamepadInput"),
 	bRequireFocusForGamepadInput,
 	TEXT("")
 );
 
+#if PLATFORM_UI_NEEDS_TOOLTIPS
+static bool bEnableTooltips = true;
+#else
+static bool bEnableTooltips = false;
+#endif
+FAutoConsoleVariableRef CVarEnableTooltips(
+	TEXT("Slate.EnableTooltips"),
+	bEnableTooltips,
+	TEXT("Whether to allow tooltips to spawn at all."));
+
 #if !UE_BUILD_SHIPPING
 
 static void HandleGlobalInvalidateCVarTriggered(const TArray<FString>& Args)
@@ -483,517 +513,174 @@ static FAutoConsoleCommand GlobalInvalidateCommand(
 );
 #endif
 
-
-
-
 //////////////////////////////////////////////////////////////////////////
-bool FSlateApplication::MouseCaptorHelper::HasCapture() const
-{
-	for (auto PointerPathPair : PointerIndexToMouseCaptorWeakPathMap)
-	{
-		if (PointerPathPair.Value.IsValid())
-		{
-			return true;
-		}
-	}
-	return false;
-}
 
-bool FSlateApplication::MouseCaptorHelper::HasCaptureForUser(uint32 UserIndex) const
+FDelegateHandle FPopupSupport::RegisterClickNotification( const TSharedRef<SWidget>& NotifyWhenClickedOutsideMe, const FOnClickedOutside& InNotification )
 {
-	for ( auto PointerPathPair : PointerIndexToMouseCaptorWeakPathMap )
-	{
-		const FUserAndPointer& UserAndPointer = PointerPathPair.Key;
-		if ( UserAndPointer.UserIndex == UserIndex )
+	// If the subscriber or a zone object is destroyed, the subscription is
+	// no longer active. Clean it up here so that consumers of this API have an
+	// easy time with resource management.
+	struct { void operator()( TArray<FClickSubscriber>& Notifications ) {
+		for ( int32 SubscriberIndex=0; SubscriberIndex < Notifications.Num(); )
 		{
-			if ( PointerPathPair.Value.IsValid() )
+			if ( !Notifications[SubscriberIndex].ShouldKeep() )
 			{
-				return true;
+				Notifications.RemoveAtSwap(SubscriberIndex);
 			}
+			else
+			{
+				SubscriberIndex++;
+			}		
 		}
-	}
+	}} ClearOutStaleNotifications;
+	
+	ClearOutStaleNotifications( ClickZoneNotifications );
 
-	return false;
+	// Add a new notification.
+	ClickZoneNotifications.Add( FClickSubscriber( NotifyWhenClickedOutsideMe, InNotification ) );
+
+	return ClickZoneNotifications.Last().Notification.GetHandle();
 }
 
-bool FSlateApplication::MouseCaptorHelper::HasCaptureForPointerIndex(uint32 UserIndex, uint32 PointerIndex) const
+void FPopupSupport::UnregisterClickNotification( FDelegateHandle Handle )
 {
-	const FWeakWidgetPath* MouseCaptorWeakPath = PointerIndexToMouseCaptorWeakPathMap.Find(FUserAndPointer(UserIndex, PointerIndex));
-	return MouseCaptorWeakPath && MouseCaptorWeakPath->IsValid();
+	for (int32 SubscriptionIndex=0; SubscriptionIndex < ClickZoneNotifications.Num();)
+	{
+		if (ClickZoneNotifications[SubscriptionIndex].Notification.GetHandle() == Handle)
+		{
+			ClickZoneNotifications.RemoveAtSwap(SubscriptionIndex);
+		}
+		else
+		{
+			SubscriptionIndex++;
+		}
+	}	
 }
 
-bool FSlateApplication::MouseCaptorHelper::DoesWidgetHaveMouseCaptureByUser(const TSharedPtr<const SWidget> Widget, int32 UserIndex, TOptional<int32> PointerIndex) const
+void FPopupSupport::SendNotifications( const FWidgetPath& WidgetsUnderCursor )
 {
-	for ( const auto& PointerPathPair : PointerIndexToMouseCaptorWeakPathMap )
+	struct FArrangedWidgetMatcher
 	{
-		const FUserAndPointer& UserAndPointer = PointerPathPair.Key;
-		if ( UserAndPointer.UserIndex == UserIndex )
-		{
-			// If the pointer index is set, filter on that as well.
-			if ( PointerIndex.IsSet() && UserAndPointer.PointerIndex != PointerIndex.GetValue() )
-			{
-				continue;
-			}
+		FArrangedWidgetMatcher( const TSharedRef<SWidget>& InWidgetToMatch )
+		: WidgetToMatch( InWidgetToMatch )
+		{}
 
-			if ( PointerPathPair.Value.IsValid() )
-			{
-				TSharedPtr<SWidget> LastWidget = PointerPathPair.Value.GetLastWidget().Pin();
-				if ( LastWidget == Widget )
-				{
-					return true;
-				}
-			}
+		bool operator()(const FArrangedWidget& Candidate) const
+		{
+			return WidgetToMatch == Candidate.Widget;
 		}
-	}
 
-	return false;
-}
+		const TSharedRef<SWidget>& WidgetToMatch;
+	};
 
-bool FSlateApplication::MouseCaptorHelper::DoesWidgetHaveMouseCapture(const TSharedPtr<const SWidget> Widget) const
-{
-	for ( const auto& IndexPathPair : PointerIndexToMouseCaptorWeakPathMap )
+	// For each subscription, if the widget in question is not being clicked, send the notification.
+	// i.e. Notifications are saying "some widget outside you was clicked".
+	for (int32 SubscriberIndex=0; SubscriberIndex < ClickZoneNotifications.Num(); ++SubscriberIndex)
 	{
-		TSharedPtr<SWidget> LastWidget = IndexPathPair.Value.GetLastWidget().Pin();
-		if ( LastWidget == Widget )
+		FClickSubscriber& Subscriber = ClickZoneNotifications[SubscriberIndex];
+		if (Subscriber.DetectClicksOutsideMe.IsValid())
 		{
-			return true;
+			// Did we click outside the region in this subscription? If so send the notification.
+			FArrangedWidgetMatcher Matcher(Subscriber.DetectClicksOutsideMe.Pin().ToSharedRef());
+			const bool bClickedOutsideOfWidget = WidgetsUnderCursor.Widgets.GetInternalArray().IndexOfByPredicate(Matcher) == INDEX_NONE;
+			if ( bClickedOutsideOfWidget )
+			{
+				Subscriber.Notification.ExecuteIfBound();
+			}
 		}
 	}
-
-	return false;
 }
 
-TSharedPtr< SWidget > FSlateApplication::MouseCaptorHelper::ToSharedWidget(uint32 UserIndex, uint32 PointerIndex) const
+void FSlateApplication::SetPlatformApplication(const TSharedRef<class GenericApplication>& InPlatformApplication)
 {
-	// If the path is valid then get the last widget, this is the current mouse captor
-	TSharedPtr< SWidget > SharedWidgetPtr;
-	const FWeakWidgetPath* MouseCaptorWeakPath = PointerIndexToMouseCaptorWeakPathMap.Find( FUserAndPointer(UserIndex,PointerIndex) );
-	if (MouseCaptorWeakPath && MouseCaptorWeakPath->IsValid() )
-	{
-		TWeakPtr< SWidget > WeakWidgetPtr = MouseCaptorWeakPath->GetLastWidget();
-		SharedWidgetPtr = WeakWidgetPtr.Pin();
-	}
+	PlatformApplication->SetMessageHandler(MakeShareable(new FGenericApplicationMessageHandler()));
+#if WITH_ACCESSIBILITY
+	PlatformApplication->SetAccessibleMessageHandler(MakeShareable(new FGenericAccessibleMessageHandler()));
+#endif
 
-	return SharedWidgetPtr;
+	PlatformApplication = InPlatformApplication;
+	PlatformApplication->SetMessageHandler(CurrentApplication.ToSharedRef());
+#if WITH_ACCESSIBILITY
+	PlatformApplication->SetAccessibleMessageHandler(CurrentApplication->GetAccessibleMessageHandler());
+#endif
 }
 
-TArray<TSharedRef<SWidget>> FSlateApplication::MouseCaptorHelper::ToSharedWidgets() const
+void FSlateApplication::OverridePlatformApplication(TSharedPtr<class GenericApplication> InPlatformApplication)
 {
-	TArray<TSharedRef<SWidget>> Widgets;
-	Widgets.Empty(PointerIndexToMouseCaptorWeakPathMap.Num());
-
-	for (const auto& IndexPathPair : PointerIndexToMouseCaptorWeakPathMap)
-	{
-		TSharedPtr<SWidget> LastWidget = IndexPathPair.Value.GetLastWidget().Pin();
-		if (LastWidget.IsValid())
-		{
-			Widgets.Add(LastWidget.ToSharedRef());
-		}
-	}
-	return Widgets;
+	PlatformApplication = InPlatformApplication;
 }
 
-TSharedPtr< SWidget > FSlateApplication::MouseCaptorHelper::ToSharedWindow(uint32 UserIndex, uint32 PointerIndex)
+void FSlateApplication::Create()
 {
-	// if the path is valid then we can get the window the current mouse captor belongs to
-	FWidgetPath MouseCaptorPath = ToWidgetPath( UserIndex, PointerIndex );
-	if ( MouseCaptorPath.IsValid() )
-	{
-		return MouseCaptorPath.GetWindow();
-	}
+	GSlateFastWidgetPath = GIsEditor ? 0 : 1;
 
-	return TSharedPtr< SWidget >();
+	Create(MakeShareable(FPlatformApplicationMisc::CreateApplication()));
 }
 
-bool FSlateApplication::MouseCaptorHelper::SetMouseCaptor(uint32 UserIndex, uint32 PointerIndex, const FWidgetPath& EventPath, TSharedPtr< SWidget > Widget)
+TSharedRef<FSlateApplication> FSlateApplication::Create(const TSharedRef<class GenericApplication>& InPlatformApplication)
 {
-	// Caller is trying to set a new mouse captor, so invalidate the current one - when the function finishes
-	// it still may not have a valid captor widget, this is ok
-	InvalidateCaptureForPointer(UserIndex, PointerIndex);
-
-	if ( Widget.IsValid() )
-	{
-		TSharedRef< SWidget > WidgetRef = Widget.ToSharedRef();
-		FWidgetPath NewMouseCaptorPath = EventPath.GetPathDownTo( WidgetRef );
+	EKeys::Initialize();
 
-		const auto IsPathToCaptorFound = []( const FWidgetPath& PathToTest, const TSharedRef<SWidget>& WidgetToFind )
-		{
-			return PathToTest.Widgets.Num() > 0 && PathToTest.Widgets.Last().Widget == WidgetToFind;
-		};
+	FCoreStyle::ResetToDefault();
 
-		FWeakWidgetPath MouseCaptorWeakPath;
-		if ( IsPathToCaptorFound( NewMouseCaptorPath, WidgetRef ) )
-		{
-			MouseCaptorWeakPath = NewMouseCaptorPath;
-		}
-		else if (EventPath.Widgets.Num() > 0)
-		{
-			// If the target widget wasn't found on the event path then start the search from the root
-			NewMouseCaptorPath = EventPath.GetPathDownTo( EventPath.Widgets[0].Widget );
-			NewMouseCaptorPath.ExtendPathTo( FWidgetMatcher( WidgetRef ) );
-			
-			MouseCaptorWeakPath = IsPathToCaptorFound( NewMouseCaptorPath, WidgetRef )
-				? NewMouseCaptorPath
-				: FWeakWidgetPath();
-		}
-		else
-		{
-			ensureMsgf(EventPath.Widgets.Num() > 0, TEXT("An unknown widget is attempting to set capture to %s"), *Widget->ToString() );
-		}
+	// Note: Important to establish the static PlatformApplication property first, as the FSlateApplication ctor relies on it
+	PlatformApplication = InPlatformApplication;
 
-		if (MouseCaptorWeakPath.IsValid())
-		{
-			PointerIndexToMouseCaptorWeakPathMap.Add(FUserAndPointer(UserIndex,PointerIndex), MouseCaptorWeakPath);
+	CurrentApplication = MakeShareable( new FSlateApplication() );
+	CurrentBaseApplication = CurrentApplication;
 
-#if WITH_SLATE_DEBUGGING
-			FSlateDebugging::BroadcastMouseCapture(UserIndex, PointerIndex, Widget);
+	PlatformApplication->SetMessageHandler( CurrentApplication.ToSharedRef() );
+#if WITH_ACCESSIBILITY
+	PlatformApplication->SetAccessibleMessageHandler(CurrentApplication->GetAccessibleMessageHandler());
 #endif
 
-			return true;
-		}
+	// The grid needs to know the size and coordinate system of the desktop.
+	// Some monitor setups have a primary monitor on the right and below the
+	// left one, so the leftmost upper right monitor can be something like (-1280, -200)Synt
+	{
+		// Get an initial value for the VirtualDesktop geometry
+		CurrentApplication->VirtualDesktopRect = []()
+		{
+			FDisplayMetrics DisplayMetrics;
+			FSlateApplicationBase::Get().GetDisplayMetrics(DisplayMetrics);
+			const FPlatformRect& VirtualDisplayRect = DisplayMetrics.VirtualDisplayRect;
+			return FSlateRect(VirtualDisplayRect.Left, VirtualDisplayRect.Top, VirtualDisplayRect.Right, VirtualDisplayRect.Bottom);
+		}();
+
+		// Sign up for updates from the OS. Polling this every frame is too expensive on at least some OSs.
+		PlatformApplication->OnDisplayMetricsChanged().AddSP(CurrentApplication.ToSharedRef(), &FSlateApplication::OnVirtualDesktopSizeChanged);
 	}
 
-	return false;
-}
+	FAsyncTaskNotificationFactory::Get().RegisterFactory(TEXT("Slate"), []() -> FAsyncTaskNotificationFactory::FImplPointerType { return new FSlateAsyncTaskNotificationImpl(); });
 
-void FSlateApplication::MouseCaptorHelper::InvalidateCaptureForAllPointers()
-{
-	TArray<FUserAndPointer> PointerIndices;
-	PointerIndexToMouseCaptorWeakPathMap.GenerateKeyArray(PointerIndices);
-	for (FUserAndPointer UserAndPointer : PointerIndices)
-	{
-		InvalidateCaptureForPointer(UserAndPointer.UserIndex, UserAndPointer.PointerIndex);
-	}
+	return CurrentApplication.ToSharedRef();
 }
 
-void FSlateApplication::MouseCaptorHelper::InvalidateCaptureForUser(uint32 UserIndex)
+void FSlateApplication::Shutdown(bool bShutdownPlatform)
 {
-	TArray<FUserAndPointer> PointerIndices;
-	PointerIndexToMouseCaptorWeakPathMap.GenerateKeyArray(PointerIndices);
-	for ( FUserAndPointer UserAndPointer : PointerIndices )
+	if (FSlateApplication::IsInitialized())
 	{
-		if ( UserAndPointer.UserIndex == UserIndex )
+		FAsyncTaskNotificationFactory::Get().UnregisterFactory(TEXT("Slate"));
+
+		CurrentApplication->OnShutdown();
+		CurrentApplication->DestroyRenderer();
+		CurrentApplication->Renderer.Reset();
+
+		if (bShutdownPlatform)
 		{
-			InvalidateCaptureForPointer(UserAndPointer.UserIndex, UserAndPointer.PointerIndex);
+			PlatformApplication->DestroyApplication();
 		}
-	}
-}
-
-FWidgetPath FSlateApplication::MouseCaptorHelper::ToWidgetPath( FWeakWidgetPath::EInterruptedPathHandling::Type InterruptedPathHandling, const FPointerEvent* PointerEvent )
-{
-	FWidgetPath WidgetPath;
-	const FWeakWidgetPath* MouseCaptorWeakPath = PointerIndexToMouseCaptorWeakPathMap.Find(FUserAndPointer(PointerEvent->GetUserIndex(), PointerEvent->GetPointerIndex()));
-	if ( MouseCaptorWeakPath && MouseCaptorWeakPath->IsValid() )
-	{
-		if ( MouseCaptorWeakPath->ToWidgetPath( WidgetPath, InterruptedPathHandling, PointerEvent ) == FWeakWidgetPath::EPathResolutionResult::Truncated )
-		{
-			// If the path was truncated then it means this widget is no longer part of the active set,
-			// so we make sure to invalidate its capture
-			InvalidateCaptureForPointer(PointerEvent->GetUserIndex(), PointerEvent->GetPointerIndex());
-		}
-	}
-
-	return WidgetPath;
-}
-
-FWidgetPath FSlateApplication::MouseCaptorHelper::ToWidgetPath(uint32 UserIndex, uint32 PointerIndex, FWeakWidgetPath::EInterruptedPathHandling::Type InterruptedPathHandling)
-{
-	FWidgetPath WidgetPath;
-	const FWeakWidgetPath* MouseCaptorWeakPath = PointerIndexToMouseCaptorWeakPathMap.Find(FUserAndPointer(UserIndex, PointerIndex));
-	if (MouseCaptorWeakPath && MouseCaptorWeakPath->IsValid() )
-	{
-		if ( MouseCaptorWeakPath->ToWidgetPath( WidgetPath, InterruptedPathHandling ) == FWeakWidgetPath::EPathResolutionResult::Truncated )
-		{
-			// If the path was truncated then it means this widget is no longer part of the active set,
-			// so we make sure to invalidate its capture
-			InvalidateCaptureForPointer(UserIndex, PointerIndex);
-		}
-	}
-
-	return WidgetPath;
-}
-
-void FSlateApplication::MouseCaptorHelper::InvalidateCaptureForPointer(uint32 UserIndex, uint32 PointerIndex)
-{
-	InformCurrentCaptorOfCaptureLoss(UserIndex, PointerIndex);
-	PointerIndexToMouseCaptorWeakPathMap.Remove( FUserAndPointer(UserIndex, PointerIndex) );
-}
-
-TArray<FWidgetPath> FSlateApplication::MouseCaptorHelper::ToWidgetPaths()
-{
-	TArray<FWidgetPath> WidgetPaths;
-	TArray<FUserAndPointer> PointerIndices;
-	PointerIndexToMouseCaptorWeakPathMap.GenerateKeyArray(PointerIndices);
-	for (auto Index : PointerIndices)
-	{
-		WidgetPaths.Add(ToWidgetPath(Index.UserIndex,Index.PointerIndex));
-	}
-	return WidgetPaths;
-}
-
-FWeakWidgetPath FSlateApplication::MouseCaptorHelper::ToWeakPath(uint32 UserIndex, uint32 PointerIndex) const
-{
-	const FWeakWidgetPath* MouseCaptorWeakPath = PointerIndexToMouseCaptorWeakPathMap.Find(FUserAndPointer(UserIndex,PointerIndex));
-	if (MouseCaptorWeakPath)
-	{
-		return *MouseCaptorWeakPath;
-	}
-	return FWeakWidgetPath();
-}
-
-void FSlateApplication::MouseCaptorHelper::InformCurrentCaptorOfCaptureLoss(uint32 UserIndex, uint32 PointerIndex) const
-{
-	// if we have a path to a widget then it is the current mouse captor and needs to know it has lost capture
-	const FWeakWidgetPath* MouseCaptorWeakPath = PointerIndexToMouseCaptorWeakPathMap.Find(FUserAndPointer(UserIndex, PointerIndex));
-	if (MouseCaptorWeakPath && MouseCaptorWeakPath->IsValid())
-	{
-		TWeakPtr< SWidget > WeakWidgetPtr = MouseCaptorWeakPath->GetLastWidget();
-		TSharedPtr< SWidget > SharedWidgetPtr = WeakWidgetPtr.Pin();
-		if (SharedWidgetPtr.IsValid())
-		{
-			FCaptureLostEvent CaptureLostEvent(UserIndex, PointerIndex);
-			SharedWidgetPtr->OnMouseCaptureLost(CaptureLostEvent);
-		}
-
-#if WITH_SLATE_DEBUGGING
-		FSlateDebugging::BroadcastMouseCaptureLost(UserIndex, PointerIndex, SharedWidgetPtr);
-#endif
-	}
-}
-
-//////////////////////////////////////////////////////////////////////////
-
-void FSlateApplication::FDragDetector::StartDragDetection(const FWidgetPath& PathToWidget, int32 UserIndex, int32 PointerIndex, FKey DragButton, FVector2D StartLocation)
-{
-	PointerIndexToDragState.Add(FUserAndPointer(UserIndex, PointerIndex), FDragDetectionState(PathToWidget, UserIndex, PointerIndex, DragButton, StartLocation));
-}
-
-bool FSlateApplication::FDragDetector::IsDetectingDrag(const FPointerEvent& PointerEvent)
-{
-	FUserAndPointer UserAndPointer(PointerEvent.GetUserIndex(), PointerEvent.GetPointerIndex());
-	FDragDetectionState* DetectionState = PointerIndexToDragState.Find(UserAndPointer);
-
-	if ( DetectionState )
-	{
-		return true;
-	}
-
-	return false;
-}
-
-bool FSlateApplication::FDragDetector::DetectDrag(const FPointerEvent& PointerEvent, float InDragTriggerDistance, FWeakWidgetPath*& OutWeakWidgetPath)
-{
-	FUserAndPointer UserAndPointer(PointerEvent.GetUserIndex(), PointerEvent.GetPointerIndex());
-	FDragDetectionState* DetectionState = PointerIndexToDragState.Find(UserAndPointer);
-
-	if ( DetectionState )
-	{
-		if ( DetectionState->DetectDragUserIndex == PointerEvent.GetUserIndex() &&
-			 DetectionState->DetectDragPointerIndex == PointerEvent.GetPointerIndex() )
-		{
-			const FVector2D DragDelta = DetectionState->DetectDragStartLocation - PointerEvent.GetScreenSpacePosition();
-			const bool bDragDetected = ( DragDelta.SizeSquared() > FMath::Square(InDragTriggerDistance) );
-			if ( bDragDetected )
-			{
-				OutWeakWidgetPath = &DetectionState->DetectDragForWidget;
-				return true;
-			}
-		}
-	}
-
-	OutWeakWidgetPath = nullptr;
-	return false;
-}
-
-void FSlateApplication::FDragDetector::OnPointerRelease(const FPointerEvent& PointerEvent)
-{
-	FUserAndPointer UserAndPointer(PointerEvent.GetUserIndex(), PointerEvent.GetPointerIndex());
-	FDragDetectionState* DetectionState = PointerIndexToDragState.Find(UserAndPointer);
-
-	if ( DetectionState )
-	{
-		if ( DetectionState->DetectDragButton == PointerEvent.GetEffectingButton() && 
-			 DetectionState->DetectDragUserIndex == PointerEvent.GetUserIndex() && 
-			 DetectionState->DetectDragPointerIndex == PointerEvent.GetPointerIndex() )
-		{
-			// The user has released the button (or finger) that was supposed to start the drag; stop detecting it.
-			PointerIndexToDragState.Remove(UserAndPointer);
-		}
-	}
-}
-
-void FSlateApplication::FDragDetector::ResetDetection()
-{
-	PointerIndexToDragState.Reset();
-}
-
-//////////////////////////////////////////////////////////////////////////
-
-FDelegateHandle FPopupSupport::RegisterClickNotification( const TSharedRef<SWidget>& NotifyWhenClickedOutsideMe, const FOnClickedOutside& InNotification )
-{
-	// If the subscriber or a zone object is destroyed, the subscription is
-	// no longer active. Clean it up here so that consumers of this API have an
-	// easy time with resource management.
-	struct { void operator()( TArray<FClickSubscriber>& Notifications ) {
-		for ( int32 SubscriberIndex=0; SubscriberIndex < Notifications.Num(); )
-		{
-			if ( !Notifications[SubscriberIndex].ShouldKeep() )
-			{
-				Notifications.RemoveAtSwap(SubscriberIndex);
-			}
-			else
-			{
-				SubscriberIndex++;
-			}		
-		}
-	}} ClearOutStaleNotifications;
-	
-	ClearOutStaleNotifications( ClickZoneNotifications );
-
-	// Add a new notification.
-	ClickZoneNotifications.Add( FClickSubscriber( NotifyWhenClickedOutsideMe, InNotification ) );
-
-	return ClickZoneNotifications.Last().Notification.GetHandle();
-}
-
-void FPopupSupport::UnregisterClickNotification( FDelegateHandle Handle )
-{
-	for (int32 SubscriptionIndex=0; SubscriptionIndex < ClickZoneNotifications.Num();)
-	{
-		if (ClickZoneNotifications[SubscriptionIndex].Notification.GetHandle() == Handle)
-		{
-			ClickZoneNotifications.RemoveAtSwap(SubscriptionIndex);
-		}
-		else
-		{
-			SubscriptionIndex++;
-		}
-	}	
-}
-
-void FPopupSupport::SendNotifications( const FWidgetPath& WidgetsUnderCursor )
-{
-	struct FArrangedWidgetMatcher
-	{
-		FArrangedWidgetMatcher( const TSharedRef<SWidget>& InWidgetToMatch )
-		: WidgetToMatch( InWidgetToMatch )
-		{}
-
-		bool operator()(const FArrangedWidget& Candidate) const
-		{
-			return WidgetToMatch == Candidate.Widget;
-		}
-
-		const TSharedRef<SWidget>& WidgetToMatch;
-	};
-
-	// For each subscription, if the widget in question is not being clicked, send the notification.
-	// i.e. Notifications are saying "some widget outside you was clicked".
-	for (int32 SubscriberIndex=0; SubscriberIndex < ClickZoneNotifications.Num(); ++SubscriberIndex)
-	{
-		FClickSubscriber& Subscriber = ClickZoneNotifications[SubscriberIndex];
-		if (Subscriber.DetectClicksOutsideMe.IsValid())
-		{
-			// Did we click outside the region in this subscription? If so send the notification.
-			FArrangedWidgetMatcher Matcher(Subscriber.DetectClicksOutsideMe.Pin().ToSharedRef());
-			const bool bClickedOutsideOfWidget = WidgetsUnderCursor.Widgets.GetInternalArray().IndexOfByPredicate(Matcher) == INDEX_NONE;
-			if ( bClickedOutsideOfWidget )
-			{
-				Subscriber.Notification.ExecuteIfBound();
-			}
-		}
-	}
-}
-
-void FSlateApplication::SetPlatformApplication(const TSharedRef<class GenericApplication>& InPlatformApplication)
-{
-	PlatformApplication->SetMessageHandler(MakeShareable(new FGenericApplicationMessageHandler()));
-#if WITH_ACCESSIBILITY
-	PlatformApplication->SetAccessibleMessageHandler(MakeShareable(new FGenericAccessibleMessageHandler()));
-#endif
-
-	PlatformApplication = InPlatformApplication;
-	PlatformApplication->SetMessageHandler(CurrentApplication.ToSharedRef());
-#if WITH_ACCESSIBILITY
-	PlatformApplication->SetAccessibleMessageHandler(CurrentApplication->GetAccessibleMessageHandler());
-#endif
-}
-
-void FSlateApplication::OverridePlatformApplication(TSharedPtr<class GenericApplication> InPlatformApplication)
-{
-	PlatformApplication = InPlatformApplication;
-}
-
-void FSlateApplication::Create()
-{
-	GSlateFastWidgetPath = GIsEditor ? 0 : 1;
-
-	Create(MakeShareable(FPlatformApplicationMisc::CreateApplication()));
-}
-
-TSharedRef<FSlateApplication> FSlateApplication::Create(const TSharedRef<class GenericApplication>& InPlatformApplication)
-{
-	EKeys::Initialize();
-
-	FCoreStyle::ResetToDefault();
-
-	CurrentApplication = MakeShareable( new FSlateApplication() );
-	CurrentBaseApplication = CurrentApplication;
-
-	PlatformApplication = InPlatformApplication;
-	PlatformApplication->SetMessageHandler( CurrentApplication.ToSharedRef() );
-#if WITH_ACCESSIBILITY
-	PlatformApplication->SetAccessibleMessageHandler(CurrentApplication->GetAccessibleMessageHandler());
-#endif
-
-	// The grid needs to know the size and coordinate system of the desktop.
-	// Some monitor setups have a primary monitor on the right and below the
-	// left one, so the leftmost upper right monitor can be something like (-1280, -200)Synt
-	{
-		// Get an initial value for the VirtualDesktop geometry
-		CurrentApplication->VirtualDesktopRect = []()
-		{
-			FDisplayMetrics DisplayMetrics;
-			FSlateApplicationBase::Get().GetDisplayMetrics(DisplayMetrics);
-			const FPlatformRect& VirtualDisplayRect = DisplayMetrics.VirtualDisplayRect;
-			return FSlateRect(VirtualDisplayRect.Left, VirtualDisplayRect.Top, VirtualDisplayRect.Right, VirtualDisplayRect.Bottom);
-		}();
-
-		// Sign up for updates from the OS. Polling this every frame is too expensive on at least some OSs.
-		PlatformApplication->OnDisplayMetricsChanged().AddSP(CurrentApplication.ToSharedRef(), &FSlateApplication::OnVirtualDesktopSizeChanged);
-	}
-
-	FAsyncTaskNotificationFactory::Get().RegisterFactory(TEXT("Slate"), []() -> FAsyncTaskNotificationFactory::FImplPointerType { return new FSlateAsyncTaskNotificationImpl(); });
-
-	return CurrentApplication.ToSharedRef();
-}
-
-void FSlateApplication::Shutdown(bool bShutdownPlatform)
-{
-	if (FSlateApplication::IsInitialized())
-	{
-		FAsyncTaskNotificationFactory::Get().UnregisterFactory(TEXT("Slate"));
-
-		CurrentApplication->OnShutdown();
-		CurrentApplication->DestroyRenderer();
-		CurrentApplication->Renderer.Reset();
-
-		if (bShutdownPlatform)
-		{
-			PlatformApplication->DestroyApplication();
-		}
-
-		PlatformApplication.Reset();
-		CurrentApplication.Reset();
-		CurrentBaseApplication.Reset();
+
+		PlatformApplication.Reset();
+		CurrentApplication.Reset();
+		CurrentBaseApplication.Reset();
 	}
 }
 
 TSharedPtr<FSlateApplication> FSlateApplication::CurrentApplication = nullptr;
 
 FSlateApplication::FSlateApplication()
-	: SynthesizeMouseMovePending(0)
-	, bAppIsActive(true)
+	: bAppIsActive(true)
 	, bSlateWindowActive(true)
 	, Scale( 1.0f )
 	, DragTriggerDistance( 0 )
@@ -1008,30 +695,8 @@ FSlateApplication::FSlateApplication()
 	, AverageDeltaTimeForResponsiveness( 1.0f / 30.0f )
 	, OnExitRequested()
 	, NumExternalModalWindowsActive( 0 )
-#if PLATFORM_UI_NEEDS_TOOLTIPS
-	, bAllowToolTips( true )
-#else
-	, bAllowToolTips( false )
-#endif
-	, ToolTipDelay( 0.15f )
-	, ToolTipFadeInDuration( 0.1f )
-	, ToolTipSummonTime( 0.0 )
-	, DesiredToolTipLocation( FVector2D::ZeroVector )
-	, ToolTipOffsetDirection( EToolTipOffsetDirection::Undetermined )
 	, bRequestLeaveDebugMode( false )
 	, bLeaveDebugForSingleStep( false )
-	, CVarAllowToolTips(
-		TEXT( "Slate.AllowToolTips" ),
-		bAllowToolTips,
-		TEXT( "Whether to allow tool-tips to spawn at all." ) )
-	, CVarToolTipDelay(
-		TEXT( "Slate.ToolTipDelay" ),
-		ToolTipDelay,
-		TEXT( "Delay in seconds before a tool-tip is displayed near the mouse cursor when hovering over widgets that supply tool-tip data." ) )
-	, CVarToolTipFadeInDuration(
-		TEXT( "Slate.ToolTipFadeInDuration" ),
-		ToolTipFadeInDuration,
-		TEXT( "How long it takes for a tool-tip to fade in, in seconds." ) )
 	, bIsExternalUIOpened( false )
 	, SlateTextField( nullptr )
 	, bIsFakingTouch(FParse::Param(FCommandLine::Get(), TEXT("simmobile")) || FParse::Param(FCommandLine::Get(), TEXT("faketouches")))
@@ -1039,7 +704,6 @@ FSlateApplication::FSlateApplication()
 	, bIsFakingTouched( false )
 	, bTouchFallbackToMouse( true )
 	, bSoftwareCursorAvailable( false )	
-	, bQueryCursorRequested( false )
 	, bMenuAnimationsEnabled( true )
 	, AppIcon( FCoreStyle::Get().GetBrush("DefaultAppIcon") )
 	, VirtualDesktopRect( 0,0,0,0 )
@@ -1073,11 +737,8 @@ FSlateApplication::FSlateApplication()
 
 	NormalExecutionGetter.BindRaw( this, &FSlateApplication::IsNormalExecution );
 
-	PointerIndexPositionMap.Add(FUserAndPointer(0, CursorPointerIndex), FVector2D::ZeroVector);
-	PointerIndexLastPositionMap.Add(FUserAndPointer(0, CursorPointerIndex), FVector2D::ZeroVector);
-
-	// Add the standard 'default' user because there's always 1 user.
-	RegisterUser(MakeShareable(new FSlateUser(0, false)));
+	// Add the standard 'default' user (there's always guaranteed to be at least one)
+	RegisterNewUser(CursorUserIndex);
 
 	NavigationConfig->OnRegister();
 #if WITH_EDITOR
@@ -1222,27 +883,21 @@ float FSlateApplication::GetSoundDuration(const FSlateSound& Sound) const
 
 FVector2D FSlateApplication::GetCursorPos() const
 {
-	if ( ICursor* Cursor = PlatformApplication->Cursor.Get() )
-	{
-		return !Cursor->GetPosition().ContainsNaN() ? Cursor->GetPosition() : FVector2D::ZeroVector;
-	}
-
-	const int32 UserIndex = 0;
-	return PointerIndexPositionMap.FindRef(FUserAndPointer(UserIndex, CursorPointerIndex));
+	TSharedPtr<const FSlateUser> CursorUser = GetUser(CursorUserIndex);
+	check(CursorUser);
+	return CursorUser->GetCursorPosition();
 }
 
 FVector2D FSlateApplication::GetLastCursorPos() const
 {
-	const int32 UserIndex = 0;
-	return PointerIndexLastPositionMap.FindRef(FUserAndPointer(UserIndex, CursorPointerIndex));
+	TSharedPtr<const FSlateUser> CursorUser = GetUser(CursorUserIndex);
+	check(CursorUser);
+	return CursorUser->GetPreviousCursorPosition();
 }
 
-void FSlateApplication::SetCursorPos( const FVector2D& MouseCoordinate )
+void FSlateApplication::SetCursorPos(const FVector2D& MouseCoordinate)
 {
-	if (ICursor* Cursor = PlatformApplication->Cursor.Get())
-	{
-		return Cursor->SetPosition( MouseCoordinate.X, MouseCoordinate.Y );
-	}
+	GetOrCreateUser(CursorUserIndex)->SetCursorPosition(MouseCoordinate);
 }
 
 FWidgetPath FSlateApplication::LocateWindowUnderMouse( FVector2D ScreenspaceMouseCoordinate, const TArray< TSharedRef< SWindow > >& Windows, bool bIgnoreEnabledStatus )
@@ -1295,14 +950,14 @@ FWidgetPath FSlateApplication::LocateWindowUnderMouse( FVector2D ScreenspaceMous
 
 bool FSlateApplication::IsWindowHousingInteractiveTooltip(const TSharedRef<const SWindow>& WindowToTest) const
 {
-	const TSharedPtr<IToolTip> ActiveToolTipPtr = ActiveToolTip.Pin();
-	const TSharedPtr<SWindow> ToolTipWindowPtr = ToolTipWindow.Pin();
-	const bool bIsHousingInteractiveTooltip =
-		WindowToTest == ToolTipWindowPtr &&
-		ActiveToolTipPtr.IsValid() &&
-		ActiveToolTipPtr->IsInteractive();
-
-	return bIsHousingInteractiveTooltip;
+	for (TSharedPtr<const FSlateUser> User : Users)
+	{
+		if (User && User->IsWindowHousingInteractiveTooltip(WindowToTest))
+		{
+			return true;
+		}
+	}
+	return false;
 }
 
 void FSlateApplication::DrawWindows()
@@ -1359,61 +1014,15 @@ void FSlateApplication::DrawWindowAndChildren( const TSharedRef<SWindow>& Window
 #endif
 			}
 
-			// Draw drag drop operation if it's windowless.
-			if ( IsDragDropping() && DragDropContent->IsWindowlessOperation() )
-			{
-				TSharedPtr<SWindow> DragDropWindow = DragDropWindowPtr.Pin();
-				if ( DragDropWindow.IsValid() && DragDropWindow == WindowToDraw )
-				{
-					TSharedPtr<SWidget> DecoratorWidget = DragDropContent->GetDefaultDecorator();
-					if ( DecoratorWidget.IsValid() && DecoratorWidget->GetVisibility().IsVisible() )
-					{
-						const float WindowRootScale = GetApplicationScale() * DragDropWindow->GetNativeWindow()->GetDPIScaleFactor();
-
-						DecoratorWidget->SetVisibility(EVisibility::HitTestInvisible);
-						DecoratorWidget->SlatePrepass(WindowRootScale);
-
-						FVector2D DragDropContentInWindowSpace = WindowToDraw->GetWindowGeometryInScreen().AbsoluteToLocal(DragDropContent->GetDecoratorPosition()) * WindowRootScale;
-						const FGeometry DragDropContentGeometry = FGeometry::MakeRoot(DecoratorWidget->GetDesiredSize(), FSlateLayoutTransform(DragDropContentInWindowSpace));
-
-						DecoratorWidget->Paint(
-							FPaintArgs(&WindowToDraw.Get(), WindowToDraw->GetHittestGrid(), WindowToDraw->GetPositionInScreen(), GetCurrentTime(), GetDeltaTime()),
-							DragDropContentGeometry, WindowToDraw->GetClippingRectangleInWindow(),
-							WindowElementList,
-							++MaxLayerId,
-							FWidgetStyle(),
-							WindowToDraw->IsEnabled());
-					}
-				}
-			}
-
-			// Draw Software Cursor
-			TSharedPtr<SWindow> CursorWindow = CursorWindowPtr.Pin();
-			if (CursorWindow.IsValid() && WindowToDraw == CursorWindow)
-			{
-				TSharedPtr<SWidget> CursorWidget = CursorWidgetPtr.Pin();
+			// Draw windowless drag drop operations
+			ForEachUser([&WindowToDraw, &WindowElementList, &MaxLayerId](FSlateUser& User) {
+				User.DrawWindowlessDragDropContent(WindowToDraw, WindowElementList, MaxLayerId);
+			});
 
-				if (CursorWidget.IsValid())
-				{
-					const float WindowRootScale = GetApplicationScale() * CursorWindow->GetNativeWindow()->GetDPIScaleFactor();
-
-					CursorWidget->SetVisibility(EVisibility::HitTestInvisible);
-					CursorWidget->SlatePrepass(WindowRootScale);
-
-					FVector2D CursorInScreen = GetCursorPos();
-					FVector2D CursorPosInWindowSpace = WindowToDraw->GetWindowGeometryInScreen().AbsoluteToLocal(CursorInScreen) * WindowRootScale;
-					CursorPosInWindowSpace += (CursorWidget->GetDesiredSize() * -0.5) * WindowRootScale;
-					const FGeometry CursorGeometry = FGeometry::MakeRoot(CursorWidget->GetDesiredSize(), FSlateLayoutTransform(CursorPosInWindowSpace));
-
-					CursorWidget->Paint(
-						FPaintArgs(&WindowToDraw.Get(), WindowToDraw->GetHittestGrid(), WindowToDraw->GetPositionInScreen(), GetCurrentTime(), GetDeltaTime()),
-						CursorGeometry, WindowToDraw->GetClippingRectangleInWindow(),
-						WindowElementList,
-						++MaxLayerId,
-						FWidgetStyle(),
-						WindowToDraw->IsEnabled());
-				}
-			}
+			// Draw Software Cursors
+			ForEachUser([&WindowToDraw, &WindowElementList, &MaxLayerId](FSlateUser& User) {
+				User.DrawCursor(WindowToDraw, WindowElementList, MaxLayerId);
+			});
 		}
 
 #if SLATE_HAS_WIDGET_REFLECTOR
@@ -1569,11 +1178,9 @@ void FSlateApplication::PrivateDrawWindows( TSharedPtr<SWindow> DrawOnlyThisWind
 	
 #if SLATE_HAS_WIDGET_REFLECTOR
 	// Is user expecting visual feedback from the Widget Reflector?
-	const bool bVisualizeLayoutUnderCursor = WidgetReflectorPtr.IsValid() && WidgetReflectorPtr.Pin()->IsVisualizingLayoutUnderCursor();
-	if (bVisualizeLayoutUnderCursor)
+	if (WidgetReflectorPtr.IsValid() && WidgetReflectorPtr.Pin()->IsVisualizingLayoutUnderCursor())
 	{
-		WidgetsToVisualizeUnderCursor =
-			WidgetsUnderCursorLastEvent.FindRef(FUserAndPointer(CursorUserIndex, CursorPointerIndex)).ToWidgetPath();
+		WidgetsToVisualizeUnderCursor = GetOrCreateUser(CursorUserIndex)->GetLastWidgetsUnderCursor().ToWidgetPath();
 	}
 #endif
 
@@ -1651,52 +1258,47 @@ void FSlateApplication::FinishedInputThisFrame()
 {
 	const float DeltaTime = GetDeltaTime();
 
+	// Any preprocessors are given a chance to process accumulated values (or do whatever other tick things they want)
+	// after we've finished processing all of the input for the frame
 	if (PlatformApplication->Cursor.IsValid())
 	{
 		InputPreProcessors.Tick(DeltaTime, *this, PlatformApplication->Cursor.ToSharedRef());
 	}
 
-	// All the input events have been processed.
-
-	// Any widgets that may have received pointer input events
-	// are given a chance to process accumulated values.
-	if (MouseCaptor.HasCapture())
-	{
-		TArray<TSharedRef<SWidget>> Captors = MouseCaptor.ToSharedWidgets();
-		for (const auto & Captor : Captors )
+	// Any widgets that may have received pointer input events are given a chance to process accumulated values.
+	ForEachUser([](FSlateUser& User) {
+		if (User.HasAnyCapture())
 		{
-			Captor->OnFinishedPointerInput();
+			for (const TSharedRef<SWidget>& Captor : User.GetCaptorWidgets())
+			{
+				Captor->OnFinishedPointerInput();
+			}
 		}
-	}
-	else
-	{
-		for( auto LastWidgetIterator = WidgetsUnderCursorLastEvent.CreateConstIterator(); LastWidgetIterator; ++LastWidgetIterator )
+		else
 		{
-			FWeakWidgetPath Path = LastWidgetIterator.Value();
-
-			for ( const TWeakPtr<SWidget>& WidgetPtr : Path.Widgets )
+			for (const auto& IndexPathPair : User.GetWidgetsUnderPointerLastEventByIndex())
 			{
-				const TSharedPtr<SWidget>& Widget = WidgetPtr.Pin();
-				if (Widget.IsValid())
-				{
-					Widget->OnFinishedPointerInput();
-				}
-				else
+				for (const TWeakPtr<SWidget>& WidgetPtr : IndexPathPair.Value.Widgets)
 				{
-					break;
+					if (TSharedPtr<SWidget> Widget = WidgetPtr.Pin())
+					{
+						Widget->OnFinishedPointerInput();
+					}
+					else
+					{
+						break;
+					}
 				}
 			}
 		}
-	}
-
-	// Any widgets that may have received key events
-	// are given a chance to process accumulated values.
-	ForEachUser([&] (FSlateUser* User) {
-		const FWeakWidgetPath& WidgetPath = User->GetWeakFocusPath();
+	});
+	
+	// Any widgets that may have received key events are given a chance to process accumulated values.
+	ForEachUser([] (FSlateUser& User) {
+		const FWeakWidgetPath& WidgetPath = User.GetWeakFocusPath();
 		for ( const TWeakPtr<SWidget>& WidgetPtr : WidgetPath.Widgets )
 		{
-			const TSharedPtr<SWidget>& Widget = WidgetPtr.Pin();
-			if ( Widget.IsValid() )
+			if (TSharedPtr<SWidget> Widget = WidgetPtr.Pin())
 			{
 				Widget->OnFinishedKeyInput();
 			}
@@ -1707,9 +1309,7 @@ void FSlateApplication::FinishedInputThisFrame()
 		}
 	});
 
-	ForEachUser([&] (FSlateUser* User) {
-		User->FinishFrame();
-	});
+	ForEachUser([] (FSlateUser& User) { User.FinishFrame(); });
 }
 
 void FSlateApplication::Tick(ESlateTickType TickType)
@@ -1783,26 +1383,10 @@ void FSlateApplication::TickApplication(ESlateTickType TickType, float DeltaTime
 			PreTickEvent.Broadcast(DeltaTime);
 		}
 
-		// The widget locking the cursor to its bounds may have been reshaped.
-		// Check if the widget was reshaped and update the cursor lock
-		// bounds if needed.
-		UpdateCursorLockRegion();
-
-		// When Slate captures the mouse, it is up to us to set the cursor 
-		// because the OS assumes that we own the mouse.
-		if (GEnableCursorQueries && (MouseCaptor.HasCapture() || bQueryCursorRequested))
-		{
-			SCOPE_CYCLE_COUNTER(STAT_SlateQueryCursor);
-			QueryCursor();
-		}
-
-		{
-			SCOPE_CYCLE_COUNTER(STAT_SlateUpdateTooltip);
-
-			// Update tool tip, if we have one
-			const bool AllowSpawningOfToolTips = false;
-			UpdateToolTip(AllowSpawningOfToolTips);
-		}
+		ForEachUser([this](FSlateUser& User) {
+			User.UpdateCursor();
+			User.UpdateTooltip(MenuStack, false);
+		});
 	}
 
 	// Advance time
@@ -1841,15 +1425,11 @@ void FSlateApplication::TickApplication(ESlateTickType TickType, float DeltaTime
 
 	if (TickType == ESlateTickType::All)
 	{
-		const bool bNeedsSyntheticMouseMouse = SynthesizeMouseMovePending > 0;
-
-		if (GEnableSyntheticMouseMoves && bNeedsSyntheticMouseMouse && (!GIsGameThreadIdInitialized || IsInGameThread())) 
-		{
-			// Force a mouse move event to make sure all widgets know whether there is a mouse cursor hovering over them
-			SynthesizeMouseMove();
-			--SynthesizeMouseMovePending;
-		}
-
+		bool bSynthesizedCursorMove = false;
+		ForEachUser([&bSynthesizedCursorMove](FSlateUser& User) {
+			bSynthesizedCursorMove |= User.SynthesizeCursorMoveIfNeeded();
+		});
+		
 		// Update auto-throttling based on elapsed time since user interaction
 		ThrottleApplicationBasedOnMouseMovement();
 
@@ -1863,14 +1443,13 @@ void FSlateApplication::TickApplication(ESlateTickType TickType, float DeltaTime
 		const bool bAnyActiveTimersPending = AnyActiveTimersArePending();
 		if (bAnyActiveTimersPending)
 		{
-			// Some UI might slide under the cursor. To a widget, this is
-			// as if the cursor moved over it.
-			QueueSynthesizedMouseMove();
+			// Some UI might slide under the cursor. To a widget, this is as if the cursor moved over it.
+			ForEachUser([](FSlateUser& User) { User.QueueSyntheticCursorMove(); });
 		}
 
 		// Generate any simulated gestures that we've detected.
-		ForEachUser([&] (FSlateUser* User) {
-			User->GestureDetector.GenerateGestures(*this, SimulateGestures);
+		ForEachUser([this] (FSlateUser& User) {
+			User.GetGestureDetector().GenerateGestures(*this, SimulateGestures);
 		});
 
 		FPlatformMisc::EndNamedEvent();
@@ -1879,7 +1458,7 @@ void FSlateApplication::TickApplication(ESlateTickType TickType, float DeltaTime
 		// This relies on Widgets properly registering for Active timer when they need something to happen even
 		// when the user is not providing any input (ie, animations, viewport rendering, async polling, etc).
 		bIsSlateAsleep = true;
-		if (!AllowSlateToSleep.GetValueOnGameThread() || bAnyActiveTimersPending || !bIsUserIdle || bNeedsSyntheticMouseMouse || FApp::UseVRFocus())
+		if (!AllowSlateToSleep.GetValueOnGameThread() || bAnyActiveTimersPending || !bIsUserIdle || bSynthesizedCursorMove || FApp::UseVRFocus())
 		{
 			bIsSlateAsleep = false; // if we get here, then Slate is not sleeping
 
@@ -1949,7 +1528,7 @@ void FSlateApplication::ThrottleApplicationBasedOnMouseMovement()
 				if( TimeSinceLastUserInteraction > MinTimeSinceButtonPressToThrottle )
 				{
 					// If a widget has the mouse captured, then we won't bother throttling
-					if( !MouseCaptor.HasCapture() )
+					if( !HasAnyMouseCaptor() )
 					{
 						// If there is no Slate window under the mouse, then we won't engage throttling
 						if( LocateWindowUnderMouse( GetCursorPos(), GetInteractiveTopLevelWindows() ).IsValid() )
@@ -2144,12 +1723,10 @@ EUINavigation FSlateApplication::GetNavigationDirectionFromKey(const FKeyEvent&
 {
 #if WITH_EDITOR
 	// Check if the focused widget is an editor widget or a PIE widget so we know which config to use.
-	if (const FSlateUser* User = GetUser(GetUserIndexForKeyboard()))
+	TSharedPtr<const FSlateUser> User = GetUser(GetUserIndexForKeyboard());
+	if (User && !IsFocusInViewport(AllGameViewports, User->GetWeakFocusPath()))
 	{
-		if (!IsFocusInViewport(AllGameViewports, User->GetWeakFocusPath()))
-		{
-			return EditorNavigationConfig->GetNavigationDirectionFromKey(InKeyEvent);
-		}
+		return EditorNavigationConfig->GetNavigationDirectionFromKey(InKeyEvent);
 	}
 #endif
 	return NavigationConfig->GetNavigationDirectionFromKey(InKeyEvent);
@@ -2159,12 +1736,10 @@ EUINavigation FSlateApplication::GetNavigationDirectionFromAnalog(const FAnalogI
 {
 #if WITH_EDITOR
 	// Check if the focused widget is an editor widget or a PIE widget so we know which config to use.
-	if (const FSlateUser* User = GetUser(GetUserIndexForKeyboard()))
+	TSharedPtr<const FSlateUser> User = GetUser(GetUserIndexForKeyboard());
+	if (User && !IsFocusInViewport(AllGameViewports, User->GetWeakFocusPath()))
 	{
-		if (!IsFocusInViewport(AllGameViewports, User->GetWeakFocusPath()))
-		{
-			return EditorNavigationConfig->GetNavigationDirectionFromAnalog(InAnalogEvent);
-		}
+		return EditorNavigationConfig->GetNavigationDirectionFromKey(InAnalogEvent);
 	}
 #endif
 	return NavigationConfig->GetNavigationDirectionFromAnalog(InAnalogEvent);
@@ -2203,8 +1778,11 @@ void FSlateApplication::AddModalWindow( TSharedRef<SWindow> InSlateWindow, const
     // Push the active modal window onto the stack.
 	ActiveModalWindows.AddUnique( InSlateWindow );
 
-	// Close the open tooltip when a new window is open.  Tooltips from non-modal windows can be dangerous and cause rentrancy into code that shouldnt execute in a modal state.
-	CloseToolTip();
+	// Close all the open tooltips when a new window is opened.  
+	// Tooltips from non-modal windows can be dangerous and re-enter into code that shouldn't execute in a modal state.
+	ForEachUser([](FSlateUser& User) {
+		User.CloseTooltip();
+	});
 
 	// Set the modal flag on the window
 	InSlateWindow->SetAsModalWindow();
@@ -2239,7 +1817,7 @@ void FSlateApplication::AddModalWindow( TSharedRef<SWindow> InSlateWindow, const
 
 	// Release mouse capture here in case the new modal window has been added in one of the mouse button
 	// event callbacks. Otherwise it will be unresponsive until the next mouse up event.
-	ReleaseMouseCapture();
+	ReleaseAllPointerCapture();
 
 	// Clear the cached pressed mouse buttons, in case a new modal window has been added between the mouse down and mouse up of another window.
 	PressedMouseButtons.Empty();
@@ -2452,8 +2030,8 @@ void FSlateApplication::DismissMenuByWidget(const TSharedRef<SWidget>& InWidgetI
 
 void FSlateApplication::RequestDestroyWindow( TSharedRef<SWindow> InWindowToDestroy )
 {
-	ForEachUser([&] (FSlateUser* User) {
-		User->NotifyWindowDestroyed(InWindowToDestroy);
+	ForEachUser([&InWindowToDestroy] (FSlateUser& User) {
+		User.NotifyWindowDestroyed(InWindowToDestroy);
 	});
 
 	// Logging to track down window shutdown issues with movie loading threads. Too spammy in editor builds with all the windows
@@ -2503,7 +2081,9 @@ void FSlateApplication::ExternalModalStart()
 		DismissAllMenus();
 
 		// Close tool-tips
-		CloseToolTip();
+		ForEachUser([](FSlateUser& User) {
+			User.CloseTooltip();
+		});
 
 		// Tick and render Slate so that it can destroy any menu windows if necessary before we disable.
 		Tick();
@@ -2642,6 +2222,7 @@ int32 FSlateApplication::GetUserIndexForController(int32 ControllerId) const
 
 void FSlateApplication::SetInputManager(TSharedRef<ISlateInputManager> InInputManager)
 {
+	//@todo DanH/NickD: Should we be worried at all about potential invalidation of SlateUsers we've created but that no longer occupy an index that can be referenced?
 	InputManager = InInputManager;
 }
 
@@ -2663,9 +2244,7 @@ void FSlateApplication::SetAllUserFocusToGameViewport(EFocusCause ReasonFocusIsC
 		FWidgetPath PathToWidget;
 		FSlateWindowHelper::FindPathToWidget(SlateWindows, CurrentGameViewportWidget.ToSharedRef(), /*OUT*/ PathToWidget);
 
-		ForEachUser([&] (FSlateUser* User) {
-			SetUserFocus(User, PathToWidget, ReasonFocusIsChanging);
-		});
+		SetAllUserFocus(PathToWidget, ReasonFocusIsChanging);
 	}
 }
 
@@ -2699,32 +2278,24 @@ void FSlateApplication::ActivateGameViewport()
 	}
 }
 
-void FSlateApplication::ResetDragDropState()
-{
-	DragDropContent.Reset();
-}
-
 bool FSlateApplication::SetUserFocus(uint32 UserIndex, const TSharedPtr<SWidget>& WidgetToFocus, EFocusCause ReasonFocusIsChanging /* = EFocusCause::SetDirectly*/)
 {
 	const bool bValidWidget = WidgetToFocus.IsValid();
 	ensureMsgf(bValidWidget, TEXT("Attempting to focus an invalid widget. If your intent is to clear focus use ClearUserFocus()"));
 	if (bValidWidget)
 	{
-		if (FSlateUser* User = GetOrCreateUser(UserIndex))
+		FWidgetPath PathToWidget;
+		const bool bFound = FSlateWindowHelper::FindPathToWidget(SlateWindows, WidgetToFocus.ToSharedRef(), /*OUT*/ PathToWidget);
+		if (bFound)
 		{
-			FWidgetPath PathToWidget;
-			const bool bFound = FSlateWindowHelper::FindPathToWidget(SlateWindows, WidgetToFocus.ToSharedRef(), /*OUT*/ PathToWidget);
-			if (bFound)
-			{
-				return SetUserFocus(User, PathToWidget, ReasonFocusIsChanging);
-			}
-			else
+			return SetUserFocus(*GetOrCreateUser(UserIndex), PathToWidget, ReasonFocusIsChanging);
+		}
+		else
+		{
+			const bool bFoundVirtual = FSlateWindowHelper::FindPathToWidget(SlateVirtualWindows, WidgetToFocus.ToSharedRef(), /*OUT*/ PathToWidget);
+			if (bFoundVirtual)
 			{
-				const bool bFoundVirtual = FSlateWindowHelper::FindPathToWidget(SlateVirtualWindows, WidgetToFocus.ToSharedRef(), /*OUT*/ PathToWidget);
-				if (bFoundVirtual)
-				{
-					return SetUserFocus(User, PathToWidget, ReasonFocusIsChanging);
-				}
+				return SetUserFocus(*GetOrCreateUser(UserIndex), PathToWidget, ReasonFocusIsChanging);
 			}
 		}
 	}
@@ -2759,12 +2330,8 @@ void FSlateApplication::SetAllUserFocus(const TSharedPtr<SWidget>& WidgetToFocus
 
 TSharedPtr<SWidget> FSlateApplication::GetUserFocusedWidget(uint32 UserIndex) const
 {
-	if ( const FSlateUser* User = GetUser(UserIndex) )
-	{
-		return User->GetFocusedWidget();
-	}
-
-	return TSharedPtr<SWidget>();
+	TSharedPtr<const FSlateUser> User = GetUser(UserIndex);
+	return User ? User->GetFocusedWidget() : TSharedPtr<SWidget>();
 }
 
 void FSlateApplication::ClearUserFocus(uint32 UserIndex, EFocusCause ReasonFocusIsChanging /* = EFocusCause::SetDirectly*/)
@@ -2795,18 +2362,18 @@ void FSlateApplication::ResetToDefaultInputSettings()
 
 void FSlateApplication::ResetToDefaultPointerInputSettings()
 {
-	TArray<FWidgetPath> CaptorPaths = MouseCaptor.ToWidgetPaths();
-	for (const FWidgetPath& MouseCaptorPath : CaptorPaths )
-	{
-		ProcessReply(MouseCaptorPath, FReply::Handled().ReleaseMouseCapture(), nullptr, nullptr);
-	}
+	//@todo DanH: Leaving unchanged during the FSlateUser updates, but this seems to be overkill to loop through every single captor
+	//		and process a ReleaseCapture reply for each. Can't we just ReleaseAllCapture() on every user? (We'd still need to update each user's cursor type)
+	//ReleaseMouseCapture();
+	ForEachUser([this](FSlateUser& User) {
+		for (const FWidgetPath& CaptorPath : User.GetCaptorPaths())
+		{
+			ProcessReply(CaptorPath, FReply::Handled().ReleaseMouseCapture(), nullptr, nullptr);
+		}
+		User.GetCursor()->SetType(EMouseCursor::Default);
+	});
 
 	ProcessReply(FWidgetPath(), FReply::Handled().ReleaseMouseLock(), nullptr, nullptr);
-
-	if ( PlatformApplication->Cursor.IsValid() )
-	{
-		PlatformApplication->Cursor->SetType(EMouseCursor::Default);
-	}
 }
 
 void* FSlateApplication::GetMouseCaptureWindow() const
@@ -2816,12 +2383,25 @@ void* FSlateApplication::GetMouseCaptureWindow() const
 
 void FSlateApplication::ReleaseMouseCapture()
 {
-	MouseCaptor.InvalidateCaptureForAllPointers();
+	ReleaseAllPointerCapture();
+}
+
+void FSlateApplication::ReleaseAllPointerCapture()
+{
+	ForEachUser([](FSlateUser& User) { User.ReleaseAllCapture(); });
+}
+
+void FSlateApplication::ReleaseAllPointerCapture(int32 UserIndex)
+{
+	if (TSharedPtr<FSlateUser> User = GetUser(UserIndex))
+	{
+		User->ReleaseAllCapture();
+	}
 }
 
 void FSlateApplication::ReleaseMouseCaptureForUser(int32 UserIndex)
 {
-	MouseCaptor.InvalidateCaptureForUser(UserIndex);
+	ReleaseAllPointerCapture(UserIndex);
 }
 
 FDelegateHandle FSlateApplication::RegisterOnWindowActionNotification(const FOnWindowAction& Notification)
@@ -2910,22 +2490,18 @@ bool FSlateApplication::SetKeyboardFocus(const FWidgetPath& InFocusPath, const E
 
 bool FSlateApplication::SetUserFocus(const uint32 InUserIndex, const FWidgetPath& InFocusPath, const EFocusCause InCause)
 {
-	return SetUserFocus(GetOrCreateUser(InUserIndex), InFocusPath, InCause);
+	TSharedPtr<FSlateUser> User = GetUser(InUserIndex);
+	return User && SetUserFocus(*User, InFocusPath, InCause);
 }
 
-bool FSlateApplication::SetUserFocus(FSlateUser* User, const FWidgetPath& InFocusPath, const EFocusCause InCause)
+bool FSlateApplication::SetUserFocus(FSlateUser& User, const FWidgetPath& InFocusPath, const EFocusCause InCause)
 {
-	if ( !User )
-	{
-		return false;
-	}
-
 	if (InFocusPath.IsValid())
 	{
 		TSharedRef<SWindow> Window = InFocusPath.GetWindow();
 		if (ActiveModalWindows.Num() != 0 && !(Window->IsDescendantOf(GetActiveModalWindow()) || ActiveModalWindows.Top() == Window))
 		{
-			UE_LOG(LogSlate, Warning, TEXT("Ignoring SetUserFocus because it's not an active modal Window (user %i not set to %s."), User->GetUserIndex(), *InFocusPath.GetLastWidget()->ToString());
+			UE_LOG(LogSlate, Warning, TEXT("Ignoring SetUserFocus because it's not an active modal Window (user %i not set to %s."), User.GetUserIndex(), *InFocusPath.GetLastWidget()->ToString());
 			return false;
 		}
 	}
@@ -2934,7 +2510,7 @@ bool FSlateApplication::SetUserFocus(FSlateUser* User, const FWidgetPath& InFocu
 	const bool bReflectorShowingFocus = WidgetReflector.IsValid() && WidgetReflector->IsShowingFocus();
 
 	// Get the old Widget information
-	const FWeakWidgetPath OldFocusedWidgetPath = User->GetWeakFocusPath();
+	const FWeakWidgetPath OldFocusedWidgetPath = User.GetWeakFocusPath();
 	TSharedPtr<SWidget> OldFocusedWidget = OldFocusedWidgetPath.IsValid() ? OldFocusedWidgetPath.GetLastWidget().Pin() : TSharedPtr< SWidget >();
 
 	// Get the new widget information by finding the first widget in the path that supports focus
@@ -2965,10 +2541,10 @@ bool FSlateApplication::SetUserFocus(FSlateUser* User, const FWidgetPath& InFocu
 		}
 	}
 
-	User->FocusVersion++;
-	int32 CurrentFocusVersion = User->FocusVersion;
+	User.IncrementFocusVersion();
+	int32 CurrentFocusVersion = User.GetFocusVersion();
 
-	FFocusEvent FocusEvent(InCause, User->GetUserIndex());
+	FFocusEvent FocusEvent(InCause, User.GetUserIndex());
 
 #if WITH_SLATE_DEBUGGING
 	FSlateDebugging::BroadcastFocusChanging(FocusEvent, OldFocusedWidgetPath, OldFocusedWidget, NewFocusedWidgetPath, NewFocusedWidget);
@@ -2987,7 +2563,7 @@ bool FSlateApplication::SetUserFocus(FSlateUser* User, const FWidgetPath& InFocu
 				SomeWidget->OnFocusChanging(OldFocusedWidgetPath, NewFocusedWidgetPath, FocusEvent);
 
 				// If focus setting is interrupted, stop what we're doing, as someone has already changed the focus path.
-				if ( CurrentFocusVersion != User->FocusVersion )
+				if ( CurrentFocusVersion != User.GetFocusVersion())
 				{
 					return false;
 				}
@@ -3008,7 +2584,7 @@ bool FSlateApplication::SetUserFocus(FSlateUser* User, const FWidgetPath& InFocu
 				SomeWidget->OnFocusChanging(OldFocusedWidgetPath, NewFocusedWidgetPath, FocusEvent);
 
 				// If focus setting is interrupted, stop what we're doing, as someone has already changed the focus path.
-				if ( CurrentFocusVersion != User->FocusVersion )
+				if ( CurrentFocusVersion != User.GetFocusVersion())
 				{
 					return false;
 				}
@@ -3035,7 +2611,7 @@ bool FSlateApplication::SetUserFocus(FSlateUser* User, const FWidgetPath& InFocu
 	}
 
 	// Store a weak widget path to the widget that's taking focus
-	User->SetFocusPath(NewFocusedWidgetPath, InCause, ShowFocus);
+	User.SetFocusPath(NewFocusedWidgetPath, InCause, ShowFocus);
 
 	// Let the old widget know that it lost keyboard focus
 	if (OldFocusedWidget.IsValid())
@@ -3081,7 +2657,7 @@ bool FSlateApplication::SetUserFocus(FSlateUser* User, const FWidgetPath& InFocu
 		FReply Reply = NewFocusedWidget->OnFocusReceived(WidgetToFocus.Geometry, FocusEvent);
 		if (Reply.IsEventHandled())
 		{
-			ProcessReply(InFocusPath, Reply, nullptr, nullptr, User->GetUserIndex());
+			ProcessReply(InFocusPath, Reply, nullptr, nullptr, User.GetUserIndex());
 		}
 
 		NavigationConfig->OnNavigationChangedFocus(OldFocusedWidget, NewFocusedWidget, FocusEvent);
@@ -3097,7 +2673,7 @@ bool FSlateApplication::SetUserFocus(FSlateUser* User, const FWidgetPath& InFocu
 
 void FSlateApplication::SetAllUserFocus(const FWidgetPath& InFocusPath, const EFocusCause InCause)
 {
-	ForEachUser([&] (FSlateUser* User) {
+	ForEachUser([&] (FSlateUser& User) {
 		SetUserFocus(User, InFocusPath, InCause);
 	});
 }
@@ -3106,9 +2682,8 @@ void FSlateApplication::SetAllUserFocusAllowingDescendantFocus(const FWidgetPath
 {
 	TSharedRef<SWidget> FocusWidget = InFocusPath.Widgets.Last().Widget;
 
-	ForEachUser([&] (FSlateUser* User) {
-		const FWeakWidgetPath& WidgetPath = User->GetWeakFocusPath();
-		if (!WidgetPath.ContainsWidget(FocusWidget))
+	ForEachUser([&] (FSlateUser& User) {
+		if (!User.GetWeakFocusPath().ContainsWidget(FocusWidget))
 		{
 			SetUserFocus(User, InFocusPath, InCause);
 		}
@@ -3128,15 +2703,10 @@ void FSlateApplication::OnShutdown()
 
 void FSlateApplication::CloseAllWindowsImmediately()
 {
-	ResetDragDropState();
-
-	// Clean up our tooltip window
-	TSharedPtr< SWindow > PinnedToolTipWindow(ToolTipWindow.Pin());
-	if (PinnedToolTipWindow.IsValid())
-	{
-		PinnedToolTipWindow->RequestDestroyWindow();
-		ToolTipWindow.Reset();
-	}
+	ForEachUser([](FSlateUser& User) { 
+		User.ResetDragDropContent();
+		User.ResetTooltipWindow();
+	});
 
 	// Destroy all top level windows.  
 	// ::RequestDestroyWindow will remove the window from the TArray SlateWindows so we
@@ -3255,23 +2825,25 @@ TSharedPtr<SWindow> FSlateApplication::FindWidgetWindow( TSharedRef<const SWidge
 
 void FSlateApplication::ProcessExternalReply(const FWidgetPath& CurrentEventPath, const FReply TheReply, const uint32 UserIndex, const uint32 PointerIndex)
 {
-	if (PointerIndex == FSlateApplicationBase::CursorPointerIndex)
+	if (PointerIndex == CursorPointerIndex)
 	{
+		TSharedRef<FSlateUser> SlateUser = GetOrCreateUser(UserIndex);
+		const bool bIsPrimaryUser = UserIndex == CursorUserIndex;
+
 		FPointerEvent MouseEvent(
 			UserIndex,
 			PointerIndex,
-			GetCursorPos(),
-			GetLastCursorPos(),
-			PressedMouseButtons,
+			SlateUser->GetCursorPosition(),
+			SlateUser->GetPreviousCursorPosition(),
+			bIsPrimaryUser ? PressedMouseButtons : TSet<FKey>(),
 			EKeys::Invalid,
 			0,
-			PlatformApplication->GetModifierKeys()
+			bIsPrimaryUser ? PlatformApplication->GetModifierKeys() : FModifierKeysState()
 		);
 
-		const FWeakWidgetPath& LastWidgetsUnderCursor = WidgetsUnderCursorLastEvent.FindRef(FUserAndPointer(UserIndex, PointerIndex));
-	
-		const FWidgetPath* PathToWidgetPtr = nullptr;
 		FWidgetPath PathToWidget;
+		const FWidgetPath* PathToWidgetPtr = nullptr;
+		FWeakWidgetPath LastWidgetsUnderCursor = SlateUser->GetLastWidgetsUnderCursor();
 		if (LastWidgetsUnderCursor.IsValid())
 		{
 			PathToWidget = LastWidgetsUnderCursor.ToWidgetPath();
@@ -3296,13 +2868,10 @@ void FSlateApplication::ProcessReply( const FWidgetPath& CurrentEventPath, const
 	// Make sure to only clobber WidgetsUnderCursor if we actually had a mouse capture.
 	uint32 PointerIndex = InMouseEvent != nullptr ? InMouseEvent->GetPointerIndex() : CursorPointerIndex;
 
-	if (MouseCaptor.HasCaptureForPointerIndex(UserIndex,PointerIndex) && (TheReply.ShouldReleaseMouse() || bStartingDragDrop) )
+	TSharedRef<FSlateUser> SlateUser = GetOrCreateUser(UserIndex);
+	if (SlateUser->HasCapture(PointerIndex) && (TheReply.ShouldReleaseMouse() || bStartingDragDrop))
 	{
-		WidgetsUnderCursorLastEvent.Add( FUserAndPointer(UserIndex,PointerIndex), MouseCaptor.ToWeakPath(UserIndex,PointerIndex) );
-		MouseCaptor.InvalidateCaptureForPointer(UserIndex,PointerIndex);
-
-		// If mouse capture changes, we should refresh the cursor state.
-		bQueryCursorRequested = true;
+		SlateUser->ReleaseCapture(PointerIndex);
 	}
 
 	// Clear focus is requested.
@@ -3310,30 +2879,29 @@ void FSlateApplication::ProcessReply( const FWidgetPath& CurrentEventPath, const
 	{
 		if (TheReply.AffectsAllUsers())
 		{
-			ForEachUser([&] (FSlateUser* User) {
-				SetUserFocus(User, FWidgetPath(), TheReply.GetFocusCause());
-			});
+			ClearAllUserFocus(TheReply.GetFocusCause());
 		}
 		else
 		{
-			SetUserFocus(UserIndex, FWidgetPath(), TheReply.GetFocusCause());
+			SlateUser->ClearFocus(TheReply.GetFocusCause());
 		}
 	}
 
 	if (TheReply.ShouldEndDragDrop())
 	{
-		CancelDragDrop();
+		SlateUser->CancelDragDrop();
 	}
 
-	if ( bStartingDragDrop )
+	if (bStartingDragDrop)
 	{
-		checkf( !this->DragDropContent.IsValid(), TEXT("Drag and Drop already in progress!") );
-		check( true == TheReply.IsEventHandled() );
-		check( WidgetsUnderMouse != nullptr );
-		check( InMouseEvent != nullptr );
-		DragDropContent = ReplyDragDropContent;
+		checkf(!SlateUser->IsDragDropping(), TEXT("Drag and Drop already in progress!"));
+		check(TheReply.IsEventHandled());
+		check(WidgetsUnderMouse);
+		check(InMouseEvent);
 
-		const FWeakWidgetPath& LastWidgetsUnderCursor = WidgetsUnderCursorLastEvent.FindRef(FUserAndPointer(UserIndex, PointerIndex));
+		SlateUser->SetDragDropContent(ReplyDragDropContent.ToSharedRef());
+		
+		const FWeakWidgetPath LastWidgetsUnderCursor = SlateUser->GetLastWidgetsUnderPointer(PointerIndex);
 
 		// We have entered drag and drop mode.
 		// LastWidgetsUnderCursor.ToWidgetPath(), CurrentEventPath, and
@@ -3396,37 +2964,34 @@ void FSlateApplication::ProcessReply( const FWidgetPath& CurrentEventPath, const
 		// Do not capture the mouse if we are also starting a drag and drop.
 		if (RequestedMouseCaptor.IsValid() && !bStartingDragDrop)
 		{
-			if ( MouseCaptor.SetMouseCaptor(UserIndex, PointerIndex, CurrentEventPath, RequestedMouseCaptor) )
+			if (SlateUser->SetPointerCaptor(PointerIndex, RequestedMouseCaptor.ToSharedRef(), CurrentEventPath))
 			{
 				if (WidgetsUnderMouse)
 				{
-					const FWeakWidgetPath& LastWidgetsUnderCursor = WidgetsUnderCursorLastEvent.FindRef(FUserAndPointer(UserIndex, PointerIndex));
-
+					const FWeakWidgetPath LastWidgetsUnderCursor = SlateUser->GetLastWidgetsUnderPointer(PointerIndex);
 					if (LastWidgetsUnderCursor.IsValid())
 					{
 						for (int32 WidgetIndex = LastWidgetsUnderCursor.Widgets.Num() - 1; WidgetIndex >= 0; --WidgetIndex)
 						{
-							TSharedPtr<SWidget> SomeWidgetPreviouslyUnderCursor = LastWidgetsUnderCursor.Widgets[WidgetIndex].Pin();
+							TSharedPtr<SWidget> WidgetPreviouslyUnderCursor = LastWidgetsUnderCursor.Widgets[WidgetIndex].Pin();
 
-							if (SomeWidgetPreviouslyUnderCursor.IsValid())
+							if (WidgetPreviouslyUnderCursor.IsValid())
 							{
-								if (SomeWidgetPreviouslyUnderCursor != RequestedMouseCaptor)
+								if (WidgetPreviouslyUnderCursor != RequestedMouseCaptor)
 								{
 									// It's possible for mouse event to be null if we end up here from a keyboard event. If so, we should synthesize an event
 									if (InMouseEvent)
 									{
 										// Note that the event's pointer position is not translated.
-										SomeWidgetPreviouslyUnderCursor->OnMouseLeave(*InMouseEvent);
+										WidgetPreviouslyUnderCursor->OnMouseLeave(*InMouseEvent);
 									}
 									else
 									{
 										const FPointerEvent& SimulatedPointer = FPointerEvent();
-										SomeWidgetPreviouslyUnderCursor->OnMouseLeave(SimulatedPointer);
+										WidgetPreviouslyUnderCursor->OnMouseLeave(SimulatedPointer);
 									}
-									
-
 #if WITH_SLATE_DEBUGGING
-									FSlateDebugging::BroadcastInputEvent(ESlateDebuggingInputEvent::MouseLeave, SomeWidgetPreviouslyUnderCursor);
+									FSlateDebugging::BroadcastInputEvent(ESlateDebuggingInputEvent::MouseLeave, WidgetPreviouslyUnderCursor);
 #endif
 								}
 								else
@@ -3447,7 +3012,7 @@ void FSlateApplication::ProcessReply( const FWidgetPath& CurrentEventPath, const
 			}
 
 			// When the cursor capture state changes we need to refresh cursor state.
-			bQueryCursorRequested = true;
+			SlateUser->RequestCursorQuery();
 		}
 
 		if ( !bIsVirtualInteraction && CurrentEventPath.IsValid() && RequestedMouseCaptor.IsValid())
@@ -3465,675 +3030,154 @@ void FSlateApplication::ProcessReply( const FWidgetPath& CurrentEventPath, const
 #endif
 
 				// When the cursor capture state changes we need to refresh cursor state.
-				bQueryCursorRequested = true;
+				SlateUser->RequestCursorQuery();
 			}
 		}
 
 		TOptional<FIntPoint> RequestedMousePos = TheReply.GetRequestedMousePos();
 		if (RequestedMousePos.IsSet())
 		{
-			const FVector2D Position = RequestedMousePos.GetValue();
-			PointerIndexPositionMap.Add(FUserAndPointer(UserIndex, PointerIndex), Position);
-			PointerIndexLastPositionMap.Add(FUserAndPointer(UserIndex, PointerIndex), Position);
-			SetCursorPos(Position);
-		}
-
-		if (TheReply.GetMouseLockWidget().IsValid())
-		{
-			// The reply requested mouse lock so tell the native application to lock the mouse to the widget receiving the event
-			LockCursor(TheReply.GetMouseLockWidget());
-		}
-	}
-
-	// Releasing high precision mode.  @HACKISH We can only support high precision mode on true mouse hardware cursors
-	// but if the user index isn't 0, there's no way it's the real mouse so we should ignore this if it's not user 0,
-	// because that means it's a virtual controller.
-	if ( UserIndex == 0 && !bIsVirtualInteraction )
-	{
-		if ( CurrentEventPath.IsValid() && TheReply.ShouldReleaseMouse() && !TheReply.ShouldUseHighPrecisionMouse() )
-		{
-			if ( PlatformApplication->IsUsingHighPrecisionMouseMode() )
-			{
-				PlatformApplication->SetHighPrecisionMouseMode(false, nullptr);
-				PlatformApplication->SetCapture(nullptr);
-
-#if WITH_SLATE_DEBUGGING
-				//TODO Release Capture
-#endif
-
-				// When the cursor capture state changes we need to refresh cursor state.
-				bQueryCursorRequested = true;
-			}
-		}
-	}
-
-	// Releasing Mouse Lock
-	if (TheReply.ShouldReleaseMouseLock())
-	{
-		LockCursor(nullptr);
-	}
-	
-	// If we have a valid Navigation request attempt the navigation.
-	if (TheReply.GetNavigationDestination().IsValid() || TheReply.GetNavigationType() != EUINavigation::Invalid)
-	{
-		FWidgetPath NavigationSource;
-		if (TheReply.GetNavigationSource() == ENavigationSource::WidgetUnderCursor)
-		{
-			NavigationSource = *WidgetsUnderMouse;
-		}
-		else if(const FSlateUser* User = GetOrCreateUser(UserIndex))
-		{
-			TSharedRef<FWidgetPath> EventPathRef = User->GetFocusPath();
-			NavigationSource = EventPathRef.Get();
-		}
-
-		if (NavigationSource.IsValid())
-		{
-			if (TheReply.GetNavigationDestination().IsValid())
-			{
-				const bool bAlwaysHandleNavigationAttempt = false;
-				ExecuteNavigation(NavigationSource, TheReply.GetNavigationDestination(), UserIndex, bAlwaysHandleNavigationAttempt);
-			}
-			else
-			{
-				TSharedRef<SWindow> NavigationWindow = NavigationSource.GetDeepestWindow();
-
-				FNavigationEvent NavigationEvent(PlatformApplication->GetModifierKeys(), UserIndex, TheReply.GetNavigationType(), TheReply.GetNavigationGenesis());
-
-				FNavigationReply NavigationReply = FNavigationReply::Escape();
-
-				for (int32 WidgetIndex = NavigationSource.Widgets.Num() - 1; WidgetIndex >= 0; --WidgetIndex)
-				{
-					FArrangedWidget& SomeWidgetGettingEvent = NavigationSource.Widgets[WidgetIndex];
-					if (SomeWidgetGettingEvent.Widget->IsEnabled())
-					{
-						NavigationReply = SomeWidgetGettingEvent.Widget->OnNavigation(SomeWidgetGettingEvent.Geometry, NavigationEvent).SetHandler(SomeWidgetGettingEvent.Widget);
-						if (NavigationReply.GetBoundaryRule() != EUINavigationRule::Escape || SomeWidgetGettingEvent.Widget == NavigationWindow || WidgetIndex == 0)
-						{
-							AttemptNavigation(NavigationSource, NavigationEvent, NavigationReply, SomeWidgetGettingEvent);
-							break;
-						}
-					}
-				}
-			}
-		}
-	}
-
-	if ( TheReply.GetDetectDragRequest().IsValid() )
-	{
-		checkSlow(InMouseEvent);
-
-		DragDetector.StartDragDetection(
-			WidgetsUnderMouse->GetPathDownTo(TheReply.GetDetectDragRequest().ToSharedRef()),
-			InMouseEvent->GetUserIndex(),
-			InMouseEvent->GetPointerIndex(),
-			TheReply.GetDetectDragRequestButton(),
-			InMouseEvent->GetScreenSpacePosition());
-	}
-
-	// Set focus if requested.
-	TSharedPtr<SWidget> RequestedFocusRecepient = TheReply.GetUserFocusRecepient();
-	if (TheReply.ShouldSetUserFocus() || RequestedFocusRecepient.IsValid())
-	{
-		if (TheReply.AffectsAllUsers())
-		{
-			ForEachUser([&] (FSlateUser* User) {
-				SetUserFocus(User->GetUserIndex(), RequestedFocusRecepient, TheReply.GetFocusCause());
-			});
-		}
-		else
-		{
-			SetUserFocus(UserIndex, RequestedFocusRecepient, TheReply.GetFocusCause());
-		}
-	}
-}
-
-void FSlateApplication::LockCursor(const TSharedPtr<SWidget>& Widget)
-{
-	if (PlatformApplication->Cursor.IsValid())
-	{
-		if (Widget.IsValid())
-		{
-			// Get a path to this widget so we know the position and size of its geometry
-			FWidgetPath WidgetPath;
-			const bool bFoundWidget = GeneratePathToWidgetUnchecked(Widget.ToSharedRef(), WidgetPath);
-			if ( ensureMsgf(bFoundWidget, TEXT("Attempting to LockCursor() to widget but could not find widget %s"), *Widget->ToString()) )
-			{
-				LockCursorToPath(WidgetPath);
-			}
-		}
-		else
-		{
-			UnlockCursor();
-		}
-	}
-}
-
-void FSlateApplication::LockCursorToPath(const FWidgetPath& WidgetPath)
-{
-	// The last widget in the path should be the widget we are locking the cursor to
-	const FArrangedWidget& WidgetGeom = WidgetPath.Widgets[WidgetPath.Widgets.Num() - 1];
-
-	TSharedRef<SWindow> Window = WidgetPath.GetWindow();
-	// Do not attempt to lock the cursor to the window if its not in the foreground.  It would cause annoying side effects
-	if (Window->GetNativeWindow()->IsForegroundWindow())
-	{
-		const FSlateRect SlateClipRect = WidgetGeom.Geometry.GetLayoutBoundingRect();
-		CursorLock.LastComputedBounds = SlateClipRect;
-		CursorLock.PathToLockingWidget = WidgetPath;
-			
-		// Generate a screen space clip rect based on the widgets geometry
-#if PLATFORM_DESKTOP
-		const bool bIsBorderlessGameWindow = Window->GetNativeWindow()->IsDefinitionValid() && Window->GetNativeWindow()->GetDefinition().Type == EWindowType::GameWindow && !Window->GetNativeWindow()->GetDefinition().HasOSWindowBorder;
-		const int32 ClipRectAdjustment = bIsBorderlessGameWindow ? 1 : 0;
-#else
-		const int32 ClipRectAdjustment = 0;
-#endif
-
-		// Note: We round the upper left coordinate of the clip rect so we guarantee the rect is inside the geometry of the widget.  If we truncated when there is a half pixel we would cause the clip
-		// rect to be half a pixel larger than the geometry and cause the mouse to go outside of the geometry.
-		RECT ClipRect;
-		ClipRect.left = FMath::RoundToInt(SlateClipRect.Left + ClipRectAdjustment);
-		ClipRect.top = FMath::RoundToInt(SlateClipRect.Top + ClipRectAdjustment);
-		ClipRect.right = FMath::TruncToInt(SlateClipRect.Right - ClipRectAdjustment);
-		ClipRect.bottom = FMath::TruncToInt(SlateClipRect.Bottom - ClipRectAdjustment);
-
-		// Lock the mouse to the widget
-		PlatformApplication->Cursor->Lock(&ClipRect);
-	}
-}
-
-void FSlateApplication::UnlockCursor()
-{
-	// Unlock the mouse
-	PlatformApplication->Cursor->Lock(nullptr);
-	CursorLock.PathToLockingWidget = FWeakWidgetPath();
-}
-
-void FSlateApplication::UpdateCursorLockRegion()
-{
-	if(CursorLock.PathToLockingWidget.IsValid())
-	{
-	const FWidgetPath PathToWidget = CursorLock.PathToLockingWidget.ToWidgetPath(FWeakWidgetPath::EInterruptedPathHandling::ReturnInvalid);
-	if (PathToWidget.IsValid())
-	{
-		const FSlateRect ComputedClipRect = PathToWidget.Widgets.Last().Geometry.GetLayoutBoundingRect();
-		if (ComputedClipRect != CursorLock.LastComputedBounds)
-		{
-			LockCursorToPath(PathToWidget);
-		}
-	}
-}
-}
-
-void FSlateApplication::SetLastUserInteractionTime(double InCurrentTime)
-{
-	if (LastUserInteractionTime != InCurrentTime)
-	{
-		LastUserInteractionTime = InCurrentTime;
-		LastUserInteractionTimeUpdateEvent.Broadcast(LastUserInteractionTime);
-	}
-}
-
-void FSlateApplication::QueryCursor()
-{
-	bQueryCursorRequested = false;
-
-	// The slate loading widget thread is not allowed to execute this code 
-	// as it is unsafe to read the hittest grid in another thread
-	if ( PlatformApplication->Cursor.IsValid() && IsInGameThread() )
-	{
-		// drag-drop overrides cursor
-		FCursorReply CursorReply = FCursorReply::Unhandled();
-		if ( IsDragDropping() )
-		{
-			CursorReply = DragDropContent->OnCursorQuery();
-		}
-		
-		if (!CursorReply.IsEventHandled())
-		{
-			FWidgetPath WidgetsToQueryForCursor;
-			const TSharedPtr<SWindow> ActiveModalWindow = GetActiveModalWindow();
-
-			const FVector2D CurrentCursorPosition = GetCursorPos();
-			const FVector2D LastCursorPosition = GetLastCursorPos();
-			const FPointerEvent CursorEvent(
-				CursorPointerIndex,
-				CurrentCursorPosition,
-				LastCursorPosition,
-				CurrentCursorPosition - LastCursorPosition,
-				PressedMouseButtons,
-				PlatformApplication->GetModifierKeys()
-				);
-
-			// Query widgets with mouse capture for the cursor
-			if (MouseCaptor.HasCaptureForPointerIndex(CursorUserIndex, CursorPointerIndex))
-			{
-				FWidgetPath MouseCaptorPath = MouseCaptor.ToWidgetPath( FWeakWidgetPath::EInterruptedPathHandling::Truncate, &CursorEvent);
-				if ( MouseCaptorPath.IsValid() )
-				{
-					TSharedRef< SWindow > CaptureWindow = MouseCaptorPath.GetWindow();
-
-					// Never query the mouse captor path if it is outside an active modal window
-					if ( !ActiveModalWindow.IsValid() || ( CaptureWindow == ActiveModalWindow || CaptureWindow->IsDescendantOf(ActiveModalWindow) ) )
-					{
-						WidgetsToQueryForCursor = MouseCaptorPath;
-					}
-				}
-			}
-			else
-			{
-				WidgetsToQueryForCursor = LocateWindowUnderMouse( GetCursorPos(), GetInteractiveTopLevelWindows() );
-			}
-
-			if (WidgetsToQueryForCursor.IsValid())
-			{
-				// Switch worlds for widgets in the current path
-				FScopedSwitchWorldHack SwitchWorld( WidgetsToQueryForCursor );
-
-				for (int32 WidgetIndex = WidgetsToQueryForCursor.Widgets.Num() - 1; WidgetIndex >= 0; --WidgetIndex)
-				{
-					const FArrangedWidget& ArrangedWidget = WidgetsToQueryForCursor.Widgets[WidgetIndex];
-
-					CursorReply = ArrangedWidget.Widget->OnCursorQuery(ArrangedWidget.Geometry, CursorEvent);
-					if (CursorReply.IsEventHandled())
-					{
-						if (!CursorReply.GetCursorWidget().IsValid())
-						{
-							for (; WidgetIndex >= 0; --WidgetIndex)
-							{
-								TOptional<TSharedRef<SWidget>> CursorWidget = WidgetsToQueryForCursor.Widgets[WidgetIndex].Widget->OnMapCursor(CursorReply);
-								if (CursorWidget.IsSet())
-								{
-									CursorReply.SetCursorWidget(WidgetsToQueryForCursor.GetWindow(), CursorWidget.GetValue());
-									break;
-								}
-							}
-						}
-						break;
-					}
-				}
-
-				if (!CursorReply.IsEventHandled() && WidgetsToQueryForCursor.IsValid())
-				{
-					// Query was NOT handled, and we are still over a slate window.
-					CursorReply = FCursorReply::Cursor(EMouseCursor::Default);
-				}
-			}
-			else
-			{
-				// Set the default cursor when there isn't an active window under the cursor and the mouse isn't captured
-				CursorReply = FCursorReply::Cursor(EMouseCursor::Default);
-			}
-		}
-		ProcessCursorReply(CursorReply);
-	}
-}
-
-void FSlateApplication::ProcessCursorReply(const FCursorReply& CursorReply)
-{
-	if (CursorReply.IsEventHandled())
-	{
-		CursorWidgetPtr = CursorReply.GetCursorWidget();
-		if (CursorReply.GetCursorWidget().IsValid())
-		{
-			CursorReply.GetCursorWidget()->SetVisibility(EVisibility::HitTestInvisible);
-			CursorWindowPtr = CursorReply.GetCursorWindow();
-			if (!IsFakingTouchEvents())
-			{
-				PlatformApplication->Cursor->SetType(EMouseCursor::Custom);
-			}
-		}
-		else
-		{
-			CursorWindowPtr.Reset();
-			PlatformApplication->Cursor->SetType(CursorReply.GetCursorType());
-		}
-	}
-	else
-	{
-		CursorWindowPtr.Reset();
-		CursorWidgetPtr.Reset();
-	}
-}
-
-void FSlateApplication::SpawnToolTip( const TSharedRef<IToolTip>& InToolTip, const FVector2D& InSpawnLocation )
-{
-	// Close existing tool tip, if we have one
-	CloseToolTip();
-
-	// Spawn the new tool tip
-	{
-		TSharedPtr< SWindow > NewToolTipWindow( ToolTipWindow.Pin() );
-		if( !NewToolTipWindow.IsValid() )
-		{
-			// Create the tool tip window
-			NewToolTipWindow = SWindow::MakeToolTipWindow();
-
-			// Don't show the window yet.  We'll set it up with some content first!
-			const bool bShowImmediately = false;
-			AddWindow( NewToolTipWindow.ToSharedRef(), bShowImmediately );
-		}
-
-		NewToolTipWindow->SetContent
-		(
-			SNew(SWeakWidget)
-			.PossiblyNullContent(InToolTip->AsWidget())
-		);
-
-		// Move the window again to recalculate popup window position if necessary (tool tip may spawn outside of the monitors work area)
-		// and in that case we need to adjust it
-		DesiredToolTipLocation = InSpawnLocation;
-		{
-			// Make sure the desired size is valid
-			NewToolTipWindow->SlatePrepass(FSlateApplication::Get().GetApplicationScale()*NewToolTipWindow->GetNativeWindow()->GetDPIScaleFactor());
-
-			// already handled
-			const bool bAutoAdjustForDPIScale = false;
-
-			FSlateRect Anchor(DesiredToolTipLocation.X, DesiredToolTipLocation.Y, DesiredToolTipLocation.X, DesiredToolTipLocation.Y);
-			DesiredToolTipLocation = CalculateTooltipWindowPosition(Anchor, NewToolTipWindow->GetDesiredSizeDesktopPixels(), bAutoAdjustForDPIScale);
-
-			// MoveWindowTo will adjust the window's position, if needed
-			NewToolTipWindow->MoveWindowTo( DesiredToolTipLocation );
-		}
-
-		NewToolTipWindow->SetOpacity(0.0f);
-
-		// Show the window
-		NewToolTipWindow->ShowWindow();
-
-		// Keep a weak reference to the tool tip window
-		ToolTipWindow = NewToolTipWindow;
-
-		// Keep track of when this tool tip was spawned
-		ToolTipSummonTime = FPlatformTime::Seconds();
-	}
-}
-
-void FSlateApplication::CloseToolTip()
-{
-	// Notify the source widget that its tooltip is closing
-	{
-		TSharedPtr<SWidget> SourceWidget = ActiveToolTipWidgetSource.Pin();
-		if (SourceWidget.IsValid())
-		{
-			SourceWidget->OnToolTipClosing();
-		}
-	}
-
-	// Notify the active tooltip that it's being closed.
-	TSharedPtr<IToolTip> StableActiveToolTip = ActiveToolTip.Pin();
-	if ( StableActiveToolTip.IsValid() )
-	{
-		StableActiveToolTip->OnClosed();
-	}
-
-	// If the tooltip had a new window holding it, hide the window.
-	TSharedPtr< SWindow > PinnedToolTipWindow( ToolTipWindow.Pin() );
-	if( PinnedToolTipWindow.IsValid() && PinnedToolTipWindow->IsVisible() )
-	{
-		// Hide the tool tip window.  We don't destroy the window, because we want to reuse it for future tool tips.
-		PinnedToolTipWindow->HideWindow();
-	}
-
-	ActiveToolTip.Reset();
-	ActiveToolTipWidgetSource.Reset();
-
-	ToolTipOffsetDirection = EToolTipOffsetDirection::Undetermined;
-}
-
-void FSlateApplication::UpdateToolTip( bool AllowSpawningOfNewToolTips )
-{
-	// Don't do anything if tooltips are not enabled.
-	if ( !bAllowToolTips )
-	{
-		// The user may have disabled this while a tooltip was visible, like during a menu to game transition, if this
-		// happens we need to close the tool tip if it's still visible.
-		if ( ActiveToolTip.IsValid() )
-		{
-			CloseToolTip();
+			SlateUser->SetCursorPosition(RequestedMousePos.GetValue());
 		}
 
-		return;
-	}
-
-	SCOPE_CYCLE_COUNTER(STAT_SlateUpdateTooltip);
-
-	const bool bCheckForToolTipChanges =
-		IsInGameThread() &&					// We should never allow the slate loading thread to create new windows or interact with the hittest grid
-		!IsUsingHighPrecisionMouseMovment() && // If we are using HighPrecision movement then we can't rely on the OS cursor to be accurate
-		!IsDragDropping() &&				// We must not currently be in the middle of a drag-drop action
-		PlatformApplication->IsCursorDirectlyOverSlateWindow(); // The cursor must be over a Slate window
-	
-	// We still want to show tooltips for widgets that are disabled
-	const bool bIgnoreEnabledStatus = true;
-
-	float DPIScaleFactor = 1.0f;
-	
-	FWidgetPath WidgetsToQueryForToolTip;
-	// We don't show any tooltips when drag and dropping or when another app is active
-	if (bCheckForToolTipChanges)
-	{
-		// Ask each widget under the Mouse if they have a tool tip to show.
-		FWidgetPath WidgetsUnderMouse = LocateWindowUnderMouse( GetCursorPos(), GetInteractiveTopLevelWindows(), bIgnoreEnabledStatus );
-		// Don't attempt to show tooltips inside an existing tooltip
-		if (!WidgetsUnderMouse.IsValid() || WidgetsUnderMouse.GetWindow() != ToolTipWindow.Pin())
+		if (TheReply.GetMouseLockWidget())
 		{
-			WidgetsToQueryForToolTip = WidgetsUnderMouse;
-			if (false)// @na (WidgetsUnderMouse.IsValid())
-			{
-				DPIScaleFactor = WidgetsUnderMouse.GetWindow()->GetNativeWindow()->GetDPIScaleFactor();
-			}
+			// The reply requested mouse lock so tell the native application to lock the mouse to the widget receiving the event
+			SlateUser->LockCursor(TheReply.GetMouseLockWidget().ToSharedRef());
 		}
 	}
 
-	bool bHaveForceFieldRect = false;
-	FSlateRect ForceFieldRect;
-
-	TSharedPtr<IToolTip> NewToolTip;
-	TSharedPtr<SWidget> WidgetProvidingNewToolTip;
-	for ( int32 WidgetIndex=WidgetsToQueryForToolTip.Widgets.Num()-1; WidgetIndex >= 0; --WidgetIndex )
+	// Releasing high precision mode.  @HACKISH We can only support high precision mode on true mouse hardware cursors
+	// but if the user index isn't 0, there's no way it's the real mouse so we should ignore this if it's not user 0,
+	// because that means it's a virtual controller.
+	if ( UserIndex == CursorUserIndex && !bIsVirtualInteraction )
 	{
-		FArrangedWidget* CurWidgetGeometry = &WidgetsToQueryForToolTip.Widgets[WidgetIndex];
-		const TSharedRef<SWidget>& CurWidget = CurWidgetGeometry->Widget;
-
-		if( !NewToolTip.IsValid() )
+		if ( CurrentEventPath.IsValid() && TheReply.ShouldReleaseMouse() && !TheReply.ShouldUseHighPrecisionMouse() )
 		{
-			TSharedPtr< IToolTip > WidgetToolTip = CurWidget->GetToolTip();
-
-			// Make sure the tool-tip currently is displaying something before spawning it.
-			if( WidgetToolTip.IsValid() && !WidgetToolTip->IsEmpty() )
+			if ( PlatformApplication->IsUsingHighPrecisionMouseMode() )
 			{
-				WidgetProvidingNewToolTip = CurWidget;
-				NewToolTip = WidgetToolTip;
-			}
-		}
+				PlatformApplication->SetHighPrecisionMouseMode(false, nullptr);
+				PlatformApplication->SetCapture(nullptr);
 
-		// Keep track of the root most widget with a tool-tip force field enabled
-		if( CurWidget->HasToolTipForceField() )
-		{
-			if( !bHaveForceFieldRect )
-			{
-				bHaveForceFieldRect = true;
-				ForceFieldRect = CurWidgetGeometry->Geometry.GetLayoutBoundingRect();
-			}
-			else
-			{
-				// Grow the rect to encompass this geometry.  Usually, the parent's rect should always be inclusive
-				// of it's child though.  Just is kind of just being paranoid.
-				ForceFieldRect = ForceFieldRect.Expand( CurWidgetGeometry->Geometry.GetLayoutBoundingRect() );
+#if WITH_SLATE_DEBUGGING
+				//TODO Release Capture
+#endif
+
+				// When the cursor capture state changes we need to refresh cursor state.
+				SlateUser->RequestCursorQuery();
 			}
-			ForceFieldRect = (1.0f / DPIScaleFactor) * ForceFieldRect;
 		}
 	}
 
-	// Did the tool tip change from last time?
-	const bool bToolTipChanged = (NewToolTip != ActiveToolTip.Pin());
+	// Releasing Mouse Lock
+	if (TheReply.ShouldReleaseMouseLock())
+	{
+		SlateUser->UnlockCursor();
+	}
 	
-	// Any widgets that wish to handle visualizing the tooltip get a chance here.
-	TSharedPtr<SWidget> NewTooltipVisualizer;
-	if (bToolTipChanged)
+	// If we have a valid Navigation request attempt the navigation.
+	if (TheReply.GetNavigationDestination().IsValid() || TheReply.GetNavigationType() != EUINavigation::Invalid)
 	{
-		// Remove existing tooltip if there is one.
-		if (TooltipVisualizerPtr.IsValid())
+		FWidgetPath NavigationSource;
+		if (TheReply.GetNavigationSource() == ENavigationSource::WidgetUnderCursor)
 		{
-			TooltipVisualizerPtr.Pin()->OnVisualizeTooltip( nullptr );
+			NavigationSource = *WidgetsUnderMouse;
 		}
-
-		// Notify the new tooltip that it's about to be opened.
-		if ( NewToolTip.IsValid() )
+		else
 		{
-			NewToolTip->OnOpening();
+			NavigationSource = SlateUser->GetFocusPath().Get();
 		}
 
-		TSharedPtr<SWidget> NewToolTipWidget = NewToolTip.IsValid() ? NewToolTip->AsWidget() : TSharedPtr<SWidget>();
-
-		bool bOnVisualizeTooltipHandled = false;
-		// Some widgets might want to provide an alternative Tooltip Handler.
-		for ( int32 WidgetIndex=WidgetsToQueryForToolTip.Widgets.Num()-1; !bOnVisualizeTooltipHandled && WidgetIndex >= 0; --WidgetIndex )
+		if (NavigationSource.IsValid())
 		{
-			const FArrangedWidget& CurWidgetGeometry = WidgetsToQueryForToolTip.Widgets[WidgetIndex];
-			bOnVisualizeTooltipHandled = CurWidgetGeometry.Widget->OnVisualizeTooltip( NewToolTipWidget );
-			if (bOnVisualizeTooltipHandled)
+			if (TheReply.GetNavigationDestination().IsValid())
 			{
-				// Someone is taking care of visualizing this tooltip
-				NewTooltipVisualizer = CurWidgetGeometry.Widget;
+				const bool bAlwaysHandleNavigationAttempt = false;
+				ExecuteNavigation(NavigationSource, TheReply.GetNavigationDestination(), UserIndex, bAlwaysHandleNavigationAttempt);
 			}
-		}
-	}
-
+			else
+			{
+				TSharedRef<SWindow> NavigationWindow = NavigationSource.GetDeepestWindow();
 
-	// If a widget under the cursor has a tool-tip forcefield active, then go through any menus
-	// in the menu stack that are above that widget's window, and make sure those windows also
-	// prevent the tool-tip from encroaching.  This prevents tool-tips from drawing over sub-menus
-	// spawned from menu items in a different window, for example.
-	if( bHaveForceFieldRect && WidgetsToQueryForToolTip.IsValid() )
-	{
-		TSharedPtr<IMenu> MenuInPath = MenuStack.FindMenuInWidgetPath(WidgetsToQueryForToolTip);
+				FNavigationEvent NavigationEvent(PlatformApplication->GetModifierKeys(), UserIndex, TheReply.GetNavigationType(), TheReply.GetNavigationGenesis());
 
-		if (MenuInPath.IsValid())
-		{
-			ForceFieldRect = ForceFieldRect.Expand(MenuStack.GetToolTipForceFieldRect(MenuInPath.ToSharedRef(), WidgetsToQueryForToolTip));
-		}
-	}
+				FNavigationReply NavigationReply = FNavigationReply::Escape();
 
-	{
-		TSharedPtr<IToolTip> ActiveToolTipPtr = ActiveToolTip.Pin();
-		if ( ( ActiveToolTipPtr.IsValid() && !ActiveToolTipPtr->IsInteractive() ) || ( NewToolTip.IsValid() && NewToolTip != ActiveToolTip.Pin() ) )
-		{
-			// Keep track of where we want tool tips to be positioned
-			DesiredToolTipLocation = GetLastCursorPos() + SlateDefs::ToolTipOffsetFromMouse;
+				for (int32 WidgetIndex = NavigationSource.Widgets.Num() - 1; WidgetIndex >= 0; --WidgetIndex)
+				{
+					FArrangedWidget& SomeWidgetGettingEvent = NavigationSource.Widgets[WidgetIndex];
+					if (SomeWidgetGettingEvent.Widget->IsEnabled())
+					{
+						NavigationReply = SomeWidgetGettingEvent.Widget->OnNavigation(SomeWidgetGettingEvent.Geometry, NavigationEvent).SetHandler(SomeWidgetGettingEvent.Widget);
+						if (NavigationReply.GetBoundaryRule() != EUINavigationRule::Escape || SomeWidgetGettingEvent.Widget == NavigationWindow || WidgetIndex == 0)
+						{
+							AttemptNavigation(NavigationSource, NavigationEvent, NavigationReply, SomeWidgetGettingEvent);
+							break;
+						}
+					}
+				}
+			}
 		}
 	}
 
-	TSharedPtr< SWindow > ToolTipWindowPtr = ToolTipWindow.Pin();
-	if ( ToolTipWindowPtr.IsValid() )
+	if ( TheReply.GetDetectDragRequest().IsValid() )
 	{
-		// already handled
-		const bool bAutoAdjustForDPIScale = false;
+		checkSlow(InMouseEvent);
 
-		FSlateRect Anchor(DesiredToolTipLocation.X, DesiredToolTipLocation.Y, DesiredToolTipLocation.X, DesiredToolTipLocation.Y);
-		DesiredToolTipLocation = CalculatePopupWindowPosition( Anchor, ToolTipWindowPtr->GetDesiredSizeDesktopPixels(), bAutoAdjustForDPIScale );
+		SlateUser->StartDragDetection(
+			WidgetsUnderMouse->GetPathDownTo(TheReply.GetDetectDragRequest().ToSharedRef()),
+			InMouseEvent->GetPointerIndex(),
+			TheReply.GetDetectDragRequestButton(),
+			InMouseEvent->GetScreenSpacePosition());
 	}
 
-	// Repel tool-tip from a force field, if necessary
-	if( bHaveForceFieldRect )
+	// Set focus if requested.
+	TSharedPtr<SWidget> RequestedFocusRecepient = TheReply.GetUserFocusRecepient();
+	if (TheReply.ShouldSetUserFocus() || RequestedFocusRecepient.IsValid())
 	{
-		FVector2D ToolTipShift;
-		ToolTipShift.X = ( ForceFieldRect.Right + SlateDefs::ToolTipOffsetFromForceField.X ) - DesiredToolTipLocation.X;
-		ToolTipShift.Y = ( ForceFieldRect.Bottom + SlateDefs::ToolTipOffsetFromForceField.Y ) - DesiredToolTipLocation.Y;
-
-		// Make sure the tool-tip needs to be offset
-		if( ToolTipShift.X > 0.0f && ToolTipShift.Y > 0.0f )
+		if (TheReply.AffectsAllUsers())
 		{
-			// Find the best edge to move the tool-tip towards
-			if( ToolTipOffsetDirection == EToolTipOffsetDirection::Right ||
-				( ToolTipOffsetDirection == EToolTipOffsetDirection::Undetermined && ToolTipShift.X < ToolTipShift.Y ) )
-			{
-				// Move right
-				DesiredToolTipLocation.X += ToolTipShift.X;
-				ToolTipOffsetDirection = EToolTipOffsetDirection::Right;
-			}
-			else
-			{
-				// Move down
-				DesiredToolTipLocation.Y += ToolTipShift.Y;
-				ToolTipOffsetDirection = EToolTipOffsetDirection::Down;
-			}
+			SetAllUserFocus(RequestedFocusRecepient, TheReply.GetFocusCause());
 		}
-	}
-
-	// The tool tip changed...
-	if ( bToolTipChanged )
-	{
-		// Close any existing tooltips; Unless the current tooltip is interactive and we don't have a valid tooltip to replace it
-		TSharedPtr<IToolTip> ActiveToolTipPtr = ActiveToolTip.Pin();
-		if ( NewToolTip.IsValid() || ( ActiveToolTipPtr.IsValid() && !ActiveToolTipPtr->IsInteractive() ) )
+		else
 		{
-			CloseToolTip();
-			
-			if (NewTooltipVisualizer.IsValid())
-			{
-				TooltipVisualizerPtr = NewTooltipVisualizer;
-			}
-			else if( bAllowToolTips && AllowSpawningOfNewToolTips )
-			{
-				// Spawn a new one if we have it
-				if( NewToolTip.IsValid() )
-				{
-					SpawnToolTip( NewToolTip.ToSharedRef(), DesiredToolTipLocation );
-				}
-			}
-			else
-			{
-				NewToolTip = nullptr;
-			}
-
-			ActiveToolTip = NewToolTip;
-			ActiveToolTipWidgetSource = WidgetProvidingNewToolTip;
+			SlateUser->SetFocus(RequestedFocusRecepient.ToSharedRef(), TheReply.GetFocusCause());
 		}
 	}
+}
 
-	// Do we have a tool tip window?
-	if( ToolTipWindow.IsValid() )
+void FSlateApplication::SetLastUserInteractionTime(double InCurrentTime)
+{
+	if (LastUserInteractionTime != InCurrentTime)
 	{
-		// Only enable tool-tip transitions if we're running at a decent frame rate
-		const bool bAllowInstantToolTips = false;
-		const bool bAllowAnimations = !bAllowInstantToolTips && FSlateApplication::Get().IsRunningAtTargetFrameRate();
-
-		// How long since the tool tip was summoned?
-		const float TimeSinceSummon = FPlatformTime::Seconds() - ToolTipDelay - ToolTipSummonTime;
-		const float ToolTipOpacity = bAllowInstantToolTips ? 1.0f : FMath::Clamp< float >( TimeSinceSummon / ToolTipFadeInDuration, 0.0f, 1.0f );
-
-		// Update window opacity
-		TSharedRef< SWindow > PinnedToolTipWindow( ToolTipWindow.Pin().ToSharedRef() );
-		PinnedToolTipWindow->SetOpacity( ToolTipOpacity );
+		LastUserInteractionTime = InCurrentTime;
+		LastUserInteractionTimeUpdateEvent.Broadcast(LastUserInteractionTime);
+	}
+}
 
-		// How far tool tips should slide
-		const FVector2D SlideDistance( 30.0f, 5.0f );
+void FSlateApplication::QueryCursor()
+{
+	GetOrCreateUser(CursorUserIndex)->QueryCursor();
+}
 
-		// Apply steep inbound curve to the movement, so it looks like it quickly decelerating
-		const float SlideProgress = bAllowAnimations ? FMath::Pow( 1.0f - ToolTipOpacity, 3.0f ) : 0.0f;
+void FSlateApplication::ProcessCursorReply(const FCursorReply& CursorReply)
+{
+	GetOrCreateUser(CursorUserIndex)->ProcessCursorReply(CursorReply);
+}
 
-		FVector2D WindowLocation = DesiredToolTipLocation + SlideProgress * SlideDistance;
-		if( WindowLocation != PinnedToolTipWindow->GetPositionInScreen() )
-		{
-			// already handled
-			const bool bAutoAdjustForDPIScale = false;
+void FSlateApplication::SpawnToolTip(const TSharedRef<IToolTip>& InToolTip, const FVector2D& InSpawnLocation)
+{
+	GetOrCreateUser(CursorUserIndex)->ShowTooltip(InToolTip, InSpawnLocation);
+}
 
-			// Avoid the edges of the desktop
-			FSlateRect Anchor(WindowLocation.X, WindowLocation.Y, WindowLocation.X, WindowLocation.Y);
-			WindowLocation = CalculateTooltipWindowPosition(Anchor, PinnedToolTipWindow->GetDesiredSizeDesktopPixels(), bAutoAdjustForDPIScale);
+void FSlateApplication::CloseToolTip()
+{
+	GetOrCreateUser(CursorUserIndex)->CloseTooltip();
+}
 
-			// Update the tool tip window positioning
-			// SetCachedScreenPosition is a hack (issue tracked as TTP #347070) which is needed because code in TickWindowAndChildren()/DrawPrepass()
-			// assumes GetPositionInScreen() to correspond to the new window location in the same tick. This is true on Windows, but other
-			// OSes (Linux in particular) may not update cached screen position until next time events are polled.
-			PinnedToolTipWindow->SetCachedScreenPosition( WindowLocation );
-			PinnedToolTipWindow->MoveWindowTo( WindowLocation );
-		}
-	}
+void FSlateApplication::UpdateToolTip(bool bAllowSpawningOfNewToolTips)
+{
+	GetOrCreateUser(CursorUserIndex)->UpdateTooltip(MenuStack, bAllowSpawningOfNewToolTips);
 }
 
 TArray< TSharedRef<SWindow> > FSlateApplication::GetInteractiveTopLevelWindows()
@@ -4178,51 +3222,6 @@ void FSlateApplication::GetAllVisibleChildWindows(TArray< TSharedRef<SWindow> >&
 	}
 }
 
-bool FSlateApplication::IsDragDropping() const
-{
-	return DragDropContent.IsValid();
-}
-
-bool FSlateApplication::IsDragDroppingAffected(const FPointerEvent& InPointerEvent) const
-{
-	return DragDropContent.IsValid() && DragDropContent->AffectedByPointerEvent(InPointerEvent);
-}
-
-TSharedPtr<FDragDropOperation> FSlateApplication::GetDragDroppingContent() const
-{
-	return DragDropContent;
-}
-
-void FSlateApplication::CancelDragDrop()
-{
-	if (!IsDragDropping())
-	{
-		return;
-	}
-
-	FPointerEvent EmptyPointerEvent;
-
-	for( auto LastWidgetIterator = WidgetsUnderCursorLastEvent.CreateConstIterator(); LastWidgetIterator; ++LastWidgetIterator)
-	{
-		
-		FWidgetPath WidgetsToDragLeave = LastWidgetIterator.Value().ToWidgetPath(FWeakWidgetPath::EInterruptedPathHandling::Truncate);
-		if(WidgetsToDragLeave.IsValid())
-		{
-			const FDragDropEvent DragDropEvent(EmptyPointerEvent, DragDropContent);
-			for(int32 WidgetIndex = WidgetsToDragLeave.Widgets.Num() - 1; WidgetIndex >= 0; --WidgetIndex)
-			{
-				WidgetsToDragLeave.Widgets[WidgetIndex].Widget->OnDragLeave(DragDropEvent);
-			}
-		}
-	}
-
-	// cancel dragdrop operation correctly firing off callbacks
-	DragDropContent->OnDrop(false, EmptyPointerEvent);
-
-	WidgetsUnderCursorLastEvent.Empty();
-	ResetDragDropState();
-}
-
 void FSlateApplication::EnterDebuggingMode()
 {
 	bRequestLeaveDebugMode = false;
@@ -4309,39 +3308,9 @@ bool FSlateApplication::IsWindowInDestroyQueue(TSharedRef<SWindow> Window) const
 	return WindowDestroyQueue.Contains(Window);
 }
 
-void FSlateApplication::SynthesizeMouseMove()
-{
-	// The slate loading widget thread is not allowed to execute this code 
-	// as it is unsafe to read the hittest grid in another thread
-	if (PlatformApplication->Cursor.IsValid() && IsInGameThread())
-	{
-		// Synthetic mouse events accomplish two goals:
-		// 1) The UI can change even if the mouse doesn't move.
-		//    Synthesizing a mouse move sends out events.
-		//    In this case, the current and previous position will be the same.
-		//
-		// 2) The mouse moves, but the OS decided not to send us an event.
-		//    e.g. Mouse moved outside of our window.
-		//    In this case, the previous and current positions differ.
-
-		FPointerEvent MouseEvent
-		(
-			CursorPointerIndex,
-			GetCursorPos(),
-			GetLastCursorPos(),
-			PressedMouseButtons,
-			EKeys::Invalid,
-			0,
-			PlatformApplication->GetModifierKeys()
-		);
-
-		ProcessMouseMoveEvent(MouseEvent, true);
-	}
-}
-
-void FSlateApplication::QueueSynthesizedMouseMove()
+void FSlateApplication::QueueSynthesizedMouseMove(const FInputEvent& SourceEvent)
 {
-	SynthesizeMouseMovePending = 2;
+	GetOrCreateUser(SourceEvent)->QueueSyntheticCursorMove();
 }
 
 void FSlateApplication::SetUnhandledKeyDownEventHandler( const FOnKeyEvent& NewHandler )
@@ -4402,12 +3371,12 @@ float FSlateApplication::GetCursorRadius() const
 
 void FSlateApplication::SetAllowTooltips(bool bCanShow)
 {
-	bAllowToolTips = bCanShow ? 1 : 0;
+	bEnableTooltips = bCanShow;
 }
 
 bool FSlateApplication::GetAllowTooltips() const
 {
-	return bAllowToolTips != 0;
+	return bEnableTooltips;
 }
 
 FVector2D FSlateApplication::CalculateTooltipWindowPosition( const FSlateRect& InAnchorRect, const FVector2D& InSize, bool bAutoAdjustForDPIScale) const
@@ -4576,45 +3545,43 @@ bool FSlateApplication::AllowMoveCursor()
 
 FSlateRect FSlateApplication::GetPreferredWorkArea() const
 {
-	if ( const FSlateUser* User = GetUser(GetUserIndexForKeyboard()) )
+	if (TSharedPtr<const FSlateUser> KeyboardUser = GetUser(GetUserIndexForKeyboard()))
 	{
-		const FWeakWidgetPath& FocusedWidgetPath = User->GetWeakFocusPath();
+		const FWeakWidgetPath& FocusedWidgetPath = KeyboardUser->GetWeakFocusPath();
 
 		// First see if we have a focused widget
-		if ( FocusedWidgetPath.IsValid() && FocusedWidgetPath.Window.IsValid() )
+		if (FocusedWidgetPath.IsValid() && FocusedWidgetPath.Window.IsValid())
 		{
 			const FVector2D WindowPos = FocusedWidgetPath.Window.Pin()->GetPositionInScreen();
 			const FVector2D WindowSize = FocusedWidgetPath.Window.Pin()->GetSizeInScreen();
 			return GetWorkArea(FSlateRect(WindowPos.X, WindowPos.Y, WindowPos.X + WindowSize.X, WindowPos.Y + WindowSize.Y));
 		}
-	}
 
-	// no focus widget, so use mouse position if there are windows present in the work area
-	const FVector2D CursorPos = GetCursorPos();
-	const FSlateRect WorkArea = GetWorkArea(FSlateRect(CursorPos.X, CursorPos.Y, CursorPos.X + 1.0f, CursorPos.Y + 1.0f));
+		// no focus widget, so use cursor position if there are windows present in the work area
+		const FVector2D CursorPos = KeyboardUser->GetCursorPosition();
+		const FSlateRect WorkArea = GetWorkArea(FSlateRect(CursorPos.X, CursorPos.Y, CursorPos.X + 1.0f, CursorPos.Y + 1.0f));
 
-	if ( FSlateWindowHelper::CheckWorkAreaForWindows(SlateWindows, WorkArea) )
-	{
-		return WorkArea;
-	}
+		if (FSlateWindowHelper::CheckWorkAreaForWindows(SlateWindows, WorkArea))
+		{
+			return WorkArea;
+		}
 
-	// If we can't find a window where the cursor is at, try finding a main window.
-	TSharedPtr<SWindow> ActiveTop = GetActiveTopLevelWindow();
-	if ( ActiveTop.IsValid() )
-	{
-		// Use the current top level windows rect
-		return GetWorkArea(ActiveTop->GetRectInScreen());
-	}
+		// If we can't find a window where the cursor is at, try finding a main window.
+		if (TSharedPtr<SWindow> ActiveTop = GetActiveTopLevelWindow())
+		{
+			// Use the current top level windows rect
+			return GetWorkArea(ActiveTop->GetRectInScreen());
+		}
 
-	// If we can't find a top level window check for an active modal window
-	TSharedPtr<SWindow> ActiveModal = GetActiveModalWindow();
-	if ( ActiveModal.IsValid() )
-	{
-		// Use the current active modal windows rect
-		return GetWorkArea(ActiveModal->GetRectInScreen());
+		// If we can't find a top level window check for an active modal window
+		if (TSharedPtr<SWindow> ActiveModal = GetActiveModalWindow())
+		{
+			// Use the current active modal windows rect
+			return GetWorkArea(ActiveModal->GetRectInScreen());
+		}
 	}
-
-	// no windows on work area - default to primary display
+	
+	// Either there is no keyboard user or there are no windows on work area - fall back to primary display
 	FDisplayMetrics DisplayMetrics;
 	GetCachedDisplayMetrics(DisplayMetrics);
 
@@ -4699,90 +3666,79 @@ bool FSlateApplication::TakeScreenshot(const TSharedRef<SWidget>& Widget, const
 	return (OutSize.X != 0 && OutSize.Y != 0);
 }
 
-TSharedRef<FSlateVirtualUser> FSlateApplication::FindOrCreateVirtualUser(int32 VirtualUserIndex)
+TSharedRef<FSlateVirtualUserHandle> FSlateApplication::FindOrCreateVirtualUser(int32 VirtualUserIndex)
 {
-	// Ensure we have a large enough array to add the new virtual user.
+	// Ensure we have a large enough array to add the new virtual user
 	if ( VirtualUserIndex >= VirtualUsers.Num() )
 	{
 		VirtualUsers.SetNum(VirtualUserIndex + 1);
 	}
 
-	TSharedPtr<FSlateVirtualUser> VirtualUser = VirtualUsers[VirtualUserIndex].Pin();
-	if ( VirtualUser.IsValid() )
+	TSharedPtr<FSlateVirtualUserHandle> VirtualUserHandle = VirtualUsers[VirtualUserIndex].Pin();
+	if (!VirtualUserHandle.IsValid())
 	{
-		return VirtualUser.ToSharedRef();
-	}
-
-	// Register new virtual user with slates standard set of users.
-	int32 NextVirtualUserIndex = SlateApplicationDefs::MaxHardwareUsers;
-	while ( GetUser(NextVirtualUserIndex) )
-	{
-		NextVirtualUserIndex++;
-	}
+		// Register at the next available user index (beyond those potentially occupied by real users)
+		int32 NextVirtualUserIndex = SlateApplicationDefs::MaxHardwareUsers;
+		while (GetUser(NextVirtualUserIndex))
+		{
+			NextVirtualUserIndex++;
+		}
 
-	TSharedRef<FSlateUser> NewUser = MakeShareable(new FSlateUser(NextVirtualUserIndex, true));
-	RegisterUser(NewUser);
+		TSharedRef<FSlateUser> NewUser = RegisterNewUser(NextVirtualUserIndex, true);
 
-	// Make a virtual user handle that can be released automatically when all virtual users
-	// of this same user index are collected.
-	VirtualUser = MakeShareable(new FSlateVirtualUser(NewUser->GetUserIndex(), VirtualUserIndex));
+		// Make a virtual user handle that will unregister the virtual user upon destruction
+		VirtualUserHandle = MakeShareable(new FSlateVirtualUserHandle(NewUser->GetUserIndex(), VirtualUserIndex));
 
-	// Update the virtual user array, so we can get this user back later.
-	VirtualUsers[VirtualUserIndex] = VirtualUser;
+		// Update the virtual user array, so we can get this user back later.
+		VirtualUsers[VirtualUserIndex] = VirtualUserHandle;
+	}
 
-	return VirtualUser.ToSharedRef();
+	return VirtualUserHandle.ToSharedRef();
 }
 
-FSlateUser* FSlateApplication::GetOrCreateUser(int32 UserIndex)
+TSharedRef<FSlateUser> FSlateApplication::GetOrCreateUser(int32 UserIndex)
 {
-	if (UserIndex < 0)
-	{
-		return nullptr;
-	}
-
-	if (FSlateUser* User = GetUser(UserIndex))
+	if (TSharedPtr<FSlateUser> FoundUser = GetUser(UserIndex))
 	{
-		return User;
+		return FoundUser.ToSharedRef();
 	}
-
-	TSharedRef<FSlateUser> NewUser = MakeShared<FSlateUser>(UserIndex, false);
-	RegisterUser(NewUser);
-
-	return &NewUser.Get();
+	return RegisterNewUser(UserIndex);
 }
 
-void FSlateApplication::RegisterUser(TSharedRef<FSlateUser> NewUser)
+TSharedRef<FSlateUser> FSlateApplication::RegisterNewUser(int32 UserIndex, bool bIsVirtual)
 {
-	if ( NewUser->UserIndex == -1 )
+	// We tolerate no shenanigans with inappropriate arguments here
+	// New users must be registered at a valid non-negative index that is not already occupied by another user
+	check(UserIndex >= 0);
+	check(!Users.IsValidIndex(UserIndex) || !Users[UserIndex]);
+
+	TSharedPtr<ICursor> UserCursor;
+	if (UserIndex == CursorUserIndex && PlatformApplication && PlatformApplication->Cursor)
 	{
-		int32 Index = Users.Add(NewUser);
-		NewUser->UserIndex = Index;
+		check(!bIsVirtual);
+		UserCursor = PlatformApplication->Cursor;
 	}
-	else
+	else if (!bIsVirtual)
 	{
-		// Ensure we have a large enough array to add the new user.
-		if ( NewUser->GetUserIndex() >= Users.Num() )
-		{
-			Users.SetNum(NewUser->GetUserIndex() + 1);
-		}
+		// Real users all control a cursor, but there's only one platform cursor (and it *always* belongs to the 0th user)
+		// Everyone else gets a faux cursor instead
+		UserCursor = MakeShared<FFauxSlateCursor>();
+	}
 
-		if ( FSlateUser* ExistingUser = Users[NewUser->GetUserIndex()].Get() )
-		{
-			// Migrate any state we know about that needs to be maintained if the
-			// user is replaced.
-			NewUser->FocusWidgetPathWeak = ExistingUser->FocusWidgetPathWeak;
-			NewUser->FocusCause = ExistingUser->FocusCause;
-			NewUser->ShowFocus = ExistingUser->ShowFocus;
-		}
+	TSharedRef<FSlateUser> NewUser = FSlateUser::Create(UserIndex, UserCursor);
 
-		// Replace the user that's at this index with the new user.
-		Users[NewUser->GetUserIndex()] = NewUser;
+	// Ensure we have a large enough array to add the new User
+	if (UserIndex >= Users.Num())
+	{
+		Users.SetNum(UserIndex + 1);
 	}
 
-	//NavigationConfig->OnUserAdded(NewUser->GetUserIndex());
+	Users[UserIndex] = NewUser;
 
-	UE_LOG(LogSlate, Log, TEXT("Slate User Registered.  User Index %d, Is Virtual User: %d"), NewUser->UserIndex, NewUser->bVirtualUser);
-	UserRegisteredEvent.Broadcast(NewUser->GetUserIndex());
+	UE_LOG(LogSlate, Log, TEXT("Slate User Registered.  User Index %d, Is Virtual User: %d"), UserIndex, bIsVirtual);
+	UserRegisteredEvent.Broadcast(UserIndex);
+
+	return NewUser;
 }
 
 void FSlateApplication::UnregisterUser(int32 UserIndex)
@@ -4801,23 +3757,25 @@ void FSlateApplication::UnregisterUser(int32 UserIndex)
 	}
 }
 
-void FSlateApplication::ForEachUser(TFunctionRef<void(FSlateUser*)> InPredicate, bool bIncludeVirtualUsers)
+void FSlateApplication::ForEachUser(TFunctionRef<void(FSlateUser&)> InPredicate, bool bIncludeVirtualUsers)
 {
-	for ( int32 UserIndex = 0; UserIndex < Users.Num(); UserIndex++ )
+	for (const TSharedPtr<FSlateUser>& User : Users)
 	{
-		if ( FSlateUser* User = Users[UserIndex].Get() )
+		// Ignore virtual users unless told not to.
+		if (User && (bIncludeVirtualUsers || !User->IsVirtualUser()))
 		{
-			// Ignore virtual users unless told not to.
-			if ( !bIncludeVirtualUsers && User->IsVirtualUser() )
-			{
-				continue;
-			}
-
-			InPredicate(User);
+			InPredicate(*User);
 		}
 	}
 }
 
+void FSlateApplication::ForEachUser(TFunctionRef<void(FSlateUser*)> InPredicate, bool bIncludeVirtualUsers /*= false*/)
+{
+	ForEachUser([&InPredicate](FSlateUser& User)
+		{
+			InPredicate(&User);
+		}, bIncludeVirtualUsers);
+}
 
 /* FSlateApplicationBase interface
  *****************************************************************************/
@@ -4845,135 +3803,116 @@ EVisibility FSlateApplication::GetSoftwareCursorVis( ) const
 	return EVisibility::Hidden;
 }
 
-TSharedPtr< SWidget > FSlateApplication::GetKeyboardFocusedWidget() const
+TSharedPtr<SWidget> FSlateApplication::GetKeyboardFocusedWidget() const
 {
-	if ( const FSlateUser* User = GetUser(GetUserIndexForKeyboard()) )
-	{
-		return User->GetFocusedWidget();
-	}
-
-	return TSharedPtr< SWidget >();
+	TSharedPtr<const FSlateUser> KeyboardUser = GetUser(GetUserIndexForKeyboard());
+	return KeyboardUser ? KeyboardUser->GetFocusedWidget() : nullptr;
 }
 
 TSharedPtr<SWidget> FSlateApplication::GetMouseCaptorImpl() const
 {
-	return MouseCaptor.ToSharedWidget(CursorUserIndex, CursorPointerIndex);
+	TSharedPtr<const FSlateUser> FoundUser = GetUser(CursorUserIndex);
+	return FoundUser ? FoundUser->GetCursorCaptor() : nullptr;
 }
 
 bool FSlateApplication::HasAnyMouseCaptor() const
 {
-	return MouseCaptor.HasCapture();
+	for (const TSharedPtr<FSlateUser>& User : Users)
+	{
+		if (User && User->HasAnyCapture())
+		{
+			return true;
+		}
+	}
+	return false;
 }
 
 bool FSlateApplication::HasUserMouseCapture(int32 UserIndex) const
 {
-	return MouseCaptor.HasCaptureForUser(UserIndex);
+	TSharedPtr<const FSlateUser> FoundUser = GetUser(UserIndex);
+	return FoundUser && FoundUser->HasAnyCapture();
 }
 
 bool FSlateApplication::DoesWidgetHaveMouseCaptureByUser(const TSharedPtr<const SWidget> Widget, int32 UserIndex, TOptional<int32> PointerIndex) const
 {
-	return MouseCaptor.DoesWidgetHaveMouseCaptureByUser(Widget, UserIndex, PointerIndex);
+	if (TSharedPtr<const FSlateUser> FoundUser = GetUser(UserIndex))
+	{
+		return PointerIndex.IsSet() ? FoundUser->DoesWidgetHaveCapture(Widget, PointerIndex.GetValue()) : FoundUser->DoesWidgetHaveAnyCapture(Widget);
+	}
+	return false;
 }
 
 bool FSlateApplication::DoesWidgetHaveMouseCapture(const TSharedPtr<const SWidget> Widget) const
 {
-	return MouseCaptor.DoesWidgetHaveMouseCapture(Widget);
-}
-
-TOptional<EFocusCause> FSlateApplication::HasUserFocus(const TSharedPtr<const SWidget> Widget, int32 UserIndex) const
-{
-	if ( const FSlateUser* User = GetUser(UserIndex) )
+	for (const TSharedPtr<FSlateUser>& User : Users)
 	{
-		if ( User->GetFocusedWidget() == Widget )
+		if (User && User->DoesWidgetHaveAnyCapture(Widget))
 		{
-			return User->FocusCause;
+			return true;
 		}
 	}
-
-	return TOptional<EFocusCause>();
+	return false;
 }
 
-TOptional<EFocusCause> FSlateApplication::HasAnyUserFocus(const TSharedPtr<const SWidget> Widget) const
+TOptional<EFocusCause> FSlateApplication::HasUserFocus(const TSharedPtr<const SWidget> Widget, int32 UserIndex) const
 {
-	for ( int32 UserIndex = 0; UserIndex < Users.Num(); UserIndex++ )
-	{
-		if ( const FSlateUser* User = Users[UserIndex].Get() )
-		{
-			if ( User->GetFocusedWidget() == Widget )
-			{
-				return User->FocusCause;
-			}
-		}
-	}
-
-	return TOptional<EFocusCause>();
+	TSharedPtr<const FSlateUser> FoundUser = GetUser(UserIndex);
+	return FoundUser ? FoundUser->HasFocus(Widget) : TOptional<EFocusCause>();
 }
 
-bool FSlateApplication::IsWidgetDirectlyHovered(const TSharedPtr<const SWidget> Widget) const
+TOptional<EFocusCause> FSlateApplication::HasAnyUserFocus(const TSharedPtr<const SWidget> Widget) const
 {
-	for( auto LastWidgetIterator = WidgetsUnderCursorLastEvent.CreateConstIterator(); LastWidgetIterator; ++LastWidgetIterator )
+	TOptional<EFocusCause> FocusCause;
+	for (const TSharedPtr<FSlateUser>& User : Users)
 	{
-		const FWeakWidgetPath& WeakPath = LastWidgetIterator.Value();
-		if( WeakPath.IsValid() && Widget == WeakPath.GetLastWidget().Pin() )
+		FocusCause = User ? User->HasFocus(Widget) : TOptional<EFocusCause>();
+		if (FocusCause.IsSet())
 		{
-			return true;
+			break;
 		}
 	}
-	return false;
+	return FocusCause;
 }
 
-bool FSlateApplication::ShowUserFocus(const TSharedPtr<const SWidget> Widget) const
+bool FSlateApplication::IsWidgetDirectlyHovered(const TSharedPtr<const SWidget> Widget) const
 {
-	for ( int32 UserIndex = 0; UserIndex < Users.Num(); UserIndex++ )
+	for (const TSharedPtr<FSlateUser>& User : Users)
 	{
-		if ( const FSlateUser* User = Users[UserIndex].Get() )
+		if (User && User->IsWidgetDirectlyUnderAnyPointer(Widget))
 		{
-			TSharedPtr<SWidget> FocusedWidget = User->GetFocusedWidget();
-			if ( FocusedWidget == Widget )
-			{
-				return User->ShowFocus;
-			}
+			return true;
 		}
 	}
-
 	return false;
 }
 
-bool FSlateApplication::HasUserFocusedDescendants(const TSharedRef< const SWidget >& Widget, int32 UserIndex) const
+bool FSlateApplication::ShowUserFocus(const TSharedPtr<const SWidget> Widget) const
 {
-	if ( const FSlateUser* User = GetUser(UserIndex) )
+	for (const TSharedPtr<FSlateUser>& User : Users)
 	{
-		TSharedPtr<SWidget> FocusedWidget = User->GetFocusedWidget();
-		if ( FocusedWidget != Widget )
-		{
-			const FWeakWidgetPath& FocusedWidgetPath = User->GetWeakFocusPath();
-			if ( FocusedWidgetPath.ContainsWidget(Widget) )
+		if (User && User->ShouldShowFocus(Widget))
 		{
 			return true;
 		}
 	}
-	}
-
 	return false;
 }
 
+bool FSlateApplication::HasUserFocusedDescendants(const TSharedRef< const SWidget >& Widget, int32 UserIndex) const
+{
+	TSharedPtr<const FSlateUser> User = GetUser(UserIndex);
+	return User && User->HasFocusedDescendants(Widget);
+}
+
 bool FSlateApplication::HasFocusedDescendants( const TSharedRef< const SWidget >& Widget ) const
 {
-	for ( int32 UserIndex = 0; UserIndex < Users.Num(); UserIndex++ )
+	for (const TSharedPtr<FSlateUser>& User : Users)
 	{
-		if ( const FSlateUser* User = Users[UserIndex].Get() )
+		if (User && User->HasFocusedDescendants(Widget))
 		{
-			TSharedPtr<SWidget> FocusedWidget = User->GetFocusedWidget();
-			if ( FocusedWidget != Widget )
-			{
-				const FWeakWidgetPath& FocusedWidgetPath = User->GetWeakFocusPath();
-				if ( FocusedWidgetPath.ContainsWidget(Widget) )
-			{
-				return true;
-			}
+			return true;
 		}
 	}
-	}
 	return false;
 }
 
@@ -5053,17 +3992,16 @@ bool FSlateApplication::ProcessKeyCharEvent( const FCharacterEvent& InCharacterE
 	// NOTE: We intentionally don't reset LastUserInteractionTimeForThrottling here so that the UI can be responsive while typing
 
 	// Bubble the key event
-	if ( FSlateUser* User = GetOrCreateUser(InCharacterEvent.GetUserIndex()) )
-	{
-		TSharedRef<FWidgetPath> EventPathRef = User->GetFocusPath();
-		const FWidgetPath& EventPath = EventPathRef.Get();
+	TSharedRef<FSlateUser> User = GetOrCreateUser(InCharacterEvent);
+	TSharedRef<FWidgetPath> EventPathRef = User->GetFocusPath();
+	const FWidgetPath& EventPath = EventPathRef.Get();
 
-		// Switch worlds for widgets in the current path
-		FScopedSwitchWorldHack SwitchWorld(EventPath);
+	// Switch worlds for widgets in the current path
+	FScopedSwitchWorldHack SwitchWorld(EventPath);
 
-		{
-			SCOPE_CYCLE_COUNTER(STAT_ProcessKeyChar_RouteAlongFocusPath);
-			Reply = FEventRouter::RouteAlongFocusPath(this, FEventRouter::FBubblePolicy(EventPath), InCharacterEvent, [] (const FArrangedWidget& SomeWidgetGettingEvent, const FCharacterEvent& Event)
+	{
+		SCOPE_CYCLE_COUNTER(STAT_ProcessKeyChar_RouteAlongFocusPath);
+		Reply = FEventRouter::RouteAlongFocusPath(this, FEventRouter::FBubblePolicy(EventPath), InCharacterEvent, [](const FArrangedWidget& SomeWidgetGettingEvent, const FCharacterEvent& Event)
 			{
 				SCOPE_CYCLE_COUNTER(STAT_ProcessKeyChar_Call_OnKeyChar);
 
@@ -5078,7 +4016,6 @@ bool FSlateApplication::ProcessKeyCharEvent( const FCharacterEvent& InCharacterE
 
 				return FReply::Unhandled();
 			});
-		}
 	}
 
 	return Reply.IsEventHandled();
@@ -5095,9 +4032,10 @@ bool FSlateApplication::OnKeyDown( const int32 KeyCode, const uint32 CharacterCo
 bool FSlateApplication::ProcessKeyDownEvent( const FKeyEvent& InKeyEvent )
 {
 	TScopeCounter<int32> BeginInput(ProcessingInput);
-
 	SCOPE_CYCLE_COUNTER(STAT_ProcessKeyDown);
 
+	TSharedRef<FSlateUser> SlateUser = GetOrCreateUser(InKeyEvent);
+
 #if WITH_EDITOR
 	//Send the key input to all pre input key down listener function
 	if (OnApplicationPreInputKeyDownListenerEvent.IsBound())
@@ -5106,7 +4044,7 @@ bool FSlateApplication::ProcessKeyDownEvent( const FKeyEvent& InKeyEvent )
 	}
 #endif //WITH_EDITOR
 
-	QueueSynthesizedMouseMove();
+	QueueSynthesizedMouseMove(InKeyEvent);
 
 	// Analog cursor gets first chance at the input
 	if (InputPreProcessors.HandleKeyDownEvent(*this, InKeyEvent))
@@ -5118,10 +4056,11 @@ bool FSlateApplication::ProcessKeyDownEvent( const FKeyEvent& InKeyEvent )
 
 	SetLastUserInteractionTime(this->GetCurrentTime());
 	
-	if (IsDragDropping() && InKeyEvent.GetKey() == EKeys::Escape)
+	
+	if (SlateUser->IsDragDropping() && InKeyEvent.GetKey() == EKeys::Escape)
 	{
 		// Pressing ESC while drag and dropping terminates the drag drop.
-		CancelDragDrop();
+		SlateUser->CancelDragDrop();
 		Reply = FReply::Handled();
 	}
 	else
@@ -5158,49 +4097,46 @@ bool FSlateApplication::ProcessKeyDownEvent( const FKeyEvent& InKeyEvent )
 #endif //!(UE_BUILD_SHIPPING || UE_BUILD_TEST)
 
 		// Bubble the keyboard event
-		if ( FSlateUser* User = GetOrCreateUser(InKeyEvent.GetUserIndex()) )
-		{
-			TSharedRef<FWidgetPath> EventPathRef = User->GetFocusPath();
-			const FWidgetPath& EventPath = EventPathRef.Get();
+		TSharedRef<FWidgetPath> EventPathRef = SlateUser->GetFocusPath();
+		const FWidgetPath& EventPath = EventPathRef.Get();
 
-			// Switch worlds for widgets inOnPreviewMouseButtonDown the current path
-			FScopedSwitchWorldHack SwitchWorld(EventPath);
+		// Switch worlds for widgets inOnPreviewMouseButtonDown the current path
+		FScopedSwitchWorldHack SwitchWorld(EventPath);
 
-			// Tunnel the keyboard event
-			Reply = FEventRouter::RouteAlongFocusPath(this, FEventRouter::FTunnelPolicy(EventPath), InKeyEvent, [] (const FArrangedWidget& CurrentWidget, const FKeyEvent& Event)
-			{
+		// Tunnel the keyboard event
+		Reply = FEventRouter::RouteAlongFocusPath(this, FEventRouter::FTunnelPolicy(EventPath), InKeyEvent, [] (const FArrangedWidget& CurrentWidget, const FKeyEvent& Event)
+		{
 #if WITH_SLATE_DEBUGGING
-				// TODO
+			// TODO
 #endif
 
-				return ( CurrentWidget.Widget->IsEnabled() )
-					? CurrentWidget.Widget->OnPreviewKeyDown(CurrentWidget.Geometry, Event)
-					: FReply::Unhandled();
-			});
+			return ( CurrentWidget.Widget->IsEnabled() )
+				? CurrentWidget.Widget->OnPreviewKeyDown(CurrentWidget.Geometry, Event)
+				: FReply::Unhandled();
+		});
 
-			// Send out key down events.
-			if ( !Reply.IsEventHandled() )
+		// Send out key down events.
+		if ( !Reply.IsEventHandled() )
+		{
+			Reply = FEventRouter::RouteAlongFocusPath(this, FEventRouter::FBubblePolicy(EventPath), InKeyEvent, [] (const FArrangedWidget& SomeWidgetGettingEvent, const FKeyEvent& Event)
 			{
-				Reply = FEventRouter::RouteAlongFocusPath(this, FEventRouter::FBubblePolicy(EventPath), InKeyEvent, [] (const FArrangedWidget& SomeWidgetGettingEvent, const FKeyEvent& Event)
+				if (SomeWidgetGettingEvent.Widget->IsEnabled())
 				{
-					if (SomeWidgetGettingEvent.Widget->IsEnabled())
-					{
-						const FReply TempReply = SomeWidgetGettingEvent.Widget->OnKeyDown(SomeWidgetGettingEvent.Geometry, Event);
+					const FReply TempReply = SomeWidgetGettingEvent.Widget->OnKeyDown(SomeWidgetGettingEvent.Geometry, Event);
 #if WITH_SLATE_DEBUGGING
-						FSlateDebugging::BroadcastInputEvent(ESlateDebuggingInputEvent::KeyDown, TempReply, SomeWidgetGettingEvent.Widget, Event.GetKey().ToString());
+					FSlateDebugging::BroadcastInputEvent(ESlateDebuggingInputEvent::KeyDown, TempReply, SomeWidgetGettingEvent.Widget, Event.GetKey().ToString());
 #endif
-						return TempReply;
-					}
+					return TempReply;
+				}
 
-					return FReply::Unhandled();
-				});
-			}
+				return FReply::Unhandled();
+			});
+		}
 
-			// If the key event was not processed by any widget...
-			if ( !Reply.IsEventHandled() && UnhandledKeyDownEventHandler.IsBound() )
-			{
-				Reply = UnhandledKeyDownEventHandler.Execute(InKeyEvent);
-			}
+		// If the key event was not processed by any widget...
+		if ( !Reply.IsEventHandled() && UnhandledKeyDownEventHandler.IsBound() )
+		{
+			Reply = UnhandledKeyDownEventHandler.Execute(InKeyEvent);
 		}
 	}
 
@@ -5221,7 +4157,7 @@ bool FSlateApplication::ProcessKeyUpEvent( const FKeyEvent& InKeyEvent )
 
 	TScopeCounter<int32> BeginInput(ProcessingInput);
 
-	QueueSynthesizedMouseMove();
+	QueueSynthesizedMouseMove(InKeyEvent);
 
 	// Analog cursor gets first chance at the input
 	if (InputPreProcessors.HandleKeyUpEvent(*this, InKeyEvent))
@@ -5236,15 +4172,13 @@ bool FSlateApplication::ProcessKeyUpEvent( const FKeyEvent& InKeyEvent )
 	LastUserInteractionTimeForThrottling = LastUserInteractionTime;
 
 	// Bubble the key event
-	if ( FSlateUser* User = GetOrCreateUser(InKeyEvent.GetUserIndex()) )
-	{
-		TSharedRef<FWidgetPath> EventPathRef = User->GetFocusPath();
-		const FWidgetPath& EventPath = EventPathRef.Get();
+	TSharedRef<FWidgetPath> EventPathRef = GetOrCreateUser(InKeyEvent)->GetFocusPath();
+	const FWidgetPath& EventPath = EventPathRef.Get();
 
-		// Switch worlds for widgets in the current path
-		FScopedSwitchWorldHack SwitchWorld(EventPath);
+	// Switch worlds for widgets in the current path
+	FScopedSwitchWorldHack SwitchWorld(EventPath);
 
-		Reply = FEventRouter::RouteAlongFocusPath(this, FEventRouter::FBubblePolicy(EventPath), InKeyEvent, [] (const FArrangedWidget& SomeWidgetGettingEvent, const FKeyEvent& Event)
+	Reply = FEventRouter::RouteAlongFocusPath(this, FEventRouter::FBubblePolicy(EventPath), InKeyEvent, [](const FArrangedWidget& SomeWidgetGettingEvent, const FKeyEvent& Event)
 		{
 			if (SomeWidgetGettingEvent.Widget->IsEnabled())
 			{
@@ -5257,7 +4191,6 @@ bool FSlateApplication::ProcessKeyUpEvent( const FKeyEvent& InKeyEvent )
 
 			return FReply::Unhandled();
 		});
-	}
 
 	return Reply.IsEventHandled();
 }
@@ -5268,7 +4201,7 @@ bool FSlateApplication::ProcessAnalogInputEvent(const FAnalogInputEvent& InAnalo
 
 	TScopeCounter<int32> BeginInput(ProcessingInput);
 
-	QueueSynthesizedMouseMove();
+	QueueSynthesizedMouseMove(InAnalogInputEvent);
 
 	FReply Reply = FReply::Unhandled();
 
@@ -5280,18 +4213,16 @@ bool FSlateApplication::ProcessAnalogInputEvent(const FAnalogInputEvent& InAnalo
 
 	if (!Reply.IsEventHandled())
 	{
-		if (FSlateUser* User = GetOrCreateUser(InAnalogInputEvent.GetUserIndex()))
-		{
-			TSharedRef<FWidgetPath> EventPathRef = User->GetFocusPath();
-			const FWidgetPath& EventPath = EventPathRef.Get();
+		TSharedRef<FWidgetPath> EventPathRef = GetOrCreateUser(InAnalogInputEvent)->GetFocusPath();
+		const FWidgetPath& EventPath = EventPathRef.Get();
 
-			FAnalogInputEvent ModifiedEvent(InAnalogInputEvent);
-			ModifiedEvent.SetEventPath(EventPath);
+		FAnalogInputEvent ModifiedEvent(InAnalogInputEvent);
+		ModifiedEvent.SetEventPath(EventPath);
 
-			// Switch worlds for widgets in the current path
-			FScopedSwitchWorldHack SwitchWorld(EventPath);
+		// Switch worlds for widgets in the current path
+		FScopedSwitchWorldHack SwitchWorld(EventPath);
 
-			Reply = FEventRouter::RouteAlongFocusPath(this, FEventRouter::FBubblePolicy(EventPath), ModifiedEvent, [](const FArrangedWidget& SomeWidgetGettingEvent, const FAnalogInputEvent& Event)
+		Reply = FEventRouter::RouteAlongFocusPath(this, FEventRouter::FBubblePolicy(EventPath), ModifiedEvent, [](const FArrangedWidget& SomeWidgetGettingEvent, const FAnalogInputEvent& Event)
 			{
 				if (SomeWidgetGettingEvent.Widget->IsEnabled())
 				{
@@ -5305,8 +4236,7 @@ bool FSlateApplication::ProcessAnalogInputEvent(const FAnalogInputEvent& InAnalo
 				return FReply::Unhandled();
 			});
 
-			QueueSynthesizedMouseMove();
-		}
+		QueueSynthesizedMouseMove(ModifiedEvent);
 	}
 
 	// If no one handled this, it was probably motion in the deadzone.  Don't treat it as activity.
@@ -5413,7 +4343,7 @@ bool FSlateApplication::ProcessMouseButtonDownEvent( const TSharedPtr< FGenericW
 	}
 #endif //WITH_EDITOR
 
-	QueueSynthesizedMouseMove();
+	QueueSynthesizedMouseMove(MouseEvent);
 	SetLastUserInteractionTime(this->GetCurrentTime());
 	LastUserInteractionTimeForThrottling = LastUserInteractionTime;
 	
@@ -5421,7 +4351,11 @@ bool FSlateApplication::ProcessMouseButtonDownEvent( const TSharedPtr< FGenericW
 	{
 		PlatformApplication->SetCapture(PlatformWindow);
 	}
-	PressedMouseButtons.Add( MouseEvent.GetEffectingButton() );
+
+	if (MouseEvent.GetUserIndex() == CursorUserIndex)
+	{
+		PressedMouseButtons.Add(MouseEvent.GetEffectingButton());
+	}
 
 	// Input preprocessors get the first chance at the input
 	if (InputPreProcessors.HandleMouseButtonDownEvent(*this, MouseEvent))
@@ -5432,12 +4366,13 @@ bool FSlateApplication::ProcessMouseButtonDownEvent( const TSharedPtr< FGenericW
 	bool bInGame = false;
 
 	// Only process mouse down messages if we are not drag/dropping
-	if ( !IsDragDroppingAffected(MouseEvent) )
+	TSharedRef<FSlateUser> SlateUser = GetOrCreateUser(MouseEvent);
+	if (!SlateUser->IsDragDropping())
 	{
 		FReply Reply = FReply::Unhandled();
-		if (MouseCaptor.HasCaptureForPointerIndex(MouseEvent.GetUserIndex(), MouseEvent.GetPointerIndex()))
+		if (SlateUser->HasCapture(MouseEvent.GetPointerIndex()))
 		{
-			FWidgetPath MouseCaptorPath = MouseCaptor.ToWidgetPath( FWeakWidgetPath::EInterruptedPathHandling::Truncate, &MouseEvent );
+			FWidgetPath MouseCaptorPath = SlateUser->GetCaptorPath(MouseEvent.GetPointerIndex(), FWeakWidgetPath::EInterruptedPathHandling::Truncate, &MouseEvent);
 			FArrangedWidget& MouseCaptorWidget = MouseCaptorPath.Widgets.Last();
 
 			// Switch worlds widgets in the current path
@@ -5483,6 +4418,7 @@ bool FSlateApplication::ProcessMouseButtonDownEvent( const TSharedPtr< FGenericW
 			FScopedSwitchWorldHack SwitchWorld(WidgetsUnderCursor);
 			bInGame = FApp::IsGame();
 
+			//@todo NickD: Route API should be private; update Process methods to accept an FWidgetPath
 			Reply = RoutePointerDownEvent(WidgetsUnderCursor, MouseEvent);
 		}
 
@@ -5510,9 +4446,8 @@ FReply FSlateApplication::RoutePointerDownEvent(const FWidgetPath& WidgetsUnderP
 {
 	TScopeCounter<int32> BeginInput(ProcessingInput);
 	
-	// Ensure the cursor location(s) get set to an initial value
-	PointerIndexPositionMap.Add(FUserAndPointer(PointerEvent.GetUserIndex(), PointerEvent.GetPointerIndex()), PointerEvent.GetScreenSpacePosition());
-	PointerIndexLastPositionMap.Add(FUserAndPointer(PointerEvent.GetUserIndex(), PointerEvent.GetPointerIndex()), PointerEvent.GetScreenSpacePosition());
+	TSharedRef<FSlateUser> SlateUser = GetOrCreateUser(PointerEvent);
+	SlateUser->UpdatePointerPosition(PointerEvent);
 
 #if PLATFORM_MAC
 	NSWindow* ActiveWindow = [ NSApp keyWindow ];
@@ -5568,16 +4503,17 @@ FReply FSlateApplication::RoutePointerDownEvent(const FWidgetPath& WidgetsUnderP
 		}
 	}
 
+#if PLATFORM_MAC
+	const bool bIsDetectingLMBDrag = PointerEvent.GetEffectingButton() == EKeys::LeftMouseButton && SlateUser->IsDetectingDrag(PointerEvent.GetPointerIndex());
+#endif
+
 	// If none of the widgets requested keyboard focus to be set (or set the keyboard focus explicitly), set it to the leaf-most widget under the mouse.
 	// On Mac we prevent the OS from activating the window on mouse down, so we have full control and can activate only if there's nothing draggable under the mouse cursor.
 	const bool bFocusChangedByEventHandler = PreviouslyFocusedWidget != GetKeyboardFocusedWidget();
 	if( ( !bFocusChangedByEventHandler || bNeedToActivateWindow ) &&
 		( !Reply.GetUserFocusRecepient().IsValid()
 #if PLATFORM_MAC
-			|| (
-				PointerEvent.GetEffectingButton() == EKeys::LeftMouseButton &&
-				!DragDetector.IsDetectingDrag(PointerEvent)
-			)
+			|| bIsDetectingLMBDrag)
 #endif
 		)
 	)
@@ -5603,15 +4539,16 @@ FReply FSlateApplication::RoutePointerDownEvent(const FWidgetPath& WidgetsUnderP
 				// Clicking on a context menu should not activate anything
 				// @todo: This needs to be updated when we have window type in SWindow and we no longer have to guess if WidgetsUnderCursor.TopLevelWindow is a menu
 				const bool bIsContextMenu = TopLevelWindow.IsValid() && !TopLevelWindow->IsRegularWindow() && TopLevelWindow->HasMinimizeBox() && TopLevelWindow->HasMaximizeBox();
-				if ( !bIsContextMenu && PointerEvent.GetEffectingButton() == EKeys::LeftMouseButton && !DragDetector.IsDetectingDrag(PointerEvent) && ActiveWindow == [NSApp keyWindow] )
+				if ( !bIsContextMenu && bIsDetectingLMBDrag && ActiveWindow == [NSApp keyWindow] )
 				{
-					MouseCaptorHelper Captor = MouseCaptor;
+					TMap<TSharedRef<FSlateUser>, TMap<uint32, FWeakWidgetPath>> AllPointerCaptors;
+					ForEachUser([&AllPointerCaptors](FSlateUser& User) { AllPointerCaptors.Add(User.AsShared(), User.GetCaptorPathsByIndex()); });
 					FPlatformApplicationMisc::ActivateApplication();
 					if ( TopLevelWindow.IsValid() )
 					{
 						TopLevelWindow->BringToFront(true);
 					}
-					MouseCaptor = Captor;
+					ForEachUser([&AllPointerCaptors](FSlateUser& User) { User.RestoreCaptorPathsByIndex(AllPointerCaptors.FindChecked(User.AsShared())); });
 				}
 			}
 		}
@@ -5627,15 +4564,16 @@ FReply FSlateApplication::RoutePointerUpEvent(const FWidgetPath& WidgetsUnderPoi
 	TScopeCounter<int32> BeginInput(ProcessingInput);
 
 	FReply Reply = FReply::Unhandled();
+	TSharedRef<FSlateUser> SlateUser = GetOrCreateUser(PointerEvent);
+	const bool bIsDragDropping = SlateUser->IsDragDroppingAffected(PointerEvent);
+	TSharedPtr<FDragDropOperation> LocalDragDropContent;
+	
+	bool bDropWasHandled = false;
+	FWidgetPath LocalWidgetsUnderPointer = WidgetsUnderPointer;
 
-	// Update the drag detector, this release may stop a drag detection.
-	DragDetector.OnPointerRelease(PointerEvent);
-
-	const bool bIsDragDropping = IsDragDroppingAffected(PointerEvent);
-
-	if (MouseCaptor.HasCaptureForPointerIndex(PointerEvent.GetUserIndex(), PointerEvent.GetPointerIndex()))
+	if (SlateUser->HasCapture(PointerEvent.GetPointerIndex()))
 	{
-		FWidgetPath MouseCaptorPath = MouseCaptor.ToWidgetPath( FWeakWidgetPath::EInterruptedPathHandling::Truncate, &PointerEvent );
+		FWidgetPath MouseCaptorPath = SlateUser->GetCaptorPath(PointerEvent.GetPointerIndex(), FWeakWidgetPath::EInterruptedPathHandling::Truncate, &PointerEvent);
 		if ( ensureMsgf(MouseCaptorPath.Widgets.Num() > 0, TEXT("A window had a widget with mouse capture. That entire window has been dismissed before the mouse up could be processed.")) )
 		{
 			// Switch worlds widgets in the current path
@@ -5672,31 +4610,27 @@ FReply FSlateApplication::RoutePointerUpEvent(const FWidgetPath& WidgetsUnderPoi
 
 					return TempReply;
 				});
-
-			// For touch events, we always invalidate capture for the pointer.  There's no reason to ever maintain capture for
-			// fingers no longer in contact with the screen.
-			if ( PointerEvent.IsTouchEvent() )
-			{
-				MouseCaptor.InvalidateCaptureForPointer(PointerEvent.GetUserIndex(), PointerEvent.GetPointerIndex());
-			}
 		}
 	}
 	else
 	{
-		FWidgetPath LocalWidgetsUnderCursor = WidgetsUnderPointer.IsValid() ? WidgetsUnderPointer : LocateWindowUnderMouse(PointerEvent.GetScreenSpacePosition(), GetInteractiveTopLevelWindows());
-
+		if (!LocalWidgetsUnderPointer.IsValid())
+		{
+			LocalWidgetsUnderPointer = LocateWindowUnderMouse(PointerEvent.GetScreenSpacePosition(), GetInteractiveTopLevelWindows());
+		}
+		
 		// Switch worlds widgets in the current path
-		FScopedSwitchWorldHack SwitchWorld(LocalWidgetsUnderCursor);
+		FScopedSwitchWorldHack SwitchWorld(LocalWidgetsUnderPointer);
 
 		// Cache the drag drop content and reset the pointer in case OnMouseButtonUpMessage re-enters as a result of OnDrop
-		TSharedPtr< FDragDropOperation > LocalDragDropContent = DragDropContent;
-
+		// In such a case, we want the re-entrant call to skip any drag-drop stuff (otherwise we'd execute the drop action twice)
 		if (bIsDragDropping)
 		{
-			ResetDragDropState();
+			LocalDragDropContent = SlateUser->GetDragDropContent();
+			SlateUser->ResetDragDropContent();
 		}
 
-		Reply = FEventRouter::Route<FReply>(this, FEventRouter::FBubblePolicy(LocalWidgetsUnderCursor), PointerEvent, [&](const FArrangedWidget& CurWidget, const FPointerEvent& Event)
+		Reply = FEventRouter::Route<FReply>(this, FEventRouter::FBubblePolicy(LocalWidgetsUnderPointer), PointerEvent, [&](const FArrangedWidget& CurWidget, const FPointerEvent& Event)
 		{
 			if (bIsDragDropping)
 			{
@@ -5727,28 +4661,10 @@ FReply FSlateApplication::RoutePointerUpEvent(const FWidgetPath& WidgetsUnderPoi
 
 			return TempReply;
 		});
-
-		// When we perform a touch end, we need to also send a mouse leave as if it were a cursor.
-		if (PointerEvent.IsTouchEvent() && !IsFakingTouchEvents())
-		{
-			for (int32 WidgetIndex = LocalWidgetsUnderCursor.Widgets.Num() - 1; WidgetIndex >= 0; --WidgetIndex)
-			{
-				LocalWidgetsUnderCursor.Widgets[WidgetIndex].Widget->OnMouseLeave(PointerEvent);
-			}
-		}
-
-		// If we were dragging, notify the content
-		if (bIsDragDropping)
-		{
-			// @todo slate : depending on SetEventPath() is not ideal.
-			FPointerEvent ModifiedEvent(PointerEvent);
-			ModifiedEvent.SetEventPath(LocalWidgetsUnderCursor);
-			LocalDragDropContent->OnDrop(Reply.IsEventHandled(), ModifiedEvent);
-
-			WidgetsUnderCursorLastEvent.Remove(FUserAndPointer(ModifiedEvent.GetUserIndex(), ModifiedEvent.GetPointerIndex()));
-		}
 	}
 
+	SlateUser->NotifyPointerReleased(PointerEvent, LocalWidgetsUnderPointer, LocalDragDropContent, Reply.IsEventHandled());
+
 #if PLATFORM_MAC
 	// Make sure the application and its front window are activated if user wasn't drag & dropping between windows
 	if (PointerEvent.GetEffectingButton() == EKeys::LeftMouseButton && !bIsDragDropping)
@@ -5777,16 +4693,14 @@ bool FSlateApplication::RoutePointerMoveEvent(const FWidgetPath& WidgetsUnderPoi
 {
 	TScopeCounter<int32> BeginInput(ProcessingInput);
 
-	PointerIndexPositionMap.Add(FUserAndPointer(PointerEvent.GetUserIndex(), PointerEvent.GetPointerIndex()), PointerEvent.GetScreenSpacePosition());
-
 	bool bHandled = false;
+	FWeakWidgetPath LastWidgetsUnderPointer;
 
-	FWeakWidgetPath LastWidgetsUnderCursor;
+	TSharedRef<FSlateUser> SlateUser = GetOrCreateUser(PointerEvent);
+	SlateUser->NotifyPointerMoveBegin(PointerEvent);
 
-	// User asked us to detect a drag.
-	bool bDragDetected = false;
-	// Currently we can only support one dragged widget at a time.
-	bool bShouldStartDetectingDrag = !DragDropContent.IsValid();
+	// Currently we support only one dragged widget at a time per user
+	bool bShouldStartDetectingDrag = !SlateUser->IsDragDropping();
 
 #if WITH_EDITOR
 	//@TODO VREDITOR - Remove and move to interaction component
@@ -5796,29 +4710,22 @@ bool FSlateApplication::RoutePointerMoveEvent(const FWidgetPath& WidgetsUnderPoi
 	}
 #endif 
 
-	if ( !bIsSynthetic && bShouldStartDetectingDrag )
+	if (!bIsSynthetic && bShouldStartDetectingDrag)
 	{
-		FWeakWidgetPath* DetectDragForWidget;
-		bDragDetected = DragDetector.DetectDrag(PointerEvent, GetDragTriggerDistance(), DetectDragForWidget);
-		if ( bDragDetected )
+		FWidgetPath DragDetectPath = SlateUser->DetectDrag(PointerEvent, GetDragTriggerDistance());
+		if (DragDetectPath.IsValid())
 		{
-			FWidgetPath DragDetectPath = DetectDragForWidget->ToWidgetPath(FWeakWidgetPath::EInterruptedPathHandling::ReturnInvalid);
-			const TSharedPtr<SWidget> DragDetectRequestor = DetectDragForWidget->IsValid() ? DetectDragForWidget->GetLastWidget().Pin() : nullptr;
-			if ( DragDetectPath.IsValid() && DragDetectRequestor.IsValid() )
-			{
-				FWidgetAndPointer DetectDragForMe = DragDetectPath.FindArrangedWidgetAndCursor(DragDetectRequestor.ToSharedRef()).Get(FWidgetAndPointer());
+			FWidgetAndPointer DetectDragForMe = DragDetectPath.FindArrangedWidgetAndCursor(DragDetectPath.GetLastWidget()).Get(FWidgetAndPointer());
 
-				// A drag has been triggered. The cursor exited some widgets as a result.
-				// This assignment ensures that we will send OnLeave notifications to those widgets.
-				LastWidgetsUnderCursor = *DetectDragForWidget;
+			// A drag has been triggered. The cursor exited some widgets as a result.
+			// This assignment ensures that we will send OnLeave notifications to those widgets.
+			LastWidgetsUnderPointer = DragDetectPath;
 
-				DragDetector.ResetDetection();
-				
-				// Switch worlds widgets in the current path
-				FScopedSwitchWorldHack SwitchWorld(DragDetectPath);
+			// Switch worlds widgets in the current path
+			FScopedSwitchWorldHack SwitchWorld(DragDetectPath);
 
-				// Send an OnDragDetected to the widget that requested drag-detection.
-				FReply Reply = FEventRouter::Route<FReply>(this, FEventRouter::FDirectPolicy(DetectDragForMe, DragDetectPath, &WidgetsUnderPointer), PointerEvent, [] (const FArrangedWidget& InDetectDragForMe, const FPointerEvent& TranslatedMouseEvent)
+			// Send an OnDragDetected to the widget that requested drag detection.
+			FReply Reply = FEventRouter::Route<FReply>(this, FEventRouter::FDirectPolicy(DetectDragForMe, DragDetectPath, &WidgetsUnderPointer), PointerEvent, [](const FArrangedWidget& InDetectDragForMe, const FPointerEvent& TranslatedMouseEvent)
 				{
 					const FReply TempReply = InDetectDragForMe.Widget->OnDragDetected(InDetectDragForMe.Geometry, TranslatedMouseEvent);
 #if WITH_SLATE_DEBUGGING
@@ -5826,80 +4733,64 @@ bool FSlateApplication::RoutePointerMoveEvent(const FWidgetPath& WidgetsUnderPoi
 #endif
 					return TempReply;
 				});
-
-				// FSlateApplication::ProcessReply() (called from
-				// FEventRouter::Route() above) will have assigned
-				// DragDropContent if it is processing the drag-and-drop
-				// operation. Otherwise, we could use
-				// Reply.GetDragDropContent().IsValid().
-				bDragDetected = DragDropContent.IsValid();
-			}
-			else
-			{
-				bDragDetected = false;
-			}
 		}
 	}
 
-
-	if (bDragDetected)
+	// A drag was detected if the user is now executing a drag-drop action
+	if (bShouldStartDetectingDrag && SlateUser->IsDragDropping())
 	{
 		// When a drag was detected, we pretend that the widgets under the mouse last time around.
 		// We have set LastWidgetsUnderCursor accordingly when the drag was detected above.
 	}
 	else
 	{
-		// No Drag Detection
-		LastWidgetsUnderCursor = WidgetsUnderCursorLastEvent.FindRef( FUserAndPointer( PointerEvent.GetUserIndex(), PointerEvent.GetPointerIndex() ) );
+		LastWidgetsUnderPointer = SlateUser->GetWidgetsUnderPointerLastEventByIndex().FindRef(PointerEvent.GetPointerIndex());
 	}
 
-	FWidgetPath MouseCaptorPath;
-	if ( MouseCaptor.HasCaptureForPointerIndex(PointerEvent.GetUserIndex(), PointerEvent.GetPointerIndex()) )
-	{
-		MouseCaptorPath = MouseCaptor.ToWidgetPath(FWeakWidgetPath::EInterruptedPathHandling::ReturnInvalid, &PointerEvent);
-	}
+	FWidgetPath MouseCaptorPath = SlateUser->GetCaptorPath(PointerEvent.GetPointerIndex(), FWeakWidgetPath::EInterruptedPathHandling::ReturnInvalid, &PointerEvent);
 
 	// Send out mouse leave events
 	// If we are doing a drag and drop, we will send this event instead.
+	TSharedPtr<FDragDropOperation> DragDropContent = SlateUser->GetDragDropContent();
 	{
-		FDragDropEvent DragDropEvent( PointerEvent, DragDropContent );
+		FDragDropEvent DragDropEvent(PointerEvent, DragDropContent);
 		// Switch worlds widgets in the current path
-		FScopedSwitchWorldHack SwitchWorld( LastWidgetsUnderCursor.Window.Pin() );
+		FScopedSwitchWorldHack SwitchWorld(LastWidgetsUnderPointer.Window.Pin());
 
-		for ( int32 WidgetIndex = LastWidgetsUnderCursor.Widgets.Num()-1; WidgetIndex >=0; --WidgetIndex )
+		for (int32 WidgetIndex = LastWidgetsUnderPointer.Widgets.Num()-1; WidgetIndex >=0; --WidgetIndex)
 		{
 			// Guards for cases where WidgetIndex can become invalid due to MouseMove being re-entrant.
-			while ( WidgetIndex >= LastWidgetsUnderCursor.Widgets.Num() )
+			while (WidgetIndex >= LastWidgetsUnderPointer.Widgets.Num())
 			{
 				WidgetIndex--;
 			}
 
-			if ( WidgetIndex >= 0 )
+			if (WidgetIndex >= 0)
 			{
-				const TSharedPtr<SWidget>& SomeWidgetPreviouslyUnderCursor = LastWidgetsUnderCursor.Widgets[WidgetIndex].Pin();
-				if( SomeWidgetPreviouslyUnderCursor.IsValid() )
+				const TSharedPtr<SWidget>& SomeWidgetPreviouslyUnderCursor = LastWidgetsUnderPointer.Widgets[WidgetIndex].Pin();
+				if (SomeWidgetPreviouslyUnderCursor.IsValid())
 				{
-					TOptional<FArrangedWidget> FoundWidget = WidgetsUnderPointer.FindArrangedWidget( SomeWidgetPreviouslyUnderCursor.ToSharedRef() );
+					TOptional<FArrangedWidget> FoundWidget = WidgetsUnderPointer.FindArrangedWidget(SomeWidgetPreviouslyUnderCursor.ToSharedRef());
 					const bool bWidgetNoLongerUnderMouse = !FoundWidget.IsSet();
-					if ( bWidgetNoLongerUnderMouse )
+					if (bWidgetNoLongerUnderMouse)
 					{
 						// Widget is no longer under cursor, so send a MouseLeave.
 						// The widget might not even be in the hierarchy any more!
 						// Thus, we cannot translate the PointerPosition into the appropriate space for this event.
-						if (IsDragDroppingAffected(PointerEvent))
+						if (SlateUser->IsDragDroppingAffected(PointerEvent))
 						{
 							// Note that the event's pointer position is not translated.
-							SomeWidgetPreviouslyUnderCursor->OnDragLeave( DragDropEvent );
+							SomeWidgetPreviouslyUnderCursor->OnDragLeave(DragDropEvent);
 #if WITH_SLATE_DEBUGGING
 							FSlateDebugging::BroadcastInputEvent(ESlateDebuggingInputEvent::DragLeave, SomeWidgetPreviouslyUnderCursor);
 #endif
 							// Reset the cursor override
-							DragDropEvent.GetOperation()->SetCursorOverride( TOptional<EMouseCursor::Type>() );
+							DragDropEvent.GetOperation()->SetCursorOverride(TOptional<EMouseCursor::Type>());
 						}
 						else
 						{
 							// Only fire mouse leave events for widgets inside the captor path, or whoever if there is no captor path.
-							if ( MouseCaptorPath.IsValid() == false || MouseCaptorPath.ContainsWidget(SomeWidgetPreviouslyUnderCursor.ToSharedRef()) )
+							if (MouseCaptorPath.IsValid() == false || MouseCaptorPath.ContainsWidget(SomeWidgetPreviouslyUnderCursor.ToSharedRef()))
 							{
 								// Note that the event's pointer position is not translated.
 								SomeWidgetPreviouslyUnderCursor->OnMouseLeave(PointerEvent);
@@ -5907,70 +4798,69 @@ bool FSlateApplication::RoutePointerMoveEvent(const FWidgetPath& WidgetsUnderPoi
 								FSlateDebugging::BroadcastInputEvent(ESlateDebuggingInputEvent::MouseLeave, SomeWidgetPreviouslyUnderCursor);
 #endif
 							}
-						}			
+						}
 					}
 				}
 			}
 		}
 	}
 
-
 	if (MouseCaptorPath.IsValid())
 	{
-		if ( !bIsSynthetic )
+		if (!bIsSynthetic)
 		{
 			// Switch worlds widgets in the current path
-			FScopedSwitchWorldHack SwitchWorld( MouseCaptorPath );
+			FScopedSwitchWorldHack SwitchWorld(MouseCaptorPath);
 
-			FEventRouter::Route<FNoReply>(this, FEventRouter::FBubblePolicy(WidgetsUnderPointer), PointerEvent, [&MouseCaptorPath, &LastWidgetsUnderCursor] (const FArrangedWidget& WidgetUnderCursor, const FPointerEvent& Event)
-			{
-				if ( !LastWidgetsUnderCursor.ContainsWidget(WidgetUnderCursor.Widget) )
+			FEventRouter::Route<FNoReply>(this, FEventRouter::FBubblePolicy(WidgetsUnderPointer), PointerEvent, [&MouseCaptorPath, &LastWidgetsUnderPointer](const FArrangedWidget& WidgetUnderCursor, const FPointerEvent& Event)
 				{
-					if ( MouseCaptorPath.ContainsWidget(WidgetUnderCursor.Widget) )
+					if (!LastWidgetsUnderPointer.ContainsWidget(WidgetUnderCursor.Widget))
 					{
-						WidgetUnderCursor.Widget->OnMouseEnter(WidgetUnderCursor.Geometry, Event);
+						if (MouseCaptorPath.ContainsWidget(WidgetUnderCursor.Widget))
+						{
+							WidgetUnderCursor.Widget->OnMouseEnter(WidgetUnderCursor.Geometry, Event);
+						}
 					}
-				}
-				return FNoReply();
-			});
+					return FNoReply();
+				});
 
-			FReply Reply = FEventRouter::Route<FReply>(this, FEventRouter::FToLeafmostPolicy(MouseCaptorPath), PointerEvent, [this] (const FArrangedWidget& MouseCaptorWidget, const FPointerEvent& Event)
-			{
-				FReply TempReply = FReply::Unhandled();
-				
-				bool bAllowMouseFallback = true;
-				if (Event.IsTouchEvent())
+			FReply Reply = FEventRouter::Route<FReply>(this, FEventRouter::FToLeafmostPolicy(MouseCaptorPath), PointerEvent, [this](const FArrangedWidget& MouseCaptorWidget, const FPointerEvent& Event)
 				{
-					if (Event.IsTouchForceChangedEvent())
+					FReply TempReply = FReply::Unhandled();
+
+					bool bAllowMouseFallback = true;
+					if (Event.IsTouchEvent())
 					{
-						TempReply = MouseCaptorWidget.Widget->OnTouchForceChanged(MouseCaptorWidget.Geometry, Event);
+						if (Event.IsTouchForceChangedEvent())
+						{
+							TempReply = MouseCaptorWidget.Widget->OnTouchForceChanged(MouseCaptorWidget.Geometry, Event);
 #if WITH_SLATE_DEBUGGING
-						//TODO FSlateDebugging
+							//TODO FSlateDebugging
 #endif
-						bAllowMouseFallback = false;
-					}
-					else if (Event.IsTouchFirstMoveEvent())
-					{
-						TempReply = MouseCaptorWidget.Widget->OnTouchFirstMove(MouseCaptorWidget.Geometry, Event);
+							bAllowMouseFallback = false;
+						}
+						else if (Event.IsTouchFirstMoveEvent())
+						{
+							TempReply = MouseCaptorWidget.Widget->OnTouchFirstMove(MouseCaptorWidget.Geometry, Event);
 #if WITH_SLATE_DEBUGGING
-						//TODO FSlateDebugging
+							//TODO FSlateDebugging
 #endif
-						bAllowMouseFallback = false;
+							bAllowMouseFallback = false;
+						}
+						else
+						{
+							TempReply = MouseCaptorWidget.Widget->OnTouchMoved(MouseCaptorWidget.Geometry, Event);
+						}
 					}
-					else
+					if ((!Event.IsTouchEvent() && bAllowMouseFallback) || (!TempReply.IsEventHandled() && this->bTouchFallbackToMouse))
 					{
-						TempReply = MouseCaptorWidget.Widget->OnTouchMoved(MouseCaptorWidget.Geometry, Event);
-					}
-				}
-				if ( (!Event.IsTouchEvent() && bAllowMouseFallback) || ( !TempReply.IsEventHandled() && this->bTouchFallbackToMouse ) )
-				{
-					TempReply = MouseCaptorWidget.Widget->OnMouseMove(MouseCaptorWidget.Geometry, Event);
+						TempReply = MouseCaptorWidget.Widget->OnMouseMove(MouseCaptorWidget.Geometry, Event);
 #if WITH_SLATE_DEBUGGING
-					FSlateDebugging::BroadcastInputEvent(ESlateDebuggingInputEvent::MouseMove, TempReply, MouseCaptorWidget.Widget);
+						FSlateDebugging::BroadcastInputEvent(ESlateDebuggingInputEvent::MouseMove, TempReply, MouseCaptorWidget.Widget);
 #endif
-				}
-				return TempReply;
-			});
+					}
+					return TempReply;
+				});
 			bHandled = Reply.IsEventHandled();
 		}
 	}
@@ -5979,13 +4869,15 @@ bool FSlateApplication::RoutePointerMoveEvent(const FWidgetPath& WidgetsUnderPoi
 		// Switch worlds widgets in the current path
 		FScopedSwitchWorldHack SwitchWorld(WidgetsUnderPointer);
 
+		const bool bIsDragDroppingAffected = SlateUser->IsDragDroppingAffected(PointerEvent);
+
 		// Send out mouse enter events.
-		if (IsDragDroppingAffected(PointerEvent))
+		if (bIsDragDroppingAffected)
 		{
 			FDragDropEvent DragDropEvent(PointerEvent, DragDropContent);
-			FEventRouter::Route<FNoReply>(this, FEventRouter::FBubblePolicy(WidgetsUnderPointer), DragDropEvent, [&LastWidgetsUnderCursor](const FArrangedWidget& WidgetUnderCursor, const FDragDropEvent& InDragDropEvent)
+			FEventRouter::Route<FNoReply>(this, FEventRouter::FBubblePolicy(WidgetsUnderPointer), DragDropEvent, [&LastWidgetsUnderPointer](const FArrangedWidget& WidgetUnderCursor, const FDragDropEvent& InDragDropEvent)
 			{
-				if (!LastWidgetsUnderCursor.ContainsWidget(WidgetUnderCursor.Widget))
+				if (!LastWidgetsUnderPointer.ContainsWidget(WidgetUnderCursor.Widget))
 				{
 					WidgetUnderCursor.Widget->OnDragEnter(WidgetUnderCursor.Geometry, InDragDropEvent);
 				}
@@ -5994,9 +4886,9 @@ bool FSlateApplication::RoutePointerMoveEvent(const FWidgetPath& WidgetsUnderPoi
 		}
 		else
 		{
-			FEventRouter::Route<FNoReply>(this, FEventRouter::FBubblePolicy(WidgetsUnderPointer), PointerEvent, [&LastWidgetsUnderCursor](const FArrangedWidget& WidgetUnderCursor, const FPointerEvent& Event)
+			FEventRouter::Route<FNoReply>(this, FEventRouter::FBubblePolicy(WidgetsUnderPointer), PointerEvent, [&LastWidgetsUnderPointer](const FArrangedWidget& WidgetUnderCursor, const FPointerEvent& Event)
 			{
-				if (!LastWidgetsUnderCursor.ContainsWidget(WidgetUnderCursor.Widget))
+				if (!LastWidgetsUnderPointer.ContainsWidget(WidgetUnderCursor.Widget))
 				{
 					WidgetUnderCursor.Widget->OnMouseEnter(WidgetUnderCursor.Geometry, Event);
 				}
@@ -6009,7 +4901,7 @@ bool FSlateApplication::RoutePointerMoveEvent(const FWidgetPath& WidgetsUnderPoi
 		{
 			FReply TempReply = FReply::Unhandled();
 
-			if (IsDragDroppingAffected(PointerEvent))
+			if (bIsDragDroppingAffected)
 			{
 				TempReply = CurWidget.Widget->OnDragOver(CurWidget.Geometry, FDragDropEvent(Event, DragDropContent));
 #if WITH_SLATE_DEBUGGING
@@ -6060,54 +4952,7 @@ bool FSlateApplication::RoutePointerMoveEvent(const FWidgetPath& WidgetsUnderPoi
 		bHandled = Reply.IsEventHandled();
 	}
 
-	// Give the current drag drop operation a chance to do something
-	// custom (e.g. update the Drag/Drop preview based on content)
-	if (IsDragDroppingAffected(PointerEvent))
-	{
-		FDragDropEvent DragDropEvent( PointerEvent, DragDropContent );
-		//@TODO VREDITOR - Remove and move to interaction component
-#if WITH_EDITOR
-		if (OnDragDropCheckOverride.IsBound() && DragDropEvent.GetOperation().IsValid())
-		{
-			DragDropEvent.GetOperation()->SetDecoratorVisibility(false);
-			DragDropEvent.GetOperation()->SetCursorOverride(EMouseCursor::None);
-			DragDropContent->SetCursorOverride(EMouseCursor::None);
-		}
-#endif
-		FScopedSwitchWorldHack SwitchWorld( WidgetsUnderPointer );
-		DragDropContent->OnDragged( DragDropEvent );
-
-		// Update the window we're under for rendering the drag drop operation if
-		// it's a windowless drag drop operation.
-		if ( WidgetsUnderPointer.IsValid() )
-		{
-			DragDropWindowPtr = WidgetsUnderPointer.GetWindow();
-		}
-		else
-		{
-			DragDropWindowPtr = nullptr;
-		}
-
-		// Don't update the cursor for the platform if we don't have a valid cursor on this platform
-		if ( PlatformApplication->Cursor.IsValid() )
-		{
-			FCursorReply CursorReply = DragDropContent->OnCursorQuery();
-			if ( !CursorReply.IsEventHandled() )
-			{
-				// Set the default cursor when there isn't an active window under the cursor and the mouse isn't captured
-				CursorReply = FCursorReply::Cursor(EMouseCursor::Default);
-			}
-
-			ProcessCursorReply(CursorReply);
-		}
-	}
-	else if (!IsDragDropping())
-	{
-		DragDropWindowPtr = nullptr;
-	}
-
-	WidgetsUnderCursorLastEvent.Add(FUserAndPointer(PointerEvent.GetUserIndex(), PointerEvent.GetPointerIndex()), FWeakWidgetPath(WidgetsUnderPointer));
-	PointerIndexLastPositionMap.Add(FUserAndPointer(PointerEvent.GetUserIndex(), PointerEvent.GetPointerIndex()), PointerEvent.GetScreenSpacePosition());
+	SlateUser->NotifyPointerMoveComplete(PointerEvent, WidgetsUnderPointer);
 
 	return bHandled;
 }
@@ -6144,12 +4989,16 @@ bool FSlateApplication::ProcessMouseButtonDoubleClickEvent( const TSharedPtr< FG
 {
 	SCOPE_CYCLE_COUNTER(STAT_ProcessMouseButtonDoubleClick);
 
-	QueueSynthesizedMouseMove();
+	QueueSynthesizedMouseMove(InMouseEvent);
 	SetLastUserInteractionTime(this->GetCurrentTime());
 	LastUserInteractionTimeForThrottling = LastUserInteractionTime;
 
 	PlatformApplication->SetCapture( PlatformWindow );
-	PressedMouseButtons.Add( InMouseEvent.GetEffectingButton() );
+
+	if (InMouseEvent.GetUserIndex() == CursorUserIndex)
+	{
+		PressedMouseButtons.Add( InMouseEvent.GetEffectingButton() );
+	}
 
 	// Input preprocessors get the first chance at the input
 	if (InputPreProcessors.HandleMouseButtonDoubleClickEvent(*this, InMouseEvent))
@@ -6157,7 +5006,7 @@ bool FSlateApplication::ProcessMouseButtonDoubleClickEvent( const TSharedPtr< FG
 		return true;
 	}
 
-	if (MouseCaptor.HasCaptureForPointerIndex(InMouseEvent.GetUserIndex(), InMouseEvent.GetPointerIndex()))
+	if (GetOrCreateUser(InMouseEvent)->HasCapture(InMouseEvent.GetPointerIndex()))
 	{
 		// If a widget has mouse capture, we've opted to simply treat this event as a mouse down
 		return ProcessMouseButtonDownEvent(PlatformWindow, InMouseEvent);
@@ -6233,10 +5082,15 @@ bool FSlateApplication::ProcessMouseButtonUpEvent( const FPointerEvent& MouseEve
 		FSlateThrottleManager::Get().LeaveResponsiveMode(MouseButtonDownResponsivnessThrottle);
 	}
 
-	QueueSynthesizedMouseMove();
+	QueueSynthesizedMouseMove(MouseEvent);
 	SetLastUserInteractionTime(this->GetCurrentTime());
 	LastUserInteractionTimeForThrottling = LastUserInteractionTime;
-	PressedMouseButtons.Remove( MouseEvent.GetEffectingButton() );
+
+	const bool bIsCursorUser = MouseEvent.GetUserIndex() == CursorUserIndex;
+	if (bIsCursorUser)
+	{
+		PressedMouseButtons.Remove(MouseEvent.GetEffectingButton());
+	}
 
 	// Input preprocessors get the first chance at the input
 	if (InputPreProcessors.HandleMouseButtonUpEvent(*this, MouseEvent))
@@ -6248,11 +5102,9 @@ bool FSlateApplication::ProcessMouseButtonUpEvent( const FPointerEvent& MouseEve
 	FWidgetPath EmptyPath;
 	const bool bHandled = RoutePointerUpEvent( EmptyPath, MouseEvent ).IsEventHandled();
 
-	if ( PressedMouseButtons.Num() == 0 )
+	if ( bIsCursorUser && PressedMouseButtons.Num() == 0 )
 	{
-		// Release Capture
 		PlatformApplication->SetCapture( nullptr );
-//		MouseCaptor.InvalidateCaptureForPointer(MouseEvent.GetUserIndex(), MouseEvent.GetPointerIndex());
 	}
 
 	return bHandled;
@@ -6282,7 +5134,7 @@ bool FSlateApplication::ProcessMouseWheelOrGestureEvent( const FPointerEvent& In
 {
 	SCOPE_CYCLE_COUNTER(STAT_ProcessMouseWheelGesture);
 
-	QueueSynthesizedMouseMove();
+	QueueSynthesizedMouseMove(InWheelEvent);
 
 	bool bShouldProcessEvent = false;
 
@@ -6328,11 +5180,13 @@ FReply FSlateApplication::RouteMouseWheelOrGestureEvent(const FWidgetPath& Widge
 	TScopeCounter<int32> BeginInput(ProcessingInput);
 
 	FWidgetPath MouseCaptorPath;
-	if ( MouseCaptor.HasCaptureForPointerIndex(InWheelEvent.GetUserIndex(), InWheelEvent.GetPointerIndex()) )
+
+	TSharedRef<FSlateUser> User = GetOrCreateUser(InWheelEvent);
+	if (User->HasCapture(InWheelEvent.GetPointerIndex()))
 	{
-		MouseCaptorPath = MouseCaptor.ToWidgetPath(FWeakWidgetPath::EInterruptedPathHandling::ReturnInvalid, &InWheelEvent);
+		MouseCaptorPath = User->GetCaptorPath(InWheelEvent.GetPointerIndex(), FWeakWidgetPath::EInterruptedPathHandling::ReturnInvalid, &InWheelEvent);
 	}
-
+		
 	const FWidgetPath& EventPath = MouseCaptorPath.IsValid() ? MouseCaptorPath : WidgetsUnderPointer;
 
 	// Switch worlds widgets in the current path
@@ -6453,11 +5307,8 @@ bool FSlateApplication::ProcessMouseMoveEvent( const FPointerEvent& MouseEvent,
 	{
 		QUICK_SCOPE_CYCLE_COUNTER(STAT_ProcessMouseMove_Tooltip);
 
-		QueueSynthesizedMouseMove();
-
-		// Detecting a mouse move of zero delta is our way of filtering out synthesized move events
-		const bool AllowSpawningOfToolTips = true;
-		UpdateToolTip( AllowSpawningOfToolTips );
+		QueueSynthesizedMouseMove(MouseEvent);
+		GetOrCreateUser(MouseEvent)->UpdateTooltip(MenuStack, /*bCanSpawnNewTooltip =*/true);
 		
 		// Guard against synthesized mouse moves and only track user interaction if the cursor pos changed
 		SetLastUserInteractionTime(this->GetCurrentTime());
@@ -6484,7 +5335,7 @@ bool FSlateApplication::ProcessMouseMoveEvent( const FPointerEvent& MouseEvent,
 
 bool FSlateApplication::OnCursorSet()
 {
-	bQueryCursorRequested = true;
+	GetOrCreateUser(CursorUserIndex)->RequestCursorQuery();
 	return true;
 }
 
@@ -6495,11 +5346,11 @@ void FSlateApplication::NavigateToWidget(const uint32 UserIndex, const TSharedPt
 		FWidgetPath NavigationSourceWP;
 		if (NavigationSource == ENavigationSource::WidgetUnderCursor)
 		{
-			NavigationSourceWP = LocateWindowUnderMouse(GetCursorPos(), GetInteractiveTopLevelWindows());
+			NavigationSourceWP = LocateWindowUnderMouse(GetOrCreateUser(UserIndex)->GetCursorPosition(), GetInteractiveTopLevelWindows());
 		}
-		else if (const FSlateUser* User = GetOrCreateUser(UserIndex))
+		else
 		{
-			NavigationSourceWP = User->FocusWidgetPathWeak.ToWidgetPath();
+			NavigationSourceWP = *GetOrCreateUser(UserIndex)->GetFocusPath();
 		}
 
 		if (NavigationSourceWP.IsValid())
@@ -6702,10 +5553,15 @@ bool FSlateApplication::OnTouchGesture( EGestureEvent GestureType, const FVector
 	return ProcessMouseWheelOrGestureEvent( MouseWheelEvent, &GestureEvent );
 }
 
+bool ValidateTouchIndex(int32 TouchIndex)
+{
+	return true;
+}
+
 bool FSlateApplication::OnTouchStarted( const TSharedPtr< FGenericWindow >& PlatformWindow, const FVector2D& Location, float Force, int32 TouchIndex, int32 ControllerId )
 {
 	// Don't process touches that overlap or surpass with the cursor pointer index.
-	if (TouchIndex >= ((int32)ETouchIndex::CursorPointerIndex))
+	if (TouchIndex >= (int32)ETouchIndex::CursorPointerIndex)
 	{
 #if !(UE_BUILD_SHIPPING || UE_BUILD_TEST)
 		// Only log when the touch starts, we don't want to spam the logs.
@@ -6723,48 +5579,32 @@ bool FSlateApplication::OnTouchStarted( const TSharedPtr< FGenericWindow >& Plat
 		true);
 	ProcessTouchStartedEvent( PlatformWindow, PointerEvent );
 
-	if ( FSlateUser* User = GetUser(ControllerId) )
-	{
-		User->GestureDetector.OnTouchStarted(TouchIndex, Location);
-	}
-
 	return true;
 }
 
 void FSlateApplication::ProcessTouchStartedEvent( const TSharedPtr< FGenericWindow >& PlatformWindow, const FPointerEvent& InTouchEvent )
 {
-	// Add or Update the entry if the finger has been added to the surface.
-	FUserAndPointer UserAndIndex(InTouchEvent.GetUserIndex(), InTouchEvent.GetPointerIndex());
-	PointerIndexLastPositionMap.Add(UserAndIndex, InTouchEvent.GetScreenSpacePosition());
-
+	GetOrCreateUser(InTouchEvent)->GetGestureDetector().OnTouchStarted(InTouchEvent.GetPointerIndex(), InTouchEvent.GetScreenSpacePosition());
 	ProcessMouseButtonDownEvent(PlatformWindow, InTouchEvent);
 }
 
 bool FSlateApplication::OnTouchMoved( const FVector2D& Location, float Force, int32 TouchIndex, int32 ControllerId )
 {
 	// Don't process touches that overlap or surpass with the cursor pointer index.
-	if (TouchIndex >= ((int32)ETouchIndex::CursorPointerIndex))
+	if (TouchIndex >= (int32)ETouchIndex::CursorPointerIndex)
 	{
 		return false;
 	}
 
-	const FVector2D* LastLocationPtr = PointerIndexLastPositionMap.Find(FUserAndPointer(ControllerId, TouchIndex));
-	const FVector2D LastLocation = LastLocationPtr ? *LastLocationPtr : Location;
-
 	FPointerEvent PointerEvent(
 		ControllerId,
 		TouchIndex,
 		Location,
-		LastLocation,
+		GetOrCreateUser(ControllerId)->GetPreviousPointerPosition(TouchIndex),
 		Force,
 		true);
 	ProcessTouchMovedEvent(PointerEvent);
 
-	if ( FSlateUser* User = GetUser(ControllerId) )
-	{
-		User->GestureDetector.OnTouchMoved(TouchIndex, Location);
-	}
-
 	return true;
 }
 
@@ -6790,14 +5630,14 @@ bool FSlateApplication::OnTouchEnded( const FVector2D& Location, int32 TouchInde
 		true);
 	ProcessTouchEndedEvent(PointerEvent);
 
-	if ( FSlateUser* User = GetUser(ControllerId) )
-	{
-		User->GestureDetector.OnTouchEnded(TouchIndex, Location);
-	}
-
 	return true;
 }
 
+void FSlateApplication::ProcessTouchEndedEvent(const FPointerEvent& PointerEvent)
+{
+	ProcessMouseButtonUpEvent(PointerEvent);
+}
+
 bool FSlateApplication::OnTouchForceChanged(const FVector2D& Location, float Force, int32 TouchIndex, int32 ControllerId)
 {
 	// Don't process touches that overlap or surpass with the cursor pointer index.
@@ -6817,11 +5657,6 @@ bool FSlateApplication::OnTouchForceChanged(const FVector2D& Location, float For
 		false);
 	ProcessTouchMovedEvent(PointerEvent);
 
-	if (FSlateUser* User = GetUser(ControllerId))
-	{
-		User->GestureDetector.OnTouchMoved(TouchIndex, Location);
-	}
-
 	return true;
 
 }
@@ -6834,25 +5669,17 @@ bool FSlateApplication::OnTouchFirstMove(const FVector2D& Location, float Force,
 		return false;
 	}
 
-	const FVector2D* LastLocationPtr = PointerIndexLastPositionMap.Find(FUserAndPointer(ControllerId, TouchIndex));
-	const FVector2D LastLocation = LastLocationPtr ? *LastLocationPtr : Location;
-
 	FPointerEvent PointerEvent(
 		ControllerId,
 		TouchIndex,
 		Location,
-		LastLocation,
+		GetOrCreateUser(ControllerId)->GetPreviousPointerPosition(TouchIndex),
 		Force,
 		true,
 		false,
 		true);
 	ProcessTouchMovedEvent(PointerEvent);
 
-	if (FSlateUser* User = GetUser(ControllerId))
-	{
-		User->GestureDetector.OnTouchMoved(TouchIndex, Location);
-	}
-
 	return true;
 
 }
@@ -6864,14 +5691,6 @@ void FSlateApplication::ShouldSimulateGesture(EGestureEvent Gesture, bool bEnabl
 	SimulateGestures[(uint8)Gesture] = bEnable;
 }
 
-void FSlateApplication::ProcessTouchEndedEvent( const FPointerEvent& PointerEvent )
-{
-	ProcessMouseButtonUpEvent(PointerEvent);
-
-	// Remove the entry if the finger has been removed from the surface.
-	PointerIndexLastPositionMap.Remove(FUserAndPointer(PointerEvent.GetUserIndex(), PointerEvent.GetPointerIndex()));
-}
-
 bool FSlateApplication::OnMotionDetected(const FVector& Tilt, const FVector& RotationRate, const FVector& Gravity, const FVector& Acceleration, int32 ControllerId)
 {
 	FMotionEvent MotionEvent( 
@@ -6888,17 +5707,13 @@ bool FSlateApplication::OnMotionDetected(const FVector& Tilt, const FVector& Rot
 
 void FSlateApplication::ProcessMotionDetectedEvent( const FMotionEvent& MotionEvent )
 {
-	QueueSynthesizedMouseMove();
+	QueueSynthesizedMouseMove(MotionEvent);
 	SetLastUserInteractionTime(this->GetCurrentTime());
 	
-	if ( FSlateUser* User = GetOrCreateUser(MotionEvent.GetUserIndex()) )
+	if (!InputPreProcessors.HandleMotionDetectedEvent(*this, MotionEvent))
 	{
-		if (InputPreProcessors.HandleMotionDetectedEvent(*this, MotionEvent))
-		{
-			return;
-		}
-
-		if ( User->HasValidFocusPath() )
+		TSharedRef<FSlateUser> User = GetOrCreateUser(MotionEvent);
+		if (User->HasValidFocusPath())
 		{
 			/* Get the controller focus target for this user */
 			TSharedRef<FWidgetPath> EventPathRef = User->GetFocusPath();
@@ -6906,10 +5721,10 @@ void FSlateApplication::ProcessMotionDetectedEvent( const FMotionEvent& MotionEv
 
 			FScopedSwitchWorldHack SwitchWorld(EventPath);
 
-			FReply Reply = FEventRouter::Route<FReply>(this, FEventRouter::FBubblePolicy(EventPath), MotionEvent, [] (const FArrangedWidget& SomeWidget, const FMotionEvent& InMotionEvent)
-			{
-				return SomeWidget.Widget->OnMotionDetected(SomeWidget.Geometry, InMotionEvent);
-			});
+			FReply Reply = FEventRouter::Route<FReply>(this, FEventRouter::FBubblePolicy(EventPath), MotionEvent, [](const FArrangedWidget& SomeWidget, const FMotionEvent& InMotionEvent)
+				{
+					return SomeWidget.Widget->OnMotionDetected(SomeWidget.Geometry, InMotionEvent);
+				});
 		}
 	}
 }
@@ -7091,7 +5906,7 @@ bool FSlateApplication::ProcessWindowActivatedEvent( const FWindowActivateEvent&
 
 	if ( ActivateEvent.GetActivationType() != FWindowActivateEvent::EA_Deactivate )
 	{
-		ReleaseMouseCapture();
+		ReleaseAllPointerCapture();
 
 		const bool bActivatedByMouse = ActivateEvent.GetActivationType() == FWindowActivateEvent::EA_ActivateByMouse;
 		
@@ -7103,7 +5918,7 @@ bool FSlateApplication::ProcessWindowActivatedEvent( const FWindowActivateEvent&
 		
 		// Widgets that happen to be under the mouse need to update if activation changes
 		// This also serves as a force redraw which is needed when restoring a window that was previously inactive.
-		QueueSynthesizedMouseMove();
+		ForEachUser([](FSlateUser& User) { User.QueueSyntheticCursorMove(); });
 
 		// NOTE: The window is brought to front even when a modal window is active and this is not the modal window one of its children 
 		// The reason for this is so that the Slate window order is in sync with the OS window order when a modal window is open.  This is important so that when the modal window closes the proper window receives input from Slate.
@@ -7143,9 +5958,6 @@ bool FSlateApplication::ProcessWindowActivatedEvent( const FWindowActivateEvent&
 			// This allows us to force ourselves back to the correct resolution after alt-tabbing out of a fullscreen
 			// window and then going back in again.
 			Renderer->RestoreSystemResolution(ActivateEvent.GetAffectedWindow());
-
-			// Synthesize mouse move to resume rendering in the next tick if Slate is sleeping
-			QueueSynthesizedMouseMove();
 		}
 		else
 		{
@@ -7231,16 +6043,20 @@ void FSlateApplication::ProcessApplicationActivationEvent(bool InAppActivated)
 		DismissAllMenus();
 
 		// Close tool-tips
-		CloseToolTip();
-
+		ForEachUser([](FSlateUser& User) {
+			User.CloseTooltip();
+		});
+		
 		// No slate window is active when our entire app becomes inactive
 		bSlateWindowActive = false;
 
-		// If we have a slate-only drag-drop occurring, stop the drag drop.
-		if (IsDragDropping() && !DragDropContent->IsExternalOperation())
-		{
-			ResetDragDropState();
-		}
+		// Stop all slate-only drag-drop operations
+		ForEachUser([](FSlateUser& User) { 
+			if (User.IsDragDropping() && User.GetDragDropContent()->IsExternalOperation())
+			{
+				User.ResetDragDropContent();
+			}
+		});
 
 		// Clear the pressed buttons when we deactivate the application, the button state can no longer be trusted.
 		PressedMouseButtons.Reset();
@@ -7248,7 +6064,7 @@ void FSlateApplication::ProcessApplicationActivationEvent(bool InAppActivated)
 	else
 	{
 		//Ensure that slate ticks/renders next frame
-		QueueSynthesizedMouseMove();
+		ForEachUser([](FSlateUser& User) { User.QueueSyntheticCursorMove(); });
 
 		//TODO Requery pushed button state?
 	}
@@ -7443,7 +6259,7 @@ bool FSlateApplication::ProcessDragEnterEvent( TSharedRef<SWindow> WindowEntered
 	FReply TriggerDragDropReply = FReply::Handled().BeginDragDrop( DragDropEvent.GetOperation().ToSharedRef() );
 	ProcessReply( WidgetsUnderCursor, TriggerDragDropReply, &WidgetsUnderCursor, &DragDropEvent );
 
-	PointerIndexLastPositionMap.Add(FUserAndPointer(DragDropEvent.GetUserIndex(), DragDropEvent.GetPointerIndex()), DragDropEvent.GetScreenSpacePosition());
+	GetOrCreateUser(DragDropEvent)->UpdatePointerPosition(DragDropEvent);
 
 	return true;
 }
@@ -7452,7 +6268,8 @@ EDropEffect::Type FSlateApplication::OnDragOver( const TSharedPtr< FGenericWindo
 {
 	EDropEffect::Type Result = EDropEffect::None;
 
-	if ( IsDragDropping() )
+	TSharedRef<FSlateUser> CursorUser = GetOrCreateUser(CursorUserIndex);
+	if (CursorUser->IsDragDropping())
 	{
 		bool MouseMoveHandled = true;
 		FVector2D CursorMovementDelta( 0, 0 );
@@ -7493,14 +6310,14 @@ EDropEffect::Type FSlateApplication::OnDragOver( const TSharedPtr< FGenericWindo
 
 void FSlateApplication::OnDragLeave( const TSharedPtr< FGenericWindow >& Window )
 {
-	ResetDragDropState();
+	GetOrCreateUser(CursorUserIndex)->ResetDragDropContent();
 }
 
 EDropEffect::Type FSlateApplication::OnDragDrop(const TSharedPtr< FGenericWindow >& Window)
 {
 	EDropEffect::Type Result = EDropEffect::None;
 
-	if (IsDragDropping())
+	if (GetOrCreateUser(CursorUserIndex)->IsDragDropping())
 	{
 		FPointerEvent MouseEvent(
 			CursorPointerIndex,
@@ -7600,18 +6417,44 @@ void FSlateApplication::SetWidgetReflector(const TSharedRef<IWidgetReflector>& W
 	WidgetReflectorPtr = WidgetReflector;
 }
 
+bool FSlateApplication::IsDragDropping() const
+{
+	TSharedPtr<const FSlateUser> CursorUser = GetUser(CursorUserIndex);
+	return CursorUser && CursorUser->IsDragDropping();
+}
+
+bool FSlateApplication::IsDragDroppingAffected(const FPointerEvent& InPointerEvent) const
+{
+	TSharedPtr<const FSlateUser> CursorUser = GetUser(CursorUserIndex);
+	return CursorUser && CursorUser->IsDragDroppingAffected(InPointerEvent);
+}
+
+TSharedPtr<FDragDropOperation> FSlateApplication::GetDragDroppingContent() const
+{
+	TSharedPtr<const FSlateUser> CursorUser = GetUser(CursorUserIndex);
+	return CursorUser ? CursorUser->GetDragDropContent() : nullptr;
+}
+
+void FSlateApplication::CancelDragDrop()
+{
+	if (TSharedPtr<FSlateUser> CursorUser = GetUser(CursorUserIndex))
+	{
+		CursorUser->CancelDragDrop();
+	}
+}
+
 void FSlateApplication::NavigateFromWidgetUnderCursor(const uint32 InUserIndex, EUINavigation InNavigationType, TSharedRef<SWindow> InWindow)
 {
 	if (InNavigationType != EUINavigation::Invalid)
 	{
-		FWidgetPath PathToLocatedWidget = LocateWidgetInWindow(GetCursorPos(), InWindow, false);
+		FWidgetPath PathToLocatedWidget = LocateWidgetInWindow(GetOrCreateUser(InUserIndex)->GetCursorPosition(), InWindow, false);
 		if (PathToLocatedWidget.IsValid())
 		{
 			TSharedPtr<SWidget> WidgetToNavFrom = PathToLocatedWidget.Widgets.Last().Widget;
 
 			if (WidgetToNavFrom.IsValid())
 			{
-				FSlateApplication::Get().ProcessReply(PathToLocatedWidget, FReply::Handled().SetNavigation(InNavigationType, ENavigationGenesis::User, ENavigationSource::WidgetUnderCursor), &PathToLocatedWidget, nullptr, InUserIndex);
+				ProcessReply(PathToLocatedWidget, FReply::Handled().SetNavigation(InNavigationType, ENavigationGenesis::User, ENavigationSource::WidgetUnderCursor), &PathToLocatedWidget, nullptr, InUserIndex);
 			}
 		}
 	}
diff --git a/Engine/Source/Runtime/Slate/Private/Framework/Application/SlateUser.cpp b/Engine/Source/Runtime/Slate/Private/Framework/Application/SlateUser.cpp
new file mode 100644
--- /dev/null
+++ b/Engine/Source/Runtime/Slate/Private/Framework/Application/SlateUser.cpp
@@ -0,0 +1,1268 @@
+// Copyright 1998-2019 Epic Games, Inc. All Rights Reserved.
+
+#include "Framework/Application/SlateUser.h"
+#include "Framework/Application/SlateApplication.h"
+#include "Widgets/SWindow.h"
+#include "Widgets/SWeakWidget.h"
+
+#if PLATFORM_WINDOWS || PLATFORM_HOLOLENS
+// Needed to be able to use RECT
+#include "Windows/WindowsHWrapper.h"
+#endif
+
+DECLARE_CYCLE_STAT(TEXT("QueryCursor"), STAT_SlateQueryCursor, STATGROUP_Slate);
+DECLARE_CYCLE_STAT(TEXT("Update Tooltip Time"), STAT_SlateUpdateTooltip, STATGROUP_Slate);
+
+namespace SlateDefs
+{
+	// How far tooltips should be offset from the mouse cursor position, in pixels
+	static const FVector2D TooltipOffsetFromMouse(12.0f, 8.0f);
+
+	// How far tooltips should be pushed out from a force field border, in pixels
+	static const FVector2D TooltipOffsetFromForceField(4.0f, 3.0f);
+}
+
+static bool bEnableSyntheticCursorMoves = true;
+FAutoConsoleVariableRef CVarEnableSyntheticCursorMoves(
+	TEXT("Slate.EnableSyntheticCursorMoves"),
+	bEnableSyntheticCursorMoves,
+	TEXT("")
+);
+
+static bool bEnableCursorQueries = true;
+FAutoConsoleVariableRef CVarEnableCursorQueries(
+	TEXT("Slate.EnableCursorQueries"),
+	bEnableCursorQueries,
+	TEXT(""));
+
+static float TooltipSummonDelay = 0.15f;
+FAutoConsoleVariableRef CVarTooltipSummonDelay(
+	TEXT("Slate.TooltipSummonDelay"),
+	TooltipSummonDelay,
+	TEXT("Delay in seconds before a tooltip is displayed near the mouse cursor when hovering over widgets that supply tooltip data."));
+	
+static float TooltipIntroDuration = 0.1f;
+FAutoConsoleVariableRef CVarTooltipIntroDuration(
+	TEXT("Slate.TooltipIntroDuration"),
+	TooltipIntroDuration,
+	TEXT("How long it takes for a tooltip to animate into view, in seconds."));
+
+//////////////////////////////////////////////////////////////////////////
+// FSlateVirtualUserHandle
+//////////////////////////////////////////////////////////////////////////
+
+FSlateVirtualUserHandle::FSlateVirtualUserHandle(int32 InUserIndex, int32 InVirtualUserIndex)
+	: UserIndex(InUserIndex)
+	, VirtualUserIndex(InVirtualUserIndex)
+{
+}
+
+FSlateVirtualUserHandle::~FSlateVirtualUserHandle()
+{
+	if (FSlateApplication::IsInitialized())
+	{
+		FSlateApplication::Get().UnregisterUser(UserIndex);
+	}
+}
+
+//////////////////////////////////////////////////////////////////////////
+// FActiveTooltipInfo
+//////////////////////////////////////////////////////////////////////////
+
+void FSlateUser::FActiveTooltipInfo::Reset()
+{
+	if (SourceWidget.IsValid())
+	{
+		SourceWidget.Pin()->OnToolTipClosing();
+	}
+
+	if (Tooltip.IsValid())
+	{
+		Tooltip.Pin()->OnClosed();
+	}
+
+	Tooltip.Reset();
+	SourceWidget.Reset();
+	OffsetDirection = ETooltipOffsetDirection::Undetermined;
+}
+
+//////////////////////////////////////////////////////////////////////////
+// FSlateUser
+//////////////////////////////////////////////////////////////////////////
+
+TSharedRef<FSlateUser> FSlateUser::Create(int32 InUserIndex, TSharedPtr<ICursor> InCursor)
+{
+	return MakeShareable(new FSlateUser(InUserIndex, InCursor));
+}
+
+FSlateUser::FSlateUser(int32 InUserIndex, TSharedPtr<ICursor> InCursor)
+	: UserIndex(InUserIndex)
+	, Cursor(InCursor)
+{
+	UE_LOG(LogSlate, Log, TEXT("New Slate User Created.  User Index %d, Is Virtual User: %d"), UserIndex, IsVirtualUser());
+
+	PointerPositionsByIndex.Add(FSlateApplication::CursorPointerIndex, FVector2D::ZeroVector);
+	PreviousPointerPositionsByIndex.Add(FSlateApplication::CursorPointerIndex, FVector2D::ZeroVector);
+}
+
+FSlateUser::~FSlateUser()
+{
+	UE_LOG(LogSlate, Log, TEXT("Slate User Destroyed.  User Index %d, Is Virtual User: %d"), UserIndex, IsVirtualUser());
+}
+
+TSharedPtr<SWidget> FSlateUser::GetFocusedWidget() const
+{
+	return WeakFocusPath.IsValid() ? WeakFocusPath.GetLastWidget().Pin() : nullptr;
+}
+
+TOptional<EFocusCause> FSlateUser::HasFocus(TSharedPtr<const SWidget> Widget) const
+{
+	return Widget && GetFocusedWidget() == Widget ? FocusCause : TOptional<EFocusCause>();
+}
+
+bool FSlateUser::ShouldShowFocus(TSharedPtr<const SWidget> Widget) const
+{
+	return HasFocus(Widget).IsSet() && bShowFocus;
+}
+
+bool FSlateUser::HasFocusedDescendants(TSharedRef<const SWidget> Widget) const
+{
+	return WeakFocusPath.IsValid() && WeakFocusPath.GetLastWidget().Pin() != Widget && WeakFocusPath.ContainsWidget(Widget);
+}
+
+bool FSlateUser::SetFocus(const TSharedRef<SWidget>& WidgetToFocus, EFocusCause ReasonFocusIsChanging)
+{
+	return FSlateApplication::Get().SetUserFocus(UserIndex, WidgetToFocus, ReasonFocusIsChanging);
+}
+
+void FSlateUser::ClearFocus(EFocusCause ReasonFocusIsChanging)
+{
+	FSlateApplication::Get().ClearUserFocus(UserIndex, ReasonFocusIsChanging);
+}
+
+bool FSlateUser::HasAnyCapture() const
+{
+	return PointerCaptorPathsByIndex.Num() > 0;
+}
+
+bool FSlateUser::HasCursorCapture() const
+{
+	return HasCapture(FSlateApplication::CursorPointerIndex);
+}
+
+bool FSlateUser::HasCapture(uint32 PointerIndex) const
+{
+	const FWeakWidgetPath* CaptorPath = PointerCaptorPathsByIndex.Find(PointerIndex);
+	return CaptorPath && CaptorPath->IsValid();
+}
+
+bool FSlateUser::DoesWidgetHaveAnyCapture(TSharedPtr<const SWidget> Widget) const
+{
+	for (const auto& IndexPathPair : PointerCaptorPathsByIndex)
+	{
+		if (IndexPathPair.Value.GetLastWidget().Pin() == Widget)
+		{
+			return true;
+		}
+	}
+	return false;
+}
+
+bool FSlateUser::DoesWidgetHaveCursorCapture(TSharedPtr<const SWidget> Widget) const
+{
+	return DoesWidgetHaveCapture(Widget, FSlateApplication::CursorPointerIndex);
+}
+
+bool FSlateUser::DoesWidgetHaveCapture(TSharedPtr<const SWidget> Widget, uint32 PointerIndex) const
+{
+	const FWeakWidgetPath* CaptorPath = PointerCaptorPathsByIndex.Find(PointerIndex);
+	return CaptorPath && CaptorPath->GetLastWidget().Pin() == Widget;
+}
+
+bool FSlateUser::SetCursorCaptor(TSharedRef<const SWidget> Widget, const FWidgetPath& EventPath)
+{
+	return SetPointerCaptor(FSlateApplication::CursorPointerIndex, Widget, EventPath);
+}
+
+bool FSlateUser::SetPointerCaptor(uint32 PointerIndex, TSharedRef<const SWidget> Widget, const FWidgetPath& EventPath)
+{
+	// Bail on any current capture we may have for this pointer before trying to establish the new captor
+	ReleaseCapture(PointerIndex);
+	
+	if (ensureMsgf(EventPath.IsValid(), TEXT("An unknown widget is attempting to set capture to %s"), *Widget->ToString()))
+	{
+		FWidgetPath NewCaptorPath = EventPath.GetPathDownTo(Widget);
+		if (!NewCaptorPath.IsValid() || NewCaptorPath.GetLastWidget() != Widget)
+		{
+			// The target widget wasn't in the given event path, so try searching for it from the root of the event path
+			NewCaptorPath = EventPath.GetPathDownTo(EventPath.GetWindow());
+			NewCaptorPath.ExtendPathTo(FWidgetMatcher(Widget));
+		}
+		
+		if (NewCaptorPath.IsValid() && NewCaptorPath.GetLastWidget() == Widget)
+		{
+			PointerCaptorPathsByIndex.Add(PointerIndex, NewCaptorPath);
+
+#if WITH_SLATE_DEBUGGING
+			FSlateDebugging::BroadcastMouseCapture(UserIndex, PointerIndex, Widget);
+#endif
+			return true;
+		}
+	}
+	return false;
+}
+
+void FSlateUser::ReleaseAllCapture()
+{
+	TArray<uint32> CapturedPointerIndices;
+	PointerCaptorPathsByIndex.GenerateKeyArray(CapturedPointerIndices);
+	for (uint32 CapturedPointerIdx : CapturedPointerIndices)
+	{
+		ReleaseCapture(CapturedPointerIdx);
+	}
+}
+
+void FSlateUser::ReleaseCursorCapture()
+{
+	ReleaseCapture(FSlateApplication::CursorPointerIndex);
+}
+
+void FSlateUser::ReleaseCapture(uint32 PointerIndex)
+{
+	if (const FWeakWidgetPath* CaptorPath = PointerCaptorPathsByIndex.Find(PointerIndex))
+	{
+		WidgetsUnderPointerLastEventByIndex.Add(PointerIndex, *CaptorPath);
+
+		if (TSharedPtr<SWidget> CaptorWidget = CaptorPath->GetLastWidget().Pin())
+		{
+			CaptorWidget->OnMouseCaptureLost(FCaptureLostEvent(UserIndex, PointerIndex));
+#if WITH_SLATE_DEBUGGING
+			FSlateDebugging::BroadcastMouseCaptureLost(UserIndex, PointerIndex, CaptorWidget);
+#endif
+		}
+
+		CaptorPath = nullptr;
+		PointerCaptorPathsByIndex.Remove(PointerIndex);
+
+		if (PointerIndex == FSlateApplication::CursorPointerIndex)
+		{
+			// If cursor capture changes, we should refresh the cursor state.
+			bQueryCursorRequested = true;
+		}
+	}
+}
+
+TArray<FWidgetPath> FSlateUser::GetCaptorPaths()
+{
+	TArray<FWidgetPath> AllPaths;
+	for (const auto& IndexPathPair : PointerCaptorPathsByIndex)
+	{
+		AllPaths.Emplace(GetCaptorPath(IndexPathPair.Key));
+	}
+
+	return AllPaths;
+}
+
+FWidgetPath FSlateUser::GetCursorCaptorPath(FWeakWidgetPath::EInterruptedPathHandling::Type InterruptedPathHandling, const FPointerEvent* PointerEvent)
+{
+	return GetCaptorPath(FSlateApplication::CursorPointerIndex, InterruptedPathHandling, PointerEvent);
+}
+
+FWidgetPath FSlateUser::GetCaptorPath(uint32 PointerIndex, FWeakWidgetPath::EInterruptedPathHandling::Type InterruptedPathHandling, const FPointerEvent* PointerEvent)
+{
+	FWidgetPath CaptorPath;
+	if (const FWeakWidgetPath* WeakCaptorPath = PointerCaptorPathsByIndex.Find(PointerIndex))
+	{
+		if (WeakCaptorPath->ToWidgetPath(CaptorPath, InterruptedPathHandling, PointerEvent) == FWeakWidgetPath::EPathResolutionResult::Truncated)
+		{
+			// The path was truncated, meaning it's not actually valid anymore, so we want to clear our entry for it out immediately
+			WeakCaptorPath = nullptr;
+			ReleaseCapture(PointerIndex);
+		}
+	}
+	return CaptorPath;
+}
+
+FWeakWidgetPath FSlateUser::GetWeakCursorCapturePath() const
+{
+	return GetWeakCapturePath(FSlateApplication::CursorPointerIndex);
+}
+
+FWeakWidgetPath FSlateUser::GetWeakCapturePath(uint32 PointerIndex) const
+{
+	const FWeakWidgetPath* WeakCaptorPath = PointerCaptorPathsByIndex.Find(PointerIndex);
+	return WeakCaptorPath ? *WeakCaptorPath : FWeakWidgetPath();
+}
+
+TArray<TSharedRef<SWidget>> FSlateUser::GetCaptorWidgets() const
+{
+	TArray<TSharedRef<SWidget>> AllCaptors;
+	AllCaptors.Reserve(PointerCaptorPathsByIndex.Num());
+	for (const auto& IndexPathPair : PointerCaptorPathsByIndex)
+	{
+		if (TSharedPtr<SWidget> CaptorWidget = IndexPathPair.Value.GetLastWidget().Pin())
+		{
+			AllCaptors.Add(CaptorWidget.ToSharedRef());
+		}
+	}
+	return AllCaptors;
+}
+
+TSharedPtr<SWidget> FSlateUser::GetCursorCaptor() const
+{
+	return GetPointerCaptor(FSlateApplication::CursorPointerIndex);
+}
+
+TSharedPtr<SWidget> FSlateUser::GetPointerCaptor(uint32 PointerIndex) const
+{
+	const FWeakWidgetPath* WeakCaptorPath = PointerCaptorPathsByIndex.Find(PointerIndex);
+	return WeakCaptorPath ? WeakCaptorPath->GetLastWidget().Pin() : nullptr;
+}
+
+void FSlateUser::SetCursorPosition(int32 PosX, int32 PosY)
+{
+	SetPointerPosition(FSlateApplication::CursorPointerIndex, PosX, PosY);
+}
+
+void FSlateUser::SetCursorPosition(const FVector2D& NewCursorPos)
+{
+	SetCursorPosition((int32)NewCursorPos.X, (int32)NewCursorPos.Y);
+}
+
+void FSlateUser::SetPointerPosition(uint32 PointerIndex, int32 PosX, int32 PosY)
+{
+	if (Cursor && PointerIndex == FSlateApplication::CursorPointerIndex)
+	{
+		Cursor->SetPosition(PosX, PosY);
+	}
+
+	UpdatePointerPosition(PointerIndex, FVector2D(PosX, PosY));
+}
+
+void FSlateUser::SetPointerPosition(uint32 PointerIndex, const FVector2D& NewPointerPos)
+{
+	SetPointerPosition(PointerIndex, (int32)NewPointerPos.X, (int32)NewPointerPos.Y);
+}
+
+FVector2D FSlateUser::GetCursorPosition() const
+{
+	return GetPointerPosition(FSlateApplication::CursorPointerIndex);
+}
+
+FVector2D FSlateUser::GetPreviousCursorPosition() const
+{
+	return GetPreviousPointerPosition(FSlateApplication::CursorPointerIndex);
+}
+
+FVector2D FSlateUser::GetPointerPosition(uint32 PointerIndex) const
+{
+	if (Cursor && PointerIndex == FSlateApplication::CursorPointerIndex)
+	{
+		return Cursor->GetPosition();
+	}
+	const FVector2D* FoundPosition = PointerPositionsByIndex.Find(PointerIndex);
+	return FoundPosition ? *FoundPosition : FVector2D::ZeroVector;
+}
+
+FVector2D FSlateUser::GetPreviousPointerPosition(uint32 PointerIndex) const
+{
+	const FVector2D* FoundPosition = PreviousPointerPositionsByIndex.Find(PointerIndex);
+	return FoundPosition ? *FoundPosition : GetPointerPosition(PointerIndex);
+}
+
+bool FSlateUser::IsWidgetUnderCursor(TSharedPtr<const SWidget> Widget) const
+{
+	return IsWidgetUnderPointer(Widget, FSlateApplication::CursorPointerIndex);
+}
+
+bool FSlateUser::IsWidgetUnderPointer(TSharedPtr<const SWidget> Widget, uint32 PointerIndex) const
+{
+	const FWeakWidgetPath* WidgetsUnderPointer = WidgetsUnderPointerLastEventByIndex.Find(PointerIndex);
+	return Widget && WidgetsUnderPointer && WidgetsUnderPointer->ContainsWidget(Widget.ToSharedRef());
+}
+
+bool FSlateUser::IsWidgetUnderAnyPointer(TSharedPtr<const SWidget> Widget) const
+{
+	if (Widget)
+	{
+		for (const auto& IndexPathPair : WidgetsUnderPointerLastEventByIndex)
+		{
+			if (IndexPathPair.Value.ContainsWidget(Widget.ToSharedRef()))
+			{
+				return true;
+			}
+		}
+	}
+	return false;
+}
+
+bool FSlateUser::IsWidgetDirectlyUnderCursor(TSharedPtr<const SWidget> Widget) const
+{
+	return IsWidgetDirectlyUnderPointer(Widget, FSlateApplication::CursorPointerIndex);
+}
+
+bool FSlateUser::IsWidgetDirectlyUnderPointer(TSharedPtr<const SWidget> Widget, uint32 PointerIndex) const
+{
+	const FWeakWidgetPath* WidgetsUnderPointer = WidgetsUnderPointerLastEventByIndex.Find(PointerIndex);
+	return WidgetsUnderPointer && WidgetsUnderPointer->IsValid() && WidgetsUnderPointer->GetLastWidget().Pin() == Widget;
+}
+
+bool FSlateUser::IsWidgetDirectlyUnderAnyPointer(TSharedPtr<const SWidget> Widget) const
+{
+	for (const auto& IndexPathPair : WidgetsUnderPointerLastEventByIndex)
+	{
+		if (IndexPathPair.Value.IsValid() && IndexPathPair.Value.GetLastWidget().Pin() == Widget)
+		{
+			return true;
+		}
+	}
+	return false;
+}
+
+FWeakWidgetPath FSlateUser::GetLastWidgetsUnderCursor() const
+{
+	return GetLastWidgetsUnderPointer(FSlateApplication::CursorPointerIndex);
+}
+
+FWeakWidgetPath FSlateUser::GetLastWidgetsUnderPointer(uint32 PointerIndex) const
+{
+	return WidgetsUnderPointerLastEventByIndex.FindRef(PointerIndex);
+}
+
+bool FSlateUser::IsDragDropping() const
+{
+	return DragDropContent.IsValid();
+}
+
+bool FSlateUser::IsDragDroppingAffected(const FPointerEvent& InPointerEvent) const
+{
+	return DragDropContent.IsValid() && DragDropContent->AffectedByPointerEvent(InPointerEvent);
+}
+
+void FSlateUser::CancelDragDrop()
+{
+	if (DragDropContent.IsValid())
+	{
+		const FPointerEvent EmptyPointerEvent;
+		const FDragDropEvent DragDropEvent(EmptyPointerEvent, DragDropContent);
+
+		for (const auto& IndexPathPair : WidgetsUnderPointerLastEventByIndex)
+		{
+			FWidgetPath WidgetsToDragLeave = IndexPathPair.Value.ToWidgetPath();
+			if (WidgetsToDragLeave.IsValid())
+			{
+				for (int32 WidgetIndex = WidgetsToDragLeave.Widgets.Num() - 1; WidgetIndex >= 0; --WidgetIndex)
+				{
+					WidgetsToDragLeave.Widgets[WidgetIndex].Widget->OnDragLeave(DragDropEvent);
+				}
+			}
+		}
+
+		// Cancel dragdrop operation correctly firing off callbacks
+		DragDropContent->OnDrop(false, EmptyPointerEvent);
+
+		// We always reset the cache of widgets under our pointers whenever we enter/exit drag-drop mode
+		WidgetsUnderPointerLastEventByIndex.Reset();
+
+		DragDropContent.Reset();
+	}
+}
+
+void FSlateUser::ShowTooltip(const TSharedRef<IToolTip>& InTooltip, const FVector2D& InLocation)
+{
+	CloseTooltip();
+
+	ActiveTooltipInfo.Tooltip = InTooltip;
+
+	// Establish the tooltip content in the window
+	TSharedRef<SWindow> TooltipWindow = GetOrCreateTooltipWindow();	
+	TooltipWindow->SetContent(
+		SNew(SWeakWidget)
+		.PossiblyNullContent(InTooltip->AsWidget()));
+
+	// Make sure the desired size is valid
+	FSlateApplication& SlateApp = FSlateApplication::Get();
+	TooltipWindow->SlatePrepass(SlateApp.GetApplicationScale() * TooltipWindow->GetNativeWindow()->GetDPIScaleFactor());
+
+	// Place the window as close to the given location as possible (MoveWindowTo will adjust the window's position to stay onscreen, if needed)
+	const FSlateRect Anchor(InLocation.X, InLocation.Y, InLocation.X, InLocation.Y);
+	ActiveTooltipInfo.DesiredLocation = SlateApp.CalculateTooltipWindowPosition(Anchor, TooltipWindow->GetDesiredSizeDesktopPixels(), /*bAutoAdjustForDPIScale =*/false);
+	TooltipWindow->MoveWindowTo(ActiveTooltipInfo.DesiredLocation);
+
+	TooltipWindow->SetOpacity(0.0f);
+	TooltipWindow->ShowWindow();
+	ActiveTooltipInfo.SummonTime = FPlatformTime::Seconds();
+}
+
+void FSlateUser::CloseTooltip()
+{
+	ActiveTooltipInfo.Reset();
+
+	// Hide the tooltip window as well (don't destroy it - we'll reuse it)
+	TSharedPtr<SWindow> TooltipWindow = TooltipWindowPtr.Pin();
+	if (TooltipWindow && TooltipWindow->IsVisible())
+	{
+		TooltipWindow->HideWindow();
+	}
+}
+
+void FSlateUser::DrawWindowlessDragDropContent(const TSharedRef<SWindow>& WindowToDraw, FSlateWindowElementList& WindowElementList, int32& MaxLayerId)
+{
+	if (DragDropContent && DragDropContent->IsWindowlessOperation())
+	{
+		TSharedPtr<SWindow> DragDropWindow = DragDropWindowPtr.Pin();
+		if (DragDropWindow && DragDropWindow == WindowToDraw)
+		{
+			TSharedPtr<SWidget> DecoratorWidget = DragDropContent->GetDefaultDecorator();
+			if (DecoratorWidget && DecoratorWidget->GetVisibility().IsVisible())
+			{
+				FSlateApplication& SlateApp = FSlateApplication::Get();
+				const float WindowRootScale = SlateApp.GetApplicationScale() * DragDropWindow->GetNativeWindow()->GetDPIScaleFactor();
+
+				DecoratorWidget->SetVisibility(EVisibility::HitTestInvisible);
+				DecoratorWidget->SlatePrepass(WindowRootScale);
+
+				FVector2D DragDropContentInWindowSpace = WindowToDraw->GetWindowGeometryInScreen().AbsoluteToLocal(DragDropContent->GetDecoratorPosition()) * WindowRootScale;
+				const FGeometry DragDropContentGeometry = FGeometry::MakeRoot(DecoratorWidget->GetDesiredSize(), FSlateLayoutTransform(DragDropContentInWindowSpace));
+
+				DecoratorWidget->Paint(
+					FPaintArgs(&WindowToDraw.Get(), WindowToDraw->GetHittestGrid(), WindowToDraw->GetPositionInScreen(), SlateApp.GetCurrentTime(), SlateApp.GetDeltaTime()), 
+					DragDropContentGeometry, WindowToDraw->GetClippingRectangleInWindow(),
+					WindowElementList,
+					++MaxLayerId,
+					FWidgetStyle(),
+					WindowToDraw->IsEnabled());
+			}
+		}
+	}
+}
+
+void FSlateUser::DrawCursor(const TSharedRef<SWindow>& WindowToDraw, FSlateWindowElementList& WindowElementList, int32& MaxLayerId)
+{
+	TSharedPtr<SWindow> CursorWindow = CursorWindowPtr.Pin();
+	if (CursorWindow && WindowToDraw == CursorWindow)
+	{
+		if (TSharedPtr<SWidget> CursorWidget = CursorWidgetPtr.Pin())
+		{
+			FSlateApplication& SlateApp = FSlateApplication::Get();
+			const float WindowRootScale = FSlateApplication::Get().GetApplicationScale() * CursorWindow->GetNativeWindow()->GetDPIScaleFactor();
+
+			CursorWidget->SetVisibility(EVisibility::HitTestInvisible);
+			CursorWidget->SlatePrepass(WindowRootScale);
+
+			FVector2D CursorInScreen = GetCursorPosition();
+			FVector2D CursorPosInWindowSpace = WindowToDraw->GetWindowGeometryInScreen().AbsoluteToLocal(CursorInScreen) * WindowRootScale;
+			CursorPosInWindowSpace += (CursorWidget->GetDesiredSize() * -0.5) * WindowRootScale;
+			const FGeometry CursorGeometry = FGeometry::MakeRoot(CursorWidget->GetDesiredSize(), FSlateLayoutTransform(CursorPosInWindowSpace));
+
+			CursorWidget->Paint(
+				FPaintArgs(&WindowToDraw.Get(), WindowToDraw->GetHittestGrid(), WindowToDraw->GetPositionInScreen(), SlateApp.GetCurrentTime(), SlateApp.GetDeltaTime()),
+				CursorGeometry, WindowToDraw->GetClippingRectangleInWindow(),
+				WindowElementList,
+				++MaxLayerId,
+				FWidgetStyle(),
+				WindowToDraw->IsEnabled());
+		}
+	}
+}
+
+void FSlateUser::QueueSyntheticCursorMove()
+{
+	// Q: Wait, why 2?
+	// A: Synthesized moves are processed last, after all other inputs (other inputs are often what call this), 
+	//		but Slate won't have had a chance yet to do any actual widget updating. We need to make sure we synthesize the 
+	//		move *after* slate has had a chance to update.
+	
+	// Q: Ok, then shouldn't we only synthesize the move when we go from 1 -> 0?
+	// A: Nope - imagine a user provided input last frame without moving the mouse. 
+	//		That queued a synthetic move and was processed (possibly unnecessarily), going from 2 -> 1.
+	//	This frame, input has been provided again, resetting this value back to 2.
+	//	If we wait until we go from 1 -> 0, we would not synthesize a move until a frame without input, and
+	//		we'd indefinitely delay the very update the synthetic move is intended to trigger until the user stopped providing any input.
+	NumPendingSyntheticCursorMoves = 2;
+}
+
+bool FSlateUser::SynthesizeCursorMoveIfNeeded()
+{
+	// The slate loading widget thread is not allowed to execute this code as it is unsafe to read the hittest grid in another thread
+	if (bEnableSyntheticCursorMoves && IsInGameThread() && ensure(Cursor) && --NumPendingSyntheticCursorMoves >= 0)
+	{
+		// Synthetic cursor/mouse events accomplish two goals:
+		// 1) The UI can change even if the mouse doesn't move.
+		//    Synthesizing a mouse move sends out events.
+		//    In this case, the current and previous position will be the same.
+		//
+		// 2) The mouse moves, but the OS decided not to send us an event.
+		//    e.g. Mouse moved outside of our window.
+		//    In this case, the previous and current positions differ.
+
+		FSlateApplication& SlateApp = FSlateApplication::Get();
+		
+		const bool bHasHardwareCursor = SlateApp.GetPlatformCursor() == Cursor;
+		FPointerEvent SyntheticCursorMoveEvent(
+			FSlateApplication::CursorPointerIndex,
+			GetCursorPosition(),
+			GetPreviousCursorPosition(),
+			bHasHardwareCursor ? SlateApp.GetPressedMouseButtons() : TSet<FKey>(),
+			EKeys::Invalid,
+			0,
+			bHasHardwareCursor ? SlateApp.GetPlatformApplication()->GetModifierKeys() : FModifierKeysState());
+
+		SlateApp.ProcessMouseMoveEvent(SyntheticCursorMoveEvent, true);
+		return true;
+	}
+	return false;
+}
+
+void FSlateUser::SetFocusPath(const FWidgetPath& NewFocusPath, EFocusCause InFocusCause, bool bInShowFocus)
+{
+	StrongFocusPath.Reset();
+	WeakFocusPath = NewFocusPath;
+	FocusCause = InFocusCause;
+	bShowFocus = bInShowFocus;
+}
+
+void FSlateUser::FinishFrame()
+{
+	StrongFocusPath.Reset();
+}
+
+void FSlateUser::NotifyWindowDestroyed(TSharedRef<SWindow> DestroyedWindow)
+{
+	if (StrongFocusPath && StrongFocusPath->IsValid() && DestroyedWindow == StrongFocusPath->GetWindow())
+	{
+		StrongFocusPath.Reset();
+	}
+}
+
+void FSlateUser::QueryCursor()
+{
+	bQueryCursorRequested = false;
+
+	// The slate loading widget thread is not allowed to execute this code (it's unsafe to read the hittest grid in another thread)
+	if (Cursor && IsInGameThread())
+	{
+		SCOPE_CYCLE_COUNTER(STAT_SlateQueryCursor);
+
+		FSlateApplication& SlateApp = FSlateApplication::Get();
+		FCursorReply CursorReply = FCursorReply::Unhandled();
+		
+		// Drag-drop gets first dibs if it exists
+		if (DragDropContent)
+		{
+			CursorReply = DragDropContent->OnCursorQuery();
+		}
+
+		if (!CursorReply.IsEventHandled())
+		{
+			const bool bHasHardwareCursor = SlateApp.GetPlatformCursor() == Cursor;
+			const FVector2D CurrentCursorPosition = GetCursorPosition();
+			const FVector2D LastCursorPosition = GetPreviousCursorPosition();
+
+			const FPointerEvent CursorEvent(
+				FSlateApplication::CursorPointerIndex,
+				CurrentCursorPosition,
+				LastCursorPosition,
+				CurrentCursorPosition - LastCursorPosition,
+				bHasHardwareCursor ? SlateApp.GetPressedMouseButtons() : TSet<FKey>(),
+				bHasHardwareCursor ? SlateApp.GetModifierKeys() : FModifierKeysState());
+
+			FWidgetPath WidgetsToQueryForCursor;
+			if (HasCursorCapture())
+			{
+				// Query widgets with mouse capture for the cursor
+				FWidgetPath MouseCaptorPath = GetCursorCaptorPath(FWeakWidgetPath::EInterruptedPathHandling::Truncate, &CursorEvent);
+				if (MouseCaptorPath.IsValid())
+				{
+					TSharedRef<SWindow> CaptureWindow = MouseCaptorPath.GetWindow();
+					const TSharedPtr<SWindow> ActiveModalWindow = SlateApp.GetActiveModalWindow();
+
+					// Never query the mouse captor path if it is outside an active modal window
+					if (!ActiveModalWindow || CaptureWindow == ActiveModalWindow || CaptureWindow->IsDescendantOf(ActiveModalWindow))
+					{
+						WidgetsToQueryForCursor = MouseCaptorPath;
+					}
+				}
+			}
+			else
+			{
+				WidgetsToQueryForCursor = SlateApp.LocateWindowUnderMouse(CurrentCursorPosition, SlateApp.GetInteractiveTopLevelWindows());
+			}
+
+			if (WidgetsToQueryForCursor.IsValid())
+			{
+				// Switch worlds for widgets in the current path
+				FScopedSwitchWorldHack SwitchWorld(WidgetsToQueryForCursor);
+
+				for (int32 WidgetIndex = WidgetsToQueryForCursor.Widgets.Num() - 1; WidgetIndex >= 0; --WidgetIndex)
+				{
+					const FArrangedWidget& ArrangedWidget = WidgetsToQueryForCursor.Widgets[WidgetIndex];
+
+					CursorReply = ArrangedWidget.Widget->OnCursorQuery(ArrangedWidget.Geometry, CursorEvent);
+					if (CursorReply.IsEventHandled())
+					{
+						if (!CursorReply.GetCursorWidget().IsValid())
+						{
+							for (; WidgetIndex >= 0; --WidgetIndex)
+							{
+								TOptional<TSharedRef<SWidget>> CursorWidget = WidgetsToQueryForCursor.Widgets[WidgetIndex].Widget->OnMapCursor(CursorReply);
+								if (CursorWidget.IsSet())
+								{
+									CursorReply.SetCursorWidget(WidgetsToQueryForCursor.GetWindow(), CursorWidget.GetValue());
+									break;
+								}
+							}
+						}
+						break;
+					}
+				}
+
+				if (!CursorReply.IsEventHandled() && WidgetsToQueryForCursor.IsValid())
+				{
+					// Query was NOT handled, and we are still over a slate window.
+					CursorReply = FCursorReply::Cursor(EMouseCursor::Default);
+				}
+			}
+			else
+			{
+				// Set the default cursor when there isn't an active window under the cursor and the mouse isn't captured
+				CursorReply = FCursorReply::Cursor(EMouseCursor::Default);
+			}
+		}
+		ProcessCursorReply(CursorReply);
+	}
+}
+
+void FSlateUser::LockCursor(const TSharedRef<SWidget>& Widget)
+{
+	if (Cursor)
+	{
+		// Get a path to this widget so we know the position and size of its geometry
+		FWidgetPath WidgetPath;
+		const bool bFoundWidget = FSlateApplication::Get().GeneratePathToWidgetUnchecked(Widget, WidgetPath);
+		if (ensureMsgf(bFoundWidget, TEXT("Attempting to LockCursor() to widget but could not find widget %s"), *Widget->ToString()))
+		{
+			LockCursorInternal(WidgetPath);
+		}
+	}
+}
+
+void FSlateUser::UnlockCursor()
+{
+	if (Cursor)
+	{
+		Cursor->Lock(nullptr);
+		LockingWidgetPath = FWeakWidgetPath();
+	}
+}
+
+void FSlateUser::UpdateCursor()
+{
+	if (!Cursor)
+	{
+		return;
+	}
+
+	if (LockingWidgetPath.IsValid())
+	{
+		const FWidgetPath PathToWidget = LockingWidgetPath.ToWidgetPath(FWeakWidgetPath::EInterruptedPathHandling::ReturnInvalid);
+		if (PathToWidget.IsValid())
+		{
+			const FSlateRect ComputedClipRect = PathToWidget.Widgets.Last().Geometry.GetLayoutBoundingRect();
+			if (ComputedClipRect != LastComputedLockBounds)
+			{
+				// The locking widget is still valid, but its bounds have changed - gotta update the lock boundaries on the cursor to match
+				LockCursorInternal(PathToWidget);
+			}
+		}
+		else
+		{
+			// Unlock immediately if the locking widget is no longer around
+			UnlockCursor();
+		}
+	}	
+
+	// When Slate captures the mouse, it is up to us to set the cursor because the OS assumes that we own the mouse.
+	if (HasAnyCapture() || bQueryCursorRequested)
+	{
+		QueryCursor();
+	}
+}
+
+void FSlateUser::ProcessCursorReply(const FCursorReply& CursorReply)
+{
+	if (Cursor && CursorReply.IsEventHandled())
+	{
+		CursorWidgetPtr = CursorReply.GetCursorWidget();
+		if (CursorReply.GetCursorWidget().IsValid())
+		{
+			CursorReply.GetCursorWidget()->SetVisibility(EVisibility::HitTestInvisible);
+			CursorWindowPtr = CursorReply.GetCursorWindow();
+			if (!FSlateApplication::Get().IsFakingTouchEvents())
+			{
+				Cursor->SetType(EMouseCursor::Custom);
+			}
+		}
+		else
+		{
+			CursorWindowPtr.Reset();
+			Cursor->SetType(CursorReply.GetCursorType());
+		}
+	}
+	else
+	{
+		CursorWindowPtr.Reset();
+		CursorWidgetPtr.Reset();
+	}
+}
+
+void FSlateUser::LockCursorInternal(const FWidgetPath& WidgetPath)
+{
+	check(Cursor);
+	check(WidgetPath.IsValid());
+
+	// Do not attempt to lock the cursor to the window if its not in the foreground.  It would cause annoying side effects
+	TSharedPtr<const FGenericWindow> NativeWindow = WidgetPath.GetWindow()->GetNativeWindow();
+	if (NativeWindow && NativeWindow->IsForegroundWindow())
+	{
+		// The last widget in the path should be the widget we are locking the cursor to
+		const FSlateRect SlateClipRect = WidgetPath.Widgets.Last().Geometry.GetLayoutBoundingRect();
+		LastComputedLockBounds = SlateClipRect;
+		LockingWidgetPath = WidgetPath;
+
+		// Generate a screen space clip rect based on the widget's geometry
+#if PLATFORM_DESKTOP
+		const bool bIsBorderlessGameWindow = NativeWindow->IsDefinitionValid() && NativeWindow->GetDefinition().Type == EWindowType::GameWindow && !NativeWindow->GetDefinition().HasOSWindowBorder;
+		const int32 ClipRectAdjustment = bIsBorderlessGameWindow ? 1 : 0;
+#else
+		const int32 ClipRectAdjustment = 0;
+#endif
+
+		// Note: We round the upper left coordinate of the clip rect so we guarantee the rect is inside the geometry of the widget.  If we truncated when there is a half pixel we would cause the clip
+		// rect to be half a pixel larger than the geometry and cause the mouse to go outside of the geometry.
+		RECT ClipRect;
+		ClipRect.left = FMath::RoundToInt(SlateClipRect.Left + ClipRectAdjustment);
+		ClipRect.top = FMath::RoundToInt(SlateClipRect.Top + ClipRectAdjustment);
+		ClipRect.right = FMath::TruncToInt(SlateClipRect.Right - ClipRectAdjustment);
+		ClipRect.bottom = FMath::TruncToInt(SlateClipRect.Bottom - ClipRectAdjustment);
+
+		// Lock the cursor to the widget
+		Cursor->Lock(&ClipRect);
+	}
+}
+
+TSharedRef<SWindow> FSlateUser::GetOrCreateTooltipWindow()
+{
+	if (TooltipWindowPtr.IsValid())
+	{
+		return TooltipWindowPtr.Pin().ToSharedRef();
+	}
+
+	// If we don't have a window already, make a new one and add it (but don't show it until we've put stuff in)
+	TSharedRef<SWindow> NewTooltipWindow = SWindow::MakeToolTipWindow();
+	TooltipWindowPtr = NewTooltipWindow;
+	FSlateApplication::Get().AddWindow(NewTooltipWindow, /*bShowImmediately =*/false);
+	return NewTooltipWindow;
+}
+
+void FSlateUser::NotifyPointerMoveBegin(const FPointerEvent& PointerEvent)
+{
+	if (PointerEvent.IsTouchEvent())
+	{
+		GestureDetector.OnTouchMoved(PointerEvent.GetPointerIndex(), PointerEvent.GetScreenSpacePosition());
+	}
+	PointerPositionsByIndex.FindOrAdd(PointerEvent.GetPointerIndex()) = PointerEvent.GetScreenSpacePosition();
+}
+
+void FSlateUser::NotifyPointerMoveComplete(const FPointerEvent& PointerEvent, const FWidgetPath& WidgetsUnderPointer)
+{
+	// Give the current drag drop operation a chance to do something custom (e.g. update the Drag/Drop preview based on content)
+	if (IsDragDroppingAffected(PointerEvent))
+	{
+		FDragDropEvent DragDropEvent(PointerEvent, DragDropContent);
+
+#if WITH_EDITOR
+		//@TODO VREDITOR - Remove and move to interaction component
+		if (FSlateApplication::Get().OnDragDropCheckOverride.IsBound() && DragDropEvent.GetOperation())
+		{
+			DragDropEvent.GetOperation()->SetDecoratorVisibility(false);
+			DragDropEvent.GetOperation()->SetCursorOverride(EMouseCursor::None);
+			DragDropContent->SetCursorOverride(EMouseCursor::None);
+		}
+#endif
+
+		FScopedSwitchWorldHack SwitchWorld(WidgetsUnderPointer);
+		DragDropContent->OnDragged(DragDropEvent);
+
+		// Update the window we're under for rendering the drag drop operation if it's a windowless drag drop operation.
+		if (WidgetsUnderPointer.IsValid())
+		{
+			DragDropWindowPtr = WidgetsUnderPointer.GetWindow();
+		}
+		else
+		{
+			DragDropWindowPtr.Reset();
+		}
+		
+		if (ensure(Cursor))
+		{
+			FCursorReply CursorReply = DragDropContent->OnCursorQuery();
+			if (!CursorReply.IsEventHandled())
+			{
+				// Set the default cursor when there isn't an active window under the cursor and the mouse isn't captured
+				CursorReply = FCursorReply::Cursor(EMouseCursor::Default);
+			}
+
+			FSlateApplication::Get().ProcessCursorReply(CursorReply);
+		}
+		
+	}
+	else if (!IsDragDropping())
+	{
+		DragDropWindowPtr.Reset();
+	}
+
+	PreviousPointerPositionsByIndex.Add(PointerEvent.GetPointerIndex(), PointerEvent.GetScreenSpacePosition());
+	WidgetsUnderPointerLastEventByIndex.Add(PointerEvent.GetPointerIndex(), FWeakWidgetPath(WidgetsUnderPointer));
+}
+
+void FSlateUser::UpdatePointerPosition(const FPointerEvent& PointerEvent)
+{
+	UpdatePointerPosition(PointerEvent.GetPointerIndex(), PointerEvent.GetScreenSpacePosition());
+}
+
+void FSlateUser::UpdatePointerPosition(uint32 PointerIndex, const FVector2D& Position)
+{
+	PointerPositionsByIndex.FindOrAdd(PointerIndex) = Position;
+	PreviousPointerPositionsByIndex.FindOrAdd(PointerIndex) = Position;
+}
+
+void FSlateUser::StartDragDetection(const FWidgetPath& PathToWidget, int32 PointerIndex, FKey DragButton, FVector2D StartLocation)
+{
+	DragStatesByPointerIndex.Add(PointerIndex, FDragDetectionState(PathToWidget, PointerIndex, DragButton, StartLocation));
+}
+
+FWidgetPath FSlateUser::DetectDrag(const FPointerEvent& PointerEvent, float DragTriggerDistance)
+{
+	if (FDragDetectionState* DragState = DragStatesByPointerIndex.Find(PointerEvent.GetPointerIndex()))
+	{
+		const FVector2D DragDelta = DragState->DragStartLocation - PointerEvent.GetScreenSpacePosition();
+		if (DragDelta.SizeSquared() > FMath::Square(DragTriggerDistance))
+		{
+			FWidgetPath DragDetectionPath = DragState->DetectDragForWidget.ToWidgetPath(FWeakWidgetPath::EInterruptedPathHandling::ReturnInvalid);
+			if (DragDetectionPath.IsValid())
+			{
+				ResetDragDetection();
+				return DragDetectionPath;
+			}
+		}
+	}
+
+	return FWidgetPath();
+}
+
+bool FSlateUser::IsDetectingDrag(uint32 PointerIndex) const
+{
+	return DragStatesByPointerIndex.Contains(PointerIndex);
+}
+
+void FSlateUser::NotifyPointerReleased(const FPointerEvent& PointerEvent, const FWidgetPath& WidgetsUnderCursor, TSharedPtr<FDragDropOperation> DroppedContent, bool bWasHandled)
+{
+	const int32 PointerIdx = PointerEvent.GetPointerIndex();
+
+	const FDragDetectionState* DragState = DragStatesByPointerIndex.Find(PointerIdx);
+	if (DragState && DragState->TriggerButton == PointerEvent.GetEffectingButton())
+	{
+		// The user has released the button (or finger) that was supposed to start the drag; stop detecting it.
+		DragState = nullptr;
+		DragStatesByPointerIndex.Remove(PointerIdx);
+	}
+
+	if (!HasCapture(PointerIdx))
+	{
+		// When we perform a touch end, we need to also send a mouse leave as if it were a cursor.
+		if (PointerEvent.IsTouchEvent() && !FSlateApplication::Get().IsFakingTouchEvents())
+		{
+			for (int32 WidgetIndex = WidgetsUnderCursor.Widgets.Num() - 1; WidgetIndex >= 0; --WidgetIndex)
+			{
+				WidgetsUnderCursor.Widgets[WidgetIndex].Widget->OnMouseLeave(PointerEvent);
+			}
+		}
+#if !UE_BUILD_SHIPPING
+		const bool bIsFingerCapturedPostRelease = PointerEvent.IsTouchEvent() && HasCapture(PointerEvent.GetPointerIndex());
+		ensureMsgf(!bIsFingerCapturedPostRelease, TEXT("Touch pointer was captured while in the process of being released. Since touch pointers cease to exist when the touch ends, this makes no sense."));
+#endif
+
+		// Note: FSlateApplication caches this content off BEFORE routing the pointer up, as OnDrop can result in re-entrance of the pointer up routing
+		//		To avoid executing the drop operation twice, our cached DragDropContent is wiped before the first routing.
+		//		Thus, we rely on the provided DroppedContent here, which will never be provided as valid more than once
+		if (DroppedContent && DroppedContent->AffectedByPointerEvent(PointerEvent))
+		{
+			// @todo slate : depending on SetEventPath() is not ideal.
+			FPointerEvent ModifiedEvent(PointerEvent);
+			ModifiedEvent.SetEventPath(WidgetsUnderCursor);
+			DroppedContent->OnDrop(bWasHandled, ModifiedEvent);
+			WidgetsUnderPointerLastEventByIndex.Remove(PointerIdx);
+		}
+	}
+	
+	if (PointerEvent.IsTouchEvent())
+	{
+		GestureDetector.OnTouchEnded(PointerIdx, PointerEvent.GetScreenSpacePosition());
+
+		// For touch events, we always invalidate capture for the pointer.  There's no reason to ever maintain capture for
+		// fingers no longer in contact with the screen.
+		ReleaseCapture(PointerIdx);
+
+		// When touch pointers are released, they also effectively cease to exist
+		PointerPositionsByIndex.Remove(PointerIdx);
+		PreviousPointerPositionsByIndex.Remove(PointerIdx);
+	}
+}
+
+void FSlateUser::ResetDragDetection()
+{
+	DragStatesByPointerIndex.Reset();
+}
+
+void FSlateUser::SetDragDropContent(TSharedRef<FDragDropOperation> InDragDropContent)
+{
+	checkf(!IsDragDropping(), TEXT("Drag and Drop already in progress!"));
+	DragDropContent = InDragDropContent;
+}
+
+void FSlateUser::ResetDragDropContent()
+{
+	DragDropContent.Reset();
+}
+
+void FSlateUser::UpdateTooltip(const FMenuStack& MenuStack, bool bCanSpawnNewTooltip)
+{
+	FSlateApplication& SlateApp = FSlateApplication::Get();
+	if (!SlateApp.GetAllowTooltips())
+	{
+		CloseTooltip();
+		return;
+	}
+
+	SCOPE_CYCLE_COUNTER(STAT_SlateUpdateTooltip);
+
+	float DPIScaleFactor = 1.0f;
+	FWidgetPath WidgetsToQueryForTooltip;
+
+	const bool bCheckForTooltipChanges =
+		IsInGameThread() &&					// We should never allow the slate loading thread to create new windows or interact with the hittest grid
+		!SlateApp.IsUsingHighPrecisionMouseMovment() && // If we are using HighPrecision movement then we can't rely on the OS cursor to be accurate
+		!IsDragDropping() &&				// We must not currently be in the middle of a drag-drop action
+		
+		//@todo DanH: We need to check if OUR cursor is over a slate window, not just the platform cursor. 
+		//		See about adding FPlatformApplication::GetSlateWindowUnderPoint(FVector2D) or something.
+		SlateApp.GetPlatformApplication()->IsCursorDirectlyOverSlateWindow(); // The cursor must be over a Slate window
+
+	if (bCheckForTooltipChanges)
+	{
+		// We're gonna check each widget under the cursor (including disabled widgets) until we find one with a tooltip to show
+		FWidgetPath WidgetsUnderCursor = SlateApp.LocateWindowUnderMouse(GetCursorPosition(), SlateApp.GetInteractiveTopLevelWindows(), /*bIgnoreEnabledStatus =*/true);
+		if (WidgetsUnderCursor.IsValid() && WidgetsUnderCursor.GetWindow() != TooltipWindowPtr.Pin())
+		{
+			WidgetsToQueryForTooltip = WidgetsUnderCursor;
+		}
+	}
+
+	TOptional<FSlateRect> ForceFieldRect;
+	TSharedPtr<IToolTip> NewTooltip;
+	TSharedPtr<SWidget> WidgetProvidingNewTooltip;
+	for (int32 WidgetIndex = WidgetsToQueryForTooltip.Widgets.Num() - 1; WidgetIndex >= 0; --WidgetIndex)
+	{
+		const FArrangedWidget* ArrangedWidget = &WidgetsToQueryForTooltip.Widgets[WidgetIndex];
+		const TSharedRef<SWidget>& CurWidget = ArrangedWidget->Widget;
+
+		if (!NewTooltip.IsValid())
+		{
+			// Make sure the tooltip has something to show before we pick it
+			TSharedPtr<IToolTip> WidgetTooltip = CurWidget->GetToolTip();
+			if (WidgetTooltip.IsValid() && !WidgetTooltip->IsEmpty())
+			{
+				WidgetProvidingNewTooltip = CurWidget;
+				NewTooltip = WidgetTooltip;
+			}
+		}
+
+		// Make sure we account for all widgets in the path with a force field, even if we've already found the tooltip we'll be showing
+		if (CurWidget->HasToolTipForceField())
+		{
+			if (!ForceFieldRect.IsSet())
+			{
+				ForceFieldRect = ArrangedWidget->Geometry.GetLayoutBoundingRect();
+			}
+			else
+			{
+				// Grow the rect to encompass this geometry to be super safe
+				// The parent's rect should always be inclusive of its child, so this should usually be overkill.
+				ForceFieldRect = ForceFieldRect->Expand(ArrangedWidget->Geometry.GetLayoutBoundingRect());
+			}
+			ForceFieldRect = (1.0f / DPIScaleFactor) * ForceFieldRect.GetValue();
+		}
+	}
+
+	TSharedPtr<SWidget> NewTooltipVisualizer;
+	TSharedPtr<IToolTip> ActiveTooltip = ActiveTooltipInfo.Tooltip.Pin();
+	const bool bTooltipChanged = NewTooltip != ActiveTooltip;
+	if (bTooltipChanged)
+	{
+		// Remove existing tooltip if there is one.
+		if (ActiveTooltipInfo.TooltipVisualizer.IsValid())
+		{
+			ActiveTooltipInfo.TooltipVisualizer.Pin()->OnVisualizeTooltip(nullptr);
+		}
+
+		// Notify the new tooltip that it's about to be opened.
+		if (NewTooltip)
+		{
+			NewTooltip->OnOpening();
+		}
+
+		// Some widgets might want to provide an alternative Tooltip Handler.
+		TSharedPtr<SWidget> NewTooltipWidget = NewTooltip ? NewTooltip->AsWidget() : TSharedPtr<SWidget>();
+		for (int32 WidgetIndex = WidgetsToQueryForTooltip.Widgets.Num() - 1; WidgetIndex >= 0; --WidgetIndex)
+		{
+			const TSharedRef<SWidget>& CurWidget = WidgetsToQueryForTooltip.Widgets[WidgetIndex].Widget;
+			if (CurWidget->OnVisualizeTooltip(NewTooltipWidget))
+			{
+				// Someone is taking care of visualizing this tooltip
+				NewTooltipVisualizer = CurWidget;
+				break;
+			}
+		}
+	}
+	
+	// If a widget under the cursor has a tool-tip forcefield active, then go through any menus
+	// in the menu stack that are above that widget's window, and make sure those windows also
+	// prevent the tool-tip from encroaching.  This prevents tool-tips from drawing over sub-menus
+	// spawned from menu items in a different window, for example.
+	if (ForceFieldRect.IsSet() && WidgetsToQueryForTooltip.IsValid())
+	{
+		if (TSharedPtr<IMenu> MenuInPath = MenuStack.FindMenuInWidgetPath(WidgetsToQueryForTooltip))
+		{
+			ForceFieldRect = ForceFieldRect->Expand(MenuStack.GetToolTipForceFieldRect(MenuInPath.ToSharedRef(), WidgetsToQueryForTooltip));
+		}
+	}
+
+	FVector2D DesiredLocation = ActiveTooltipInfo.DesiredLocation;
+	if ((ActiveTooltip && !ActiveTooltip->IsInteractive()) || (NewTooltip && NewTooltip != ActiveTooltip))
+	{
+		// Keep track of where we want tooltips to be positioned
+		DesiredLocation = GetPreviousCursorPosition() + SlateDefs::TooltipOffsetFromMouse;
+	}
+
+	if (TooltipWindowPtr.IsValid())
+	{
+		FSlateRect Anchor(DesiredLocation.X, DesiredLocation.Y, DesiredLocation.X, DesiredLocation.Y);
+		DesiredLocation = SlateApp.CalculatePopupWindowPosition(Anchor, TooltipWindowPtr.Pin()->GetDesiredSizeDesktopPixels(), /*bAutoAdjustForDPIScale =*/false);
+	}
+
+	// Repel tooltip from a force field, if necessary
+	if (ForceFieldRect.IsSet())
+	{
+		FVector2D TooltipShift;
+		TooltipShift.X = (ForceFieldRect->Right + SlateDefs::TooltipOffsetFromForceField.X) - DesiredLocation.X;
+		TooltipShift.Y = (ForceFieldRect->Bottom + SlateDefs::TooltipOffsetFromForceField.Y) - DesiredLocation.Y;
+
+		// Make sure the tooltip needs to be offset
+		if (TooltipShift.X > 0.0f && TooltipShift.Y > 0.0f)
+		{
+			// Find the best edge to move the tooltip towards
+			if (ActiveTooltipInfo.OffsetDirection == ETooltipOffsetDirection::Right ||
+				(ActiveTooltipInfo.OffsetDirection == ETooltipOffsetDirection::Undetermined && TooltipShift.X < TooltipShift.Y))
+			{
+				// Move right
+				DesiredLocation.X += TooltipShift.X;
+				ActiveTooltipInfo.OffsetDirection = ETooltipOffsetDirection::Right;
+			}
+			else
+			{
+				// Move down
+				DesiredLocation.Y += TooltipShift.Y;
+				ActiveTooltipInfo.OffsetDirection = ETooltipOffsetDirection::Down;
+			}
+		}
+	}
+
+	// The tool tip changed...
+	if (bTooltipChanged)
+	{
+		// Close any existing tooltips; Unless the current tooltip is interactive and we don't have a valid tooltip to replace it
+		if (NewTooltip || (ActiveTooltip && !ActiveTooltip->IsInteractive()))
+		{
+			CloseTooltip();
+
+			if (NewTooltipVisualizer)
+			{
+				ActiveTooltipInfo.TooltipVisualizer = NewTooltipVisualizer;
+			}
+			else if (bCanSpawnNewTooltip && NewTooltip)
+			{
+				ShowTooltip(NewTooltip.ToSharedRef(), DesiredLocation);
+				ActiveTooltipInfo.SourceWidget = WidgetProvidingNewTooltip;
+			}
+		}
+	}
+
+	if (TooltipWindowPtr.IsValid())
+	{
+		// Only enable tooltip transitions if we're running at a decent frame rate
+		const bool bAllowAnimations = FSlateApplication::Get().IsRunningAtTargetFrameRate();
+
+		// How long since the tooltip was summoned?
+		const double PlatformSeconds = FPlatformTime::Seconds();
+		const float TimeSinceSummon = PlatformSeconds - TooltipSummonDelay - ActiveTooltipInfo.SummonTime;
+		const float TooltipOpacity = FMath::Clamp<float>(TimeSinceSummon / TooltipIntroDuration, 0.0f, 1.0f);
+
+		// Update window opacity
+		TSharedRef<SWindow> TooltipWindow = TooltipWindowPtr.Pin().ToSharedRef();
+		TooltipWindow->SetOpacity(TooltipOpacity);
+
+		// How far tool tips should slide
+		const FVector2D SlideDistance(30.0f, 5.0f);
+
+		// Apply steep inbound curve to the movement, so it looks like it quickly decelerating
+		const float SlideProgress = bAllowAnimations ? FMath::Pow(1.0f - TooltipOpacity, 3.0f) : 0.0f;
+
+		FVector2D WindowLocation = DesiredLocation + SlideProgress * SlideDistance;
+		if (WindowLocation != TooltipWindow->GetPositionInScreen())
+		{
+			// already handled
+			const bool bAutoAdjustForDPIScale = false;
+
+			// Avoid the edges of the desktop
+			FSlateRect Anchor(WindowLocation.X, WindowLocation.Y, WindowLocation.X, WindowLocation.Y);
+			WindowLocation = SlateApp.CalculateTooltipWindowPosition(Anchor, TooltipWindow->GetDesiredSizeDesktopPixels(), bAutoAdjustForDPIScale);
+
+			// Update the tool tip window positioning
+			// SetCachedScreenPosition is a hack (issue tracked as TTP #347070) which is needed because code in TickWindowAndChildren()/DrawPrepass()
+			// assumes GetPositionInScreen() to correspond to the new window location in the same tick. This is true on Windows, but other
+			// OSes (Linux in particular) may not update cached screen position until next time events are polled.
+			TooltipWindow->SetCachedScreenPosition(WindowLocation);
+			TooltipWindow->MoveWindowTo(WindowLocation);
+		}
+	}
+}
+
+void FSlateUser::ResetTooltipWindow()
+{
+	if (TooltipWindowPtr.IsValid())
+	{
+		TooltipWindowPtr.Pin()->RequestDestroyWindow();
+		TooltipWindowPtr.Reset();
+	}
+}
+
+bool FSlateUser::IsWindowHousingInteractiveTooltip(const TSharedRef<const SWindow>& WindowToTest) const
+{
+	if (WindowToTest == TooltipWindowPtr.Pin())
+	{
+		const TSharedPtr<IToolTip> ActiveTooltip = ActiveTooltipInfo.Tooltip.Pin();
+		return ActiveTooltip && ActiveTooltip->IsInteractive();
+	}
+	return false;
+}
diff --git a/Engine/Source/Runtime/Slate/Private/Widgets/Input/SMenuAnchor.cpp b/Engine/Source/Runtime/Slate/Private/Widgets/Input/SMenuAnchor.cpp
--- a/Engine/Source/Runtime/Slate/Private/Widgets/Input/SMenuAnchor.cpp
+++ b/Engine/Source/Runtime/Slate/Private/Widgets/Input/SMenuAnchor.cpp
@@ -395,7 +395,7 @@ void SMenuAnchor::SetIsOpen( bool InIsOpen, const bool bFocusMenu, const int32 F
 
 							if (bFocusMenu)
 							{
-								FSlateApplication::Get().ReleaseMouseCaptureForUser(FocusUserIndex);
+								FSlateApplication::Get().ReleaseAllPointerCapture(FocusUserIndex);
 							}
 
 							TSharedRef<SMenuAnchor> SharedThis = StaticCastSharedRef<SMenuAnchor>(AsShared());
@@ -445,7 +445,7 @@ void SMenuAnchor::SetIsOpen( bool InIsOpen, const bool bFocusMenu, const int32 F
 							// Release the mouse so that context can be properly restored upon closing menus.  See CL 1411833 before changing this.
 							if (bFocusMenu)
 							{
-								FSlateApplication::Get().ReleaseMouseCaptureForUser(FocusUserIndex);
+								FSlateApplication::Get().ReleaseAllPointerCapture(FocusUserIndex);
 							}
 
 							// Create a new window for the menu
@@ -487,7 +487,7 @@ void SMenuAnchor::SetIsOpen( bool InIsOpen, const bool bFocusMenu, const int32 F
 
 							if (bFocusMenu)
 							{
-								FSlateApplication::Get().ReleaseMouseCaptureForUser(FocusUserIndex);
+								FSlateApplication::Get().ReleaseAllPointerCapture(FocusUserIndex);
 							}
 
 							TSharedRef<SMenuAnchor> SharedThis = StaticCastSharedRef<SMenuAnchor>(AsShared());
diff --git a/Engine/Source/Runtime/Slate/Public/Framework/Application/SlateApplication.h b/Engine/Source/Runtime/Slate/Public/Framework/Application/SlateApplication.h
--- a/Engine/Source/Runtime/Slate/Public/Framework/Application/SlateApplication.h
+++ b/Engine/Source/Runtime/Slate/Public/Framework/Application/SlateApplication.h
@@ -41,6 +41,8 @@ class ITextInputMethodSystem;
 class IVirtualKeyboardEntry;
 class IWidgetReflector;
 class SViewport;
+class FSlateUser;
+class FSlateVirtualUserHandle;
 
 /** A Delegate for querying whether source code access is possible */
 DECLARE_DELEGATE_RetVal(bool, FQueryAccessSourceCode);
@@ -107,83 +109,6 @@ class SLATE_API FPopupSupport
 	TArray<FClickSubscriber> ClickZoneNotifications;
 };
 
-
-
-/**
- * A representation of a slate input providing user.  We allocate a slate user as new input sources are
- * discovered.
- */
-class SLATE_API FSlateUser
-{
-public:
-	FSlateUser(int32 InUserIndex, bool InVirtualUser);
-	virtual ~FSlateUser();
-
-	FORCEINLINE int32 GetUserIndex() const { return UserIndex; }
-	FORCEINLINE bool IsVirtualUser() const { return bVirtualUser; }
-
-	TSharedPtr<SWidget> GetFocusedWidget() const;
-
-	FORCEINLINE uint64 GetFocusVersion() const { return FocusVersion; }
-	FORCEINLINE void UpdateFocusVersion() { FocusVersion++; }
-
-	FGestureDetector GestureDetector;
-
-private:
-	FORCEINLINE bool HasValidFocusPath() const
-	{
-		return FocusWidgetPathWeak.IsValid();
-	}
-
-	FORCEINLINE const FWeakWidgetPath& GetWeakFocusPath() const
-	{
-		return FocusWidgetPathWeak;
-	}
-
-	FORCEINLINE TSharedRef<FWidgetPath> GetFocusPath() const
-	{
-		if ( !FocusWidgetPathStrong.IsValid() )
-		{
-			FocusWidgetPathStrong = FocusWidgetPathWeak.ToWidgetPathRef();
-		}
-
-		return FocusWidgetPathStrong.ToSharedRef();
-	}
-
-	void SetFocusPath(const FWidgetPath& InWidgetPath, EFocusCause InCause, bool InShowFocus);
-
-	void FinishFrame();
-
-	void NotifyWindowDestroyed(TSharedRef<SWindow> DestroyedWindow);
-
-private:
-	/** The index the user was assigned. */
-	int32 UserIndex;
-
-	/** Is this a virtual user?  Virtual users are generally ignored in most operations that affect all users. */
-	bool bVirtualUser;
-
-	/** A weak path to the widget currently focused by a user, if any. */
-	FWeakWidgetPath FocusWidgetPathWeak;
-
-	/** A strong widget path to a widget, this is cleared after the end of pumping messages. */
-	mutable TSharedPtr<FWidgetPath> FocusWidgetPathStrong;
-
-	/** Reason a widget was focused by a user, if any. */
-	EFocusCause FocusCause;
-
-	/** If we should show this focus */
-	bool ShowFocus;
-
-	/**
-	 * The FocusVersion is used to know if the focus state is modified for a user while processing focus
-	 * events, that way upon returning from focus calls, we know if we should abandon the remainder of the event.
-	 */
-	int32 FocusVersion;
-
-	friend class FSlateApplication;
-};
-
 /**
  * Interface for a Slate Input Mapping.
  */
@@ -201,27 +126,6 @@ public:
 	virtual int32 GetUserIndexForController(int32 ControllerId) const override { return ControllerId; }
 };
 
-/**
- * Represents a virtual user of slate.
- */
-class SLATE_API FSlateVirtualUser
-{
-public:
-	FSlateVirtualUser(int32 InUserIndex, int32 InVirtualUserIndex);
-	virtual ~FSlateVirtualUser();
-
-	FORCEINLINE int32 GetUserIndex() const { return UserIndex; }
-	FORCEINLINE int32 GetVirtualUserIndex() const { return VirtualUserIndex; }
-
-private:
-
-	/** The index the user was assigned. */
-	int32 UserIndex;
-
-	/** The virtual index the user was assigned. */
-	int32 VirtualUserIndex;
-};
-
 enum class ESlateTickType : uint8
 {
 	/** Tick time only */
@@ -685,10 +589,14 @@ public:
 	/** returning platform-specific value designating window that captures mouse, or nullptr if mouse isn't captured */
 	virtual void* GetMouseCaptureWindow() const;
 
-	/** Releases the mouse capture from whatever it currently is on - for all users for all pointers. */
+	/** Releases capture for every pointer on every user from whatever it currently is on. */
+	void ReleaseAllPointerCapture();
+	UE_DEPRECATED(4.23, "ReleaseMouseCapture has been renamed to ReleaseAllPointerCapture()")
 	void ReleaseMouseCapture();
 
-	/** Releases the mouse capture from whatever it currently is on for a particular user. */
+	/** Releases capture for every pointer belonging to the given user index particular user. */
+	void ReleaseAllPointerCapture(int32 UserIndex);
+	UE_DEPRECATED(4.23, "ReleaseMouseCaptureForUser has been renamed to ReleaseAllPointerCapture(int32 UserIndex)")
 	void ReleaseMouseCaptureForUser(int32 UserIndex);
 
 	/** @return The active modal window or nullptr if there is no modal window. */
@@ -948,51 +856,47 @@ public:
 	 */
 	bool TakeScreenshot(const TSharedRef<SWidget>& Widget, const FIntRect& InnerWidgetArea, TArray<FColor>& OutColorData, FIntVector& OutSize);
 
+	/** Gets the user at the given index, null if the user does not exist. */
+	FORCEINLINE TSharedPtr<const FSlateUser> GetUser(int32 UserIndex) const
+	{
+		return Users.IsValidIndex(UserIndex) ? Users[UserIndex] : nullptr;
+	}
+	FORCEINLINE TSharedPtr<FSlateUser> GetUser(int32 UserIndex)
+	{
+		return Users.IsValidIndex(UserIndex) ? Users[UserIndex] : nullptr;
+	}
+	FORCEINLINE TSharedPtr<const FSlateUser> GetUser(const FInputEvent& InputEvent) const { return GetUser(InputEvent.GetUserIndex()); }
+	FORCEINLINE TSharedPtr<FSlateUser> GetUser(const FInputEvent& InputEvent) { return GetUser(InputEvent.GetUserIndex()); }
+
 	/**
 	 * @return a handle for the existing or newly created virtual slate user.  This is handy when you need to create
 	 * virtual hardware users for slate components in the virtual world that may need to be interacted with with virtual hardware.
 	 */
-	TSharedRef<FSlateVirtualUser> FindOrCreateVirtualUser(int32 VirtualUserIndex);
+	TSharedRef<FSlateVirtualUserHandle> FindOrCreateVirtualUser(int32 VirtualUserIndex);
 
-	/**
-	 * 
-	 */
+	//@todo DanH: This is only getting called by the WidgetInteractionComponent. There's some piping to be set up here for this to be fully accurate for splitscreen
 	void UnregisterUser(int32 UserIndex);
 
-	/**
-	 * Allows you do some operations for every registered user.
-	 */
+	/** Allows you do some operations for every registered user. */
+	void ForEachUser(TFunctionRef<void(FSlateUser&)> InPredicate, bool bIncludeVirtualUsers = false);
+	
+	UE_DEPRECATED(4.23, "ForEachUser now provides an FSlateUser& parameter to the lambda instead of an FSlateUser*")
 	void ForEachUser(TFunctionRef<void(FSlateUser*)> InPredicate, bool bIncludeVirtualUsers = false);
 
 protected:
 	/**
-	 * Register a user with Slate.  Normally this is unnecessary as Slate automatically adds
+	 * Register a new user with Slate.  Normally this is unnecessary as Slate automatically adds
 	 * a user entry if it gets input from a controller for that index.  Might happen if the user
 	 * allocates the virtual user.
 	 */
-	void RegisterUser(TSharedRef<FSlateUser> User);
+	TSharedRef<FSlateUser> RegisterNewUser(int32 UserIndex, bool bIsVirtual = false);
 
 	/**
-	 * Gets the user at the given index, null if the user does not exist.
+	 * Locates the SlateUser object corresponding to the index, creating a new one if it doesn't exist.
+	 * Asserts if given an invalid (ie, negative) index.
 	 */
-	FORCEINLINE const FSlateUser* GetUser(int32 UserIndex) const
-	{
-		return (UserIndex >= 0 && UserIndex < Users.Num()) ? Users[UserIndex].Get() : nullptr;
-	}
-
-	/**
-	 * Gets the user at the given index, null if the user does not exist.
-	 */
-	FORCEINLINE FSlateUser* GetUser(int32 UserIndex)
-	{
-		return ( UserIndex >= 0 && UserIndex < Users.Num() ) ? Users[UserIndex].Get() : nullptr;
-	}
-
-	/**
-	 * Locates the SlateUser object corresponding to the index, if one can't be found, it will create a slate user at
-	 * the provided index.  If the index is less than 0, null is returned.
-	 */
-	FSlateUser* GetOrCreateUser(int32 UserIndex);
+	TSharedRef<FSlateUser> GetOrCreateUser(int32 UserIndex);
+	FORCEINLINE TSharedRef<FSlateUser> GetOrCreateUser(const FInputEvent& InputEvent) { return GetOrCreateUser(InputEvent.GetUserIndex()); }
 
 	friend class FEventRouter;
 
@@ -1257,10 +1161,7 @@ public:
 	/** Closes the open tool-tip, if a tool-tip is open */
 	void CloseToolTip();
 
-	/**
-	 * Updates tool tip
-	 */
-	void UpdateToolTip( bool AllowSpawningOfNewToolTips );
+	void UpdateToolTip( bool bAllowSpawningOfNewToolTips );
 
 	/** @return an array of top-level windows that can be interacted with. e.g. when a modal window is up, only return the modal window */
 	TArray< TSharedRef<SWindow> > GetInteractiveTopLevelWindows();
@@ -1366,21 +1267,14 @@ public:
 	virtual TSharedPtr<SWindow> GetActiveTopLevelWindow() const override;
 	virtual const FSlateBrush* GetAppIcon() const override;
 
-	virtual float GetApplicationScale() const override
-	{
-		return Scale;
-	}
+	virtual float GetApplicationScale() const override { return Scale; }
+	virtual bool GetSoftwareCursorAvailable() const override { return bSoftwareCursorAvailable; }
+	virtual EVisibility GetSoftwareCursorVis() const override;
 
 	virtual FVector2D GetCursorPos() const override;
 	virtual FVector2D GetLastCursorPos() const override;
 	virtual FVector2D GetCursorSize() const override;
 
-	virtual bool GetSoftwareCursorAvailable() const override
-	{
-		return bSoftwareCursorAvailable;
-	}
-
-	virtual EVisibility GetSoftwareCursorVis() const override;	
 	virtual TSharedPtr<SWidget> GetKeyboardFocusedWidget() const override;
 
 	virtual EWindowTransparency GetWindowTransparencySupport() const override
@@ -1522,19 +1416,13 @@ public:
 	 */
 	FReply RouteMouseWheelOrGestureEvent(const FWidgetPath& WidgetsUnderPointer, const FPointerEvent& InWheelEvent, const FPointerEvent* InGestureEvent = nullptr);
 
-	/**
-	 * @return int user index that the keyboard is mapped to. -1 if the keyboard isn't mapped
-	 */
+	/** @return int user index that the keyboard is mapped to. */
 	int32 GetUserIndexForKeyboard() const;
 
-	/** 
-	 * @return int user index that this controller is mapped to. -1 if the controller isn't mapped
-	 */
+	/** @return int user index that this controller is mapped to. */
 	int32 GetUserIndexForController(int32 ControllerId) const;
 
-	/** 
-	 * @return int user index that this controller is mapped to. -1 if the controller isn't mapped
-	 */
+	/** Establishes the input mapping object used to map input sources to SlateUser indices */
 	void SetInputManager(TSharedRef<ISlateInputManager> InputManager);
 
 	/**
@@ -1558,13 +1446,13 @@ public:
 
 	/**
 	* Unregister the notification because it is no longer desired.
-	*
-	* @param Handle                Hanlde to the delegate to unregister.
+
+	* @param Handle Handle to the delegate to unregister.
 	*/
 	void UnregisterOnWindowActionNotification(FDelegateHandle Handle);
 
 	/**
-	 * Destroys an SWindow, removing it and all its children from the Slate window list.  Notifies the native window to destroy itself and releases rendering resources
+	 * Attempts to navigate directly to the given widget
 	 *
 	 * @param InUserIndex The user that is doing the navigation
 	 * @param NavigationDestination The navigation destination widget
@@ -1573,7 +1461,7 @@ public:
 	void NavigateToWidget(const uint32 UserIndex, const TSharedPtr<SWidget>& NavigationDestination, ENavigationSource NavigationSource = ENavigationSource::FocusedWidget);
 
 	/**
-	 * Destroys an SWindow, removing it and all its children from the Slate window list.  Notifies the native window to destroy itself and releases rendering resources
+	 * Attempts to navigate directly to the widget currently under the given user's cursor
 	 *
 	 * @param InUserIndex The user that is doing the navigation
 	 * @param InNavigationType The navigation type / direction
@@ -1626,79 +1514,13 @@ private:
 	bool ExecuteNavigation(const FWidgetPath& NavigationSource, TSharedPtr<SWidget> DestinationWidget, const uint32 UserIndex, bool bAlwaysHandleNavigationAttempt);
 
 private:
-
-	void ResetDragDropState();
-
-	bool SetUserFocus(FSlateUser* User, const FWidgetPath& InFocusPath, const EFocusCause InCause);
-
-	/** Lock the cursor such that it cannot leave the bounds of the specified widget. Null widget implies no cursor lock. */
-	void LockCursor(const TSharedPtr<SWidget>& Widget);
-
-	/**
-	 * Lock the cursor such that it cannot leave the bounds of the leaf-most widget of the specified path.
-	 * Assumes a valid path.
-	 */
-	void LockCursorToPath(const FWidgetPath& WidgetPath);
-
-	/** Clear any cursor locks */
-	void UnlockCursor();
-
-	/** Make sure that the cursor lock region matcher the locking widget's geometry */
-	void UpdateCursorLockRegion();
-
-	/** State related to cursor locking. */
-	struct
-	{
-		/** Path to widget that currently holds the cursor lock; invalid path if no cursor lock. */
-		FWeakWidgetPath PathToLockingWidget;
-
-		/** Desktop Space Rect that bounds the cursor. */
-		FSlateRect LastComputedBounds;
-	} CursorLock;
-
-private:
-
-	/** Sets the LastUserInteractionTime and fires off the LastUserInteractionTimeUpdateEvent */
-	void SetLastUserInteractionTime(const double InCurrentTime);
-
-private:
-
-	// Hidden default constructor.
 	FSlateApplication();
+	void SetLastUserInteractionTime(const double InCurrentTime);	
+	bool SetUserFocus(FSlateUser& User, const FWidgetPath& InFocusPath, const EFocusCause InCause);
 
 private:
-	/** Represents a single user and pointer index for a device 
-	 *  Used to uniquely track widget state per user and per device
-	 */
-	struct FUserAndPointer
-	{
-		uint32 UserIndex;
-		uint32 PointerIndex;
-		
-		FUserAndPointer( uint32 InUserIndex, uint32 InPointerIndex )
-			: UserIndex( InUserIndex )
-			, PointerIndex( InPointerIndex )
-		{}
-
-		bool operator==( const FUserAndPointer& Other ) const
-		{
-			return UserIndex == Other.UserIndex && PointerIndex == Other.PointerIndex;
-		}
-
-		friend uint32 GetTypeHash( const FUserAndPointer& UserAndPointer )
-		{
-			return UserAndPointer.UserIndex << 16 | UserAndPointer.PointerIndex;
-		}
-	};
-
-	/**
-	 * Creates a mouse move event for the last known cursor position.  This should be called every tick to make
-	 * sure that widgets that appear (or vanish from) underneath the cursor have hover state set appropriately.
-	 */
-	void SynthesizeMouseMove();
-
 	/** Signal that a synthesized mouse move will be required after this operation. */
-	void QueueSynthesizedMouseMove();
+	void QueueSynthesizedMouseMove(const FInputEvent& SourceEvent);
 
 	/**
 	 * Will be invoked when the size of the geometry of the virtual
@@ -1711,9 +1533,6 @@ private:
 
 	TSet<FKey> PressedMouseButtons;
 
-	/** After processing an event or performing an active timer, we need to synthesize a mouse move. @see SynthesizeMouseMove */
-	int32 SynthesizeMouseMovePending;
-
 	/** true when the slate app is active; i.e. the current foreground window is from our Slate app*/
 	bool bAppIsActive;
 
@@ -1747,121 +1566,19 @@ private:
 	/** The stack of menus that are open */
 	FMenuStack MenuStack;
 
-	/** A vertical slice through the tree of widgets on screen; it represents widgets that were under the cursor last time an event was processed */
-	TMap<FUserAndPointer, FWeakWidgetPath> WidgetsUnderCursorLastEvent;
-
-	/** Stores the position for the pointer position. */
-	TMap<FUserAndPointer, FVector2D> PointerIndexPositionMap;
-
-	/** Stores the position for the last pointer position. */
-	TMap<FUserAndPointer, FVector2D> PointerIndexLastPositionMap;
-
-	/**
-	 * A helper class to wrap the weak path functionality. The advantage of using this
-	 * class is that the path can be validated and the current mouse captor (if any) can
-	 * be informed they have lost mouse capture in the case of a number of events, such as
-	 * application losing focus, or the widget being hidden.
-	 */
-	class MouseCaptorHelper
-	{
-	public:
-		/**
-		 * Returns whether or not there are any active pointer captures.
-		 */
-		bool HasCapture() const;
-
-		/**
-		* Returns whether or not the particular UserIndex has capture.
-		*/
-		bool HasCaptureForUser(uint32 UserIndex) const;
-
-		/**
-		 * Returns whether or not the particular PointerIndex has capture.
-		 */
-		bool HasCaptureForPointerIndex(uint32 UserIndex, uint32 PointerIndex) const;
-
-		bool DoesWidgetHaveMouseCapture(const TSharedPtr<const SWidget> Widget) const;
-		bool DoesWidgetHaveMouseCaptureByUser(const TSharedPtr<const SWidget> Widget, int32 UserIndex, TOptional<int32> PointerIndex) const;
-
-		/**
-		 * Sets a new mouse captor widget for a specific pointer index, invalidating the previous one if any and calling
-		 * its OnMouseCaptureLost() handler.
-		 *
-		 * @param PointerIndex	The index of the pointer which initiated the capture.
-		 * @param EventPath		The path to the event.
-		 * @param Widget		The widget that wants to capture the mouse.
-		 */
-		bool SetMouseCaptor(uint32 UserIndex, uint32 PointerIndex, const FWidgetPath& EventPath, TSharedPtr< SWidget > Widget );
-
-		/** Invalidates all current mouse captors. Calls OnMouseCaptureLost() on the current mouse captor if one exists */
-		void InvalidateCaptureForAllPointers();
-
-		/** Invalidates a specific mouse captor. Calls OnMouseCaptureLost() on the specific mouse captor if one exists */
-		void InvalidateCaptureForPointer(uint32 UserIndex, uint32 PointIndex);
-
-		/** Invalidates a specific mouse captor. Calls OnMouseCaptureLost() on the specific mouse captor if one exists */
-		void InvalidateCaptureForUser(uint32 UserIndex);
-
-		/**
-		 * Retrieves a resolved FWidgetPath for a specific pointer index, if possible.
-		 * If the path is invalid or truncated (i.e. the widget is no longer relevant) then the current mouse captor's
-		 * OnMouseCaptureLost() handler is called and the path is invalidated.
-		 *
-		 * @param PointerIndex				The index of the pointer which has capture.
-		 * @param InterruptedPathHandling	How to handled incomplete paths. "Truncate" will return a partial path, "ReturnInvalid" will return an empty path.
-		 */
-		FWidgetPath ToWidgetPath(uint32 UserIndex, uint32 PointerIndex, FWeakWidgetPath::EInterruptedPathHandling::Type InterruptedPathHandling = FWeakWidgetPath::EInterruptedPathHandling::Truncate );
-
-		FWidgetPath ToWidgetPath( FWeakWidgetPath::EInterruptedPathHandling::Type InterruptedPathHandling = FWeakWidgetPath::EInterruptedPathHandling::Truncate, const FPointerEvent* PointerEvent = nullptr );
-
-		/**
-		 * Retrieves an array of the resolved widget paths for any active captures.
-		 */
-		TArray<FWidgetPath> ToWidgetPaths();
-
-		/** Retrieves the weak path for a current mouse captor with a specific pointer index */
-		FWeakWidgetPath ToWeakPath(uint32 UserIndex, uint32 PointerIndex) const;
-		
-
-		/* Walks the weak path and retrieves the widget that is set as the current mouse captor with a specific pointer index */
-		TSharedPtr< SWidget > ToSharedWidget(uint32 UserIndex, uint32 PointerIndex) const;
-
-		/*
-		 * Retrieves an array of shared widget pointers for the active mouse captures.
-		 */
-		TArray<TSharedRef<SWidget>> ToSharedWidgets() const;
-
-		/* Walks the weak path and retrieves the window for the widget belonging to the mouse captor with the specified pointer index */
-		TSharedPtr< SWidget > ToSharedWindow(uint32 UserIndex, uint32 PointerIndex);
-
-	protected:
-		/** Call the OnMouseCaptureLost() handler for the widget captured by the specific pointer index */
-		void InformCurrentCaptorOfCaptureLoss(uint32 UserIndex, uint32 PointerIndex) const;
-
-		/** A map of pointer indices to weak widget paths for the active mouse captures */
-		TMap<FUserAndPointer, FWeakWidgetPath> PointerIndexToMouseCaptorWeakPathMap;
-	};
-	/** The current mouse captor for the application, if any. */
-	MouseCaptorHelper MouseCaptor;
-
-	/** The cursor widget and window to render that cursor for the current software cursor.*/
-	TWeakPtr<SWindow> CursorWindowPtr;
-	TWeakPtr<SWidget> CursorWidgetPtr;
-
 	/** The hit-test radius of the cursor. Default value is 0. */
 	float CursorRadius;
 
 	/**
-	 * All users currently registered with Slate.  Normally this is 1, but in a 
-	 * situation where multiple users are providing input you need to track ui state
-	 * of each user separately.
+	 * All users currently registered with Slate.
+	 * Normally there's just one, but any case where multiple users can provide input to the same application will register multiple users.
+	 * 
+	 * Note: The array may contain invalid entries. Users have associated indices that they expect to maintain, independent of the existence of other users.
 	 */
 	TArray<TSharedPtr<FSlateUser>> Users;
 
-	/**
-	 * Weak pointers to the allocated virtual users.
-	 */
-	TArray<TWeakPtr<FSlateVirtualUser>> VirtualUsers;
+	/** Weak pointers to the allocated virtual users. */
+	TArray<TWeakPtr<FSlateVirtualUserHandle>> VirtualUsers;
 
 	/**
 	 * Application throttling
@@ -1885,59 +1602,6 @@ private:
 	/** Used when considering whether to put Slate to sleep */
 	double LastMouseMoveTime;
 
-	/** Helper for detecting when a drag should begin */
-	class FDragDetector
-	{
-	public:
-		FDragDetector()
-		{
-		}
-
-		void StartDragDetection(const FWidgetPath& PathToWidget, int32 UserIndex, int32 PointerIndex, FKey DragButton, FVector2D StartLocation);
-		bool IsDetectingDrag(const FPointerEvent& PointerEvent);
-		bool DetectDrag(const FPointerEvent& PointerEvent, float DragTriggerDistance, FWeakWidgetPath*& OutWeakWidgetPath);
-		void OnPointerRelease(const FPointerEvent& PointerEvent);
-		void ResetDetection();
-
-	private:
-
-		struct FDragDetectionState
-		{
-			FDragDetectionState()
-				: DetectDragStartLocation(FVector2D::ZeroVector)
-				, DetectDragButton(EKeys::Invalid)
-				, DetectDragUserIndex(INDEX_NONE)
-				, DetectDragPointerIndex(INDEX_NONE)
-			{
-			}
-
-			FDragDetectionState(const FWidgetPath& PathToWidget, int32 UserIndex, int32 PointerIndex, FKey DragButton, FVector2D StartLocation)
-				: DetectDragForWidget(PathToWidget)
-				, DetectDragStartLocation(StartLocation)
-				, DetectDragButton(DragButton)
-				, DetectDragUserIndex(UserIndex)
-				, DetectDragPointerIndex(PointerIndex)
-			{
-			}
-
-			/** If not null, a widget has request that we detect a drag being triggered in this widget and send an OnDragDetected() event*/
-			FWeakWidgetPath DetectDragForWidget;
-			/** Location from which be begin detecting the drag */
-			FVector2D DetectDragStartLocation;
-			/** Button that must be pressed to trigger the drag */
-			FKey DetectDragButton;
-			/** User index of the drag operation */
-			int32 DetectDragUserIndex;
-			/** Pointer index of the drag operation */
-			int32 DetectDragPointerIndex;
-		};
-
-		/** A map of pointer indices drag states currently being tracked. */
-		TMap<FUserAndPointer, FDragDetectionState> PointerIndexToDragState;
-	};
-
-	FDragDetector DragDetector;
-
 	/** Support for auto-dismissing pop-ups */
 	FPopupSupport PopupSupport;
 	
@@ -1989,77 +1653,15 @@ private:
 	/** List of delegates that need to be called when the window action occurs. */
 	TArray<FOnWindowAction> OnWindowActionNotifications;
 
-	/**
-	 * Tool-tips
-	 */
-
-	/** Window that we'll re-use for spawned tool tips */
-	TWeakPtr< SWindow > ToolTipWindow;
-
-	/** Widget that last visualized the tooltip;  */
-	TWeakPtr< SWidget > TooltipVisualizerPtr;
-
-	/** The tool tip that is currently being displayed; can be invalid (i.e. not showing tooltip) */
-	TWeakPtr<IToolTip> ActiveToolTip;
-
-	/** The widget that sourced the currently active tooltip widget */
-	TWeakPtr<SWidget> ActiveToolTipWidgetSource;
-
-	/** Whether tool-tips are allowed to spawn at all.  Used only for debugging purposes. */
-	int32 bAllowToolTips;
-
-	/** How long before the tool tip should start fading in? */
-	float ToolTipDelay;
-
-	/** Tool-tip fade-in duration */
-	float ToolTipFadeInDuration;
-
-	/** Absolute real time that the current tool-tip window was summoned */
-	double ToolTipSummonTime;
-
-	/** Desired tool tip position in screen space, updated whenever the mouse moves */
-	FVector2D DesiredToolTipLocation;
-
-	/** Direction that tool-tip is being repelled from a force field in.  We cache this to avoid tool-tips
-	    teleporting between different offset directions as the user moves the mouse cursor around. */
-	enum class EToolTipOffsetDirection : uint8
-	{
-		Undetermined,
-		Down,
-		Right
-	};
-	EToolTipOffsetDirection ToolTipOffsetDirection;
-
 	/** The top of the Style tree. */
 	const class FStyleNode* RootStyleNode;
 
-	/**
-	 * Drag and Drop (DragDrop)
-	 */
-	
-	/** When not null, the content of the current drag drop operation. */
-	TSharedPtr< FDragDropOperation > DragDropContent;
-
-	/** The window the drag drop content is over. */
-	TWeakPtr< SWindow > DragDropWindowPtr;
-
 	/** Whether or not we are requesting that we leave debugging mode after the tick is complete */
 	bool bRequestLeaveDebugMode;
 	/** Whether or not we need to leave debug mode for single stepping */
 	bool bLeaveDebugForSingleStep;
 	TAttribute<bool> NormalExecutionGetter;
 
-	/**
-	 * Console objects
-	 */
-
-	/** AllowToolTips console variable **/
-	FAutoConsoleVariableRef CVarAllowToolTips;
-	/** ToolTipDelay console variable **/
-	FAutoConsoleVariableRef CVarToolTipDelay;
-	/** ToolTipFadeInDuration console variable **/
-	FAutoConsoleVariableRef CVarToolTipFadeInDuration;
-
 	/**
 	 * Modal Windows
 	 */
@@ -2100,9 +1702,6 @@ private:
 	/** .ini controlled option to allow or disallow software cursor rendering */
 	bool bSoftwareCursorAvailable;
 
-	/** The OS or actions taken by the user may require we refresh the current state of the cursor. */
-	bool bQueryCursorRequested;
-
 	/**
 	 * Slate look and feel
 	 */
@@ -2124,7 +1723,6 @@ private:
 	// but there could be another monitor on any of the sides.
 	FSlateRect VirtualDesktopRect;
 
-	/** This factory function creates a navigation config for each slate user. */
 	TSharedRef<FNavigationConfig> NavigationConfig;
 
 #if WITH_EDITOR
diff --git a/Engine/Source/Runtime/Slate/Public/Framework/Application/SlateUser.h b/Engine/Source/Runtime/Slate/Public/Framework/Application/SlateUser.h
new file mode 100644
--- /dev/null
+++ b/Engine/Source/Runtime/Slate/Public/Framework/Application/SlateUser.h
@@ -0,0 +1,299 @@
+// Copyright 1998-2019 Epic Games, Inc. All Rights Reserved.
+
+#pragma once
+
+#include "SlateScope.h"
+#include "Layout/WidgetPath.h"
+#include "GestureDetector.h"
+
+class SWidget;
+class SWindow;
+class ICursor;
+class FDragDropOperation;
+class FMenuStack;
+
+/** Handle to a virtual user of slate. */
+class SLATE_API FSlateVirtualUserHandle
+{
+public:
+	virtual ~FSlateVirtualUserHandle();
+
+	FORCEINLINE int32 GetUserIndex() const { return UserIndex; }
+	FORCEINLINE int32 GetVirtualUserIndex() const { return VirtualUserIndex; }
+
+SLATE_SCOPE:
+	FSlateVirtualUserHandle(int32 InUserIndex, int32 InVirtualUserIndex);
+
+private:
+	/** The index the user was assigned. */
+	int32 UserIndex;
+
+	/** The virtual index the user was assigned. */
+	int32 VirtualUserIndex;
+};
+
+/**
+ * Slate's representation of an individual input-providing user.  
+ * As new input sources are connected, new SlateUsers are created.
+ */
+class SLATE_API FSlateUser : public TSharedFromThis<FSlateUser>
+{
+public:
+	virtual ~FSlateUser();
+
+	FORCEINLINE int32 GetUserIndex() const { return UserIndex; }
+	FORCEINLINE bool IsVirtualUser() const { return !Cursor.IsValid(); }
+
+	TSharedPtr<SWidget> GetFocusedWidget() const;
+	TOptional<EFocusCause> HasFocus(TSharedPtr<const SWidget> Widget) const;
+	bool ShouldShowFocus(TSharedPtr<const SWidget> Widget) const;
+	bool HasFocusedDescendants(TSharedRef<const SWidget> Widget) const;
+	bool SetFocus(const TSharedRef<SWidget>& WidgetToFocus, EFocusCause ReasonFocusIsChanging = EFocusCause::SetDirectly);
+	void ClearFocus(EFocusCause ReasonFocusIsChanging = EFocusCause::SetDirectly);
+
+	bool HasAnyCapture() const;
+	bool HasCursorCapture() const;
+	bool HasCapture(uint32 PointerIndex) const;
+	
+	bool DoesWidgetHaveAnyCapture(TSharedPtr<const SWidget> Widget) const;
+	bool DoesWidgetHaveCursorCapture(TSharedPtr<const SWidget> Widget) const;
+	bool DoesWidgetHaveCapture(TSharedPtr<const SWidget> Widget, uint32 PointerIndex) const;
+	
+	bool SetCursorCaptor(TSharedRef<const SWidget> Widget, const FWidgetPath& EventPath);
+	bool SetPointerCaptor(uint32 PointerIndex, TSharedRef<const SWidget> Widget, const FWidgetPath& EventPath);
+	
+	void ReleaseAllCapture();
+	void ReleaseCursorCapture();
+	void ReleaseCapture(uint32 PointerIndex);
+
+	TArray<FWidgetPath> GetCaptorPaths();
+	FWidgetPath GetCursorCaptorPath(FWeakWidgetPath::EInterruptedPathHandling::Type InterruptedPathHandling = FWeakWidgetPath::EInterruptedPathHandling::Truncate, const FPointerEvent* PointerEvent = nullptr);
+	FWidgetPath GetCaptorPath(uint32 PointerIndex, FWeakWidgetPath::EInterruptedPathHandling::Type InterruptedPathHandling = FWeakWidgetPath::EInterruptedPathHandling::Truncate, const FPointerEvent* PointerEvent = nullptr);
+
+	FWeakWidgetPath GetWeakCursorCapturePath() const;
+	FWeakWidgetPath GetWeakCapturePath(uint32 PointerIndex) const;
+
+	TArray<TSharedRef<SWidget>> GetCaptorWidgets() const;
+	TSharedPtr<SWidget> GetCursorCaptor() const;
+	TSharedPtr<SWidget> GetPointerCaptor(uint32 PointerIndex) const;
+
+	void SetCursorPosition(int32 PosX, int32 PosY);
+	void SetCursorPosition(const FVector2D& NewCursorPos);
+	void SetPointerPosition(uint32 PointerIndex, int32 PosX, int32 PosY);
+	void SetPointerPosition(uint32 PointerIndex, const FVector2D& NewPointerPos);
+	
+	FVector2D GetCursorPosition() const;
+	FVector2D GetPreviousCursorPosition() const;
+	
+	FVector2D GetPointerPosition(uint32 PointerIndex) const;
+	FVector2D GetPreviousPointerPosition(uint32 PointerIndex) const;
+
+	bool IsWidgetUnderCursor(TSharedPtr<const SWidget> Widget) const;
+	bool IsWidgetUnderPointer(TSharedPtr<const SWidget> Widget, uint32 PointerIndex) const;
+	bool IsWidgetUnderAnyPointer(TSharedPtr<const SWidget> Widget) const;
+
+	bool IsWidgetDirectlyUnderCursor(TSharedPtr<const SWidget> Widget) const;
+	bool IsWidgetDirectlyUnderPointer(TSharedPtr<const SWidget> Widget, uint32 PointerIndex) const;
+	bool IsWidgetDirectlyUnderAnyPointer(TSharedPtr<const SWidget> Widget) const;
+
+	FWeakWidgetPath GetLastWidgetsUnderCursor() const;
+	FWeakWidgetPath GetLastWidgetsUnderPointer(uint32 PointerIndex) const;
+	const TMap<uint32, FWeakWidgetPath>& GetWidgetsUnderPointerLastEventByIndex() const { return WidgetsUnderPointerLastEventByIndex; }
+
+	TSharedPtr<FDragDropOperation> GetDragDropContent() const { return DragDropContent; }
+	bool IsDragDropping() const;
+	bool IsDragDroppingAffected(const FPointerEvent& InPointerEvent) const;
+	void CancelDragDrop();
+	
+	void ShowTooltip(const TSharedRef<IToolTip>& InTooltip, const FVector2D& InSpawnLocation);
+	void CloseTooltip();
+
+	const FGestureDetector& GetGestureDetector() const { return GestureDetector; }
+
+SLATE_SCOPE:
+	static TSharedRef<FSlateUser> Create(int32 InUserIndex, TSharedPtr<ICursor> InCursor);
+	
+	FORCEINLINE bool HasValidFocusPath() const { return WeakFocusPath.IsValid(); }
+	FORCEINLINE const FWeakWidgetPath& GetWeakFocusPath() const { return WeakFocusPath; }
+
+	FORCEINLINE TSharedRef<FWidgetPath> GetFocusPath() const
+	{
+		if (!StrongFocusPath.IsValid())
+		{
+			StrongFocusPath = WeakFocusPath.ToWidgetPathRef();
+		}
+		return StrongFocusPath.ToSharedRef();
+	}
+
+	void DrawWindowlessDragDropContent(const TSharedRef<SWindow>& WindowToDraw, FSlateWindowElementList& WindowElementList, int32& MaxLayerId);
+	void DrawCursor(const TSharedRef<SWindow>& WindowToDraw, FSlateWindowElementList& WindowElementList, int32& MaxLayerId);
+
+	void QueueSyntheticCursorMove();
+	bool SynthesizeCursorMoveIfNeeded();
+
+	TSharedPtr<ICursor> GetCursor() const { return Cursor; }
+	void LockCursor(const TSharedRef<SWidget>& Widget);
+	void UnlockCursor();
+	void UpdateCursor();
+	void ProcessCursorReply(const FCursorReply& CursorReply);
+	void RequestCursorQuery() { bQueryCursorRequested = true; }
+	void QueryCursor();
+	
+	void SetFocusPath(const FWidgetPath& NewFocusPath, EFocusCause InFocusCause, bool bInShowFocus);
+	
+	void FinishFrame();
+	void NotifyWindowDestroyed(TSharedRef<SWindow> DestroyedWindow);
+
+	void NotifyPointerMoveBegin(const FPointerEvent& PointerEvent);
+	void NotifyPointerMoveComplete(const FPointerEvent& PointerEvent, const FWidgetPath& WidgetsUnderPointer);
+	void NotifyPointerReleased(const FPointerEvent& PointerEvent, const FWidgetPath& WidgetsUnderCursor, TSharedPtr<FDragDropOperation> DroppedContent, bool bWasHandled);
+	void UpdatePointerPosition(const FPointerEvent& PointerEvent);
+
+	void StartDragDetection(const FWidgetPath& PathToWidget, int32 PointerIndex, FKey DragButton, FVector2D StartLocation);
+	FWidgetPath DetectDrag(const FPointerEvent& PointerEvent, float DragTriggerDistance);
+	bool IsDetectingDrag(uint32 PointerIndex) const;
+	void ResetDragDetection();
+
+	void SetDragDropContent(TSharedRef<FDragDropOperation> InDragDropContent);
+	void ResetDragDropContent();
+
+	int32 GetFocusVersion() const { return FocusVersion; }
+	void IncrementFocusVersion() { FocusVersion++; }
+
+	void UpdateTooltip(const FMenuStack& MenuStack, bool bCanSpawnNewTooltip);
+	TSharedPtr<IToolTip> GetTooltip() const;
+	void ResetTooltipWindow();
+	bool IsWindowHousingInteractiveTooltip(const TSharedRef<const SWindow>& WindowToTest) const;
+
+	FGestureDetector& GetGestureDetector() { return GestureDetector; }
+
+#if PLATFORM_MAC
+	// Unclear why from existing code, but Mac seems to need to cache & restore all mouse captor paths when activating the top level window
+	const TMap<uint32, FWeakWidgetPath> GetCaptorPathsByIndex() const { return PointerCaptorPathsByIndex; }
+	void RestoreCaptorPathsByIndex(const TMap<uint32, FWeakWidgetPath>& InPointerCaptorPathsByIndex) { PointerCaptorPathsByIndex = InPointerCaptorPathsByIndex; }
+#endif
+
+private:
+	FSlateUser(int32 InUserIndex, TSharedPtr<ICursor> InCursor);
+	void UpdatePointerPosition(uint32 PointerIndex, const FVector2D& Position);
+	void LockCursorInternal(const FWidgetPath& WidgetPath);
+	TSharedRef<SWindow> GetOrCreateTooltipWindow();
+
+	/** The index the user was assigned. */
+	int32 UserIndex = INDEX_NONE;
+
+	/** The cursor this user is in control of. Guaranteed to be valid for all real users, absence implies this is a virtual user. */
+	TSharedPtr<ICursor> Cursor;
+
+	/** The OS or actions taken by the user may require we refresh the current state of the cursor. */
+	bool bQueryCursorRequested = false;
+
+	/**
+	 * Whenever something happens that can affect Slate layout, we need to process a mouse move even if the mouse didn't move at all.
+	 * @see SynthesizeCursorMoveIfNeeded
+	 */
+	int32 NumPendingSyntheticCursorMoves = 0;
+
+	/** The cursor widget and window to render that cursor for the current software cursor.*/
+	TWeakPtr<SWindow> CursorWindowPtr;
+	TWeakPtr<SWidget> CursorWidgetPtr;
+
+	/** A weak path to the widget currently focused by a user, if any. */
+	FWeakWidgetPath WeakFocusPath;
+
+	/** A strong widget path to the focused widget, if any. This is cleared after the end of pumping messages. */
+	mutable TSharedPtr<FWidgetPath> StrongFocusPath;
+
+	/** Reason a widget was focused by a user, if any. */
+	EFocusCause FocusCause = EFocusCause::Cleared;
+
+	/** If we should show this focus */
+	bool bShowFocus = false;
+
+	/**
+	 * The FocusVersion is used to know if the focus state is modified for a user while processing focus
+	 * events, that way upon returning from focus calls, we know if we should abandon the remainder of the event.
+	 */
+	int32 FocusVersion = 0;
+
+	/** Current position of all pointers controlled by this user */
+	TMap<uint32, FVector2D> PointerPositionsByIndex;
+	TMap<uint32, FVector2D> PreviousPointerPositionsByIndex;
+
+	/** Weak paths to widgets that are currently capturing a particular pointer */
+	TMap<uint32, FWeakWidgetPath> PointerCaptorPathsByIndex;
+
+	struct FDragDetectionState
+	{
+		FDragDetectionState(const FWidgetPath& PathToWidget, int32 PointerIndex, FKey DragButton, const FVector2D& StartLocation)
+			: DetectDragForWidget(PathToWidget)
+			, DragStartLocation(StartLocation)
+			, TriggerButton(DragButton)
+			, PointerIndex(PointerIndex)
+		{
+		}
+
+		/** If not null, a widget has requested that we detect a drag being triggered in this widget and send an OnDragDetected() event*/
+		FWeakWidgetPath DetectDragForWidget;
+
+		FVector2D DragStartLocation = FVector2D::ZeroVector;
+		FKey TriggerButton = EKeys::Invalid;
+		int32 PointerIndex = INDEX_NONE;
+	};
+
+	/** Current drag status for pointers currently executing a drag/drop operation */
+	TMap<uint32, FDragDetectionState> DragStatesByPointerIndex;
+
+	/** When not null, the content of the current drag drop operation. */
+	TSharedPtr<FDragDropOperation> DragDropContent;
+
+	/** The window the drag drop content is over. */
+	TWeakPtr<SWindow> DragDropWindowPtr;
+
+	/** Weak paths to the last widget each pointer was under last time an event was processed */
+	TMap<uint32, FWeakWidgetPath> WidgetsUnderPointerLastEventByIndex;
+
+	/** Path to widget that currently holds the cursor lock; invalid path if no cursor lock. */
+	FWeakWidgetPath LockingWidgetPath;
+
+	/** Desktop Space Rect that bounds the cursor. */
+	FSlateRect LastComputedLockBounds;
+
+	/** Window that we'll re-use for spawned tool tips */
+	TWeakPtr<SWindow> TooltipWindowPtr;
+
+	enum class ETooltipOffsetDirection : uint8
+	{
+		Undetermined,
+		Down,
+		Right
+	};
+	struct FActiveTooltipInfo
+	{
+		void Reset();
+
+		/** The actual tooltip object running the show here (not necessarily the same object as the tooltip SWidget, but can be) */
+		TWeakPtr<IToolTip> Tooltip;
+
+		/** The widget that is taking care of visualizing this tooltip (can be null) */
+		TWeakPtr<SWidget> TooltipVisualizer;
+
+		/** The widget that sourced the currently active tooltip */
+		TWeakPtr<SWidget> SourceWidget;
+
+		/** Desired position of the tooltip in screen space, updated whenever the mouse moves */
+		FVector2D DesiredLocation = FVector2D::ZeroVector;
+
+		/** The time at which the tooltip was summoned */
+		double SummonTime = 0.0;
+
+		/** 
+		 * Direction in which the tooltip widget is being offset from the source widget  
+		 * Cached to prevent the direction of offset from changing as the user moves the mouse cursor (otherwise the tooltip teleports around)
+		 */
+		ETooltipOffsetDirection OffsetDirection = ETooltipOffsetDirection::Undetermined;
+	};
+	FActiveTooltipInfo ActiveTooltipInfo;
+
+	FGestureDetector GestureDetector;
+};
\ No newline at end of file
diff --git a/Engine/Source/Runtime/Slate/Public/SlateScope.h b/Engine/Source/Runtime/Slate/Public/SlateScope.h
new file mode 100644
--- /dev/null
+++ b/Engine/Source/Runtime/Slate/Public/SlateScope.h
@@ -0,0 +1,11 @@
+// Copyright 1998-2019 Epic Games, Inc. All Rights Reserved.
+
+#pragma once
+
+// SLATE_MODULE is defined private to the module in Slate.build.cs
+// This allows us to establish a scope that is public within the Slate module itself, but protected from all consumers of the module
+#ifdef SLATE_MODULE
+#define SLATE_SCOPE public
+#else
+#define SLATE_SCOPE protected
+#endif
\ No newline at end of file
diff --git a/Engine/Source/Runtime/Slate/Public/Widgets/Input/SComboBox.h b/Engine/Source/Runtime/Slate/Public/Widgets/Input/SComboBox.h
--- a/Engine/Source/Runtime/Slate/Public/Widgets/Input/SComboBox.h
+++ b/Engine/Source/Runtime/Slate/Public/Widgets/Input/SComboBox.h
@@ -16,6 +16,7 @@
 #include "Styling/CoreStyle.h"
 #include "Framework/SlateDelegates.h"
 #include "Framework/Application/SlateApplication.h"
+#include "Framework/Application/SlateUser.h"
 #include "Widgets/Text/STextBlock.h"
 #include "Widgets/Layout/SBox.h"
 #include "Widgets/Input/SComboButton.h"
@@ -454,10 +455,11 @@ private:
 			}
 
 			// Set focus back to ComboBox for users focusing the ListView that just closed
-			FSlateApplication::Get().ForEachUser([&](FSlateUser* User) {
-				if (FSlateApplication::Get().HasUserFocusedDescendants(AsShared(), User->GetUserIndex()))
+			TSharedRef<SWidget> ThisRef = AsShared();
+			FSlateApplication::Get().ForEachUser([&ThisRef](FSlateUser& User) {
+				if (User.HasFocusedDescendants(ThisRef))
 				{
-					FSlateApplication::Get().SetUserFocus(User->GetUserIndex(), AsShared(), EFocusCause::SetDirectly);
+					User.SetFocus(ThisRef, EFocusCause::SetDirectly);
 				}
 			});
 
diff --git a/Engine/Source/Runtime/Slate/Slate.Build.cs b/Engine/Source/Runtime/Slate/Slate.Build.cs
--- a/Engine/Source/Runtime/Slate/Slate.Build.cs
+++ b/Engine/Source/Runtime/Slate/Slate.Build.cs
@@ -6,7 +6,8 @@ public class Slate : ModuleRules
 {
 	public Slate(ReadOnlyTargetRules Target) : base(Target)
 	{
-		SharedPCHHeaderFile = "Public/SlateSharedPCH.h";
+        PrivateDefinitions.Add("SLATE_MODULE=1");
+        SharedPCHHeaderFile = "Public/SlateSharedPCH.h";
 
 		PublicDependencyModuleNames.AddRange(
 			new string[] {
diff --git a/Engine/Source/Runtime/SlateCore/Private/Debugging/SlateDebugging.cpp b/Engine/Source/Runtime/SlateCore/Private/Debugging/SlateDebugging.cpp
--- a/Engine/Source/Runtime/SlateCore/Private/Debugging/SlateDebugging.cpp
+++ b/Engine/Source/Runtime/SlateCore/Private/Debugging/SlateDebugging.cpp
@@ -157,7 +157,7 @@ FSlateDebuggingMouseCaptureEventArgs::FSlateDebuggingMouseCaptureEventArgs(
 	bool InCaptured,
 	uint32 InUserIndex,
 	uint32 InPointerIndex,
-	const TSharedPtr<SWidget>& InCapturingWidget)
+	const TSharedPtr<const SWidget>& InCapturingWidget)
 	: Captured(InCaptured)
 	, UserIndex(InUserIndex)
 	, PointerIndex(InPointerIndex)
@@ -252,12 +252,12 @@ void FSlateDebugging::BroadcastExecuteNavigation()
 	ExecuteNavigationEvent.Broadcast(FSlateDebuggingExecuteNavigationEventArgs());
 }
 
-void FSlateDebugging::BroadcastMouseCapture(uint32 UserIndex, uint32 PointerIndex, const TSharedPtr<SWidget>& InCapturingWidget)
+void FSlateDebugging::BroadcastMouseCapture(uint32 UserIndex, uint32 PointerIndex, TSharedPtr<const SWidget> InCapturingWidget)
 {
 	MouseCaptureEvent.Broadcast(FSlateDebuggingMouseCaptureEventArgs(true, UserIndex, PointerIndex, InCapturingWidget));
 }
 
-void FSlateDebugging::BroadcastMouseCaptureLost(uint32 UserIndex, uint32 PointerIndex, const TSharedPtr<SWidget>& InWidgetLostCapture)
+void FSlateDebugging::BroadcastMouseCaptureLost(uint32 UserIndex, uint32 PointerIndex, TSharedPtr<const SWidget> InWidgetLostCapture)
 {
 	MouseCaptureEvent.Broadcast(FSlateDebuggingMouseCaptureEventArgs(false, UserIndex, PointerIndex, InWidgetLostCapture));
 }
diff --git a/Engine/Source/Runtime/SlateCore/Public/Debugging/SlateDebugging.h b/Engine/Source/Runtime/SlateCore/Public/Debugging/SlateDebugging.h
--- a/Engine/Source/Runtime/SlateCore/Public/Debugging/SlateDebugging.h
+++ b/Engine/Source/Runtime/SlateCore/Public/Debugging/SlateDebugging.h
@@ -153,13 +153,13 @@ public:
 		bool InCaptured,
 		uint32 InUserIndex,
 		uint32 InPointerIndex,
-		const TSharedPtr<SWidget>& InCapturingWidget
+		const TSharedPtr<const SWidget>& InCapturingWidget
 	);
 
 	bool Captured;
 	uint32 UserIndex;
 	uint32 PointerIndex;
-	const TSharedPtr<SWidget>& CaptureWidget;
+	const TSharedPtr<const SWidget>& CaptureWidget;
 };
 
 
@@ -235,8 +235,8 @@ public:
 	DECLARE_MULTICAST_DELEGATE_OneParam(FWidgetMouseCaptureEvent, const FSlateDebuggingMouseCaptureEventArgs& /*EventArgs*/);
 	static FWidgetMouseCaptureEvent MouseCaptureEvent;
 
-	static void BroadcastMouseCapture(uint32 UserIndex, uint32 PointerIndex, const TSharedPtr<SWidget>& InCapturingWidget);
-	static void BroadcastMouseCaptureLost(uint32 UserIndex, uint32 PointerIndex, const TSharedPtr<SWidget>& InWidgetLostCapture);
+	static void BroadcastMouseCapture(uint32 UserIndex, uint32 PointerIndex, TSharedPtr<const SWidget> InCapturingWidget);
+	static void BroadcastMouseCaptureLost(uint32 UserIndex, uint32 PointerIndex, TSharedPtr<const SWidget> InWidgetLostCapture);
 
 public:
 	/**  */
diff --git a/Engine/Source/Runtime/UMG/Public/Components/WidgetInteractionComponent.h b/Engine/Source/Runtime/UMG/Public/Components/WidgetInteractionComponent.h
--- a/Engine/Source/Runtime/UMG/Public/Components/WidgetInteractionComponent.h
+++ b/Engine/Source/Runtime/UMG/Public/Components/WidgetInteractionComponent.h
@@ -12,7 +12,7 @@
 #include "Layout/WidgetPath.h"
 #include "WidgetInteractionComponent.generated.h"
 
-class FSlateVirtualUser;
+class FSlateVirtualUserHandle;
 class UPrimitiveComponent;
 class UWidgetComponent;
 
@@ -187,7 +187,7 @@ private:
 	 * the mouse and keyboard focus input (the viewport), so that things like the player controller receive
 	 * standard hardware input.
 	 */
-	TSharedPtr<FSlateVirtualUser> VirtualUser;
+	TSharedPtr<FSlateVirtualUserHandle> VirtualUser;
 
 public:
 
