Fixing Mac compile error

diff --git a/Engine/Source/Runtime/Slate/Private/Framework/Application/SlateApplication.cpp b/Engine/Source/Runtime/Slate/Private/Framework/Application/SlateApplication.cpp
--- a/Engine/Source/Runtime/Slate/Private/Framework/Application/SlateApplication.cpp
+++ b/Engine/Source/Runtime/Slate/Private/Framework/Application/SlateApplication.cpp
@@ -4516,7 +4516,6 @@ FReply FSlateApplication::RoutePointerDownEvent(const FWidgetPath& WidgetsUnderP
 			|| bIsDetectingLMBDrag)
 #endif
 		)
-	)
 	{
 		for ( int32 WidgetIndex = WidgetsUnderPointer.Widgets.Num() - 1; WidgetIndex >= 0; --WidgetIndex )
 		{
