#UE4 Quick fix for a crash found in automation

diff --git a/Engine/Source/Runtime/Slate/Private/Framework/Application/SlateApplication.cpp b/Engine/Source/Runtime/Slate/Private/Framework/Application/SlateApplication.cpp
--- a/Engine/Source/Runtime/Slate/Private/Framework/Application/SlateApplication.cpp
+++ b/Engine/Source/Runtime/Slate/Private/Framework/Application/SlateApplication.cpp
@@ -3698,6 +3698,12 @@ TSharedRef<FSlateVirtualUserHandle> FSlateApplication::FindOrCreateVirtualUser(i
 
 TSharedRef<FSlateUser> FSlateApplication::GetOrCreateUser(int32 UserIndex)
 {
+	// @todo For now we are allowing a little shenanigans since we are crashing in automation
+	if ( UserIndex < 0 )
+	{
+		UserIndex = 0;
+	}
+	
 	if (TSharedPtr<FSlateUser> FoundUser = GetUser(UserIndex))
 	{
 		return FoundUser.ToSharedRef();
