Resolve issue where FSlateUser::GetCaptorPaths ensures when the map is modified during iteration

diff --git a/Engine/Source/Runtime/Slate/Private/Framework/Application/SlateUser.cpp b/Engine/Source/Runtime/Slate/Private/Framework/Application/SlateUser.cpp
--- a/Engine/Source/Runtime/Slate/Private/Framework/Application/SlateUser.cpp
+++ b/Engine/Source/Runtime/Slate/Private/Framework/Application/SlateUser.cpp
@@ -254,13 +254,20 @@ void FSlateUser::ReleaseCapture(uint32 PointerIndex)
 
 TArray<FWidgetPath> FSlateUser::GetCaptorPaths()
 {
-	TArray<FWidgetPath> AllPaths;
-	for (const auto& IndexPathPair : PointerCaptorPathsByIndex)
+	TArray<FWidgetPath> CaptorPaths;
+
+	TArray<uint32> CaptureIndices;
+	PointerCaptorPathsByIndex.GenerateKeyArray(CaptureIndices);
+	for (uint32 PointerIndex : CaptureIndices)
 	{
-		AllPaths.Emplace(GetCaptorPath(IndexPathPair.Key));
+		FWidgetPath CaptorPath = GetCaptorPath(PointerIndex);
+		if (CaptorPath.IsValid())
+		{
+			CaptorPaths.Add(CaptorPath);
+		}
 	}
 
-	return AllPaths;
+	return CaptorPaths;
 }
 
 FWidgetPath FSlateUser::GetCursorCaptorPath(FWeakWidgetPath::EInterruptedPathHandling::Type InterruptedPathHandling, const FPointerEvent* PointerEvent)
