FAnalogCursor now inherits from TSharedFromThis
- Allows FAnalogCursor and subclasses to bind delegates w/ BindSP (safe since InputPreprocessors are already always tracked via shared ptr)

diff --git a/Engine/Source/Runtime/Slate/Public/Framework/Application/AnalogCursor.h b/Engine/Source/Runtime/Slate/Public/Framework/Application/AnalogCursor.h
--- a/Engine/Source/Runtime/Slate/Public/Framework/Application/AnalogCursor.h
+++ b/Engine/Source/Runtime/Slate/Public/Framework/Application/AnalogCursor.h
@@ -30,7 +30,7 @@ enum class EAnalogStick : uint8
 /**
  * A class that simulates a cursor driven by an analog stick.
  */
-class SLATE_API FAnalogCursor : public IInputProcessor
+class SLATE_API FAnalogCursor : public IInputProcessor, public TSharedFromThis<FAnalogCursor>
 {
 public:
 	FAnalogCursor();
