Removed #pragma once from the top of SlateScope.h to resolve static analysis error
- Still compiles fine, and now matches PartyPackage.h exactly (which SA has no issue with), so this should resolve the SA error

diff --git a/Engine/Source/Runtime/Slate/Public/SlateScope.h b/Engine/Source/Runtime/Slate/Public/SlateScope.h
--- a/Engine/Source/Runtime/Slate/Public/SlateScope.h
+++ b/Engine/Source/Runtime/Slate/Public/SlateScope.h
@@ -1,6 +1,6 @@
 // Copyright 1998-2019 Epic Games, Inc. All Rights Reserved.
 
-#pragma once
+// Can't be #pragma once because we want this to be interpreted differently depending on the module that is referencing the including file
 
 // SLATE_MODULE is defined private to the module in Slate.build.cs
 // This allows us to establish a scope that is public within the Slate module itself, but protected from all consumers of the module
