GitHub 5548 : Add MinAspectRatio to SizeBox

PR #5548: Add MinAspectRatio to SizeBox (Contributed by [REDACTED_PERSON])

diff --git a/Engine/Source/Runtime/Slate/Private/Widgets/Layout/SBox.cpp b/Engine/Source/Runtime/Slate/Private/Widgets/Layout/SBox.cpp
--- a/Engine/Source/Runtime/Slate/Private/Widgets/Layout/SBox.cpp
+++ b/Engine/Source/Runtime/Slate/Private/Widgets/Layout/SBox.cpp
@@ -22,23 +22,24 @@ void SBox::Construct( const FArguments& InArgs )
 	MaxDesiredWidth = InArgs._MaxDesiredWidth;
 	MaxDesiredHeight = InArgs._MaxDesiredHeight;
 
+	MinAspectRatio = InArgs._MinAspectRatio;
 	MaxAspectRatio = InArgs._MaxAspectRatio;
 
 	ChildSlot
 		.HAlign( InArgs._HAlign )
 		.VAlign( InArgs._VAlign )
 		.Padding( InArgs._Padding )
-	[
-		InArgs._Content.Widget
-	];
+		[
+			InArgs._Content.Widget
+		];
 }
 
 void SBox::SetContent(const TSharedRef< SWidget >& InContent)
 {
 	ChildSlot
-	[
-		InContent
-	];
+		[
+			InContent
+		];
 
 	Invalidate(EInvalidateWidget::Layout);
 }
@@ -124,6 +125,15 @@ void SBox::SetMaxDesiredHeight(TAttribute<FOptionalSize> InMaxDesiredHeight)
 	}
 }
 
+void SBox::SetMinAspectRatio(TAttribute<FOptionalSize> InMinAspectRatio)
+{
+	if (!MinAspectRatio.IdenticalTo(InMinAspectRatio))
+	{
+		MinAspectRatio = InMinAspectRatio;
+		Invalidate(EInvalidateWidget::LayoutAndVolatility);
+	}
+}
+
 void SBox::SetMaxAspectRatio(TAttribute<FOptionalSize> InMaxAspectRatio)
 {
 	if (!MaxAspectRatio.IdenticalTo(InMaxAspectRatio))
@@ -214,6 +224,7 @@ void SBox::OnArrangeChildren( const FGeometry& AllottedGeometry, FArrangedChildr
 	const EVisibility ChildVisibility = ChildSlot.GetWidget()->GetVisibility();
 	if ( ArrangedChildren.Accepts(ChildVisibility) )
 	{
+		const FOptionalSize CurrentMinAspectRatio = MinAspectRatio.Get();
 		const FOptionalSize CurrentMaxAspectRatio = MaxAspectRatio.Get();
 		const FMargin SlotPadding(ChildSlot.SlotPadding.Get());
 		bool bAlignChildren = true;
@@ -221,26 +232,39 @@ void SBox::OnArrangeChildren( const FGeometry& AllottedGeometry, FArrangedChildr
 		AlignmentArrangeResult XAlignmentResult(0, 0);
 		AlignmentArrangeResult YAlignmentResult(0, 0);
 
-		if ( CurrentMaxAspectRatio.IsSet() )
+		if (CurrentMaxAspectRatio.IsSet() || CurrentMinAspectRatio.IsSet())
 		{
 			float CurrentWidth = FMath::Min(AllottedGeometry.Size.X, ChildSlot.GetWidget()->GetDesiredSize().X);
 			float CurrentHeight = FMath::Min(AllottedGeometry.Size.Y, ChildSlot.GetWidget()->GetDesiredSize().Y);
 
-			float AspectRatioWidth = CurrentMaxAspectRatio.Get();
-			if ( AspectRatioWidth != 0 && CurrentHeight > 0 && CurrentWidth > 0 )
+			float MinAspectRatioWidth = CurrentMinAspectRatio.IsSet() ? CurrentMinAspectRatio.Get() : 0;
+			float MaxAspectRatioWidth = CurrentMaxAspectRatio.IsSet() ? CurrentMaxAspectRatio.Get() : 0;
+			if (CurrentHeight > 0 && CurrentWidth > 0)
 			{
-				const float AspectRatioHeight = 1.0f / AspectRatioWidth;
-
-				const float CurrentRatioWidth = ( AllottedGeometry.GetLocalSize().X / AllottedGeometry.GetLocalSize().Y );
-				const float CurrentRatioHeight = 1.0f / CurrentRatioWidth;
+				const float CurrentRatioWidth = (AllottedGeometry.GetLocalSize().X / AllottedGeometry.GetLocalSize().Y);
 
-				if ( CurrentRatioWidth > AspectRatioWidth /*|| CurrentRatioHeight > AspectRatioHeight*/ )
+				bool bFitMaxRatio = (CurrentRatioWidth > MaxAspectRatioWidth && MaxAspectRatioWidth != 0);
+				bool bFitMinRatio = (CurrentRatioWidth < MinAspectRatioWidth && MinAspectRatioWidth != 0);
+				if (bFitMaxRatio || bFitMinRatio)
 				{
 					XAlignmentResult = AlignChild<Orient_Horizontal>(AllottedGeometry.GetLocalSize().X, ChildSlot, SlotPadding);
 					YAlignmentResult = AlignChild<Orient_Vertical>(AllottedGeometry.GetLocalSize().Y, ChildSlot, SlotPadding);
 
-					float NewWidth = AspectRatioWidth * XAlignmentResult.Size;
-					float NewHeight = AspectRatioHeight * NewWidth;
+					float NewWidth;
+					float NewHeight;
+
+					if (bFitMaxRatio)
+					{
+						const float MaxAspectRatioHeight = 1.0f / MaxAspectRatioWidth;
+						NewWidth = MaxAspectRatioWidth * XAlignmentResult.Size;
+						NewHeight = MaxAspectRatioHeight * NewWidth;
+					}
+					else
+					{
+						const float MinAspectRatioHeight = 1.0f / MinAspectRatioWidth;
+						NewWidth = MinAspectRatioWidth * XAlignmentResult.Size;
+						NewHeight = MinAspectRatioHeight * NewWidth;
+					}
 
 					const float MaxWidth = AllottedGeometry.Size.X - SlotPadding.GetTotalSpaceAlong<Orient_Horizontal>();
 					const float MaxHeight = AllottedGeometry.Size.Y - SlotPadding.GetTotalSpaceAlong<Orient_Vertical>();
@@ -306,4 +330,4 @@ int32 SBox::OnPaint( const FPaintArgs& Args, const FGeometry& AllottedGeometry,
 		return TheChild.Widget->Paint( Args.WithNewParent(this), TheChild.Geometry, MyCullingRect, OutDrawElements, LayerId, InWidgetStyle, ShouldBeEnabled( bParentEnabled ) );
 	}
 	return LayerId;
-}
+}
\ No newline at end of file
diff --git a/Engine/Source/Runtime/Slate/Public/Widgets/Layout/SBox.h b/Engine/Source/Runtime/Slate/Public/Widgets/Layout/SBox.h
--- a/Engine/Source/Runtime/Slate/Public/Widgets/Layout/SBox.h
+++ b/Engine/Source/Runtime/Slate/Public/Widgets/Layout/SBox.h
@@ -21,107 +21,111 @@ class FSlateWindowElementList;
  */
 class SLATE_API SBox : public SPanel
 {
+public:
+	class FBoxSlot : public TSupportsOneChildMixin<FBoxSlot>, public TSupportsContentAlignmentMixin<FBoxSlot>, public TSupportsContentPaddingMixin<FBoxSlot>
+	{
 	public:
-		class FBoxSlot : public TSupportsOneChildMixin<FBoxSlot>, public TSupportsContentAlignmentMixin<FBoxSlot>, public TSupportsContentPaddingMixin<FBoxSlot>
+		FBoxSlot(SWidget* InOwner)
+			: TSupportsOneChildMixin<FBoxSlot>(InOwner)
+			, TSupportsContentAlignmentMixin<FBoxSlot>(HAlign_Fill, VAlign_Fill)
 		{
-			public:
-				FBoxSlot(SWidget* InOwner)
-				: TSupportsOneChildMixin<FBoxSlot>(InOwner)
-				, TSupportsContentAlignmentMixin<FBoxSlot>(HAlign_Fill, VAlign_Fill)
-				{
-				}
-		};
+		}
+	};
 
-		SLATE_BEGIN_ARGS(SBox)
-			: _HAlign( HAlign_Fill )
-			, _VAlign( VAlign_Fill )
-			, _Padding( 0.0f )
-			, _Content()
-			, _WidthOverride(FOptionalSize())
-			, _HeightOverride(FOptionalSize())
-			, _MinDesiredWidth(FOptionalSize())
-			, _MinDesiredHeight(FOptionalSize())
-			, _MaxDesiredWidth(FOptionalSize())
-			, _MaxDesiredHeight(FOptionalSize())
-			{
-				_Visibility = EVisibility::SelfHitTestInvisible;
-			}
+	SLATE_BEGIN_ARGS(SBox)
+		: _HAlign(HAlign_Fill)
+		, _VAlign(VAlign_Fill)
+		, _Padding(0.0f)
+		, _Content()
+		, _WidthOverride(FOptionalSize())
+		, _HeightOverride(FOptionalSize())
+		, _MinDesiredWidth(FOptionalSize())
+		, _MinDesiredHeight(FOptionalSize())
+		, _MaxDesiredWidth(FOptionalSize())
+		, _MaxDesiredHeight(FOptionalSize())
+	{
+		_Visibility = EVisibility::SelfHitTestInvisible;
+	}
 
-			/** Horizontal alignment of content in the area allotted to the SBox by its parent */
-			SLATE_ARGUMENT( EHorizontalAlignment, HAlign )
+	/** Horizontal alignment of content in the area allotted to the SBox by its parent */
+	SLATE_ARGUMENT(EHorizontalAlignment, HAlign)
 
-			/** Vertical alignment of content in the area allotted to the SBox by its parent */
-			SLATE_ARGUMENT( EVerticalAlignment, VAlign )
+	/** Vertical alignment of content in the area allotted to the SBox by its parent */
+	SLATE_ARGUMENT(EVerticalAlignment, VAlign)
 
-			/** Padding between the SBox and the content that it presents. Padding affects desired size. */
-			SLATE_ATTRIBUTE( FMargin, Padding )
+	/** Padding between the SBox and the content that it presents. Padding affects desired size. */
+	SLATE_ATTRIBUTE(FMargin, Padding)
 
-			/** The widget content presented by the SBox */
-			SLATE_DEFAULT_SLOT( FArguments, Content )
+	/** The widget content presented by the SBox */
+	SLATE_DEFAULT_SLOT(FArguments, Content)
 
-			/** When specified, ignore the content's desired size and report the WidthOverride as the Box's desired width. */
-			SLATE_ATTRIBUTE( FOptionalSize, WidthOverride )
+	/** When specified, ignore the content's desired size and report the WidthOverride as the Box's desired width. */
+	SLATE_ATTRIBUTE(FOptionalSize, WidthOverride)
 
-			/** When specified, ignore the content's desired size and report the HeightOverride as the Box's desired height. */
-			SLATE_ATTRIBUTE( FOptionalSize, HeightOverride)
+	/** When specified, ignore the content's desired size and report the HeightOverride as the Box's desired height. */
+	SLATE_ATTRIBUTE(FOptionalSize, HeightOverride)
 
-			/** When specified, will report the MinDesiredWidth if larger than the content's desired width. */
-			SLATE_ATTRIBUTE(FOptionalSize, MinDesiredWidth)
+	/** When specified, will report the MinDesiredWidth if larger than the content's desired width. */
+	SLATE_ATTRIBUTE(FOptionalSize, MinDesiredWidth)
 
-			/** When specified, will report the MinDesiredHeight if larger than the content's desired height. */
-			SLATE_ATTRIBUTE(FOptionalSize, MinDesiredHeight)
+	/** When specified, will report the MinDesiredHeight if larger than the content's desired height. */
+	SLATE_ATTRIBUTE(FOptionalSize, MinDesiredHeight)
 
-			/** When specified, will report the MaxDesiredWidth if smaller than the content's desired width. */
-			SLATE_ATTRIBUTE(FOptionalSize, MaxDesiredWidth)
+	/** When specified, will report the MaxDesiredWidth if smaller than the content's desired width. */
+	SLATE_ATTRIBUTE(FOptionalSize, MaxDesiredWidth)
 
-			/** When specified, will report the MaxDesiredHeight if smaller than the content's desired height. */
-			SLATE_ATTRIBUTE(FOptionalSize, MaxDesiredHeight)
+	/** When specified, will report the MaxDesiredHeight if smaller than the content's desired height. */
+	SLATE_ATTRIBUTE(FOptionalSize, MaxDesiredHeight)
 
-			SLATE_ATTRIBUTE(FOptionalSize, MaxAspectRatio)
+	SLATE_ATTRIBUTE(FOptionalSize, MinAspectRatio)
 
-		SLATE_END_ARGS()
+	SLATE_ATTRIBUTE(FOptionalSize, MaxAspectRatio)
 
-		SBox();
+	SLATE_END_ARGS()
 
-		void Construct( const FArguments& InArgs );
+	SBox();
 
-		virtual void OnArrangeChildren( const FGeometry& AllottedGeometry, FArrangedChildren& ArrangedChildren ) const override;
-		virtual FChildren* GetChildren() override;
-		virtual int32 OnPaint( const FPaintArgs& Args, const FGeometry& AllottedGeometry, const FSlateRect& MyCullingRect, FSlateWindowElementList& OutDrawElements, int32 LayerId, const FWidgetStyle& InWidgetStyle, bool bParentEnabled ) const override;
+	void Construct(const FArguments& InArgs);
 
-		/**
-		 * See the Content slot.
-		 */
-		void SetContent(const TSharedRef< SWidget >& InContent);
+	virtual void OnArrangeChildren(const FGeometry& AllottedGeometry, FArrangedChildren& ArrangedChildren) const override;
+	virtual FChildren* GetChildren() override;
+	virtual int32 OnPaint(const FPaintArgs& Args, const FGeometry& AllottedGeometry, const FSlateRect& MyCullingRect, FSlateWindowElementList& OutDrawElements, int32 LayerId, const FWidgetStyle& InWidgetStyle, bool bParentEnabled) const override;
 
-		/** See HAlign argument */
-		void SetHAlign(EHorizontalAlignment HAlign);
+	/**
+		* See the Content slot.
+		*/
+	void SetContent(const TSharedRef< SWidget >& InContent);
 
-		/** See VAlign argument */
-		void SetVAlign(EVerticalAlignment VAlign);
+	/** See HAlign argument */
+	void SetHAlign(EHorizontalAlignment HAlign);
 
-		/** See Padding attribute */
-		void SetPadding(const TAttribute<FMargin>& InPadding);
+	/** See VAlign argument */
+	void SetVAlign(EVerticalAlignment VAlign);
 
-		/** See WidthOverride attribute */
-		void SetWidthOverride(TAttribute<FOptionalSize> InWidthOverride);
+	/** See Padding attribute */
+	void SetPadding(const TAttribute<FMargin>& InPadding);
 
-		/** See HeightOverride attribute */
-		void SetHeightOverride(TAttribute<FOptionalSize> InHeightOverride);
+	/** See WidthOverride attribute */
+	void SetWidthOverride(TAttribute<FOptionalSize> InWidthOverride);
 
-		/** See MinDesiredWidth attribute */
-		void SetMinDesiredWidth(TAttribute<FOptionalSize> InMinDesiredWidth);
+	/** See HeightOverride attribute */
+	void SetHeightOverride(TAttribute<FOptionalSize> InHeightOverride);
 
-		/** See MinDesiredHeight attribute */
-		void SetMinDesiredHeight(TAttribute<FOptionalSize> InMinDesiredHeight);
+	/** See MinDesiredWidth attribute */
+	void SetMinDesiredWidth(TAttribute<FOptionalSize> InMinDesiredWidth);
 
-		/** See MaxDesiredWidth attribute */
-		void SetMaxDesiredWidth(TAttribute<FOptionalSize> InMaxDesiredWidth);
+	/** See MinDesiredHeight attribute */
+	void SetMinDesiredHeight(TAttribute<FOptionalSize> InMinDesiredHeight);
 
-		/** See MaxDesiredHeight attribute */
-		void SetMaxDesiredHeight(TAttribute<FOptionalSize> InMaxDesiredHeight);
+	/** See MaxDesiredWidth attribute */
+	void SetMaxDesiredWidth(TAttribute<FOptionalSize> InMaxDesiredWidth);
 
-		void SetMaxAspectRatio(TAttribute<FOptionalSize> InMaxAspectRatio);
+	/** See MaxDesiredHeight attribute */
+	void SetMaxDesiredHeight(TAttribute<FOptionalSize> InMaxDesiredHeight);
+
+	void SetMinAspectRatio(TAttribute<FOptionalSize> InMinAspectRatio);
+
+	void SetMaxAspectRatio(TAttribute<FOptionalSize> InMaxAspectRatio);
 
 protected:
 	// Begin SWidget overrides.
@@ -133,28 +137,29 @@ protected:
 
 protected:
 
-		FBoxSlot ChildSlot;
+	FBoxSlot ChildSlot;
 
 private:
-		/** When specified, ignore the content's desired size and report the.WidthOverride as the Box's desired width. */
-		TAttribute<FOptionalSize> WidthOverride;
+	/** When specified, ignore the content's desired size and report the.WidthOverride as the Box's desired width. */
+	TAttribute<FOptionalSize> WidthOverride;
 
-		/** When specified, ignore the content's desired size and report the.HeightOverride as the Box's desired height. */
-		TAttribute<FOptionalSize> HeightOverride;
-		
-		/** When specified, will report the MinDesiredWidth if larger than the content's desired width. */
-		TAttribute<FOptionalSize> MinDesiredWidth;
+	/** When specified, ignore the content's desired size and report the.HeightOverride as the Box's desired height. */
+	TAttribute<FOptionalSize> HeightOverride;
 
-		/** When specified, will report the MinDesiredHeight if larger than the content's desired height. */
-		TAttribute<FOptionalSize> MinDesiredHeight;
+	/** When specified, will report the MinDesiredWidth if larger than the content's desired width. */
+	TAttribute<FOptionalSize> MinDesiredWidth;
 
-		/** When specified, will report the MaxDesiredWidth if smaller than the content's desired width. */
-		TAttribute<FOptionalSize> MaxDesiredWidth;
+	/** When specified, will report the MinDesiredHeight if larger than the content's desired height. */
+	TAttribute<FOptionalSize> MinDesiredHeight;
 
-		/** When specified, will report the MaxDesiredHeight if smaller than the content's desired height. */
-		TAttribute<FOptionalSize> MaxDesiredHeight;
+	/** When specified, will report the MaxDesiredWidth if smaller than the content's desired width. */
+	TAttribute<FOptionalSize> MaxDesiredWidth;
 
-		TAttribute<FOptionalSize> MaxAspectRatio;
-};
+	/** When specified, will report the MaxDesiredHeight if smaller than the content's desired height. */
+	TAttribute<FOptionalSize> MaxDesiredHeight;
 
+	TAttribute<FOptionalSize> MinAspectRatio;
+
+	TAttribute<FOptionalSize> MaxAspectRatio;
+};
 
diff --git a/Engine/Source/Runtime/UMG/Private/Components/SizeBox.cpp b/Engine/Source/Runtime/UMG/Private/Components/SizeBox.cpp
--- a/Engine/Source/Runtime/UMG/Private/Components/SizeBox.cpp
+++ b/Engine/Source/Runtime/UMG/Private/Components/SizeBox.cpp
@@ -16,6 +16,7 @@ USizeBox::USizeBox(const FObjectInitializer& ObjectInitializer)
 {
 	bIsVariable = false;
 	Visibility = ESlateVisibility::SelfHitTestInvisible;
+	MinAspectRatio = 1;
 	MaxAspectRatio = 1;
 }
 
@@ -29,8 +30,8 @@ void USizeBox::ReleaseSlateResources(bool bReleaseChildren)
 TSharedRef<SWidget> USizeBox::RebuildWidget()
 {
 	MySizeBox = SNew(SBox);
-	
-	if ( GetChildrenCount() > 0 )
+
+	if (GetChildrenCount() > 0)
 	{
 		Cast<USizeBoxSlot>(GetContentSlot())->BuildSlot(MySizeBox.ToSharedRef());
 	}
@@ -42,7 +43,7 @@ void USizeBox::SetWidthOverride(float InWidthOverride)
 {
 	bOverride_WidthOverride = true;
 	WidthOverride = InWidthOverride;
-	if ( MySizeBox.IsValid() )
+	if (MySizeBox.IsValid())
 	{
 		MySizeBox->SetWidthOverride(InWidthOverride);
 	}
@@ -51,7 +52,7 @@ void USizeBox::SetWidthOverride(float InWidthOverride)
 void USizeBox::ClearWidthOverride()
 {
 	bOverride_WidthOverride = false;
-	if ( MySizeBox.IsValid() )
+	if (MySizeBox.IsValid())
 	{
 		MySizeBox->SetWidthOverride(FOptionalSize());
 	}
@@ -61,7 +62,7 @@ void USizeBox::SetHeightOverride(float InHeightOverride)
 {
 	bOverride_HeightOverride = true;
 	HeightOverride = InHeightOverride;
-	if ( MySizeBox.IsValid() )
+	if (MySizeBox.IsValid())
 	{
 		MySizeBox->SetHeightOverride(InHeightOverride);
 	}
@@ -70,7 +71,7 @@ void USizeBox::SetHeightOverride(float InHeightOverride)
 void USizeBox::ClearHeightOverride()
 {
 	bOverride_HeightOverride = false;
-	if ( MySizeBox.IsValid() )
+	if (MySizeBox.IsValid())
 	{
 		MySizeBox->SetHeightOverride(FOptionalSize());
 	}
@@ -80,7 +81,7 @@ void USizeBox::SetMinDesiredWidth(float InMinDesiredWidth)
 {
 	bOverride_MinDesiredWidth = true;
 	MinDesiredWidth = InMinDesiredWidth;
-	if ( MySizeBox.IsValid() )
+	if (MySizeBox.IsValid())
 	{
 		MySizeBox->SetMinDesiredWidth(InMinDesiredWidth);
 	}
@@ -89,7 +90,7 @@ void USizeBox::SetMinDesiredWidth(float InMinDesiredWidth)
 void USizeBox::ClearMinDesiredWidth()
 {
 	bOverride_MinDesiredWidth = false;
-	if ( MySizeBox.IsValid() )
+	if (MySizeBox.IsValid())
 	{
 		MySizeBox->SetMinDesiredWidth(FOptionalSize());
 	}
@@ -99,7 +100,7 @@ void USizeBox::SetMinDesiredHeight(float InMinDesiredHeight)
 {
 	bOverride_MinDesiredHeight = true;
 	MinDesiredHeight = InMinDesiredHeight;
-	if ( MySizeBox.IsValid() )
+	if (MySizeBox.IsValid())
 	{
 		MySizeBox->SetMinDesiredHeight(InMinDesiredHeight);
 	}
@@ -108,7 +109,7 @@ void USizeBox::SetMinDesiredHeight(float InMinDesiredHeight)
 void USizeBox::ClearMinDesiredHeight()
 {
 	bOverride_MinDesiredHeight = false;
-	if ( MySizeBox.IsValid() )
+	if (MySizeBox.IsValid())
 	{
 		MySizeBox->SetMinDesiredHeight(FOptionalSize());
 	}
@@ -118,7 +119,7 @@ void USizeBox::SetMaxDesiredWidth(float InMaxDesiredWidth)
 {
 	bOverride_MaxDesiredWidth = true;
 	MaxDesiredWidth = InMaxDesiredWidth;
-	if ( MySizeBox.IsValid() )
+	if (MySizeBox.IsValid())
 	{
 		MySizeBox->SetMaxDesiredWidth(InMaxDesiredWidth);
 	}
@@ -127,7 +128,7 @@ void USizeBox::SetMaxDesiredWidth(float InMaxDesiredWidth)
 void USizeBox::ClearMaxDesiredWidth()
 {
 	bOverride_MaxDesiredWidth = false;
-	if ( MySizeBox.IsValid() )
+	if (MySizeBox.IsValid())
 	{
 		MySizeBox->SetMaxDesiredWidth(FOptionalSize());
 	}
@@ -137,7 +138,7 @@ void USizeBox::SetMaxDesiredHeight(float InMaxDesiredHeight)
 {
 	bOverride_MaxDesiredHeight = true;
 	MaxDesiredHeight = InMaxDesiredHeight;
-	if ( MySizeBox.IsValid() )
+	if (MySizeBox.IsValid())
 	{
 		MySizeBox->SetMaxDesiredHeight(InMaxDesiredHeight);
 	}
@@ -146,17 +147,36 @@ void USizeBox::SetMaxDesiredHeight(float InMaxDesiredHeight)
 void USizeBox::ClearMaxDesiredHeight()
 {
 	bOverride_MaxDesiredHeight = false;
-	if ( MySizeBox.IsValid() )
+	if (MySizeBox.IsValid())
 	{
 		MySizeBox->SetMaxDesiredHeight(FOptionalSize());
 	}
 }
 
+void USizeBox::SetMinAspectRatio(float InMinAspectRatio)
+{
+	bOverride_MinAspectRatio = true;
+	MinAspectRatio = InMinAspectRatio;
+	if (MySizeBox.IsValid())
+	{
+		MySizeBox->SetMinAspectRatio(InMinAspectRatio);
+	}
+}
+
+void USizeBox::ClearMinAspectRatio()
+{
+	bOverride_MinAspectRatio = false;
+	if (MySizeBox.IsValid())
+	{
+		MySizeBox->SetMinAspectRatio(FOptionalSize());
+	}
+}
+
 void USizeBox::SetMaxAspectRatio(float InMaxAspectRatio)
 {
 	bOverride_MaxAspectRatio = true;
 	MaxAspectRatio = InMaxAspectRatio;
-	if ( MySizeBox.IsValid() )
+	if (MySizeBox.IsValid())
 	{
 		MySizeBox->SetMaxAspectRatio(InMaxAspectRatio);
 	}
@@ -165,7 +185,7 @@ void USizeBox::SetMaxAspectRatio(float InMaxAspectRatio)
 void USizeBox::ClearMaxAspectRatio()
 {
 	bOverride_MaxAspectRatio = false;
-	if ( MySizeBox.IsValid() )
+	if (MySizeBox.IsValid())
 	{
 		MySizeBox->SetMaxAspectRatio(FOptionalSize());
 	}
@@ -175,7 +195,7 @@ void USizeBox::SynchronizeProperties()
 {
 	Super::SynchronizeProperties();
 
-	if ( bOverride_WidthOverride )
+	if (bOverride_WidthOverride)
 	{
 		SetWidthOverride(WidthOverride);
 	}
@@ -184,7 +204,7 @@ void USizeBox::SynchronizeProperties()
 		ClearWidthOverride();
 	}
 
-	if ( bOverride_HeightOverride )
+	if (bOverride_HeightOverride)
 	{
 		SetHeightOverride(HeightOverride);
 	}
@@ -193,7 +213,7 @@ void USizeBox::SynchronizeProperties()
 		ClearHeightOverride();
 	}
 
-	if ( bOverride_MinDesiredWidth )
+	if (bOverride_MinDesiredWidth)
 	{
 		SetMinDesiredWidth(MinDesiredWidth);
 	}
@@ -202,7 +222,7 @@ void USizeBox::SynchronizeProperties()
 		ClearMinDesiredWidth();
 	}
 
-	if ( bOverride_MinDesiredHeight )
+	if (bOverride_MinDesiredHeight)
 	{
 		SetMinDesiredHeight(MinDesiredHeight);
 	}
@@ -211,7 +231,7 @@ void USizeBox::SynchronizeProperties()
 		ClearMinDesiredHeight();
 	}
 
-	if ( bOverride_MaxDesiredWidth )
+	if (bOverride_MaxDesiredWidth)
 	{
 		SetMaxDesiredWidth(MaxDesiredWidth);
 	}
@@ -220,7 +240,7 @@ void USizeBox::SynchronizeProperties()
 		ClearMaxDesiredWidth();
 	}
 
-	if ( bOverride_MaxDesiredHeight )
+	if (bOverride_MaxDesiredHeight)
 	{
 		SetMaxDesiredHeight(MaxDesiredHeight);
 	}
@@ -229,7 +249,16 @@ void USizeBox::SynchronizeProperties()
 		ClearMaxDesiredHeight();
 	}
 
-	if ( bOverride_MaxAspectRatio )
+	if (bOverride_MinAspectRatio)
+	{
+		SetMinAspectRatio(MinAspectRatio);
+	}
+	else
+	{
+		ClearMinAspectRatio();
+	}
+
+	if (bOverride_MaxAspectRatio)
 	{
 		SetMaxAspectRatio(MaxAspectRatio);
 	}
@@ -247,7 +276,7 @@ UClass* USizeBox::GetSlotClass() const
 void USizeBox::OnSlotAdded(UPanelSlot* InSlot)
 {
 	// Add the child to the live slot if it already exists
-	if ( MySizeBox.IsValid() )
+	if (MySizeBox.IsValid())
 	{
 		CastChecked<USizeBoxSlot>(InSlot)->BuildSlot(MySizeBox.ToSharedRef());
 	}
@@ -256,7 +285,7 @@ void USizeBox::OnSlotAdded(UPanelSlot* InSlot)
 void USizeBox::OnSlotRemoved(UPanelSlot* InSlot)
 {
 	// Remove the widget from the live slot if it exists.
-	if ( MySizeBox.IsValid() )
+	if (MySizeBox.IsValid())
 	{
 		MySizeBox->SetContent(SNullWidget::NullWidget);
 	}
@@ -273,4 +302,4 @@ const FText USizeBox::GetPaletteCategory()
 
 /////////////////////////////////////////////////////
 
-#undef LOCTEXT_NAMESPACE
+#undef LOCTEXT_NAMESPACE
\ No newline at end of file
diff --git a/Engine/Source/Runtime/UMG/Public/Components/SizeBox.h b/Engine/Source/Runtime/UMG/Public/Components/SizeBox.h
--- a/Engine/Source/Runtime/UMG/Public/Components/SizeBox.h
+++ b/Engine/Source/Runtime/UMG/Public/Components/SizeBox.h
@@ -26,108 +26,122 @@ public:
 
 	/**  */
 	UPROPERTY(EditAnywhere, Category="Child Layout", meta=(InlineEditConditionToggle))
-	uint32 bOverride_WidthOverride : 1;
+		uint32 bOverride_WidthOverride : 1;
 
 	/**  */
 	UPROPERTY(EditAnywhere, Category="Child Layout", meta=(InlineEditConditionToggle))
-	uint32 bOverride_HeightOverride : 1;
+		uint32 bOverride_HeightOverride : 1;
 
 	/**  */
 	UPROPERTY(EditAnywhere, Category="Child Layout", meta=(InlineEditConditionToggle))
-	uint32 bOverride_MinDesiredWidth : 1;
+		uint32 bOverride_MinDesiredWidth : 1;
 
 	/**  */
 	UPROPERTY(EditAnywhere, Category="Child Layout", meta=(InlineEditConditionToggle))
-	uint32 bOverride_MinDesiredHeight : 1;
+		uint32 bOverride_MinDesiredHeight : 1;
 
 	/**  */
 	UPROPERTY(EditAnywhere, Category="Child Layout", meta=(InlineEditConditionToggle))
-	uint32 bOverride_MaxDesiredWidth : 1;
+		uint32 bOverride_MaxDesiredWidth : 1;
 
 	/**  */
 	UPROPERTY(EditAnywhere, Category="Child Layout", meta=(InlineEditConditionToggle))
-	uint32 bOverride_MaxDesiredHeight : 1;
+		uint32 bOverride_MaxDesiredHeight : 1;
 
 	/**  */
 	UPROPERTY(EditAnywhere, Category="Child Layout", meta=(InlineEditConditionToggle))
+	uint32 bOverride_MinAspectRatio : 1;
+
+	/**  */
+	UPROPERTY(EditAnywhere, Category = "Child Layout", meta=(InlineEditConditionToggle))
 	uint32 bOverride_MaxAspectRatio : 1;
 
 
 	/** When specified, ignore the content's desired size and report the WidthOverride as the Box's desired width. */
 	UPROPERTY(EditAnywhere, BlueprintReadOnly, Category="Child Layout", meta=( editcondition="bOverride_WidthOverride" ))
-	float WidthOverride;
+		float WidthOverride;
 
 	/** When specified, ignore the content's desired size and report the HeightOverride as the Box's desired height. */
 	UPROPERTY(EditAnywhere, BlueprintReadOnly, Category="Child Layout", meta=( editcondition="bOverride_HeightOverride" ))
-	float HeightOverride;
+		float HeightOverride;
 
 	/** When specified, will report the MinDesiredWidth if larger than the content's desired width. */
 	UPROPERTY(EditAnywhere, BlueprintReadOnly, Category="Child Layout", meta=( editcondition="bOverride_MinDesiredWidth" ))
-	float MinDesiredWidth;
+		float MinDesiredWidth;
 
 	/** When specified, will report the MinDesiredHeight if larger than the content's desired height. */
 	UPROPERTY(EditAnywhere, BlueprintReadOnly, Category="Child Layout", meta=( editcondition="bOverride_MinDesiredHeight" ))
-	float MinDesiredHeight;
+		float MinDesiredHeight;
 
 	/** When specified, will report the MaxDesiredWidth if smaller than the content's desired width. */
 	UPROPERTY(EditAnywhere, BlueprintReadOnly, Category="Child Layout", meta=( editcondition="bOverride_MaxDesiredWidth" ))
-	float MaxDesiredWidth;
+		float MaxDesiredWidth;
 
 	/** When specified, will report the MaxDesiredHeight if smaller than the content's desired height. */
 	UPROPERTY(EditAnywhere, BlueprintReadOnly, Category="Child Layout", meta=( editcondition="bOverride_MaxDesiredHeight" ))
-	float MaxDesiredHeight;
+		float MaxDesiredHeight;
+
+	/** */
+	UPROPERTY(EditAnywhere, BlueprintReadOnly, Category = "Child Layout", meta = (editcondition = "bOverride_MinAspectRatio"))
+		float MinAspectRatio;
 
 	/** */
 	UPROPERTY(EditAnywhere, BlueprintReadOnly, Category="Child Layout", meta=( editcondition="bOverride_MaxAspectRatio" ))
-	float MaxAspectRatio;
+		float MaxAspectRatio;
 
 public:
 		
 	/** When specified, ignore the content's desired size and report the WidthOverride as the Box's desired width. */
 	UFUNCTION(BlueprintCallable, Category="Layout|Size Box")
-	void SetWidthOverride(float InWidthOverride);
+		void SetWidthOverride(float InWidthOverride);
 
 	UFUNCTION(BlueprintCallable, Category="Layout|Size Box")
-	void ClearWidthOverride();
+		void ClearWidthOverride();
 
 	/** When specified, ignore the content's desired size and report the HeightOverride as the Box's desired height. */
 	UFUNCTION(BlueprintCallable, Category="Layout|Size Box")
-	void SetHeightOverride(float InHeightOverride);
+		void SetHeightOverride(float InHeightOverride);
 
 	UFUNCTION(BlueprintCallable, Category="Layout|Size Box")
-	void ClearHeightOverride();
+		void ClearHeightOverride();
 
 	/** When specified, will report the MinDesiredWidth if larger than the content's desired width. */
 	UFUNCTION(BlueprintCallable, Category="Layout|Size Box")
-	void SetMinDesiredWidth(float InMinDesiredWidth);
+		void SetMinDesiredWidth(float InMinDesiredWidth);
 
 	UFUNCTION(BlueprintCallable, Category="Layout|Size Box")
-	void ClearMinDesiredWidth();
+		void ClearMinDesiredWidth();
 
 	/** When specified, will report the MinDesiredHeight if larger than the content's desired height. */
 	UFUNCTION(BlueprintCallable, Category="Layout|Size Box")
-	void SetMinDesiredHeight(float InMinDesiredHeight);
+		void SetMinDesiredHeight(float InMinDesiredHeight);
 
 	UFUNCTION(BlueprintCallable, Category="Layout|Size Box")
-	void ClearMinDesiredHeight();
+		void ClearMinDesiredHeight();
 
 	/** When specified, will report the MaxDesiredWidth if smaller than the content's desired width. */
 	UFUNCTION(BlueprintCallable, Category="Layout|Size Box")
-	void SetMaxDesiredWidth(float InMaxDesiredWidth);
+		void SetMaxDesiredWidth(float InMaxDesiredWidth);
 
 	UFUNCTION(BlueprintCallable, Category="Layout|Size Box")
-	void ClearMaxDesiredWidth();
+		void ClearMaxDesiredWidth();
 
 	/** When specified, will report the MaxDesiredHeight if smaller than the content's desired height. */
 	UFUNCTION(BlueprintCallable, Category="Layout|Size Box")
-	void SetMaxDesiredHeight(float InMaxDesiredHeight);
+		void SetMaxDesiredHeight(float InMaxDesiredHeight);
 
 	UFUNCTION(BlueprintCallable, Category="Layout|Size Box")
-	void ClearMaxDesiredHeight();
+		void ClearMaxDesiredHeight();
+
+	UFUNCTION(BlueprintCallable, Category="Layout|Size Box")
+	void SetMinAspectRatio(float InMinAspectRatio);
 
 	UFUNCTION(BlueprintCallable, Category="Layout|Size Box")
 	void SetMaxAspectRatio(float InMaxAspectRatio);
 
+	UFUNCTION(BlueprintCallable, Category="Layout|Size Box")
+	void ClearMinAspectRatio();
+
 	UFUNCTION(BlueprintCallable, Category="Layout|Size Box")
 	void ClearMaxAspectRatio();
 
@@ -160,4 +174,4 @@ protected:
 	// UWidget interface
 	virtual TSharedRef<SWidget> RebuildWidget() override;
 	// End of UWidget interface
-};
+};
\ No newline at end of file
