Fix to a bug in SButton::OnPaint where ColorAndOpacity property was ignored for rendering.

#rn Fix to a bug in SButton where ColorAndOpacity property was ignored for rendering. With this fix, it will be correctly applied and could lead to Buttons rendered differently than previous versions if ColorAndOpacity was set to something else than white.

diff --git a/Engine/Source/Runtime/Slate/Private/Widgets/Input/SButton.cpp b/Engine/Source/Runtime/Slate/Private/Widgets/Input/SButton.cpp
--- a/Engine/Source/Runtime/Slate/Private/Widgets/Input/SButton.cpp
+++ b/Engine/Source/Runtime/Slate/Private/Widgets/Input/SButton.cpp
@@ -97,7 +97,7 @@ int32 SButton::OnPaint(const FPaintArgs& Args, const FGeometry& AllottedGeometry
 			AllottedGeometry.ToPaintGeometry(),
 			BrushResource,
 			DrawEffects,
-			BrushResource->GetTint(InWidgetStyle) * InWidgetStyle.GetColorAndOpacityTint() * BorderBackgroundColor.Get().GetColor(InWidgetStyle)
+			BrushResource->GetTint(InWidgetStyle) * InWidgetStyle.GetColorAndOpacityTint() * BorderBackgroundColor.Get().GetColor(InWidgetStyle) * ColorAndOpacity.Get()
 			);
 	}
 
