Change 'Use mouse as Touch' only with left mouse click This change means that we no longer stop a right click event from going through.

diff --git a/Engine/Source/Runtime/Slate/Private/Framework/Application/SlateApplication.cpp b/Engine/Source/Runtime/Slate/Private/Framework/Application/SlateApplication.cpp
--- a/Engine/Source/Runtime/Slate/Private/Framework/Application/SlateApplication.cpp
+++ b/Engine/Source/Runtime/Slate/Private/Framework/Application/SlateApplication.cpp
@@ -5325,8 +5325,8 @@ bool FSlateApplication::OnMouseDown(const TSharedPtr< FGenericWindow >& Platform
 
 bool FSlateApplication::OnMouseDown( const TSharedPtr< FGenericWindow >& PlatformWindow, const EMouseButtons::Type Button, const FVector2D CursorPos )
 {
-	// convert to touch event if we are faking it	
-	if (bIsFakingTouch || bIsGameFakingTouch)
+	// convert a left mouse button click to touch event if we are faking it
+	if (bIsFakingTouch || bIsGameFakingTouch && Button == EMouseButtons::Left)
 	{
 		bIsFakingTouched = true;
 		return OnTouchStarted( PlatformWindow, PlatformApplication->Cursor->GetPosition(), 1.0f, 0, 0 );
@@ -6147,8 +6147,8 @@ bool FSlateApplication::OnMouseUp( const EMouseButtons::Type Button )
 
 bool FSlateApplication::OnMouseUp( const EMouseButtons::Type Button, const FVector2D CursorPos )
 {
-	// convert to touch event if we are faking it	
-	if (bIsFakingTouch || bIsGameFakingTouch)
+	// convert left mouse click to touch event if we are faking it	
+	if ((bIsFakingTouch || bIsGameFakingTouch) && Button == EMouseButtons::Left)
 	{
 		bIsFakingTouched = false;
 		return OnTouchEnded(PlatformApplication->Cursor->GetPosition(), 0, 0);
@@ -6314,7 +6314,8 @@ FReply FSlateApplication::RouteMouseWheelOrGestureEvent(const FWidgetPath& Widge
 
 bool FSlateApplication::OnMouseMove()
 {
-	if (bIsFakingTouched || bIsGameFakingTouch)
+	// If the left button is pressed we fake 
+	if ((bIsFakingTouched || bIsGameFakingTouch) && GetPressedMouseButtons().Contains(EKeys::LeftMouseButton))
 	{
 		// convert to touch event if we are faking it
 		if (bIsFakingTouched)
@@ -6356,7 +6357,8 @@ bool FSlateApplication::OnMouseMove()
 
 bool FSlateApplication::OnRawMouseMove( const int32 X, const int32 Y )
 {
-	if (bIsFakingTouched || bIsGameFakingTouch)
+    // We fake a move only if the left mous button is down
+	if ((bIsFakingTouch || bIsGameFakingTouch) && GetPressedMouseButtons().Contains(EKeys::LeftMouseButton))
 	{
 		// convert to touch event if we are faking it
 		if (bIsFakingTouched)
