Fix CL5315468: Remove ambiguity:  '&&' within '||'.

diff --git a/Engine/Source/Runtime/Slate/Private/Framework/Application/SlateApplication.cpp b/Engine/Source/Runtime/Slate/Private/Framework/Application/SlateApplication.cpp
--- a/Engine/Source/Runtime/Slate/Private/Framework/Application/SlateApplication.cpp
+++ b/Engine/Source/Runtime/Slate/Private/Framework/Application/SlateApplication.cpp
@@ -5326,7 +5326,7 @@ bool FSlateApplication::OnMouseDown(const TSharedPtr< FGenericWindow >& Platform
 bool FSlateApplication::OnMouseDown( const TSharedPtr< FGenericWindow >& PlatformWindow, const EMouseButtons::Type Button, const FVector2D CursorPos )
 {
 	// convert a left mouse button click to touch event if we are faking it
-	if (bIsFakingTouch || bIsGameFakingTouch && Button == EMouseButtons::Left)
+	if ((bIsFakingTouch || bIsGameFakingTouch) && Button == EMouseButtons::Left)
 	{
 		bIsFakingTouched = true;
 		return OnTouchStarted( PlatformWindow, PlatformApplication->Cursor->GetPosition(), 1.0f, 0, 0 );
