Slate - Fixing border and image to properly invalidate with the underlying resource or brush data changes when the brush is changed.  Moving that code into a shared area to make this bit reuseable, this caused some fallout fixing up places that subclassed SBorder.  Likely to experience other cases in plugins or other projects.

Matt.Kuhlenschmidt

diff --git a/Engine/Source/Developer/CollisionAnalyzer/Private/SCAQueryTableRow.cpp b/Engine/Source/Developer/CollisionAnalyzer/Private/SCAQueryTableRow.cpp
--- a/Engine/Source/Developer/CollisionAnalyzer/Private/SCAQueryTableRow.cpp
+++ b/Engine/Source/Developer/CollisionAnalyzer/Private/SCAQueryTableRow.cpp
@@ -13,7 +13,7 @@ void SCAQueryTableRow::Construct(const FArguments& InArgs, const TSharedRef<STab
 
 	if(Item->bIsGroup)
 	{
-		BorderImage = FCollisionAnalyzerStyle::Get()->GetBrush("CollisionAnalyzer.GroupBackground");
+		BorderImage = FInvalidatableBrushAttribute(FCollisionAnalyzerStyle::Get()->GetBrush("CollisionAnalyzer.GroupBackground"));
 	}
 }
 
diff --git a/Engine/Source/Runtime/Slate/Private/Widgets/Images/SSpinningImage.cpp b/Engine/Source/Runtime/Slate/Private/Widgets/Images/SSpinningImage.cpp
--- a/Engine/Source/Runtime/Slate/Private/Widgets/Images/SSpinningImage.cpp
+++ b/Engine/Source/Runtime/Slate/Private/Widgets/Images/SSpinningImage.cpp
@@ -6,7 +6,7 @@
 
 void SSpinningImage::Construct(const FArguments& InArgs)
 {
-	Image = InArgs._Image;
+	Image = FInvalidatableBrushAttribute(InArgs._Image);
 	ColorAndOpacity = InArgs._ColorAndOpacity;
 
 	if (InArgs._OnMouseButtonDown.IsBound())
diff --git a/Engine/Source/Runtime/Slate/Private/Widgets/Layout/SBorder.cpp b/Engine/Source/Runtime/Slate/Private/Widgets/Layout/SBorder.cpp
--- a/Engine/Source/Runtime/Slate/Private/Widgets/Layout/SBorder.cpp
+++ b/Engine/Source/Runtime/Slate/Private/Widgets/Layout/SBorder.cpp
@@ -131,20 +131,12 @@ FVector2D SBorder::ComputeDesiredSize(float LayoutScaleMultiplier) const
 
 void SBorder::SetBorderBackgroundColor(const TAttribute<FSlateColor>& InColorAndOpacity)
 {
-	if (!BorderBackgroundColor.IdenticalTo(InColorAndOpacity))
-	{
-		BorderBackgroundColor = InColorAndOpacity;
-		Invalidate(EInvalidateWidget::PaintAndVolatility);
-	}
+	SetAttribute(BorderBackgroundColor, InColorAndOpacity, EInvalidateWidgetReason::Paint);
 }
 
 void SBorder::SetDesiredSizeScale(const TAttribute<FVector2D>& InDesiredSizeScale)
 {
-	if (!DesiredSizeScale.IdenticalTo(InDesiredSizeScale))
-	{
-		DesiredSizeScale = InDesiredSizeScale;
-		Invalidate(EInvalidateWidget::LayoutAndVolatility);
-	}
+	SetAttribute(DesiredSizeScale, InDesiredSizeScale, EInvalidateWidgetReason::Layout);
 }
 
 void SBorder::SetHAlign(EHorizontalAlignment HAlign)
@@ -167,27 +159,15 @@ void SBorder::SetVAlign(EVerticalAlignment VAlign)
 
 void SBorder::SetPadding(const TAttribute<FMargin>& InPadding)
 {
-	if (!ChildSlot.SlotPadding.IdenticalTo(InPadding))
-	{
-		ChildSlot.SlotPadding = InPadding;
-		Invalidate(EInvalidateWidget::LayoutAndVolatility);
-	}
+	SetAttribute(ChildSlot.SlotPadding, InPadding, EInvalidateWidgetReason::Layout);
 }
 
 void SBorder::SetShowEffectWhenDisabled(const TAttribute<bool>& InShowEffectWhenDisabled)
 {
-	if (!ShowDisabledEffect.IdenticalTo(InShowEffectWhenDisabled))
-	{
-		ShowDisabledEffect = InShowEffectWhenDisabled;
-		Invalidate(EInvalidateWidget::LayoutAndVolatility);
-	}
+	SetAttribute(ShowDisabledEffect, InShowEffectWhenDisabled, EInvalidateWidgetReason::Paint);
 }
 
 void SBorder::SetBorderImage(const TAttribute<const FSlateBrush*>& InBorderImage)
 {
-	if (!BorderImage.IdenticalTo(InBorderImage))
-	{
-		BorderImage = InBorderImage;
-		Invalidate(EInvalidateWidget::LayoutAndVolatility);
-	}
+	BorderImage.SetImage(*this, InBorderImage);
 }
diff --git a/Engine/Source/Runtime/Slate/Public/Widgets/Layout/SBorder.h b/Engine/Source/Runtime/Slate/Public/Widgets/Layout/SBorder.h
--- a/Engine/Source/Runtime/Slate/Public/Widgets/Layout/SBorder.h
+++ b/Engine/Source/Runtime/Slate/Public/Widgets/Layout/SBorder.h
@@ -10,6 +10,7 @@
 #include "Layout/Margin.h"
 #include "Widgets/SCompoundWidget.h"
 #include "Styling/CoreStyle.h"
+#include "Styling/SlateTypes.h"
 
 class FPaintArgs;
 class FSlateWindowElementList;
@@ -135,7 +136,7 @@ protected:
 	// End SWidget overrides.
 
  protected:
-	TAttribute<const FSlateBrush*> BorderImage;
+	FInvalidatableBrushAttribute BorderImage;
 	TAttribute<FSlateColor> BorderBackgroundColor;
 	TAttribute<FVector2D> DesiredSizeScale;
 	/** Whether or not to show the disabled effect when this border is disabled */
diff --git a/Engine/Source/Runtime/Slate/Public/Widgets/Views/STableRow.h b/Engine/Source/Runtime/Slate/Public/Widgets/Views/STableRow.h
--- a/Engine/Source/Runtime/Slate/Public/Widgets/Views/STableRow.h
+++ b/Engine/Source/Runtime/Slate/Public/Widgets/Views/STableRow.h
@@ -958,7 +958,7 @@ protected:
 		check(InArgs._ExpanderStyleSet);
 		ExpanderStyleSet = InArgs._ExpanderStyleSet;
 
-		this->BorderImage = TAttribute<const FSlateBrush*>( this, &STableRow::GetBorder );
+		this->BorderImage = FInvalidatableBrushAttribute(TAttribute<const FSlateBrush*>(this, &STableRow::GetBorder));
 
 		this->ForegroundColor = TAttribute<FSlateColor>( this, &STableRow::GetForegroundBasedOnSelection );
 
diff --git a/Engine/Source/Runtime/SlateCore/Private/Styling/SlateTypes.cpp b/Engine/Source/Runtime/SlateCore/Private/Styling/SlateTypes.cpp
--- a/Engine/Source/Runtime/SlateCore/Private/Styling/SlateTypes.cpp
+++ b/Engine/Source/Runtime/SlateCore/Private/Styling/SlateTypes.cpp
@@ -3,7 +3,8 @@
 #include "Styling/SlateTypes.h"
 #include "Brushes/SlateNoResource.h"
 #include "Styling/StyleDefaults.h"
-
+#include "Widgets/InvalidateWidgetReason.h"
+#include "Widgets/SWidget.h"
 
 namespace SlateTypeDefs
 {
@@ -672,4 +673,22 @@ const FWindowStyle& FWindowStyle::GetDefault()
 	return Default;
 }
 
+void FInvalidatableBrushAttribute::SetImage(SWidget& ThisWidget, const TAttribute< const FSlateBrush* >& InImage)
+{
+	const bool bImagePointerChanged = SetWidgetAttribute(ThisWidget, Image, InImage, EInvalidateWidgetReason::Layout);
+
+	const FSlateBrush* ImagePtr = Image.Get();
+	const FSlateBrush NewImageToCache = ImagePtr ? *ImagePtr : FSlateBrush();
+
+	// If the slate brush pointer didn't change, that may not mean nothing changed.  We
+	// sometimes reuse the slate brush memory address and change out the texture.  In those
+	// circumstances we need to actually look at the data the brush has compared to what it had
+	// previously.
+	if (!bImagePointerChanged && ImageCache != NewImageToCache)
+	{
+		ThisWidget.Invalidate(EInvalidateWidgetReason::Layout);
+	}
 
+	// Cache the new image's value in case the memory changes later.
+	ImageCache = NewImageToCache;
+}
\ No newline at end of file
diff --git a/Engine/Source/Runtime/SlateCore/Private/Widgets/Images/SImage.cpp b/Engine/Source/Runtime/SlateCore/Private/Widgets/Images/SImage.cpp
--- a/Engine/Source/Runtime/SlateCore/Private/Widgets/Images/SImage.cpp
+++ b/Engine/Source/Runtime/SlateCore/Private/Widgets/Images/SImage.cpp
@@ -9,7 +9,7 @@
 
 void SImage::Construct( const FArguments& InArgs )
 {
-	Image = InArgs._Image;
+	Image = FInvalidatableBrushAttribute(InArgs._Image);
 	ColorAndOpacity = InArgs._ColorAndOpacity;
 	bFlipForRightToLeftFlowDirection = InArgs._FlipForRightToLeftFlowDirection;
 	SetOnMouseButtonDown(InArgs._OnMouseButtonDown);
@@ -17,7 +17,7 @@ void SImage::Construct( const FArguments& InArgs )
 
 int32 SImage::OnPaint( const FPaintArgs& Args, const FGeometry& AllottedGeometry, const FSlateRect& MyCullingRect, FSlateWindowElementList& OutDrawElements, int32 LayerId, const FWidgetStyle& InWidgetStyle, bool bParentEnabled ) const
 {
-	const FSlateBrush* ImageBrush = Image.Get();
+	const FSlateBrush* ImageBrush = Image.GetImage().Get();
 
 	if ((ImageBrush != nullptr) && (ImageBrush->DrawAs != ESlateBrushDrawType::NoDrawType))
 	{
@@ -62,22 +62,7 @@ void SImage::SetColorAndOpacity( FLinearColor InColorAndOpacity )
 
 void SImage::SetImage(TAttribute<const FSlateBrush*> InImage)
 {
-	const bool bImagePointerChanged = SetAttribute(Image, InImage, EInvalidateWidgetReason::Layout);
-
-	const FSlateBrush* ImagePtr = Image.Get();
-	const FSlateBrush NewImageToCache = ImagePtr ? *ImagePtr : FSlateBrush();
-
-	// If the slate brush pointer didn't change, that may not mean nothing changed.  We
-	// sometimes reuse the slate brush memory address and change out the texture.  In those
-	// circumstances we need to actually look at the data the brush has compared to what it had
-	// previously.
-	if (!bImagePointerChanged && ImageCache != NewImageToCache)
-	{
-		Invalidate(EInvalidateWidgetReason::Layout);
-	}
-
-	// Cache the new image's value in case the memory changes later.
-	ImageCache = NewImageToCache;
+	Image.SetImage(*this, InImage);
 }
 
 #if WITH_ACCESSIBILITY
diff --git a/Engine/Source/Runtime/SlateCore/Public/Styling/SlateTypes.h b/Engine/Source/Runtime/SlateCore/Public/Styling/SlateTypes.h
--- a/Engine/Source/Runtime/SlateCore/Public/Styling/SlateTypes.h
+++ b/Engine/Source/Runtime/SlateCore/Public/Styling/SlateTypes.h
@@ -16,6 +16,8 @@
 
 #include "SlateTypes.generated.h"
 
+class SWidget;
+
 /** Used to determine how we should handle mouse wheel input events when someone scrolls. */
 UENUM()
 enum class EConsumeMouseWheel : uint8
@@ -1758,6 +1760,34 @@ struct SLATECORE_API FWindowStyle : public FSlateWidgetStyle
 };
 
 
+class SLATECORE_API FInvalidatableBrushAttribute
+{
+public:
+	FInvalidatableBrushAttribute() { }
+	FInvalidatableBrushAttribute(const TAttribute< const FSlateBrush* >& InImage)
+		: Image(InImage)
+	{
+		// INFO don't call this in the ctor, users did not anticipate needing to have their accessors work.
+
+		//const FSlateBrush* ImagePtr = Image.Get();
+		//ImageCache = ImagePtr ? *ImagePtr : FSlateBrush();
+	}
+
+	bool IsBound() const { return Image.IsBound(); }
+
+	const FSlateBrush* Get() const { return Image.Get(); }
+	TAttribute< const FSlateBrush* > GetImage() const { return Image; }
+	void SetImage(SWidget& ThisWidget, const TAttribute< const FSlateBrush* >& InImage);
+
+private:
+	/** The slate brush to draw for the image, or a bound delegate to a brush. */
+	TAttribute< const FSlateBrush* > Image;
+
+	/** The copy of the image data, some users reuse the same FSlateBrush pointer, so we need to check it against the last true data to see what changed. */
+	FSlateBrush ImageCache;
+};
+
+
 /** HACK: We need a UClass here or UHT will complain. */
 UCLASS()
 class USlateTypes : public UObject
diff --git a/Engine/Source/Runtime/SlateCore/Public/Widgets/Images/SImage.h b/Engine/Source/Runtime/SlateCore/Public/Widgets/Images/SImage.h
--- a/Engine/Source/Runtime/SlateCore/Public/Widgets/Images/SImage.h
+++ b/Engine/Source/Runtime/SlateCore/Public/Widgets/Images/SImage.h
@@ -9,6 +9,7 @@
 #include "Styling/SlateColor.h"
 #include "Styling/CoreStyle.h"
 #include "Widgets/SLeafWidget.h"
+#include "Styling/SlateTypes.h"
 
 class FPaintArgs;
 class FSlateWindowElementList;
@@ -79,11 +80,8 @@ protected:
 
 protected:
 
-	/** The slate brush to draw for the image, or a bound delegate to a brush. */
-	TAttribute< const FSlateBrush* > Image;
-	
-	/** The copy of the image data, some users reuse the same FSlateBrush pointer, so we need to check it against the last true data to see what changed. */
-	FSlateBrush ImageCache;
+	/** The slate brush to draw for the image that we can invalidate. */
+	FInvalidatableBrushAttribute Image;
 
 	/** Color and opacity scale for this image */
 	TAttribute<FSlateColor> ColorAndOpacity;
diff --git a/Engine/Source/Runtime/SlateCore/Public/Widgets/SWidget.h b/Engine/Source/Runtime/SlateCore/Public/Widgets/SWidget.h
--- a/Engine/Source/Runtime/SlateCore/Public/Widgets/SWidget.h
+++ b/Engine/Source/Runtime/SlateCore/Public/Widgets/SWidget.h
@@ -156,6 +156,32 @@ private:
 };
 
 
+/**
+ * Performs the attribute assignment and invalidates the widget minimally based on what actually changed.  So if the boundness of the attribute didn't change
+ * volatility won't need to be recalculated.  Returns true if the value changed.
+ */
+template<typename TargetValueType, typename SourceValueType>
+static bool SetWidgetAttribute(SWidget& ThisWidget, TAttribute<TargetValueType>& TargetValue, const TAttribute<SourceValueType>& SourceValue, EInvalidateWidgetReason BaseInvalidationReason)
+{
+	if (!TargetValue.IdenticalTo(SourceValue))
+	{
+		const bool bWasBound = TargetValue.IsBound();
+		const bool bBoundnessChanged = bWasBound != SourceValue.IsBound();
+		TargetValue = SourceValue;
+
+		EInvalidateWidgetReason InvalidateReason = BaseInvalidationReason;
+		if (bBoundnessChanged)
+		{
+			InvalidateReason |= EInvalidateWidgetReason::Volatility;
+		}
+
+		ThisWidget.Invalidate(InvalidateReason);
+		return true;
+	}
+
+	return false;
+}
+
 class IToolTip;
 
 /**
@@ -1489,23 +1515,7 @@ protected:
 	template<typename TargetValueType, typename SourceValueType>
 	bool SetAttribute(TAttribute<TargetValueType>& TargetValue, const TAttribute<SourceValueType>& SourceValue, EInvalidateWidgetReason BaseInvalidationReason)
 	{
-		if (!TargetValue.IdenticalTo(SourceValue))
-		{
-			const bool bWasBound = TargetValue.IsBound();
-			const bool bBoundnessChanged = bWasBound != SourceValue.IsBound();
-			TargetValue = SourceValue;
-
-			EInvalidateWidgetReason InvalidateReason = BaseInvalidationReason;
-			if (bBoundnessChanged)
-			{
-				InvalidateReason |= EInvalidateWidgetReason::Volatility;
-			}
-
-			Invalidate(InvalidateReason);
-			return true;
-		}
-
-		return false;
+		return SetWidgetAttribute(*this, TargetValue, SourceValue, BaseInvalidationReason);
 	}
 
 protected:
