Slate - The basic SSlider needed a lot more invalidation to properly work under Global Invalidation mode.

Matt.Kuhlenschmidt

diff --git a/Engine/Source/Runtime/Slate/Private/Widgets/Input/SSlider.cpp b/Engine/Source/Runtime/Slate/Private/Widgets/Input/SSlider.cpp
--- a/Engine/Source/Runtime/Slate/Private/Widgets/Input/SSlider.cpp
+++ b/Engine/Source/Runtime/Slate/Private/Widgets/Input/SSlider.cpp
@@ -414,7 +414,8 @@ float SSlider::PositionToValue( const FGeometry& MyGeometry, const FVector2D& Ab
 	return RelativeValue;
 }
 
-const FSlateBrush* SSlider::GetBarImage() const {
+const FSlateBrush* SSlider::GetBarImage() const
+{
 	if (!IsEnabled() || LockedAttribute.Get())
 	{
 		return &Style->DisabledBarImage;
@@ -429,7 +430,8 @@ const FSlateBrush* SSlider::GetBarImage() const {
 	}
 }
 
-const FSlateBrush* SSlider::GetThumbImage() const {
+const FSlateBrush* SSlider::GetThumbImage() const
+{
 	if (!IsEnabled() || LockedAttribute.Get())
 	{
 		return &Style->DisabledThumbImage;
@@ -463,7 +465,7 @@ float SSlider::GetNormalizedValue() const
 
 void SSlider::SetValue(const TAttribute<float>& InValueAttribute)
 {
-	ValueAttribute = InValueAttribute;
+	SetAttribute(ValueAttribute, InValueAttribute, EInvalidateWidgetReason::Paint);
 }
 
 void SSlider::SetMinAndMaxValues(float InMinValue, float InMaxValue)
@@ -478,27 +480,31 @@ void SSlider::SetMinAndMaxValues(float InMinValue, float InMaxValue)
 
 void SSlider::SetIndentHandle(const TAttribute<bool>& InIndentHandle)
 {
-	IndentHandle = InIndentHandle;
+	SetAttribute(IndentHandle, InIndentHandle, EInvalidateWidgetReason::Paint);
 }
 
 void SSlider::SetLocked(const TAttribute<bool>& InLocked)
 {
-	LockedAttribute = InLocked;
+	SetAttribute(LockedAttribute, InLocked, EInvalidateWidgetReason::Paint);
 }
 
 void SSlider::SetOrientation(EOrientation InOrientation)
 {
-	Orientation = InOrientation;
+	if (Orientation != InOrientation)
+	{
+		Orientation = InOrientation;
+		Invalidate(EInvalidateWidgetReason::Layout);
+	}
 }
 
 void SSlider::SetSliderBarColor(FSlateColor InSliderBarColor)
 {
-	SliderBarColor = InSliderBarColor;
+	SetAttribute(SliderBarColor, TAttribute<FSlateColor>(InSliderBarColor), EInvalidateWidgetReason::Paint);
 }
 
 void SSlider::SetSliderHandleColor(FSlateColor InSliderHandleColor)
 {
-	SliderHandleColor = InSliderHandleColor;
+	SetAttribute(SliderHandleColor, TAttribute<FSlateColor>(InSliderHandleColor), EInvalidateWidgetReason::Paint);
 }
 
 float SSlider::GetStepSize() const
@@ -511,11 +517,13 @@ void SSlider::SetStepSize(const TAttribute<float>& InStepSize)
 	StepSize = InStepSize;
 }
 
-void SSlider::SetMouseUsesStep(bool MouseUsesStep) {
+void SSlider::SetMouseUsesStep(bool MouseUsesStep)
+{
 	bMouseUsesStep = MouseUsesStep;
 }
 
-void SSlider::SetRequiresControllerLock(bool RequiresControllerLock) {
+void SSlider::SetRequiresControllerLock(bool RequiresControllerLock)
+{
 	bRequiresControllerLock = RequiresControllerLock;
 }
 
