Slate - Updatign the richtextblock to use the new SetAttribute which is more efficent than invalidating volatility and layout straight up.

diff --git a/Engine/Source/Runtime/Slate/Private/Widgets/Text/SRichTextBlock.cpp b/Engine/Source/Runtime/Slate/Private/Widgets/Text/SRichTextBlock.cpp
--- a/Engine/Source/Runtime/Slate/Private/Widgets/Text/SRichTextBlock.cpp
+++ b/Engine/Source/Runtime/Slate/Private/Widgets/Text/SRichTextBlock.cpp
@@ -126,38 +126,32 @@ void SRichTextBlock::SetTextFlowDirection(const TOptional<ETextFlowDirection>& I
 
 void SRichTextBlock::SetWrapTextAt(const TAttribute<float>& InWrapTextAt)
 {
-	WrapTextAt = InWrapTextAt;
-	Invalidate(EInvalidateWidget::LayoutAndVolatility);
+	SetAttribute(WrapTextAt, InWrapTextAt, EInvalidateWidgetReason::Layout);
 }
 
 void SRichTextBlock::SetAutoWrapText(const TAttribute<bool>& InAutoWrapText)
 {
-	AutoWrapText = InAutoWrapText;
-	Invalidate(EInvalidateWidget::LayoutAndVolatility);
+	SetAttribute(AutoWrapText, InAutoWrapText, EInvalidateWidgetReason::Layout);
 }
 
 void SRichTextBlock::SetWrappingPolicy(const TAttribute<ETextWrappingPolicy>& InWrappingPolicy)
 {
-	WrappingPolicy = InWrappingPolicy;
-	Invalidate(EInvalidateWidget::LayoutAndVolatility);
+	SetAttribute(WrappingPolicy, InWrappingPolicy, EInvalidateWidgetReason::Layout);
 }
 
 void SRichTextBlock::SetLineHeightPercentage(const TAttribute<float>& InLineHeightPercentage)
 {
-	LineHeightPercentage = InLineHeightPercentage;
-	Invalidate(EInvalidateWidget::LayoutAndVolatility);
+	SetAttribute(LineHeightPercentage, InLineHeightPercentage, EInvalidateWidgetReason::Layout);
 }
 
 void SRichTextBlock::SetMargin(const TAttribute<FMargin>& InMargin)
 {
-	Margin = InMargin;
-	Invalidate(EInvalidateWidget::LayoutAndVolatility);
+	SetAttribute(Margin, InMargin, EInvalidateWidgetReason::Layout);
 }
 
 void SRichTextBlock::SetJustification(const TAttribute<ETextJustify::Type>& InJustification)
 {
-	Justification = InJustification;
-	Invalidate(EInvalidateWidget::LayoutAndVolatility);
+	SetAttribute(Justification, InJustification, EInvalidateWidgetReason::Layout);
 }
 
 void SRichTextBlock::SetTextStyle(const FTextBlockStyle& InTextStyle)
@@ -168,8 +162,7 @@ void SRichTextBlock::SetTextStyle(const FTextBlockStyle& InTextStyle)
 
 void SRichTextBlock::SetMinDesiredWidth(const TAttribute<float>& InMinDesiredWidth)
 {
-	MinDesiredWidth = InMinDesiredWidth;
-	Invalidate(EInvalidateWidget::LayoutAndVolatility);
+	SetAttribute(MinDesiredWidth, InMinDesiredWidth, EInvalidateWidgetReason::Layout);
 }
 
 void SRichTextBlock::SetDecoratorStyleSet(const ISlateStyle* NewDecoratorStyleSet)
