Fixed tooltips not disappearing after mousing away from the object that created it.

diff --git a/Engine/Source/Runtime/Slate/Private/Framework/Application/SlateUser.cpp b/Engine/Source/Runtime/Slate/Private/Framework/Application/SlateUser.cpp
--- a/Engine/Source/Runtime/Slate/Private/Framework/Application/SlateUser.cpp
+++ b/Engine/Source/Runtime/Slate/Private/Framework/Application/SlateUser.cpp
@@ -1208,12 +1208,13 @@ void FSlateUser::UpdateTooltip(const FMenuStack& MenuStack, bool bCanSpawnNewToo
 			if (NewTooltipVisualizer)
 			{
 				ActiveTooltipInfo.TooltipVisualizer = NewTooltipVisualizer;
+				ActiveTooltipInfo.Tooltip = NewTooltip;
 			}
 			else if (bCanSpawnNewTooltip && NewTooltip)
 			{
 				ShowTooltip(NewTooltip.ToSharedRef(), DesiredLocation);
 				ActiveTooltipInfo.SourceWidget = WidgetProvidingNewTooltip;
-			}
+			}	
 		}
 	}
 
