Removed code in GetOrCreateUser that treats invalid(negative) indices as index 0.
ProcessExternalReply now treats invalid user indices as user index 0 and takes int32's instead of uint32's.
SetUserFocus now uses GetUser rather than GetOrCreateUser, and does nothing in the event that GetUser returns null. This prevents the creation of an invalid user.

Changed ProcessExternalReply's signature to take int32's instead of uint32's.

diff --git a/Engine/Source/Runtime/Slate/Private/Framework/Application/SlateApplication.cpp b/Engine/Source/Runtime/Slate/Private/Framework/Application/SlateApplication.cpp
--- a/Engine/Source/Runtime/Slate/Private/Framework/Application/SlateApplication.cpp
+++ b/Engine/Source/Runtime/Slate/Private/Framework/Application/SlateApplication.cpp
@@ -2292,22 +2292,21 @@ void FSlateApplication::ActivateGameViewport()
 
 bool FSlateApplication::SetUserFocus(uint32 UserIndex, const TSharedPtr<SWidget>& WidgetToFocus, EFocusCause ReasonFocusIsChanging /* = EFocusCause::SetDirectly*/)
 {
-	const bool bValidWidget = WidgetToFocus.IsValid();
-	ensureMsgf(bValidWidget, TEXT("Attempting to focus an invalid widget. If your intent is to clear focus use ClearUserFocus()"));
-	if (bValidWidget)
+	TSharedPtr<FSlateUser> CurrentUser = GetUser(UserIndex);
+	if (ensureMsgf(WidgetToFocus.IsValid(), TEXT("Attempting to focus an invalid widget. If your intent is to clear focus use ClearUserFocus()")) && CurrentUser)
 	{
 		FWidgetPath PathToWidget;
 		const bool bFound = FSlateWindowHelper::FindPathToWidget(SlateWindows, WidgetToFocus.ToSharedRef(), /*OUT*/ PathToWidget);
 		if (bFound)
 		{
-			return SetUserFocus(*GetOrCreateUser(UserIndex), PathToWidget, ReasonFocusIsChanging);
+			return SetUserFocus(*CurrentUser, PathToWidget, ReasonFocusIsChanging);
 		}
 		else
 		{
 			const bool bFoundVirtual = FSlateWindowHelper::FindPathToWidget(SlateVirtualWindows, WidgetToFocus.ToSharedRef(), /*OUT*/ PathToWidget);
 			if (bFoundVirtual)
 			{
-				return SetUserFocus(*GetOrCreateUser(UserIndex), PathToWidget, ReasonFocusIsChanging);
+				return SetUserFocus(*CurrentUser, PathToWidget, ReasonFocusIsChanging);
 			}
 		}
 	}
@@ -2835,15 +2834,16 @@ TSharedPtr<SWindow> FSlateApplication::FindWidgetWindow( TSharedRef<const SWidge
 	return OutWidgetPath.TopLevelWindow;
 }
 
-void FSlateApplication::ProcessExternalReply(const FWidgetPath& CurrentEventPath, const FReply TheReply, const uint32 UserIndex, const uint32 PointerIndex)
+void FSlateApplication::ProcessExternalReply(const FWidgetPath& CurrentEventPath, const FReply TheReply, const int32 UserIndex, const int32 PointerIndex)
 {
+	const int32 ValidatedUserIndex = (UserIndex >= 0) ?  UserIndex : 0;
 	if (PointerIndex == CursorPointerIndex)
 	{
-		TSharedRef<FSlateUser> SlateUser = GetOrCreateUser(UserIndex);
-		const bool bIsPrimaryUser = UserIndex == CursorUserIndex;
+		TSharedRef<FSlateUser> SlateUser = GetOrCreateUser(ValidatedUserIndex);
+		const bool bIsPrimaryUser = ValidatedUserIndex == CursorUserIndex;
 
 		FPointerEvent MouseEvent(
-			UserIndex,
+			ValidatedUserIndex,
 			PointerIndex,
 			SlateUser->GetCursorPosition(),
 			SlateUser->GetPreviousCursorPosition(),
@@ -2862,11 +2862,11 @@ void FSlateApplication::ProcessExternalReply(const FWidgetPath& CurrentEventPath
 			PathToWidgetPtr = &PathToWidget;
 		}
 
-		ProcessReply(CurrentEventPath, TheReply, PathToWidgetPtr, &MouseEvent, UserIndex);
+		ProcessReply(CurrentEventPath, TheReply, PathToWidgetPtr, &MouseEvent, ValidatedUserIndex);
 	}
 	else
 	{
-		ProcessReply(CurrentEventPath, TheReply, nullptr, nullptr, UserIndex);
+		ProcessReply(CurrentEventPath, TheReply, nullptr, nullptr, ValidatedUserIndex);
 	}
 }
 
@@ -3715,12 +3715,6 @@ TSharedRef<FSlateVirtualUserHandle> FSlateApplication::FindOrCreateVirtualUser(i
 
 TSharedRef<FSlateUser> FSlateApplication::GetOrCreateUser(int32 UserIndex)
 {
-	// @todo For now we are allowing a little shenanigans since we are crashing in automation
-	if ( UserIndex < 0 )
-	{
-		UserIndex = 0;
-	}
-	
 	if (TSharedPtr<FSlateUser> FoundUser = GetUser(UserIndex))
 	{
 		return FoundUser.ToSharedRef();
@@ -6630,4 +6624,4 @@ bool FSlateApplication::InputPreProcessorsHelper::PreProcessInput(TFunctionRef<b
 	ProcessorsPendingRemoval.Reset();
 
 	return bShouldExit;
-}
+}
\ No newline at end of file
diff --git a/Engine/Source/Runtime/Slate/Public/Framework/Application/SlateApplication.h b/Engine/Source/Runtime/Slate/Public/Framework/Application/SlateApplication.h
--- a/Engine/Source/Runtime/Slate/Public/Framework/Application/SlateApplication.h
+++ b/Engine/Source/Runtime/Slate/Public/Framework/Application/SlateApplication.h
@@ -1135,7 +1135,7 @@ public:
 	 * @param UserIndex			 User index that generated the event we are replying to (defaults to 0, at least for now)
 	 * @param PointerIndex		 Pointer index that generated the event we are replying to
 	 */
-	void ProcessExternalReply(const FWidgetPath& CurrentEventPath, const FReply TheReply, const uint32 UserIndex = 0, const uint32 PointerIndex = 10 /* todo: use the enum */);
+	void ProcessExternalReply(const FWidgetPath& CurrentEventPath, const FReply TheReply, const int32 UserIndex = 0, const int32 PointerIndex = 10 /* todo: use the enum */);
 
 	/**
 	 * Apply any requests from the Reply to the application. E.g. Capture mouse
