Wrong FPointerEvent constructor was being called, defaulting UserIndex to 0 for all input.
Now constructor passes SlateUserIndex as FPointerEvent's UserIndex.
Second controller can now no longer cause UI shennanigans.

diff --git a/Engine/Source/Runtime/Slate/Private/Framework/Application/SlateUser.cpp b/Engine/Source/Runtime/Slate/Private/Framework/Application/SlateUser.cpp
--- a/Engine/Source/Runtime/Slate/Private/Framework/Application/SlateUser.cpp
+++ b/Engine/Source/Runtime/Slate/Private/Framework/Application/SlateUser.cpp
@@ -607,7 +607,7 @@ bool FSlateUser::SynthesizeCursorMoveIfNeeded()
 		
 		const bool bHasHardwareCursor = SlateApp.GetPlatformCursor() == Cursor;
 		FPointerEvent SyntheticCursorMoveEvent(
-			SlateApp.GetUserIndexForMouse(),
+			GetUserIndex(),
 			FSlateApplication::CursorPointerIndex,
 			GetCursorPosition(),
 			GetPreviousCursorPosition(),
