Fixed horizontal scrolling on touch
#rnx

diff --git a/Engine/Source/Runtime/Slate/Private/Widgets/Views/STableViewBase.cpp b/Engine/Source/Runtime/Slate/Private/Widgets/Views/STableViewBase.cpp
--- a/Engine/Source/Runtime/Slate/Private/Widgets/Views/STableViewBase.cpp
+++ b/Engine/Source/Runtime/Slate/Private/Widgets/Views/STableViewBase.cpp
@@ -601,8 +601,14 @@ FReply STableViewBase::OnTouchMoved( const FGeometry& MyGeometry, const FPointer
 {
 	if (bStartedTouchInteraction)
 	{
-		const float ScrollByAmount = InTouchEvent.GetCursorDelta().Y / MyGeometry.Scale;
+		// We only care about deltas along the scroll axis
+		FTableViewDimensions CursorDeltaDimensions(Orientation, InTouchEvent.GetCursorDelta());
+		CursorDeltaDimensions.LineAxis = 0.f;
+
+		const float ScrollByAmount = CursorDeltaDimensions.ScrollAxis / MyGeometry.Scale;
+
 		AmountScrolledWhileRightMouseDown += FMath::Abs( ScrollByAmount );
+		
 		TickScrollDelta -= ScrollByAmount;
 
 		if (FSlateApplication::Get().HasTraveledFarEnoughToTriggerDrag(InTouchEvent, PressedScreenSpacePosition))
