[Slate]
  - Fixed RequestScrollIntoView code for SListView.
  - Fixed navigation in ListView widget not working as expected toward the end of list, when it contains large number of items (> 500K) and partial items are visible. This is done by correcly computing SmoothScrollOffset when float numbers are very large, in STableViewBase.
  - Fixed ListView widget not scrolling to end of list when it contains large number of items (> 500K). This is done by decreasing tolerance from KINDA_SMALL_NUMBER (1e-4) to SMALL_NUMBER (1e-6) in detection for end of list, in STableViewBase.

diff --git a/Engine/Source/Runtime/Slate/Private/Widgets/Views/STableViewBase.cpp b/Engine/Source/Runtime/Slate/Private/Widgets/Views/STableViewBase.cpp
--- a/Engine/Source/Runtime/Slate/Private/Widgets/Views/STableViewBase.cpp
+++ b/Engine/Source/Runtime/Slate/Private/Widgets/Views/STableViewBase.cpp
@@ -263,10 +263,17 @@ void STableViewBase::Tick( const FGeometry& AllottedGeometry, const double InCur
 			{
 				ScrollOffset = ReGenerateResults.NewScrollOffset;
 			}
-			
-			
-			SetScrollOffset( FMath::Max(0.0, ScrollOffset) );
-			ItemsPanel->SmoothScrollOffset( FMath::Fractional(ScrollOffset / GetNumItemsWide()) );
+
+			// Compute SmoothScrollOffset before SetScrollOffset(). This is becasue SetScrollOffset() modifies the ScrollOffset (double downcast to float).
+			// We need original value in computation of SmoothScrollOffset.
+			double SmoothScrollOffset = ScrollOffset / GetNumItemsWide();
+			SmoothScrollOffset = SmoothScrollOffset - (int64)SmoothScrollOffset; // get fractional part
+
+			SetScrollOffset(FMath::Max(0.0, ScrollOffset));
+
+			// FMath::Fractional() is buggy as it casts to int32 (too small for the integer part of float).
+			//ItemsPanel->SmoothScrollOffset(FMath::Fractional(ScrollOffset / GetNumItemsWide()));
+			ItemsPanel->SmoothScrollOffset(SmoothScrollOffset);
 
 			if (AllowOverscroll == EAllowOverscroll::Yes)
 			{
@@ -303,7 +310,7 @@ void STableViewBase::Tick( const FGeometry& AllottedGeometry, const double InCur
 				ScrollBar->SetState( OffsetFraction, ThumbSizeFraction );
 			}
 
-			bWasAtEndOfList = (ScrollBar->DistanceFromBottom() < KINDA_SMALL_NUMBER);
+			bWasAtEndOfList = (ScrollBar->DistanceFromBottom() < SMALL_NUMBER);
 
 			bItemsNeedRefresh = false;
 			ItemsPanel->SetRefreshPending(false);
diff --git a/Engine/Source/Runtime/Slate/Public/Widgets/Views/SListView.h b/Engine/Source/Runtime/Slate/Public/Widgets/Views/SListView.h
--- a/Engine/Source/Runtime/Slate/Public/Widgets/Views/SListView.h
+++ b/Engine/Source/Runtime/Slate/Public/Widgets/Views/SListView.h
@@ -998,7 +998,7 @@ public:
 			// But we may still have space to fill!
 			if (bAtEndOfList && ViewHeightUsedSoFar < MyGeometry.GetLocalSize().Y)
 			{
-				float NewScrollOffsetForBackfill = StartIndex + (HeightGeneratedSoFar - MyGeometry.GetLocalSize().Y) / FirstItemHeight;
+				double NewScrollOffsetForBackfill = static_cast<double>(StartIndex) + (HeightGeneratedSoFar - MyGeometry.GetLocalSize().Y) / FirstItemHeight;
 
 				for( int32 ItemIndex = StartIndex-1; HeightGeneratedSoFar < MyGeometry.GetLocalSize().Y && ItemIndex >= 0; --ItemIndex )
 				{
@@ -1010,7 +1010,7 @@ public:
 					{
 						// Generated the item that puts us over the top.
 						// Count the fraction of this item that will stick out above the list
-						NewScrollOffsetForBackfill = ItemIndex + (HeightGeneratedSoFar + ItemHeight - MyGeometry.GetLocalSize().Y) / ItemHeight;
+						NewScrollOffsetForBackfill = static_cast<double>(ItemIndex) + (HeightGeneratedSoFar + ItemHeight - MyGeometry.GetLocalSize().Y) / ItemHeight;
 					}
 
 					// The widget used up some of the available vertical space.
@@ -1024,7 +1024,6 @@ public:
 		}
 
 		return FReGenerateResults(0.0f, 0.0f, 0.0f, false);
-
 	}
 
 	float GenerateWidgetForItem( const ItemType& CurItem, int32 ItemIndex, int32 StartIndex, float LayoutScaleMultiplier )
@@ -1400,14 +1399,15 @@ protected:
 				{
 					// Scroll the top of the listview to the item in question
 					double NewScrollOffset = IndexOfItem;
+
 					// Center the list view on the item in question.
 					NewScrollOffset -= (NumLiveWidgets / 2);
-					//we also don't want the widget being chopped off if it is at the end of the list
-					const double MoveBackBy = FMath::Clamp<double>(IndexOfItem - (NewScrollOffset + NumLiveWidgets), 0, FLT_MAX);
-					//Move to the correct center spot
-					NewScrollOffset += MoveBackBy;
 
-					SetScrollOffset( NewScrollOffset );
+					// Limit offset to top and bottom of the list.
+					const double MaxScrollOffset = FMath::Max(0.0, static_cast<double>(ItemsSource->Num()) - NumLiveWidgets);
+					NewScrollOffset = FMath::Clamp<double>(NewScrollOffset, 0.0, MaxScrollOffset);
+
+					SetScrollOffset(NewScrollOffset);
 				}
 				else if (bNavigateOnScrollIntoView)
 				{
@@ -1418,7 +1418,9 @@ protected:
 						if (RowGeometry.GetAbsolutePositionAtCoordinates(FVector2D::ZeroVector).Y < ListViewGeometry.GetAbsolutePositionAtCoordinates(FVector2D::ZeroVector).Y)
 						{
 							// This row is clipped on the top, so simply set it as the new scroll offset target to bump it down a bit
-							SetScrollOffset(IndexOfItem - NavigationScrollOffset);
+							const double MaxScrollOffset = FMath::Max(0.0, static_cast<double>(ItemsSource->Num()) - NumLiveWidgets);
+							double NewScrollOffset = FMath::Clamp<double>(static_cast<double>(IndexOfItem - NavigationScrollOffset), 0.0, MaxScrollOffset);
+							SetScrollOffset(NewScrollOffset);
 						}
 						else
 						{
@@ -1447,7 +1449,9 @@ protected:
 									}
 								}
 
-								SetScrollOffset(ScrollOffset + AdditionalOffset + NavigationScrollOffset);
+								const double MaxScrollOffset = FMath::Max(0.0, static_cast<double>(ItemsSource->Num()) - NumLiveWidgets);
+								double NewScrollOffset = FMath::Clamp<double>(ScrollOffset + AdditionalOffset + NavigationScrollOffset, 0.0, MaxScrollOffset);
+								SetScrollOffset(NewScrollOffset);
 							}
 						}
 					}
