Removed spurious qqq in DockDirection comment
#rnx

diff --git a/Engine/Source/Runtime/Slate/Private/Framework/Docking/FDockingDragOperation.h b/Engine/Source/Runtime/Slate/Private/Framework/Docking/FDockingDragOperation.h
--- a/Engine/Source/Runtime/Slate/Private/Framework/Docking/FDockingDragOperation.h
+++ b/Engine/Source/Runtime/Slate/Private/Framework/Docking/FDockingDragOperation.h
@@ -51,7 +51,7 @@ public:
 
 		/** We'll put the tab relative to this node */
 		TWeakPtr<class SDockingNode> TargetNode;
-		/** Relation to node where we will put the tab.qqq */
+		/** Relation to node where we will put the tab. */
 		SDockingNode::RelativeDirection DockDirection;
 	};
 
