Setting the default value of the Safe Zone TOptional to 1.f and setting it to 1.f if a negative value is passed in.

diff --git a/Engine/Source/Runtime/Slate/Private/Widgets/Layout/SSafeZone.cpp b/Engine/Source/Runtime/Slate/Private/Widgets/Layout/SSafeZone.cpp
--- a/Engine/Source/Runtime/Slate/Private/Widgets/Layout/SSafeZone.cpp
+++ b/Engine/Source/Runtime/Slate/Private/Widgets/Layout/SSafeZone.cpp
@@ -25,8 +25,13 @@ static FAutoConsoleVariableRef CVarEnableSafeZoneScale(
 
 void SSafeZone::SetGlobalSafeZoneScale(TOptional<float> InScale)
 {
+	float NewScale = InScale.Get(1.f);
+	if (NewScale < 0)
+	{
+		NewScale = 1.f;
+	}
 	bEnableSafeZoneScale = InScale.IsSet();
-	GSafeZoneScale = InScale.Get(0.f);
+	GSafeZoneScale = NewScale;
 
 	FCoreDelegates::OnSafeFrameChangedEvent.Broadcast();
 }
