Add logging for tracking down an assert

diff --git a/Engine/Source/Runtime/Slate/Private/Widgets/Input/SMenuAnchor.cpp b/Engine/Source/Runtime/Slate/Private/Widgets/Input/SMenuAnchor.cpp
--- a/Engine/Source/Runtime/Slate/Private/Widgets/Input/SMenuAnchor.cpp
+++ b/Engine/Source/Runtime/Slate/Private/Widgets/Input/SMenuAnchor.cpp
@@ -10,6 +10,11 @@
 #include "Framework/Application/SlateApplication.h"
 #include "Layout/LayoutUtils.h"
 
+#if WITH_EDITOR
+#include "Framework/MultiBox/MultiBox.h"
+#include "Framework/MultiBox/ToolMenuBase.h"
+#endif
+
 static FVector2D GetMenuOffsetForPlacement(const FGeometry& AllottedGeometry, EMenuPlacement PlacementMode, FVector2D PopupSizeLocalSpace)
 {
 	switch (PlacementMode)
@@ -379,6 +384,35 @@ void SMenuAnchor::SetIsOpen( bool InIsOpen, const bool bFocusMenu, const int32 F
 							
 							if (ensure(NewMenu.IsValid()))
 							{
+#if WITH_EDITOR
+								// Reporting more information for UE-81655
+								if (!NewMenu->GetOwnedWindow().IsValid())
+								{
+									UE_LOG(LogSlate, Error, TEXT("!NewMenu->GetOwnedWindow().IsValid(), WidgetPath:\n%s"), *MyWidgetPath.ToString());
+
+									static const FName SMultiBoxWidgetTypeName = "SMultiBoxWidget";
+									for (int32 WidgetIndex = MyWidgetPath.Widgets.Num() - 1; WidgetIndex >= 0; --WidgetIndex)
+									{
+										if (MyWidgetPath.Widgets[WidgetIndex].Widget->GetType() == SMultiBoxWidgetTypeName)
+										{
+											TSharedRef<SMultiBoxWidget> MultiBoxWidget = StaticCastSharedRef<SMultiBoxWidget>(MyWidgetPath.Widgets[WidgetIndex].Widget);
+											TSharedRef<const FMultiBox> MultiBox = MultiBoxWidget->GetMultiBox();
+											FString BlockText;
+											const TArray< TSharedRef< const FMultiBlock > >& Blocks = MultiBox->GetBlocks();
+											for (int32 BlockIndex = 0; BlockIndex < Blocks.Num(); ++BlockIndex)
+											{
+												const bool bIsFinalBlock = (BlockIndex == (Blocks.Num() - 1));
+												const TSharedRef< const FMultiBlock >& Block = Blocks[BlockIndex];
+												BlockText += FString::Printf(TEXT("%s (%d)%s"), *Block->GetExtensionHook().ToString(), (int32)Block->GetType(), bIsFinalBlock ? TEXT("") : TEXT(", "));
+											}
+											UE_LOG(LogSlate, Error, TEXT(" Blocks: %s"), *BlockText);
+											break;
+										}
+									}
+
+									UE_LOG(LogSlate, Error, TEXT(" MenuContentRef: %s"), *MenuContentRef->ToString());
+								}
+#endif
 								check(NewMenu->GetOwnedWindow().IsValid());
 
 								PopupMenuPtr = NewMenu;
