Copied CL# 9259124
Copied CL# 9452975
Copied CL# 9538568

Bug fixes:
1) UE-80271: Solved the issue that plugin-dependent SWidget windows appeared  blank when the Editor was loaded with the respective plugin disabled.
2) Attemp to solve UE-80344: Avoided the out-of-bounds TArray error in FTabManager.
3) Standalone SWidgets from non-loaded plugins now behave in the same way than the docked SWidgets from non-loaded plugings: They are also re-opened when the plugin is loaded again and the Editor re-started. Before, they just disappeared.

Enhancement/bug fix: Unrecognized Docked SWidgets from non-loaded plugins are not displayed on the UI rather than being showed up as "unrecognized tab", similarly to how standalone widgets behave.
Also removed 1 of 2 very similar warnings.
This code should also fix the 1-off crash of crash when cliking on "Window" in the Menu Bar in QA Game, given that we have reverted even more changes and have limited it only to the layout generated in MainFrameModule.

Fixed regression suffered in CL 9452975 for experimental "Modeling Tools Editor Mode" plugin.

diff --git a/Engine/Source/Editor/LevelEditor/Private/SLevelEditor.cpp b/Engine/Source/Editor/LevelEditor/Private/SLevelEditor.cpp
--- a/Engine/Source/Editor/LevelEditor/Private/SLevelEditor.cpp
+++ b/Engine/Source/Editor/LevelEditor/Private/SLevelEditor.cpp
@@ -1319,7 +1319,9 @@ TSharedRef<SWidget> SLevelEditor::RestoreContentArea( const TSharedRef<SDockTab>
 	LevelEditorModule.OnRegisterLayoutExtensions().Broadcast(LayoutExtender);
 	Layout->ProcessExtensions(LayoutExtender);
 
-	return LevelEditorTabManager->RestoreFrom( Layout, OwnerWindow ).ToSharedRef();
+	const bool bEmbedTitleAreaContent = false;
+	const EOutputCanBeNullptr OutputCanBeNullptr = EOutputCanBeNullptr::IfNoTabValid;
+	return LevelEditorTabManager->RestoreFrom(Layout, OwnerWindow, bEmbedTitleAreaContent, OutputCanBeNullptr).ToSharedRef();
 }
 
 void SLevelEditor::HandleExperimentalSettingChanged(FName PropertyName)
diff --git a/Engine/Source/Editor/MainFrame/Private/MainFrameModule.cpp b/Engine/Source/Editor/MainFrame/Private/MainFrameModule.cpp
--- a/Engine/Source/Editor/MainFrame/Private/MainFrameModule.cpp
+++ b/Engine/Source/Editor/MainFrame/Private/MainFrameModule.cpp
@@ -204,7 +204,8 @@ void FMainFrameModule::CreateDefaultMainFrame( const bool bStartImmersive, const
 				)
 			);
 
-			MainFrameContent = FGlobalTabmanager::Get()->RestoreFrom( LoadedLayout, RootWindow, bEmbedTitleAreaContent );
+			const EOutputCanBeNullptr OutputCanBeNullptr = EOutputCanBeNullptr::IfNoOpenTabValid;
+			MainFrameContent = FGlobalTabmanager::Get()->RestoreFrom(LoadedLayout, RootWindow, bEmbedTitleAreaContent, OutputCanBeNullptr);
 			bLevelEditorIsMainTab = true;
 		}
 
diff --git a/Engine/Source/Runtime/Slate/Private/Framework/Docking/TabManager.cpp b/Engine/Source/Runtime/Slate/Private/Framework/Docking/TabManager.cpp
--- a/Engine/Source/Runtime/Slate/Private/Framework/Docking/TabManager.cpp
+++ b/Engine/Source/Runtime/Slate/Private/Framework/Docking/TabManager.cpp
@@ -838,7 +838,7 @@ void FTabManager::UnregisterAllTabSpawners()
 	TabSpawner.Empty();
 }
 
-TSharedPtr<SWidget> FTabManager::RestoreFrom( const TSharedRef<FLayout>& Layout, const TSharedPtr<SWindow>& ParentWindow, const bool bEmbedTitleAreaContent )
+TSharedPtr<SWidget> FTabManager::RestoreFrom(const TSharedRef<FLayout>& Layout, const TSharedPtr<SWindow>& ParentWindow, const bool bEmbedTitleAreaContent, const EOutputCanBeNullptr RestoreAreaOutputCanBeNullptr)
 {
 	ActiveLayoutName = Layout->LayoutName;
 
@@ -858,13 +858,10 @@ TSharedPtr<SWidget> FTabManager::RestoreFrom( const TSharedRef<FLayout>& Layout,
 
 			if ( bHasOpenTabs )
 			{
-				const bool bCanOutputBeNullptr = true;
-				RestoredDockArea = RestoreArea(ThisArea, ParentWindow, bEmbedTitleAreaContent, bCanOutputBeNullptr);
+				RestoredDockArea = RestoreArea(ThisArea, ParentWindow, bEmbedTitleAreaContent, RestoreAreaOutputCanBeNullptr);
 				// Invalidate all tabs in ThisArea because they were not recognized
 				if (!RestoredDockArea)
 				{
-					UE_LOG(LogSlate, Warning, TEXT("The area index %d/%d (extension id \"%s\") attempted to spawn but failed."),
-						AreaIndex+1, Layout->Areas.Num(), *ThisArea->GetExtensionId().ToString());
 					SetTabsTo(ThisArea, ETabState::InvalidTab, ETabState::OpenedTab);
 					InvalidDockAreas.Add(ThisArea);
 				}
@@ -1337,9 +1334,9 @@ FTabManager::FTabManager( const TSharedPtr<SDockTab>& InOwnerTab, const TSharedR
 }
 
 TSharedPtr<SDockingArea> FTabManager::RestoreArea(
-	const TSharedRef<FArea>& AreaToRestore, const TSharedPtr<SWindow>& InParentWindow, const bool bEmbedTitleAreaContent, const bool bCanOutputBeNullptr)
+	const TSharedRef<FArea>& AreaToRestore, const TSharedPtr<SWindow>& InParentWindow, const bool bEmbedTitleAreaContent, const EOutputCanBeNullptr OutputCanBeNullptr)
 {
-	if (TSharedPtr<SDockingNode> RestoredNode = RestoreArea_Helper(AreaToRestore, InParentWindow, bEmbedTitleAreaContent, bCanOutputBeNullptr))
+	if (TSharedPtr<SDockingNode> RestoredNode = RestoreArea_Helper(AreaToRestore, InParentWindow, bEmbedTitleAreaContent, OutputCanBeNullptr))
 	{
 		TSharedRef<SDockingArea> RestoredArea = StaticCastSharedRef<SDockingArea>(RestoredNode->AsShared());
 		RestoredArea->CleanUp(SDockingNode::TabRemoval_None);
@@ -1347,32 +1344,59 @@ TSharedPtr<SDockingArea> FTabManager::RestoreArea(
 	}
 	else
 	{
-		check(bCanOutputBeNullptr);
+		check(OutputCanBeNullptr != EOutputCanBeNullptr::Never);
 		return nullptr;
 	}
 }
 
 TSharedPtr<SDockingNode> FTabManager::RestoreArea_Helper(
-	const TSharedRef<FLayoutNode>& LayoutNode, const TSharedPtr<SWindow>& ParentWindow, const bool bEmbedTitleAreaContent, const bool bCanOutputBeNullptr)
+	const TSharedRef<FLayoutNode>& LayoutNode, const TSharedPtr<SWindow>& ParentWindow, const bool bEmbedTitleAreaContent, const EOutputCanBeNullptr OutputCanBeNullptr)
 {
 	TSharedPtr<FTabManager::FStack> NodeAsStack = LayoutNode->AsStack();
 	TSharedPtr<FTabManager::FSplitter> NodeAsSplitter = LayoutNode->AsSplitter();
 	TSharedPtr<FTabManager::FArea> NodeAsArea = LayoutNode->AsArea();
+	const bool bCanOutputBeNullptr = (OutputCanBeNullptr != EOutputCanBeNullptr::Never);
 
 	if (NodeAsStack.IsValid())
 	{
 		TSharedPtr<SDockTab> WidgetToActivate;
 
 		TSharedPtr<SDockingTabStack> NewStackWidget;
-		if (!bCanOutputBeNullptr)
+		// Should we init NewStackWidget before the for loop? It depends on OutputCanBeNullptr
+		bool bIsNewStackWidgetInit = false;
+		// 1. If EOutputCanBeNullptr::Never, function cannot return nullptr
+		if (OutputCanBeNullptr == EOutputCanBeNullptr::Never)
+		{
+			bIsNewStackWidgetInit = true;
+		}
+		// 2. If EOutputCanBeNullptr::IfNoTabValid, we must init the SWidget as soon as any tab is valid for spawning
+		else if (OutputCanBeNullptr == EOutputCanBeNullptr::IfNoTabValid)
+		{
+			// Note: IsValidTabForSpawning does not check whether SpawnTab() will return nullptr
+			for (const FTab& SomeTab : NodeAsStack->Tabs)
+			{
+				if (IsValidTabForSpawning(SomeTab))
+				{
+					bIsNewStackWidgetInit = true;
+					break;
+				}
+			}
+		}
+		// 3. If EOutputCanBeNullptr::IfNoOpenTabValid, we must init the SWidget as soon as any open tab is valid for spawning. For efficiency, done in the for loop
+		// 4. Else, case not handled --> error
+		else if (OutputCanBeNullptr != EOutputCanBeNullptr::IfNoOpenTabValid)
+		{
+			check(false);
+		}
+		// Initialize the SWidget already?
+		if (bIsNewStackWidgetInit)
 		{
 			NewStackWidget = SNew(SDockingTabStack, NodeAsStack.ToSharedRef());
 			NewStackWidget->SetSizeCoefficient(LayoutNode->GetSizeCoefficient());
 		}
-		for (int32 TabIndex=0; TabIndex < NodeAsStack->Tabs.Num(); ++TabIndex )
+		// Open Tabs
+		for (const FTab& SomeTab : NodeAsStack->Tabs)
 		{
-			const FTab& SomeTab = NodeAsStack->Tabs[ TabIndex ];
-
 			if (SomeTab.TabState == ETabState::OpenedTab && IsValidTabForSpawning(SomeTab))
 			{
 				const TSharedPtr<SDockTab> NewTabWidget = SpawnTab(SomeTab.TabId, ParentWindow, bCanOutputBeNullptr);
@@ -1441,7 +1465,7 @@ TSharedPtr<SDockingNode> FTabManager::RestoreArea_Helper(
 			NewWindow->SetTitle(FGlobalTabmanager::Get()->GetApplicationTitle());
 
 			TArray<TSharedRef<SDockingNode>> DockingNodes;
-			if (!bCanOutputBeNullptr || CanRestoreSplitterContent(DockingNodes, NodeAsArea.ToSharedRef(), NewWindow))
+			if (CanRestoreSplitterContent(DockingNodes, NodeAsArea.ToSharedRef(), NewWindow, OutputCanBeNullptr))
 			{
 				// Create SplitterWidget only if it will be filled with at least 1 DockingNodes
 				// Any windows that were "pulled out" of a dock area should be children of the window in which the parent dock area resides.
@@ -1470,28 +1494,54 @@ TSharedPtr<SDockingNode> FTabManager::RestoreArea_Helper(
 		}
 		else
 		{
-			SAssignNew( NewDockAreaWidget, SDockingArea, SharedThis(this), NodeAsArea.ToSharedRef() )
-
-				// We only want to set a parent window on this dock area, if we need to have title area content
-				// embedded within it.  SDockingArea assumes that if it has a parent window set, then it needs to have
-				// title area content 
-				.ParentWindow( bEmbedTitleAreaContent ? ParentWindow : TSharedPtr<SWindow>() )
-
-				// Never manage these windows, even if a parent window is set.  The owner will take care of
-				// destroying these windows.
-				.ShouldManageParentWindow( false );
+			TArray<TSharedRef<SDockingNode>> DockingNodes;
+			if (CanRestoreSplitterContent(DockingNodes, NodeAsArea.ToSharedRef(), ParentWindow, OutputCanBeNullptr))
+			{
+				SAssignNew(NewDockAreaWidget, SDockingArea, SharedThis(this), NodeAsArea.ToSharedRef())
+					// We only want to set a parent window on this dock area, if we need to have title area content
+					// embedded within it.  SDockingArea assumes that if it has a parent window set, then it needs to have
+					// title area content 
+					.ParentWindow(bEmbedTitleAreaContent ? ParentWindow : TSharedPtr<SWindow>())
+					// Never manage these windows, even if a parent window is set.  The owner will take care of
+					// destroying these windows.
+					.ShouldManageParentWindow(false);
 
-			RestoreSplitterContent( NodeAsArea.ToSharedRef(), NewDockAreaWidget.ToSharedRef(), ParentWindow );
+				// Restore content
+				if (!bCanOutputBeNullptr)
+				{
+					RestoreSplitterContent(NodeAsArea.ToSharedRef(), NewDockAreaWidget.ToSharedRef(), ParentWindow);
+				}
+				else
+				{
+					RestoreSplitterContent(DockingNodes, NewDockAreaWidget.ToSharedRef());
+				}
+			}
 		}
 		
 		return NewDockAreaWidget;
 	}
 	else if ( NodeAsSplitter.IsValid() ) 
 	{
-		TSharedRef<SDockingSplitter> NewSplitterWidget = SNew( SDockingSplitter, NodeAsSplitter.ToSharedRef() );
-		NewSplitterWidget->SetSizeCoefficient(LayoutNode->GetSizeCoefficient());
-		RestoreSplitterContent( NodeAsSplitter.ToSharedRef(), NewSplitterWidget, ParentWindow );
-		return NewSplitterWidget;
+		TArray<TSharedRef<SDockingNode>> DockingNodes;
+		if (CanRestoreSplitterContent(DockingNodes, NodeAsSplitter.ToSharedRef(), ParentWindow, OutputCanBeNullptr))
+		{
+			TSharedRef<SDockingSplitter> NewSplitterWidget = SNew( SDockingSplitter, NodeAsSplitter.ToSharedRef() );
+			NewSplitterWidget->SetSizeCoefficient(LayoutNode->GetSizeCoefficient());
+			// Restore content
+			if (!bCanOutputBeNullptr)
+			{
+				RestoreSplitterContent(NodeAsSplitter.ToSharedRef(), NewSplitterWidget, ParentWindow);
+			}
+			else
+			{
+				RestoreSplitterContent(DockingNodes, NewSplitterWidget);
+			}
+			return NewSplitterWidget;
+		}
+		else
+		{
+			return nullptr;
+		}
 	}
 	else
 	{
@@ -1502,8 +1552,12 @@ TSharedPtr<SDockingNode> FTabManager::RestoreArea_Helper(
 	}
 }
 
-bool FTabManager::CanRestoreSplitterContent( TArray<TSharedRef<SDockingNode>>& DockingNodes, const TSharedRef<FSplitter>& SplitterNode, const TSharedPtr<SWindow>& ParentWindow )
+bool FTabManager::CanRestoreSplitterContent(TArray<TSharedRef<SDockingNode>>& DockingNodes, const TSharedRef<FSplitter>& SplitterNode, const TSharedPtr<SWindow>& ParentWindow, const EOutputCanBeNullptr OutputCanBeNullptr)
 {
+	if (OutputCanBeNullptr == EOutputCanBeNullptr::Never)
+	{
+		return true;
+	}
 	DockingNodes.Empty();
 	// Restore the contents of this splitter.
 	for ( int32 ChildNodeIndex = 0; ChildNodeIndex < SplitterNode->ChildNodes.Num(); ++ChildNodeIndex )
@@ -1511,8 +1565,7 @@ bool FTabManager::CanRestoreSplitterContent( TArray<TSharedRef<SDockingNode>>& D
 		const TSharedRef<FLayoutNode> ThisChildNode = SplitterNode->ChildNodes[ChildNodeIndex];
 
 		const bool bEmbedTitleAreaContent = false;
-		const bool bCanOutputBeNullptr = true;
-		const TSharedPtr<SDockingNode> ThisChildNodeWidget = RestoreArea_Helper(ThisChildNode, ParentWindow, bEmbedTitleAreaContent, bCanOutputBeNullptr);
+		const TSharedPtr<SDockingNode> ThisChildNodeWidget = RestoreArea_Helper(ThisChildNode, ParentWindow, bEmbedTitleAreaContent, OutputCanBeNullptr);
 		if (ThisChildNodeWidget)
 		{
 			const TSharedRef<SDockingNode> ThisChildNodeWidgetRef = StaticCastSharedRef<SDockingNode>(ThisChildNodeWidget->AsShared());
@@ -1603,17 +1656,21 @@ TSharedPtr<SDockTab> FTabManager::SpawnTab(const FTabId& TabId, const TSharedPtr
 	if (bSpawningAllowedBySpawner && !NewTabWidget.IsValid())
 	{
 		// We don't know how to spawn this tab. 2 alternatives:
-		// 1) Make a dummy tab so that things aren't entirely broken (previous versions did this).
-		// 2) Do not open the widget, but keep it saved in the layout, so this function should still return it.
-		const FString StringToDisplay = (Spawner.IsValid() && !Spawner->GetDisplayName().IsEmpty() ? Spawner->GetDisplayName().ToString() : TabId.TabType.ToString());
-		UE_LOG(LogSlate, Warning,
-			TEXT("The tab \"%s\" attempted to spawn but failed for some reason. It will not be displayed but it will be saved in the layout settings file."),
-			*(!StringToDisplay.IsEmpty() ? StringToDisplay : FString("Unknown"))
-		);
+		// 1) Make a dummy tab so that things aren't entirely broken (previous versions of UE did this in all cases).
+		// 2) Do not open the widget and return nullptr, but keep the unknown widget saved in the layout. E.g., applied when calling RestoreFrom() from MainFrameModule.
 
+		FString StringToDisplay = (Spawner.IsValid() && !Spawner->GetDisplayName().IsEmpty() ? Spawner->GetDisplayName().ToString() : TabId.TabType.ToString());
+		if (StringToDisplay.IsEmpty())
+		{
+			StringToDisplay = FString("Unknown");
+		}
 		// If an output must be generated, create an "unrecognized tab"
 		if (!bCanOutputBeNullptr)
 		{
+			UE_LOG(LogSlate, Warning,
+				TEXT("The tab \"%s\" attempted to spawn but failed for some reason. An \"unrecognized tab\" will be returned instead."), *StringToDisplay
+			);
+
 			NewTabWidget = SNew(SDockTab)
 				.Label( TabId.ToText() )
 				.ShouldAutosize( false )
@@ -1629,6 +1686,13 @@ TSharedPtr<SDockTab> FTabManager::SpawnTab(const FTabId& TabId, const TSharedPtr
 
 			NewTabWidget->SetLayoutIdentifier(TabId);
 		}
+		// If we can return nullptr, report it in the log
+		else
+		{
+			UE_LOG(LogSlate, Warning,
+				TEXT("The tab \"%s\" attempted to spawn but failed for some reason. It will not be displayed but it will still be saved in the layout settings file."), *StringToDisplay
+			);
+		}
 	}
 
 	if (NewTabWidget.IsValid())
diff --git a/Engine/Source/Runtime/Slate/Public/Framework/Docking/TabManager.h b/Engine/Source/Runtime/Slate/Public/Framework/Docking/TabManager.h
--- a/Engine/Source/Runtime/Slate/Public/Framework/Docking/TabManager.h
+++ b/Engine/Source/Runtime/Slate/Public/Framework/Docking/TabManager.h
@@ -232,6 +232,29 @@ namespace ETabState
 	};
 }
 
+enum class EOutputCanBeNullptr
+{
+	/**
+	 * RestoreArea_Helper() will always return a SWidget. It will return an "Unrecognized Tab" dummy SWidget if it cannot find the way to create the desired SWidget.
+	 * Default behavior and the only one used before UE 4.24.
+	 * This is the most strict condition and will never return nullptr.
+	 */
+	Never,
+	/**
+	 * RestoreArea_Helper() will return nullptr if its parent FTabManager::FStack does not contain at least a valid tab.
+	 * Useful for docked SWidgets that contain closed Tabs.
+	 * With IfNoTabValid, if a previously docked Tab is re-opened (i.e., its ETabState value changes from ClosedTab to OpenedTab), it will be displayed in the same place
+	 * it was placed before being initially closed. However, with IfNoOpenTabValid, if that Tab were re-opened, it might be instead displayed in a new standalone SWidget.
+	 */
+	IfNoTabValid,
+	/**
+	 * RestoreArea_Helper() will return nullptr if its parent FTabManager::FStack does not contain at least a valid tab whose ETabState is set to OpenedTab.
+	 * Useful for standalone SWidgets that otherwise will display a blank UI with no Tabs on it.
+	 * This is the most relaxed condition and the one that will return nullptr the most.
+	 */
+	IfNoOpenTabValid
+};
+
 
 class SLATE_API FTabManager : public TSharedFromThis<FTabManager>
 {
@@ -679,7 +702,7 @@ class SLATE_API FTabManager : public TSharedFromThis<FTabManager>
 		 */
 		void UnregisterAllTabSpawners();
 
-		TSharedPtr<SWidget> RestoreFrom( const TSharedRef<FLayout>& Layout, const TSharedPtr<SWindow>& ParentWindow, const bool bEmbedTitleAreaContent = false );
+		TSharedPtr<SWidget> RestoreFrom(const TSharedRef<FLayout>& Layout, const TSharedPtr<SWindow>& ParentWindow, const bool bEmbedTitleAreaContent = false, const EOutputCanBeNullptr RestoreAreaOutputCanBeNullptr = EOutputCanBeNullptr::Never);
 
 		void PopulateLocalTabSpawnerMenu( FMenuBuilder& PopulateMe );
 
@@ -808,16 +831,16 @@ class SLATE_API FTabManager : public TSharedFromThis<FTabManager>
 		FTabManager( const TSharedPtr<SDockTab>& InOwnerTab, const TSharedRef<FTabManager::FTabSpawner> & InNomadTabSpawner );
 
 		TSharedPtr<SDockingArea> RestoreArea(
-			const TSharedRef<FArea>& AreaToRestore, const TSharedPtr<SWindow>& InParentWindow, const bool bEmbedTitleAreaContent = false, const bool bCanOutputBeNullptr = false);
+			const TSharedRef<FArea>& AreaToRestore, const TSharedPtr<SWindow>& InParentWindow, const bool bEmbedTitleAreaContent = false, const EOutputCanBeNullptr OutputCanBeNullptr = EOutputCanBeNullptr::Never);
 
 		TSharedPtr<class SDockingNode> RestoreArea_Helper(
-			const TSharedRef<FLayoutNode>& LayoutNode, const TSharedPtr<SWindow>& ParentWindow, const bool bEmbedTitleAreaContent, const bool bCanOutputBeNullptr = false);
+			const TSharedRef<FLayoutNode>& LayoutNode, const TSharedPtr<SWindow>& ParentWindow, const bool bEmbedTitleAreaContent, const EOutputCanBeNullptr OutputCanBeNullptr = EOutputCanBeNullptr::Never);
 
 		/**
 		 * Use CanRestoreSplitterContent + RestoreSplitterContent when the output of its internal RestoreArea_Helper can be a nullptr.
 		 * Usage example:
 		 *		TArray<TSharedRef<SDockingNode>> DockingNodes;
-		 *		if (CanRestoreSplitterContent(DockingNodes, SplitterNode, ParentWindow))
+		 *		if (CanRestoreSplitterContent(DockingNodes, SplitterNode, ParentWindow, OutputCanBeNullptr))
 		 *		{
 		 *			// Create SplitterWidget only if it will be filled with at least 1 DockingNodes
 		 *			TSharedRef<SDockingSplitter> SplitterWidget = SNew(SDockingSplitter, SplitterNode);
@@ -825,7 +848,7 @@ class SLATE_API FTabManager : public TSharedFromThis<FTabManager>
 		 *			RestoreSplitterContent(DockingNodes, SplitterWidget);
 		 *		}
 		 */
-		bool CanRestoreSplitterContent(TArray<TSharedRef<class SDockingNode>>& DockingNodes, const TSharedRef<FSplitter>& SplitterNode, const TSharedPtr<SWindow>& ParentWindow);
+		bool CanRestoreSplitterContent(TArray<TSharedRef<class SDockingNode>>& DockingNodes, const TSharedRef<FSplitter>& SplitterNode, const TSharedPtr<SWindow>& ParentWindow, const EOutputCanBeNullptr OutputCanBeNullptr);
 		void RestoreSplitterContent(const TArray<TSharedRef<class SDockingNode>>& DockingNodes, const TSharedRef<class SDockingSplitter>& SplitterWidget);
 
 		/**
