Fix for stack-use-after-scope error caught by [REDACTED_PERSON].

In FSlateUser::SynthesizeCursorMoveIfNeeded, the SyntheticCursorMoveEvent can be constructed with an empty set of pressed mouse buttons.  The empty set was constructed as a parameter, and falling out of scope.  The problem is that FPointerEvent keeps a reference to that parameter, and accesses it later.  The fix is to construct an empty set prior to creating the SyntheticCursorMoveEvent, and just reference it in the constructor parameters, if necessary.

#rnx

diff --git a/Engine/Source/Runtime/Slate/Private/Framework/Application/SlateUser.cpp b/Engine/Source/Runtime/Slate/Private/Framework/Application/SlateUser.cpp
--- a/Engine/Source/Runtime/Slate/Private/Framework/Application/SlateUser.cpp
+++ b/Engine/Source/Runtime/Slate/Private/Framework/Application/SlateUser.cpp
@@ -606,12 +606,13 @@ bool FSlateUser::SynthesizeCursorMoveIfNeeded()
 		FSlateApplication& SlateApp = FSlateApplication::Get();
 		
 		const bool bHasHardwareCursor = SlateApp.GetPlatformCursor() == Cursor;
+		const TSet<FKey> EmptySet;
 		FPointerEvent SyntheticCursorMoveEvent(
 			GetUserIndex(),
 			FSlateApplication::CursorPointerIndex,
 			GetCursorPosition(),
 			GetPreviousCursorPosition(),
-			bHasHardwareCursor ? SlateApp.GetPressedMouseButtons() : TSet<FKey>(),
+			bHasHardwareCursor ? SlateApp.GetPressedMouseButtons() : EmptySet,
 			EKeys::Invalid,
 			0,
 			bHasHardwareCursor ? SlateApp.GetPlatformApplication()->GetModifierKeys() : FModifierKeysState());
