Slate - Listviews now won't crash if the items stop corresponding to the currently available widgets.  In the event they do, we fail safely now instead of crashing, but during development it will still check(false) if a listview item can't find the corresponding item for a widget, if the table is not pending a refresh.

#rnx

diff --git a/Engine/Source/Runtime/Slate/Public/Framework/Views/ITypedTableView.h b/Engine/Source/Runtime/Slate/Public/Framework/Views/ITypedTableView.h
--- a/Engine/Source/Runtime/Slate/Public/Framework/Views/ITypedTableView.h
+++ b/Engine/Source/Runtime/Slate/Public/Framework/Views/ITypedTableView.h
@@ -194,6 +194,9 @@ public:
 	/** @return the selection mode of this TableView */
 	virtual ESelectionMode::Type Private_GetSelectionMode() const = 0;
 
+	/** Is the list pending a refresh? */
+	virtual bool Private_IsPendingRefresh() const = 0;
+
 	/** @return Is the tableview a tree or a list? */
 	virtual ETableViewMode::Type GetTableViewMode() const = 0;
 	
diff --git a/Engine/Source/Runtime/Slate/Public/Widgets/Views/SListView.h b/Engine/Source/Runtime/Slate/Public/Widgets/Views/SListView.h
--- a/Engine/Source/Runtime/Slate/Public/Widgets/Views/SListView.h
+++ b/Engine/Source/Runtime/Slate/Public/Widgets/Views/SListView.h
@@ -898,6 +898,11 @@ public:
 		return SelectionMode.Get();
 	}
 
+	virtual bool Private_IsPendingRefresh() const override
+	{
+		return IsPendingRefresh();
+	}
+
 	virtual void Private_OnItemRightClicked( ItemType TheItem, const FPointerEvent& MouseEvent ) override
 	{
 		this->OnRightMouseButtonUp( MouseEvent );
diff --git a/Engine/Source/Runtime/Slate/Public/Widgets/Views/STableRow.h b/Engine/Source/Runtime/Slate/Public/Widgets/Views/STableRow.h
--- a/Engine/Source/Runtime/Slate/Public/Widgets/Views/STableRow.h
+++ b/Engine/Source/Runtime/Slate/Public/Widgets/Views/STableRow.h
@@ -126,7 +126,7 @@ DECLARE_DELEGATE_RetVal_OneParam(FReply, FOnTableRowDrop, FDragDropEvent const&)
 
 /**
  * The ListView is populated by Selectable widgets.
- * A Selectable widget is a way of the ListView containing it (OwnerWidget) and holds arbitrary Content (Content).
+ * A Selectable widget is a way of the ListView containing it (OwnerTable) and holds arbitrary Content (Content).
  * A Selectable works with its corresponding ListView to provide selection functionality.
  */
 template<typename ItemType>
@@ -267,19 +267,22 @@ public:
 
 	virtual int32 OnPaint( const FPaintArgs& Args, const FGeometry& AllottedGeometry, const FSlateRect& MyCullingRect, FSlateWindowElementList& OutDrawElements, int32 LayerId, const FWidgetStyle& InWidgetStyle, bool bParentEnabled ) const override
 	{
-		TSharedRef< ITypedTableView<ItemType> > OwnerWidget = OwnerTablePtr.Pin().ToSharedRef();
-		const bool bIsActive = OwnerWidget->AsWidget()->HasKeyboardFocus();
-		const ItemType* MyItem = OwnerWidget->Private_ItemFromWidget( this );
-		if ( bIsActive && OwnerWidget->Private_UsesSelectorFocus() && OwnerWidget->Private_HasSelectorFocus( *MyItem ) )
+		TSharedRef< ITypedTableView<ItemType> > OwnerTable = OwnerTablePtr.Pin().ToSharedRef();
+		const bool bIsActive = OwnerTable->AsWidget()->HasKeyboardFocus();
+
+		if (const ItemType* MyItemPtr = GetItemForThis(OwnerTable))
 		{
-			FSlateDrawElement::MakeBox(
-				OutDrawElements,
-				LayerId,
-				AllottedGeometry.ToPaintGeometry(),
-				&Style->SelectorFocusedBrush,
-				ESlateDrawEffect::None,
-				Style->SelectorFocusedBrush.GetTint( InWidgetStyle ) * InWidgetStyle.GetColorAndOpacityTint()
+			if (bIsActive && OwnerTable->Private_UsesSelectorFocus() && OwnerTable->Private_HasSelectorFocus(*MyItemPtr))
+			{
+				FSlateDrawElement::MakeBox(
+					OutDrawElements,
+					LayerId,
+					AllottedGeometry.ToPaintGeometry(),
+					&Style->SelectorFocusedBrush,
+					ESlateDrawEffect::None,
+					Style->SelectorFocusedBrush.GetTint(InWidgetStyle) * InWidgetStyle.GetColorAndOpacityTint()
 				);
+			}
 		}
 
 		LayerId = SBorder::OnPaint(Args, AllottedGeometry, MyCullingRect, OutDrawElements, LayerId, InWidgetStyle, bParentEnabled);
@@ -319,24 +322,25 @@ public:
 	 * @param  InMouseEvent  Mouse button event.
 	 * @return  Returns whether the event was handled, along with other possible actions.
 	 */
-	virtual FReply OnMouseButtonDoubleClick( const FGeometry& InMyGeometry, const FPointerEvent& InMouseEvent ) override
+	virtual FReply OnMouseButtonDoubleClick(const FGeometry& InMyGeometry, const FPointerEvent& InMouseEvent) override
 	{
-		if( InMouseEvent.GetEffectingButton() == EKeys::LeftMouseButton )
+		if (InMouseEvent.GetEffectingButton() == EKeys::LeftMouseButton)
 		{
-			TSharedRef< ITypedTableView<ItemType> > OwnerWidget = OwnerTablePtr.Pin().ToSharedRef();
+			TSharedRef< ITypedTableView<ItemType> > OwnerTable = OwnerTablePtr.Pin().ToSharedRef();
 
 			// Only one item can be double-clicked
-			const ItemType* MyItem = OwnerWidget->Private_ItemFromWidget( this );
-
-			// If we're configured to route double-click messages to the owner of the table, then
-			// do that here.  Otherwise, we'll toggle expansion.
-			const bool bWasHandled = OwnerWidget->Private_OnItemDoubleClicked( *MyItem );
-			if( !bWasHandled )
+			if (const ItemType* MyItemPtr = GetItemForThis(OwnerTable))
 			{
-				ToggleExpansion();
-			}
+				// If we're configured to route double-click messages to the owner of the table, then
+				// do that here.  Otherwise, we'll toggle expansion.
+				const bool bWasHandled = OwnerTable->Private_OnItemDoubleClicked(*MyItemPtr);
+				if (!bWasHandled)
+				{
+					ToggleExpansion();
+				}
 
-			return FReply::Handled();
+				return FReply::Handled();
+			}
 		}
 
 		return FReply::Unhandled();
@@ -351,57 +355,58 @@ public:
 	 */	
 	virtual FReply OnMouseButtonDown( const FGeometry& MyGeometry, const FPointerEvent& MouseEvent ) override
 	{
-		TSharedPtr< ITypedTableView<ItemType> > OwnerWidget = OwnerTablePtr.Pin();
+		TSharedRef< ITypedTableView<ItemType> > OwnerTable = OwnerTablePtr.Pin().ToSharedRef();
 		bChangedSelectionOnMouseDown = false;
 		bDragWasDetected = false;
 
-		check(OwnerWidget.IsValid());
-
 		if ( MouseEvent.GetEffectingButton() == EKeys::LeftMouseButton )
 		{
 			const ESelectionMode::Type SelectionMode = GetSelectionMode();
 			if (SelectionMode != ESelectionMode::None)
 			{
-				const ItemType& MyItem = *OwnerWidget->Private_ItemFromWidget(this);
-				const bool bIsSelected = OwnerWidget->Private_IsItemSelected(MyItem);
-
-				if (SelectionMode == ESelectionMode::Multi)
+				if (const ItemType* MyItemPtr = GetItemForThis(OwnerTable))
 				{
-					if (MouseEvent.IsControlDown())
+					const ItemType& MyItem = *MyItemPtr;
+					const bool bIsSelected = OwnerTable->Private_IsItemSelected(MyItem);
+
+					if (SelectionMode == ESelectionMode::Multi)
 					{
-						OwnerWidget->Private_SetItemSelection(MyItem, !bIsSelected, true);
-						bChangedSelectionOnMouseDown = true;
-						if (SignalSelectionMode == ETableRowSignalSelectionMode::Instantaneous)
+						if (MouseEvent.IsControlDown())
+						{
+							OwnerTable->Private_SetItemSelection(MyItem, !bIsSelected, true);
+							bChangedSelectionOnMouseDown = true;
+							if (SignalSelectionMode == ETableRowSignalSelectionMode::Instantaneous)
+							{
+								OwnerTable->Private_SignalSelectionChanged(ESelectInfo::OnMouseClick);
+							}
+						}
+						else if (MouseEvent.IsShiftDown())
 						{
-							OwnerWidget->Private_SignalSelectionChanged(ESelectInfo::OnMouseClick);
+							OwnerTable->Private_SelectRangeFromCurrentTo(MyItem);
+							bChangedSelectionOnMouseDown = true;
+							if (SignalSelectionMode == ETableRowSignalSelectionMode::Instantaneous)
+							{
+								OwnerTable->Private_SignalSelectionChanged(ESelectInfo::OnMouseClick);
+							}
 						}
 					}
-					else if (MouseEvent.IsShiftDown())
+
+					if (!bIsSelected && !bChangedSelectionOnMouseDown)
 					{
-						OwnerWidget->Private_SelectRangeFromCurrentTo(MyItem);
+						OwnerTable->Private_ClearSelection();
+						OwnerTable->Private_SetItemSelection(MyItem, true, true);
 						bChangedSelectionOnMouseDown = true;
 						if (SignalSelectionMode == ETableRowSignalSelectionMode::Instantaneous)
 						{
-							OwnerWidget->Private_SignalSelectionChanged(ESelectInfo::OnMouseClick);
+							OwnerTable->Private_SignalSelectionChanged(ESelectInfo::OnMouseClick);
 						}
 					}
-				}
-				
-				if (!bIsSelected && !bChangedSelectionOnMouseDown)
-				{
-					OwnerWidget->Private_ClearSelection();
-					OwnerWidget->Private_SetItemSelection(MyItem, true, true);
-					bChangedSelectionOnMouseDown = true;
-					if (SignalSelectionMode == ETableRowSignalSelectionMode::Instantaneous)
-					{
-						OwnerWidget->Private_SignalSelectionChanged(ESelectInfo::OnMouseClick);
-					}
-				}
 
-				return FReply::Handled()
-					.DetectDrag(SharedThis(this), EKeys::LeftMouseButton)
-					.SetUserFocus(OwnerWidget->AsWidget(), EFocusCause::Mouse)
-					.CaptureMouse(SharedThis(this));
+					return FReply::Handled()
+						.DetectDrag(SharedThis(this), EKeys::LeftMouseButton)
+						.SetUserFocus(OwnerTable->AsWidget(), EFocusCause::Mouse)
+						.CaptureMouse(SharedThis(this));
+				}
 			}
 		}
 
@@ -417,11 +422,10 @@ public:
 	 */
 	virtual FReply OnMouseButtonUp( const FGeometry& MyGeometry, const FPointerEvent& MouseEvent ) override
 	{
-		TSharedPtr< ITypedTableView<ItemType> > OwnerWidget = OwnerTablePtr.Pin();
-		check(OwnerWidget.IsValid());
+		TSharedRef< ITypedTableView<ItemType> > OwnerTable = OwnerTablePtr.Pin().ToSharedRef();
 
 		// Requires #include "Widgets/Views/STableViewBase.h"
-		TSharedRef< STableViewBase > OwnerTableViewBase = StaticCastSharedPtr< SListView<ItemType> >( OwnerWidget ).ToSharedRef();
+		TSharedRef< STableViewBase > OwnerTableViewBase = StaticCastSharedRef< SListView<ItemType> >(OwnerTable);
 
 		if ( MouseEvent.GetEffectingButton() == EKeys::LeftMouseButton )
 		{
@@ -443,8 +447,8 @@ public:
 						{
 							if ( !bChangedSelectionOnMouseDown )
 							{
-								OwnerWidget->Private_ClearSelection();
-								OwnerWidget->Private_SignalSelectionChanged(ESelectInfo::OnMouseClick);
+								OwnerTable->Private_ClearSelection();
+								OwnerTable->Private_SignalSelectionChanged(ESelectInfo::OnMouseClick);
 							}
 
 							Reply = FReply::Handled().ReleaseMouseCapture();
@@ -455,20 +459,20 @@ public:
 						{
 							if ( !bChangedSelectionOnMouseDown && !MouseEvent.IsControlDown() && !MouseEvent.IsShiftDown() )
 							{
-								const ItemType* MyItem = OwnerWidget->Private_ItemFromWidget( this );
-								check(MyItem != nullptr);
-
-								const bool bIsSelected = OwnerWidget->Private_IsItemSelected( *MyItem );
-								if ( bIsSelected && OwnerWidget->Private_GetNumSelectedItems() > 1 )
+								if (const ItemType* MyItemPtr = GetItemForThis(OwnerTable))
 								{
-									// We are mousing up on a previous selected item;
-									// deselect everything but this item.
-
-									OwnerWidget->Private_ClearSelection();
-									OwnerWidget->Private_SetItemSelection( *MyItem, true, true );
-									OwnerWidget->Private_SignalSelectionChanged(ESelectInfo::OnMouseClick);
-
-									Reply = FReply::Handled().ReleaseMouseCapture();
+									const bool bIsSelected = OwnerTable->Private_IsItemSelected(*MyItemPtr);
+									if (bIsSelected && OwnerTable->Private_GetNumSelectedItems() > 1)
+									{
+										// We are mousing up on a previous selected item;
+										// deselect everything but this item.
+
+										OwnerTable->Private_ClearSelection();
+										OwnerTable->Private_SetItemSelection(*MyItemPtr, true, true);
+										OwnerTable->Private_SignalSelectionChanged(ESelectInfo::OnMouseClick);
+
+										Reply = FReply::Handled().ReleaseMouseCapture();
+									}
 								}
 							}
 						}
@@ -476,15 +480,17 @@ public:
 					}
 				}
 
-				const ItemType* MyItem = OwnerWidget->Private_ItemFromWidget(this);
-				if (MyItem && OwnerWidget->Private_OnItemClicked(*MyItem))
+				if (const ItemType* MyItemPtr = GetItemForThis(OwnerTable))
 				{
-					Reply = FReply::Handled().ReleaseMouseCapture();
+					if (OwnerTable->Private_OnItemClicked(*MyItemPtr))
+					{
+						Reply = FReply::Handled().ReleaseMouseCapture();
+					}
 				}
 
 				if (bChangedSelectionOnMouseDown && !bDragWasDetected && (SignalSelectionMode == ETableRowSignalSelectionMode::Deferred))
 				{
-					OwnerWidget->Private_SignalSelectionChanged(ESelectInfo::OnMouseClick);
+					OwnerTable->Private_SignalSelectionChanged(ESelectInfo::OnMouseClick);
 				}
 
 				return Reply;
@@ -502,20 +508,22 @@ public:
 			case ESelectionMode::Multi:
 				{
 					// Only one item can be selected at a time
-					const ItemType* MyItem = OwnerWidget->Private_ItemFromWidget( this );
-					const bool bIsSelected = OwnerWidget->Private_IsItemSelected( *MyItem );
-
-					// Select the item under the cursor
-					if( !bIsSelected )
+					if (const ItemType* MyItemPtr = GetItemForThis(OwnerTable))
 					{
-						OwnerWidget->Private_ClearSelection();
-						OwnerWidget->Private_SetItemSelection( *MyItem, true, true );
-						OwnerWidget->Private_SignalSelectionChanged(ESelectInfo::OnMouseClick);
-					}
+						const bool bIsSelected = OwnerTable->Private_IsItemSelected(*MyItemPtr);
 
-					OwnerWidget->Private_OnItemRightClicked( *MyItem, MouseEvent );
+						// Select the item under the cursor
+						if (!bIsSelected)
+						{
+							OwnerTable->Private_ClearSelection();
+							OwnerTable->Private_SetItemSelection(*MyItemPtr, true, true);
+							OwnerTable->Private_SignalSelectionChanged(ESelectInfo::OnMouseClick);
+						}
 
-					return FReply::Handled();
+						OwnerTable->Private_OnItemRightClicked(*MyItemPtr, MouseEvent);
+
+						return FReply::Handled();
+					}
 				}
 				break;
 			}
@@ -542,34 +550,34 @@ public:
 		if (bProcessingSelectionTouch)
 		{
 			bProcessingSelectionTouch = false;
-			const TSharedPtr<ITypedTableView<ItemType>> OwnerWidget = OwnerTablePtr.Pin();
-			if (const ItemType* MyItem = OwnerWidget->Private_ItemFromWidget(this))
+			const TSharedRef<ITypedTableView<ItemType>> OwnerTable = OwnerTablePtr.Pin().ToSharedRef();
+			if (const ItemType* MyItemPtr = GetItemForThis(OwnerTable))
 			{
 				ESelectionMode::Type SelectionMode = GetSelectionMode();
 				if (SelectionMode != ESelectionMode::None)
 				{
-					const bool bIsSelected = OwnerWidget->Private_IsItemSelected(*MyItem);
+					const bool bIsSelected = OwnerTable->Private_IsItemSelected(*MyItemPtr);
 					if (!bIsSelected)
 					{
 						if (SelectionMode != ESelectionMode::Multi)
 						{
-							OwnerWidget->Private_ClearSelection();
+							OwnerTable->Private_ClearSelection();
 						}
-						OwnerWidget->Private_SetItemSelection(*MyItem, true, true);
-						OwnerWidget->Private_SignalSelectionChanged(ESelectInfo::OnMouseClick);
+						OwnerTable->Private_SetItemSelection(*MyItemPtr, true, true);
+						OwnerTable->Private_SignalSelectionChanged(ESelectInfo::OnMouseClick);
 
 						Reply = FReply::Handled();
 					}
 					else if (SelectionMode == ESelectionMode::SingleToggle || SelectionMode == ESelectionMode::Multi)
 					{
-						OwnerWidget->Private_SetItemSelection(*MyItem, true, true);
-						OwnerWidget->Private_SignalSelectionChanged(ESelectInfo::OnMouseClick);
+						OwnerTable->Private_SetItemSelection(*MyItemPtr, true, true);
+						OwnerTable->Private_SignalSelectionChanged(ESelectInfo::OnMouseClick);
 
 						Reply = FReply::Handled();
 					}
 				}
 
-				if (OwnerWidget->Private_OnItemClicked(*MyItem))
+				if (OwnerTable->Private_OnItemClicked(*MyItemPtr))
 				{
 					Reply = FReply::Handled();
 				}
@@ -595,8 +603,8 @@ public:
 
 			if ( bChangedSelectionOnMouseDown && SignalSelectionMode == ETableRowSignalSelectionMode::Deferred )
 			{
-				TSharedPtr< ITypedTableView<ItemType> > OwnerWidget = OwnerTablePtr.Pin();
-				OwnerWidget->Private_SignalSelectionChanged(ESelectInfo::OnMouseClick);
+				TSharedPtr< ITypedTableView<ItemType> > OwnerTable = OwnerTablePtr.Pin();
+				OwnerTable->Private_SignalSelectionChanged(ESelectInfo::OnMouseClick);
 			}
 		}
 
@@ -653,11 +661,15 @@ public:
 			const FVector2D LocalPointerPos = MyGeometry.AbsoluteToLocal(DragDropEvent.GetScreenSpacePosition());
 			const EItemDropZone ItemHoverZone = ZoneFromPointerPosition(LocalPointerPos, MyGeometry.GetLocalSize().Y);
 
-			ItemDropZone = [&]()
+			ItemDropZone = [&, this]()
 			{
-				TSharedPtr< ITypedTableView<ItemType> > OwnerWidget = OwnerTablePtr.Pin();
-				const ItemType* MyItem = OwnerWidget->Private_ItemFromWidget(this);
-				return OnCanAcceptDrop.Execute(DragDropEvent, ItemHoverZone, *MyItem);
+				TSharedRef< ITypedTableView<ItemType> > OwnerTable = OwnerTablePtr.Pin().ToSharedRef();
+				if (const ItemType* MyItemPtr = GetItemForThis(OwnerTable))
+				{
+					return OnCanAcceptDrop.Execute(DragDropEvent, ItemHoverZone, *MyItemPtr);
+				}
+
+				return TOptional<EItemDropZone>();
 			}();
 
 			return FReply::Handled();
@@ -675,33 +687,34 @@ public:
 		{
 			if (OnAcceptDrop.IsBound())
 			{
-				const TSharedRef< ITypedTableView<ItemType> > OwnerWidget = OwnerTablePtr.Pin().ToSharedRef();
+				const TSharedRef< ITypedTableView<ItemType> > OwnerTable = OwnerTablePtr.Pin().ToSharedRef();
 
 				// A drop finishes the drag/drop operation, so we are no longer providing any feedback.
 				ItemDropZone = TOptional<EItemDropZone>();
 
 				// Find item associated with this widget.
-				const ItemType* MyItem = OwnerWidget->Private_ItemFromWidget(this);
-				
-				// Which physical drop zone is the drop about to be performed onto?
-				const FVector2D LocalPointerPos = MyGeometry.AbsoluteToLocal(DragDropEvent.GetScreenSpacePosition());
-				const EItemDropZone HoveredZone = ZoneFromPointerPosition(LocalPointerPos, MyGeometry.GetLocalSize().Y);
-
-				// The row gets final say over which zone to drop onto regardless of physical location.
-				const TOptional<EItemDropZone> ReportedZone = OnCanAcceptDrop.IsBound()
-					? OnCanAcceptDrop.Execute(DragDropEvent, HoveredZone, *MyItem)
-					: HoveredZone;
-
-				if (ReportedZone.IsSet())
+				if (const ItemType* MyItemPtr = GetItemForThis(OwnerTable))
 				{
-					FReply DropReply = OnAcceptDrop.Execute(DragDropEvent, ReportedZone.GetValue(), *MyItem);
-					if (DropReply.IsEventHandled() && ReportedZone.GetValue() == EItemDropZone::OntoItem)
+					// Which physical drop zone is the drop about to be performed onto?
+					const FVector2D LocalPointerPos = MyGeometry.AbsoluteToLocal(DragDropEvent.GetScreenSpacePosition());
+					const EItemDropZone HoveredZone = ZoneFromPointerPosition(LocalPointerPos, MyGeometry.GetLocalSize().Y);
+
+					// The row gets final say over which zone to drop onto regardless of physical location.
+					const TOptional<EItemDropZone> ReportedZone = OnCanAcceptDrop.IsBound()
+						? OnCanAcceptDrop.Execute(DragDropEvent, HoveredZone, *MyItemPtr)
+						: HoveredZone;
+
+					if (ReportedZone.IsSet())
 					{
-						// Expand the drop target just in case, so that what we dropped is visible.
-						OwnerWidget->Private_SetItemExpansion(*MyItem, true);
-					}
+						FReply DropReply = OnAcceptDrop.Execute(DragDropEvent, ReportedZone.GetValue(), *MyItemPtr);
+						if (DropReply.IsEventHandled() && ReportedZone.GetValue() == EItemDropZone::OntoItem)
+						{
+							// Expand the drop target just in case, so that what we dropped is visible.
+							OwnerTable->Private_SetItemExpansion(*MyItemPtr, true);
+						}
 
-					return DropReply;
+						return DropReply;
+					}
 				}
 			}
 
@@ -727,31 +740,40 @@ public:
 
 	virtual bool IsItemExpanded() const override
 	{
-		TSharedPtr< ITypedTableView<ItemType> > OwnerWidget = OwnerTablePtr.Pin();
-		const ItemType* MyItem = OwnerWidget->Private_ItemFromWidget( this );
-		const bool bIsItemExpanded = OwnerWidget->Private_IsItemExpanded( *MyItem );
-		return bIsItemExpanded;
+		TSharedRef< ITypedTableView<ItemType> > OwnerTable = OwnerTablePtr.Pin().ToSharedRef();
+		if (const ItemType* MyItemPtr = GetItemForThis(OwnerTable))
+		{
+			return OwnerTable->Private_IsItemExpanded(*MyItemPtr);
+		}
+		
+		return false;
 	}
 
 	virtual void ToggleExpansion() override
 	{
-		TSharedPtr< ITypedTableView<ItemType> > OwnerWidget = OwnerTablePtr.Pin();
+		TSharedRef< ITypedTableView<ItemType> > OwnerTable = OwnerTablePtr.Pin().ToSharedRef();
 
-		const bool bItemHasChildren = OwnerWidget->Private_DoesItemHaveChildren( IndexInList );
+		const bool bItemHasChildren = OwnerTable->Private_DoesItemHaveChildren( IndexInList );
 		// Nothing to expand if row being clicked on doesn't have children
 		if( bItemHasChildren )
 		{
-			ItemType MyItem = *(OwnerWidget->Private_ItemFromWidget( this ));
-			const bool bIsItemExpanded = bItemHasChildren && OwnerWidget->Private_IsItemExpanded( MyItem );
-			OwnerWidget->Private_SetItemExpansion( MyItem, !bIsItemExpanded);
+			if (const ItemType* MyItemPtr = GetItemForThis(OwnerTable))
+			{
+				const bool bIsItemExpanded = bItemHasChildren && OwnerTable->Private_IsItemExpanded(*MyItemPtr);
+				OwnerTable->Private_SetItemExpansion(*MyItemPtr, !bIsItemExpanded);
+			}
 		}
 	}
 
 	virtual bool IsItemSelected() const override
 	{
-		TSharedPtr<ITypedTableView<ItemType>> OwnerTable = OwnerTablePtr.Pin();
-		const ItemType& MyItem = *OwnerTable->Private_ItemFromWidget(this);
-		return OwnerTable->Private_IsItemSelected(MyItem);
+		TSharedRef<ITypedTableView<ItemType>> OwnerTable = OwnerTablePtr.Pin().ToSharedRef();
+		if (const ItemType* MyItemPtr = GetItemForThis(OwnerTable))
+		{
+			return OwnerTable->Private_IsItemSelected(*MyItemPtr);
+		}
+
+		return false;
 	}
 
 	virtual int32 GetIndentLevel() const override
@@ -817,88 +839,91 @@ public:
 
 	virtual void Private_OnExpanderArrowShiftClicked() override
 	{
-		TSharedPtr< ITypedTableView<ItemType> > OwnerWidget = OwnerTablePtr.Pin();
+		TSharedRef< ITypedTableView<ItemType> > OwnerTable = OwnerTablePtr.Pin().ToSharedRef();
 
-		const bool bItemHasChildren = OwnerWidget->Private_DoesItemHaveChildren( IndexInList );
+		const bool bItemHasChildren = OwnerTable->Private_DoesItemHaveChildren( IndexInList );
 		// Nothing to expand if row being clicked on doesn't have children
 		if( bItemHasChildren )
 		{
-			ItemType MyItem = *(OwnerWidget->Private_ItemFromWidget( this ));
-			const bool IsItemExpanded = bItemHasChildren && OwnerWidget->Private_IsItemExpanded( MyItem );
-			OwnerWidget->Private_OnExpanderArrowShiftClicked( MyItem, !IsItemExpanded );
+			if (const ItemType* MyItemPtr = GetItemForThis(OwnerTable))
+			{
+				const bool IsItemExpanded = bItemHasChildren && OwnerTable->Private_IsItemExpanded(*MyItemPtr);
+				OwnerTable->Private_OnExpanderArrowShiftClicked(*MyItemPtr, !IsItemExpanded);
+			}
 		}
 	}
 
 	/** @return The border to be drawn around this list item */
 	virtual const FSlateBrush* GetBorder() const 
 	{
-		TSharedPtr< ITypedTableView<ItemType> > OwnerWidget = OwnerTablePtr.Pin();
+		TSharedRef< ITypedTableView<ItemType> > OwnerTable = OwnerTablePtr.Pin().ToSharedRef();
 
-		const bool bIsActive = OwnerWidget->AsWidget()->HasKeyboardFocus();
+		const bool bIsActive = OwnerTable->AsWidget()->HasKeyboardFocus();
 
 		static FName GenericWhiteBoxBrush("GenericWhiteBox");
 
 		// @todo: Slate Style - make this part of the widget style
 		const FSlateBrush* WhiteBox = FCoreStyle::Get().GetBrush(GenericWhiteBoxBrush);
 
-		const ItemType* MyItem = OwnerWidget->Private_ItemFromWidget( this );
-
-		check(MyItem);
-
-		const bool bIsSelected = OwnerWidget->Private_IsItemSelected( *MyItem );
-		const bool bIsHighlighted = OwnerWidget->Private_IsItemHighlighted( *MyItem );
+		if (const ItemType* MyItemPtr = GetItemForThis(OwnerTable))
+		{
+			const bool bIsSelected = OwnerTable->Private_IsItemSelected(*MyItemPtr);
+			const bool bIsHighlighted = OwnerTable->Private_IsItemHighlighted(*MyItemPtr);
 
-		const bool bAllowSelection = GetSelectionMode() != ESelectionMode::None;
-		const bool bEvenEntryIndex = (IndexInList % 2 == 0);
+			const bool bAllowSelection = GetSelectionMode() != ESelectionMode::None;
+			const bool bEvenEntryIndex = (IndexInList % 2 == 0);
 
-		if (bIsSelected && bShowSelection)
-		{
-			if( bIsActive )
-			{
-				return IsHovered()
-					? &Style->ActiveHoveredBrush
-					: &Style->ActiveBrush;
-			}
-			else
+			if (bIsSelected && bShowSelection)
 			{
-				return IsHovered()
-					? &Style->InactiveHoveredBrush
-					: &Style->InactiveBrush;
+				if (bIsActive)
+				{
+					return IsHovered()
+						? &Style->ActiveHoveredBrush
+						: &Style->ActiveBrush;
+				}
+				else
+				{
+					return IsHovered()
+						? &Style->InactiveHoveredBrush
+						: &Style->InactiveBrush;
+				}
 			}
-		}
-		else if (!bIsSelected && bIsHighlighted)
-		{
-			if (bIsActive)
+			else if (!bIsSelected && bIsHighlighted)
 			{
-				return IsHovered()
-					? (bEvenEntryIndex ? &Style->EvenRowBackgroundHoveredBrush : &Style->OddRowBackgroundHoveredBrush)
-					: &Style->ActiveHighlightedBrush;
+				if (bIsActive)
+				{
+					return IsHovered()
+						? (bEvenEntryIndex ? &Style->EvenRowBackgroundHoveredBrush : &Style->OddRowBackgroundHoveredBrush)
+						: &Style->ActiveHighlightedBrush;
+				}
+				else
+				{
+					return IsHovered()
+						? (bEvenEntryIndex ? &Style->EvenRowBackgroundHoveredBrush : &Style->OddRowBackgroundHoveredBrush)
+						: &Style->InactiveHighlightedBrush;
+				}
 			}
 			else
 			{
-				return IsHovered()
-					? (bEvenEntryIndex ? &Style->EvenRowBackgroundHoveredBrush : &Style->OddRowBackgroundHoveredBrush)
-					: &Style->InactiveHighlightedBrush;
-			}
-		}
-		else
-		{
-			// Add a slightly lighter background for even rows
-			if(bEvenEntryIndex)
-			{
-				return ( IsHovered() && bAllowSelection )
-					? &Style->EvenRowBackgroundHoveredBrush
-					: &Style->EvenRowBackgroundBrush;
+				// Add a slightly lighter background for even rows
+				if (bEvenEntryIndex)
+				{
+					return (IsHovered() && bAllowSelection)
+						? &Style->EvenRowBackgroundHoveredBrush
+						: &Style->EvenRowBackgroundBrush;
 
-			}
-			else
-			{
-				return ( IsHovered() && bAllowSelection )
-					? &Style->OddRowBackgroundHoveredBrush
-					: &Style->OddRowBackgroundBrush;
+				}
+				else
+				{
+					return (IsHovered() && bAllowSelection)
+						? &Style->OddRowBackgroundHoveredBrush
+						: &Style->OddRowBackgroundBrush;
 
+				}
 			}
 		}
+
+		return nullptr;
 	}
 
 	/** 
@@ -908,15 +933,19 @@ public:
 	 */
 	bool IsSelectedExclusively() const
 	{
-		TSharedPtr< ITypedTableView< ItemType > > OwnerWidget = OwnerTablePtr.Pin();
+		TSharedRef< ITypedTableView< ItemType > > OwnerTable = OwnerTablePtr.Pin().ToSharedRef();
 
-		if(!OwnerWidget->AsWidget()->HasKeyboardFocus() || OwnerWidget->Private_GetNumSelectedItems() > 1)
+		if (!OwnerTable->AsWidget()->HasKeyboardFocus() || OwnerTable->Private_GetNumSelectedItems() > 1)
 		{
 			return false;
 		}
 
-		const ItemType* MyItem = OwnerWidget->Private_ItemFromWidget( this );
-		return OwnerWidget->Private_IsItemSelected( *MyItem );
+		if (const ItemType* MyItemPtr = GetItemForThis(OwnerTable))
+		{
+			return OwnerTable->Private_IsItemSelected(*MyItemPtr);
+		}
+
+		return false;
 	}
 
 	/**
@@ -926,10 +955,14 @@ public:
 	 */
 	bool IsSelected() const
 	{
-		TSharedPtr< ITypedTableView< ItemType > > OwnerWidget = OwnerTablePtr.Pin();
+		TSharedRef< ITypedTableView< ItemType > > OwnerTable = OwnerTablePtr.Pin().ToSharedRef();
 
-		const ItemType* MyItem = OwnerWidget->Private_ItemFromWidget(this);
-		return OwnerWidget->Private_IsItemSelected(*MyItem);
+		if (const ItemType* MyItemPtr = GetItemForThis(OwnerTable))
+		{
+			return OwnerTable->Private_IsItemSelected(*MyItemPtr);
+		}
+
+		return false;
 	}
 
 	/** By default, this function does nothing, it should be implemented by derived class */
@@ -995,27 +1028,46 @@ protected:
 
 	FSlateColor GetForegroundBasedOnSelection() const
 	{
-		const TSharedPtr< ITypedTableView<ItemType> > OwnerWidget = OwnerTablePtr.Pin();
+		const TSharedPtr< ITypedTableView<ItemType> > OwnerTable = OwnerTablePtr.Pin();
 		const FSlateColor& NonSelectedForeground = Style->TextColor; 
 		const FSlateColor& SelectedForeground = Style->SelectedTextColor;
 
-		if ( !bShowSelection || !OwnerWidget.IsValid() )
+		if ( !bShowSelection || !OwnerTable.IsValid() )
 		{
 			return NonSelectedForeground;
 		}
 
-		const ItemType* MyItem = OwnerWidget->Private_ItemFromWidget( this );
-		const bool bIsSelected = OwnerWidget->Private_IsItemSelected( *MyItem );
+		if (const ItemType* MyItemPtr = GetItemForThis(OwnerTable.ToSharedRef()))
+		{
+			const bool bIsSelected = OwnerTable->Private_IsItemSelected(*MyItemPtr);
+
+			return bIsSelected
+				? SelectedForeground
+				: NonSelectedForeground;
+		}
 
-		return bIsSelected
-			? SelectedForeground
-			: NonSelectedForeground;
+		return NonSelectedForeground;
 	}
 
 	virtual ESelectionMode::Type GetSelectionMode() const override
 	{
-		const TSharedPtr< ITypedTableView<ItemType> > OwnerWidget = OwnerTablePtr.Pin();
-		return OwnerWidget->Private_GetSelectionMode();
+		const TSharedPtr< ITypedTableView<ItemType> > OwnerTable = OwnerTablePtr.Pin();
+		return OwnerTable->Private_GetSelectionMode();
+	}
+
+	const ItemType* GetItemForThis(const TSharedRef<ITypedTableView<ItemType>>& OwnerTable) const
+	{
+		const ItemType* MyItemPtr = OwnerTable->Private_ItemFromWidget(this);
+		if (MyItemPtr)
+		{
+			return MyItemPtr;
+		}
+		else
+		{
+			checkf(OwnerTable->Private_IsPendingRefresh(), TEXT("We were unable to find the item for this widget.  If it was removed from the source collection, the list should be pending a refresh."));
+		}
+
+		return nullptr;
 	}
 
 protected:
diff --git a/Engine/Source/Runtime/UMG/Public/Slate/SObjectTableRow.h b/Engine/Source/Runtime/UMG/Public/Slate/SObjectTableRow.h
--- a/Engine/Source/Runtime/UMG/Public/Slate/SObjectTableRow.h
+++ b/Engine/Source/Runtime/UMG/Public/Slate/SObjectTableRow.h
@@ -116,11 +116,13 @@ public:
 		TSharedPtr<ITypedTableView<ItemType>> OwnerTable = OwnerTablePtr.Pin();
 		if (OwnerTable.IsValid())
 		{
-			const ItemType& MyItem = *OwnerTable->Private_ItemFromWidget(this);
-			if (bIsAppearingSelected != OwnerTable->Private_IsItemSelected(MyItem))
+			if (const ItemType* MyItemPtr = GetItemForThis(OwnerTable.ToSharedRef()))
 			{
-				bIsAppearingSelected = !bIsAppearingSelected;
-				OnItemSelectionChanged(bIsAppearingSelected);
+				if (bIsAppearingSelected != OwnerTable->Private_IsItemSelected(*MyItemPtr))
+				{
+					bIsAppearingSelected = !bIsAppearingSelected;
+					OnItemSelectionChanged(bIsAppearingSelected);
+				}
 			}
 		}
 	}
@@ -162,26 +164,36 @@ public:
 
 	virtual bool IsItemExpanded() const override
 	{
-		TSharedPtr<ITypedTableView<ItemType>> OwnerTable = OwnerTablePtr.Pin();
-		const ItemType& MyItem = *OwnerTable->Private_ItemFromWidget(this);
-		return OwnerTable->Private_IsItemExpanded(MyItem);
+		TSharedRef<ITypedTableView<ItemType>> OwnerTable = OwnerTablePtr.Pin().ToSharedRef();
+		if (const ItemType* MyItemPtr = GetItemForThis(OwnerTable))
+		{
+			return OwnerTable->Private_IsItemExpanded(*MyItemPtr);
+		}
+
+		return false;
 	}
 
 	virtual void ToggleExpansion() override
 	{
-		TSharedPtr<ITypedTableView<ItemType>> OwnerTable = OwnerTablePtr.Pin();
+		TSharedRef<ITypedTableView<ItemType>> OwnerTable = OwnerTablePtr.Pin().ToSharedRef();
 		if (OwnerTable->Private_DoesItemHaveChildren(IndexInList))
 		{
-			const ItemType& MyItem = *OwnerTable->Private_ItemFromWidget(this);
-			OwnerTable->Private_SetItemExpansion(MyItem, !OwnerTable->Private_IsItemExpanded(MyItem));
+			if (const ItemType* MyItemPtr = GetItemForThis(OwnerTable))
+			{
+				OwnerTable->Private_SetItemExpansion(*MyItemPtr, !OwnerTable->Private_IsItemExpanded(*MyItemPtr));
+			}
 		}
 	}
 
 	virtual bool IsItemSelected() const override
 	{
-		TSharedPtr<ITypedTableView<ItemType>> OwnerTable = OwnerTablePtr.Pin();
-		const ItemType& MyItem = *OwnerTable->Private_ItemFromWidget(this);
-		return OwnerTable->Private_IsItemSelected(MyItem);
+		TSharedRef<ITypedTableView<ItemType>> OwnerTable = OwnerTablePtr.Pin().ToSharedRef();
+		if (const ItemType* MyItemPtr = GetItemForThis(OwnerTable))
+		{
+			return OwnerTable->Private_IsItemSelected(*MyItemPtr);
+		}
+
+		return false;
 	}
 
 	virtual TBitArray<> GetWiresNeededByDepth() const override
@@ -222,9 +234,12 @@ public:
 		if (InMouseEvent.GetEffectingButton() == EKeys::LeftMouseButton)
 		{
 			TSharedRef<ITypedTableView<ItemType>> OwnerTable = OwnerTablePtr.Pin().ToSharedRef();
-			OwnerTable->Private_OnItemDoubleClicked(*OwnerTable->Private_ItemFromWidget(this));
-			
-			return FReply::Handled();
+
+			if (const ItemType* MyItemPtr = GetItemForThis(OwnerTable))
+			{
+				OwnerTable->Private_OnItemDoubleClicked(*MyItemPtr);
+				return FReply::Handled();
+			}
 		}
 		return FReply::Unhandled();
 	}
@@ -246,33 +261,33 @@ public:
 		{
 			bProcessingSelectionTouch = false;
 			TSharedRef<ITypedTableView<ItemType>> OwnerTable = OwnerTablePtr.Pin().ToSharedRef();
-			if (const ItemType* MyItem = OwnerTable->Private_ItemFromWidget(this))
+			if (const ItemType* MyItemPtr = GetItemForThis(OwnerTable))
 			{
 				ESelectionMode::Type SelectionMode = GetSelectionMode();
 				if (SelectionMode != ESelectionMode::None)
 				{
-					const bool bIsSelected = OwnerTable->Private_IsItemSelected(*MyItem);
+					const bool bIsSelected = OwnerTable->Private_IsItemSelected(*MyItemPtr);
 					if (!bIsSelected)
 					{
 						if (SelectionMode != ESelectionMode::Multi)
 						{
 							OwnerTable->Private_ClearSelection();
 						}
-						OwnerTable->Private_SetItemSelection(*MyItem, true, true);
+						OwnerTable->Private_SetItemSelection(*MyItemPtr, true, true);
 						OwnerTable->Private_SignalSelectionChanged(ESelectInfo::OnMouseClick);
 
 						Reply = FReply::Handled();
 					}
 					else if (SelectionMode == ESelectionMode::SingleToggle || SelectionMode == ESelectionMode::Multi)
 					{
-						OwnerTable->Private_SetItemSelection(*MyItem, true, true);
+						OwnerTable->Private_SetItemSelection(*MyItemPtr, true, true);
 						OwnerTable->Private_SignalSelectionChanged(ESelectInfo::OnMouseClick);
 
 						Reply = FReply::Handled();
 					}
 				}
 
-				if (OwnerTable->Private_OnItemClicked(*MyItem))
+				if (OwnerTable->Private_OnItemClicked(*MyItemPtr))
 				{
 					Reply = FReply::Handled();
 				}
@@ -314,15 +329,18 @@ public:
 				if (IsItemSelectable())
 				{
 					// New selections are handled on mouse down, deselection is handled on mouse up
-					const ItemType& MyItem = *OwnerTable->Private_ItemFromWidget(this);
-					if (!OwnerTable->Private_IsItemSelected(MyItem))
+					if (const ItemType* MyItemPtr = GetItemForThis(OwnerTable))
 					{
-						if (SelectionMode != ESelectionMode::Multi)
+						const ItemType& MyItem = *MyItemPtr;
+						if (!OwnerTable->Private_IsItemSelected(MyItem))
 						{
-							OwnerTable->Private_ClearSelection();
+							if (SelectionMode != ESelectionMode::Multi)
+							{
+								OwnerTable->Private_ClearSelection();
+							}
+							OwnerTable->Private_SetItemSelection(MyItem, true, true);
+							bChangedSelectionOnMouseDown = true;
 						}
-						OwnerTable->Private_SetItemSelection(MyItem, true, true);
-						bChangedSelectionOnMouseDown = true;
 					}
 				}
 
@@ -341,61 +359,64 @@ public:
 		FReply Reply = SObjectWidget::OnMouseButtonUp(MyGeometry, MouseEvent);
 
 		TSharedRef<ITypedTableView<ItemType>> OwnerTable = OwnerTablePtr.Pin().ToSharedRef();
-		const ItemType* MyItem = OwnerTable->Private_ItemFromWidget(this);
-		if (!Reply.IsEventHandled() && MyItem)
+
+		if (!Reply.IsEventHandled())
 		{
-			const ESelectionMode::Type SelectionMode = GetSelectionMode();
-			if (MouseEvent.GetEffectingButton() == EKeys::LeftMouseButton && HasMouseCapture())
+			if (const ItemType* MyItemPtr = GetItemForThis(OwnerTable))
 			{
-				bool bSignalSelectionChanged = bChangedSelectionOnMouseDown;
-				if (IsItemSelectable() && MyGeometry.IsUnderLocation(MouseEvent.GetScreenSpacePosition()))
+				const ESelectionMode::Type SelectionMode = GetSelectionMode();
+				if (MouseEvent.GetEffectingButton() == EKeys::LeftMouseButton && HasMouseCapture())
 				{
-					if (SelectionMode == ESelectionMode::SingleToggle)
+					bool bSignalSelectionChanged = bChangedSelectionOnMouseDown;
+					if (IsItemSelectable() && MyGeometry.IsUnderLocation(MouseEvent.GetScreenSpacePosition()))
 					{
-						OwnerTable->Private_ClearSelection();
-						bSignalSelectionChanged = true;
+						if (SelectionMode == ESelectionMode::SingleToggle)
+						{
+							OwnerTable->Private_ClearSelection();
+							bSignalSelectionChanged = true;
+						}
+						else if (SelectionMode == ESelectionMode::Multi &&
+							OwnerTable->Private_GetNumSelectedItems() > 1 &&
+							OwnerTable->Private_IsItemSelected(*MyItemPtr))
+						{
+							// Releasing mouse over one of the multiple selected items - leave this one as the sole selected item
+							OwnerTable->Private_ClearSelection();
+							OwnerTable->Private_SetItemSelection(*MyItemPtr, true, true);
+							bSignalSelectionChanged = true;
+						}
 					}
-					else if (SelectionMode == ESelectionMode::Multi &&
-						OwnerTable->Private_GetNumSelectedItems() > 1 &&
-						OwnerTable->Private_IsItemSelected(*MyItem))
+
+					if (bSignalSelectionChanged)
 					{
-						// Releasing mouse over one of the multiple selected items - leave this one as the sole selected item
-						OwnerTable->Private_ClearSelection();
-						OwnerTable->Private_SetItemSelection(*MyItem, true, true);
-						bSignalSelectionChanged = true;
+						OwnerTable->Private_SignalSelectionChanged(ESelectInfo::OnMouseClick);
+						Reply = FReply::Handled();
 					}
-				}
 
-				if (bSignalSelectionChanged)
-				{
-					OwnerTable->Private_SignalSelectionChanged(ESelectInfo::OnMouseClick);
-					Reply = FReply::Handled();
-				}
+					if (OwnerTable->Private_OnItemClicked(*MyItemPtr))
+					{
+						Reply = FReply::Handled();
+					}
 
-				if (OwnerTable->Private_OnItemClicked(*MyItem))
-				{
-					Reply = FReply::Handled();
+					Reply = Reply.ReleaseMouseCapture();
 				}
-
-				Reply = Reply.ReleaseMouseCapture();
-			}
-			else if (SelectionMode != ESelectionMode::None && MouseEvent.GetEffectingButton() == EKeys::RightMouseButton)
-			{
-				// Ignore the right click release if it was being used for scrolling
-				TSharedRef<STableViewBase> OwnerTableViewBase = StaticCastSharedRef<SListView<ItemType>>(OwnerTable);
-				if (!OwnerTableViewBase->IsRightClickScrolling())
+				else if (SelectionMode != ESelectionMode::None && MouseEvent.GetEffectingButton() == EKeys::RightMouseButton)
 				{
-					if (IsItemSelectable() && !OwnerTable->Private_IsItemSelected(*MyItem))
+					// Ignore the right click release if it was being used for scrolling
+					TSharedRef<STableViewBase> OwnerTableViewBase = StaticCastSharedRef<SListView<ItemType>>(OwnerTable);
+					if (!OwnerTableViewBase->IsRightClickScrolling())
 					{
-						// If this item isn't selected, it becomes the sole selected item. Otherwise we leave selection untouched.
-						OwnerTable->Private_ClearSelection();
-						OwnerTable->Private_SetItemSelection(*MyItem, true, true);
-						OwnerTable->Private_SignalSelectionChanged(ESelectInfo::OnMouseClick);
-					}
+						if (IsItemSelectable() && !OwnerTable->Private_IsItemSelected(*MyItemPtr))
+						{
+							// If this item isn't selected, it becomes the sole selected item. Otherwise we leave selection untouched.
+							OwnerTable->Private_ClearSelection();
+							OwnerTable->Private_SetItemSelection(*MyItemPtr, true, true);
+							OwnerTable->Private_SignalSelectionChanged(ESelectInfo::OnMouseClick);
+						}
 
-					OwnerTable->Private_OnItemRightClicked(*MyItem, MouseEvent);
+						OwnerTable->Private_OnItemRightClicked(*MyItemPtr, MouseEvent);
 
-					Reply = FReply::Handled();
+						Reply = FReply::Handled();
+					}
 				}
 			}
 		}
@@ -407,15 +428,16 @@ public:
 protected:
 	virtual void InitializeObjectRow()
 	{
-		TSharedPtr<ITypedTableView<ItemType>> OwnerTable = OwnerTablePtr.Pin();
-		const ItemType& MyItem = *OwnerTable->Private_ItemFromWidget(this);
-		
-		InitObjectRowInternal(*WidgetObject, MyItem);
-
-		// Unselectable items should never be selected
-		if (!ensure(!OwnerTable->Private_IsItemSelected(MyItem) || IsItemSelectable()))
+		TSharedRef<ITypedTableView<ItemType>> OwnerTable = OwnerTablePtr.Pin().ToSharedRef();
+		if (const ItemType* MyItemPtr = GetItemForThis(OwnerTable))
 		{
-			OwnerTable->Private_SetItemSelection(MyItem, false, false);
+			InitObjectRowInternal(*WidgetObject, *MyItemPtr);
+
+			// Unselectable items should never be selected
+			if (!ensure(!OwnerTable->Private_IsItemSelected(*MyItemPtr) || IsItemSelectable()))
+			{
+				OwnerTable->Private_SetItemSelection(*MyItemPtr, false, false);
+			}
 		}
 	}
 
@@ -442,6 +464,21 @@ protected:
 		return NativeListEntryImpl ? NativeListEntryImpl->IsListItemSelectable() : true;
 	}
 
+	const ItemType* GetItemForThis(const TSharedRef<ITypedTableView<ItemType>>& OwnerTable) const
+	{
+		const ItemType* MyItemPtr = OwnerTable->Private_ItemFromWidget(this);
+		if (MyItemPtr)
+		{
+			return MyItemPtr;
+		}
+		else
+		{
+			checkf(OwnerTable->Private_IsPendingRefresh(), TEXT("We were unable to find the item for this widget.  If it was removed from the source collection, the list should be pending a refresh."));
+		}
+
+		return nullptr;
+	}
+
 	FOnRowHovered OnHovered;
 	FOnRowHovered OnUnhovered;
 
