Stack use after free on temp TSet which FPointerEvent stored as a member variable reference

diff --git a/Engine/Source/Runtime/Slate/Private/Framework/Application/SlateUser.cpp b/Engine/Source/Runtime/Slate/Private/Framework/Application/SlateUser.cpp
--- a/Engine/Source/Runtime/Slate/Private/Framework/Application/SlateUser.cpp
+++ b/Engine/Source/Runtime/Slate/Private/Framework/Application/SlateUser.cpp
@@ -605,6 +605,9 @@ bool FSlateUser::SynthesizeCursorMoveIfNeeded()
 
 		FSlateApplication& SlateApp = FSlateApplication::Get();
 		
+		// Need to ensure this EmptySet lifetime is as long as SyntheticCursorMoveEvent
+		TSet<FKey> EmptySet;
+
 		const bool bHasHardwareCursor = SlateApp.GetPlatformCursor() == Cursor;
 		const TSet<FKey> EmptySet;
 		FPointerEvent SyntheticCursorMoveEvent(
