Reconcile separate fixes for same problem (which now causes compile error)
#rnx

diff --git a/Engine/Source/Runtime/Slate/Private/Framework/Application/SlateUser.cpp b/Engine/Source/Runtime/Slate/Private/Framework/Application/SlateUser.cpp
--- a/Engine/Source/Runtime/Slate/Private/Framework/Application/SlateUser.cpp
+++ b/Engine/Source/Runtime/Slate/Private/Framework/Application/SlateUser.cpp
@@ -605,9 +605,6 @@ bool FSlateUser::SynthesizeCursorMoveIfNeeded()
 
 		FSlateApplication& SlateApp = FSlateApplication::Get();
 		
-		// Need to ensure this EmptySet lifetime is as long as SyntheticCursorMoveEvent
-		TSet<FKey> EmptySet;
-
 		const bool bHasHardwareCursor = SlateApp.GetPlatformCursor() == Cursor;
 		const TSet<FKey> EmptySet;
 		FPointerEvent SyntheticCursorMoveEvent(
