Slate - Fixing several places in SGridPanel that don't properly invalidate when they change.

diff --git a/Engine/Source/Runtime/Slate/Private/Widgets/Layout/SGridPanel.cpp b/Engine/Source/Runtime/Slate/Private/Widgets/Layout/SGridPanel.cpp
--- a/Engine/Source/Runtime/Slate/Private/Widgets/Layout/SGridPanel.cpp
+++ b/Engine/Source/Runtime/Slate/Private/Widgets/Layout/SGridPanel.cpp
@@ -9,6 +9,7 @@
 SGridPanel::SGridPanel()
 : Slots(this)
 {
+	SetCanTick(false);
 }
 
 SGridPanel::FSlot& SGridPanel::AddSlot( int32 Column, int32 Row, SGridPanel::Layer InLayer )
@@ -272,6 +273,8 @@ void SGridPanel::SetColumnFill( int32 ColumnId, const TAttribute<float>& Coeffic
 		ColFillCoefficients.Emplace(0);
 	}
 	ColFillCoefficients[ColumnId] = Coefficient;
+
+	Invalidate(EInvalidateWidgetReason::Layout);
 }
 
 void SGridPanel::SetRowFill( int32 RowId, const TAttribute<float>& Coefficient )
@@ -281,12 +284,16 @@ void SGridPanel::SetRowFill( int32 RowId, const TAttribute<float>& Coefficient )
 		RowFillCoefficients.Emplace(0);
 	}
 	RowFillCoefficients[RowId] = Coefficient;
+
+	Invalidate(EInvalidateWidgetReason::Layout);
 }
 
 void SGridPanel::ClearFill()
 {
 	ColFillCoefficients.Reset();
 	RowFillCoefficients.Reset();
+
+	Invalidate(EInvalidateWidgetReason::Layout);
 }
 
 void SGridPanel::ComputePartialSums( TArray<float>& TurnMeIntoPartialSums )
@@ -369,6 +376,8 @@ void SGridPanel::NotifySlotChanged(SGridPanel::FSlot* InSlot, bool bSlotLayerCha
 			return LHS.LayerParam < RHS.LayerParam;
 		});
 	}
+
+	Invalidate(EInvalidateWidgetReason::Layout);
 }
 
 void SGridPanel::ComputeDesiredCellSizes( TArray<float>& OutColumns, TArray<float>& OutRows ) const
