Added parameter to constrain menu height

diff --git a/Engine/Source/Developer/ToolMenus/Private/ToolMenu.cpp b/Engine/Source/Developer/ToolMenus/Private/ToolMenu.cpp
--- a/Engine/Source/Developer/ToolMenus/Private/ToolMenu.cpp
+++ b/Engine/Source/Developer/ToolMenus/Private/ToolMenu.cpp
@@ -19,6 +19,7 @@ UToolMenu::UToolMenu() :
 	, bToolBarForceSmallIcons(false)
 	, bRegistered(false)
 	, StyleSet(&FCoreStyle::Get())
+	, MaxHeight(INT_MAX)
 {
 }
 
@@ -49,7 +50,7 @@ void UToolMenu::InitGeneratedCopy(const UToolMenu* Source, const FName InMenuNam
 
 	SubMenuParent = Source->SubMenuParent;
 	SubMenuSourceEntryName = Source->SubMenuSourceEntryName;
-
+	MaxHeight = Source->MaxHeight;
 	if (InContext)
 	{
 		Context = *InContext;
diff --git a/Engine/Source/Developer/ToolMenus/Private/ToolMenus.cpp b/Engine/Source/Developer/ToolMenus/Private/ToolMenus.cpp
--- a/Engine/Source/Developer/ToolMenus/Private/ToolMenus.cpp
+++ b/Engine/Source/Developer/ToolMenus/Private/ToolMenus.cpp
@@ -11,6 +11,7 @@
 #include "Internationalization/Internationalization.h"
 
 #include "HAL/PlatformApplicationMisc.h" // For clipboard
+#include "Widgets/Layout/SScrollBox.h"
 
 #define LOCTEXT_NAMESPACE "ToolMenuSubsystem"
 
@@ -1455,7 +1456,7 @@ TSharedRef<SWidget> UToolMenus::GenerateWidget(UToolMenu* GeneratedMenu)
 	{
 		FMenuBuilder MenuBuilder(GeneratedMenu->bShouldCloseWindowAfterMenuSelection, GeneratedMenu->Context.CommandList, GeneratedMenu->Context.GetAllExtenders(), GeneratedMenu->bCloseSelfOnly, GeneratedMenu->StyleSet, GeneratedMenu->bSearchable, GeneratedMenu->MenuName);
 		PopulateMenuBuilder(MenuBuilder, GeneratedMenu);
-		TSharedRef<SWidget> Result = MenuBuilder.MakeWidget();
+		TSharedRef<SWidget> Result = MenuBuilder.MakeWidget(nullptr, GeneratedMenu->MaxHeight);
 		GeneratedMenuWidget.Widget = Result;
 		return Result;
 	}
diff --git a/Engine/Source/Developer/ToolMenus/Public/ToolMenu.h b/Engine/Source/Developer/ToolMenus/Public/ToolMenu.h
--- a/Engine/Source/Developer/ToolMenus/Public/ToolMenu.h
+++ b/Engine/Source/Developer/ToolMenus/Public/ToolMenu.h
@@ -74,6 +74,10 @@ public:
 	//~ End UToolMenuBase Interface
 
 	TArray<FName> GetMenuHierarchyNames(bool bIncludeSubMenuRoot) const;
+	void SetMaxHeight(uint32 InMaxHeight)
+	{
+		MaxHeight = InMaxHeight;
+	}
 
 	template <typename TContextType>
 	TContextType* FindContext() const
@@ -159,4 +163,6 @@ private:
 	bool bRegistered;
 
 	const ISlateStyle* StyleSet;
+
+	uint32 MaxHeight;
 };
diff --git a/Engine/Source/Editor/Kismet/Private/SBlueprintRevisionMenu.cpp b/Engine/Source/Editor/Kismet/Private/SBlueprintRevisionMenu.cpp
--- a/Engine/Source/Editor/Kismet/Private/SBlueprintRevisionMenu.cpp
+++ b/Engine/Source/Editor/Kismet/Private/SBlueprintRevisionMenu.cpp
@@ -242,7 +242,7 @@ void SBlueprintRevisionMenu::OnSourceControlQueryComplete(const FSourceControlOp
 	MenuBuilder.EndSection();
 	MenuBox->AddSlot() 
 	[
-		MenuBuilder.MakeWidget()
+		MenuBuilder.MakeWidget(nullptr, 500)
 	];
 
 	SourceControlQueryOp.Reset();
diff --git a/Engine/Source/Editor/MaterialEditor/Private/MaterialEditor.cpp b/Engine/Source/Editor/MaterialEditor/Private/MaterialEditor.cpp
--- a/Engine/Source/Editor/MaterialEditor/Private/MaterialEditor.cpp
+++ b/Engine/Source/Editor/MaterialEditor/Private/MaterialEditor.cpp
@@ -1241,18 +1241,7 @@ TSharedRef<SWidget> FMaterialEditor::GenerateInheritanceMenu()
 		MenuBuilder.EndSection();
 	}
 
-	TSharedRef<SWidget> ConstrainedMenu = SNew(SVerticalBox)
-		+ SVerticalBox::Slot()
-		.MaxHeight(500.0f)
-		[
-			SNew(SScrollBox)
-			+ SScrollBox::Slot()
-		[
-			MenuBuilder.MakeWidget()
-		]
-		];
-
-	return ConstrainedMenu;
+	return MenuBuilder.MakeWidget(nullptr, 500);
 }
 
 TSharedRef< SWidget > FMaterialEditor::GeneratePreviewMenuContent()
diff --git a/Engine/Source/Editor/MaterialEditor/Private/MaterialInstanceEditor.cpp b/Engine/Source/Editor/MaterialEditor/Private/MaterialInstanceEditor.cpp
--- a/Engine/Source/Editor/MaterialEditor/Private/MaterialInstanceEditor.cpp
+++ b/Engine/Source/Editor/MaterialEditor/Private/MaterialInstanceEditor.cpp
@@ -935,6 +935,7 @@ void FMaterialInstanceEditor::GenerateInheritanceMenu(UToolMenu* Menu)
 	RebuildInheritanceList();
 	Menu->bShouldCloseWindowAfterMenuSelection = true;
 	Menu->bSearchable = true;
+	Menu->SetMaxHeight(500);
 	const FName ParentName = TEXT("ParentChain");
 	FToolMenuSection& Section = Menu->AddSection(ParentName, LOCTEXT("ParentChain", "Parent Chain"));
 	if (bIsFunctionPreviewMaterial)
diff --git a/Engine/Source/Runtime/Slate/Private/Framework/MultiBox/MultiBoxBuilder.cpp b/Engine/Source/Runtime/Slate/Private/Framework/MultiBox/MultiBoxBuilder.cpp
--- a/Engine/Source/Runtime/Slate/Private/Framework/MultiBox/MultiBoxBuilder.cpp
+++ b/Engine/Source/Runtime/Slate/Private/Framework/MultiBox/MultiBoxBuilder.cpp
@@ -17,6 +17,7 @@
 #include "Framework/MultiBox/SWidgetBlock.h"
 #include "Framework/MultiBox/SGroupMarkerBlock.h"
 #include "Framework/MultiBox/ToolMenuBase.h"
+#include "Widgets/Layout/SScrollBox.h"
 
 
 FMultiBoxBuilder::FMultiBoxBuilder( const EMultiBoxType InType, FMultiBoxCustomization InCustomization, const bool bInShouldCloseWindowAfterMenuSelection, const TSharedPtr< const FUICommandList >& InCommandList, TSharedPtr<FExtender> InExtender, FName InTutorialHighlightName, FName InMenuName )
@@ -87,7 +88,7 @@ FMultiBoxCustomization FMultiBoxBuilder::GetCustomization() const
 	return FMultiBoxCustomization( MultiBox->GetCustomizationName() ); 
 }
 
-TSharedRef< class SWidget > FMultiBoxBuilder::MakeWidget( FMultiBox::FOnMakeMultiBoxBuilderOverride* InMakeMultiBoxBuilderOverride /* = nullptr */ )
+TSharedRef< class SWidget > FMultiBoxBuilder::MakeWidget( FMultiBox::FOnMakeMultiBoxBuilderOverride* InMakeMultiBoxBuilderOverride, uint32 MaxHeight /* = nullptr */ )
 {
 	return MultiBox->MakeWidget( false, InMakeMultiBoxBuilderOverride );
 }
@@ -170,10 +171,25 @@ void FBaseMenuBuilder::AddMenuEntry( const FUIAction& UIAction, const TSharedRef
 	ApplyHook(InExtensionHook, EExtensionHook::After);
 }
 
-TSharedRef< class SWidget > FMenuBuilder::MakeWidget( FMultiBox::FOnMakeMultiBoxBuilderOverride* InMakeMultiBoxBuilderOverride /* = nullptr */ )
+TSharedRef< class SWidget > FMenuBuilder::MakeWidget( FMultiBox::FOnMakeMultiBoxBuilderOverride* InMakeMultiBoxBuilderOverride /* = nullptr */, uint32 MaxHeight)
 {
 	// Make menu builders searchable (by default)
-	return MultiBox->MakeWidget( bSearchable, InMakeMultiBoxBuilderOverride );
+	TSharedRef< class SWidget > MenuWidget = MultiBox->MakeWidget(bSearchable, InMakeMultiBoxBuilderOverride);
+	if (MaxHeight < INT_MAX)
+	{
+		TSharedRef<SWidget> ConstrainedMenu = SNew(SVerticalBox)
+			+ SVerticalBox::Slot()
+			.MaxHeight((float)MaxHeight)
+			[
+				SNew(SScrollBox)
+				+ SScrollBox::Slot()
+				[
+					MenuWidget
+				]
+			];
+		return ConstrainedMenu;
+	}
+	return MenuWidget;
 }
 
 void FMenuBuilder::BeginSection( FName InExtensionHook, const TAttribute< FText >& InHeadingText )
diff --git a/Engine/Source/Runtime/Slate/Public/Framework/MultiBox/MultiBoxBuilder.h b/Engine/Source/Runtime/Slate/Public/Framework/MultiBox/MultiBoxBuilder.h
--- a/Engine/Source/Runtime/Slate/Public/Framework/MultiBox/MultiBoxBuilder.h
+++ b/Engine/Source/Runtime/Slate/Public/Framework/MultiBox/MultiBoxBuilder.h
@@ -65,7 +65,7 @@ public:
 	 *
 	 * @return  New widget object
 	 */
-	virtual TSharedRef< class SWidget > MakeWidget( FMultiBox::FOnMakeMultiBoxBuilderOverride* InMakeMultiBoxBuilderOverride = nullptr );
+	virtual TSharedRef< class SWidget > MakeWidget( FMultiBox::FOnMakeMultiBoxBuilderOverride* InMakeMultiBoxBuilderOverride = nullptr, uint32 MaxHeight = INT_MAX);
 	
 
 	/** 
@@ -230,7 +230,7 @@ public:
 	*
 	* @return  New widget object
 	*/
-	virtual TSharedRef< class SWidget > MakeWidget( FMultiBox::FOnMakeMultiBoxBuilderOverride* InMakeMultiBoxBuilderOverride = nullptr ) override;
+	virtual TSharedRef< class SWidget > MakeWidget( FMultiBox::FOnMakeMultiBoxBuilderOverride* InMakeMultiBoxBuilderOverride = nullptr, uint32 MaxHeight = INT_MAX) override;
 
 	/**
 	 * Adds a menu separator
