Slate - The SSlider now waits until you've dragged far enough before it captures the mouse, this will enable it to work better inside of a scrollable list on touch platforms.  This is similar to how buttons work under 'precise tap' mode.

diff --git a/Engine/Source/Runtime/Slate/Private/Framework/Application/SlateApplication.cpp b/Engine/Source/Runtime/Slate/Private/Framework/Application/SlateApplication.cpp
--- a/Engine/Source/Runtime/Slate/Private/Framework/Application/SlateApplication.cpp
+++ b/Engine/Source/Runtime/Slate/Private/Framework/Application/SlateApplication.cpp
@@ -3359,6 +3359,18 @@ bool FSlateApplication::HasTraveledFarEnoughToTriggerDrag(const FPointerEvent& P
 	return ( PointerEvent.GetScreenSpacePosition() - ScreenSpaceOrigin ).SizeSquared() >= ( DragTriggerDistance * DragTriggerDistance );
 }
 
+bool FSlateApplication::HasTraveledFarEnoughToTriggerDrag(const FPointerEvent& PointerEvent, const FVector2D ScreenSpaceOrigin, EOrientation Orientation) const
+{
+	if (Orientation == Orient_Horizontal)
+	{
+		return FMath::Abs(PointerEvent.GetScreenSpacePosition().X - ScreenSpaceOrigin.X) >= DragTriggerDistance;
+	}
+	else // Orientation == Orient_Vertical
+	{
+		return FMath::Abs(PointerEvent.GetScreenSpacePosition().Y - ScreenSpaceOrigin.Y) >= DragTriggerDistance;
+	}
+}
+
 void FSlateApplication::SetDragTriggerDistance( float ScreenPixels )
 {
 	DragTriggerDistance = ScreenPixels;
diff --git a/Engine/Source/Runtime/Slate/Private/Widgets/Input/SSlider.cpp b/Engine/Source/Runtime/Slate/Private/Widgets/Input/SSlider.cpp
--- a/Engine/Source/Runtime/Slate/Private/Widgets/Input/SSlider.cpp
+++ b/Engine/Source/Runtime/Slate/Private/Widgets/Input/SSlider.cpp
@@ -315,14 +315,10 @@ FReply SSlider::OnTouchStarted(const FGeometry& MyGeometry, const FPointerEvent&
 {
 	if (!IsLocked())
 	{
-		CachedCursor = Cursor.Get().Get(EMouseCursor::Default);
-		OnMouseCaptureBegin.ExecuteIfBound();
-		CommitValue(PositionToValue(MyGeometry, InTouchEvent.GetLastScreenSpacePosition()));
-
 		// Release capture for controller/keyboard when switching to mouse.
 		ResetControllerState();
 
-		return FReply::Handled().CaptureMouse(SharedThis(this));
+		PressedScreenSpaceTouchDownPosition = InTouchEvent.GetScreenSpacePosition();
 	}
 
 	return FReply::Unhandled();
@@ -332,13 +328,28 @@ FReply SSlider::OnTouchMoved(const FGeometry& MyGeometry, const FPointerEvent& I
 {
 	if (HasMouseCaptureByUser(InTouchEvent.GetUserIndex(), InTouchEvent.GetPointerIndex()))
 	{
-		CommitValue(PositionToValue(MyGeometry, InTouchEvent.GetLastScreenSpacePosition()));
+		CommitValue(PositionToValue(MyGeometry, InTouchEvent.GetScreenSpacePosition()));
 
 		// Release capture for controller/keyboard when switching to mouse
 		ResetControllerState();
 
 		return FReply::Handled();
 	}
+	else if (!HasMouseCapture())
+	{
+		if (FSlateApplication::Get().HasTraveledFarEnoughToTriggerDrag(InTouchEvent, PressedScreenSpaceTouchDownPosition, Orientation))
+		{
+			CachedCursor = Cursor.Get().Get(EMouseCursor::Default);
+			OnMouseCaptureBegin.ExecuteIfBound();
+
+			CommitValue(PositionToValue(MyGeometry, InTouchEvent.GetScreenSpacePosition()));
+
+			// Release capture for controller/keyboard when switching to mouse
+			ResetControllerState();
+
+			return FReply::Handled().CaptureMouse(SharedThis(this));
+		}
+	}
 
 	return FReply::Unhandled();
 }
@@ -350,6 +361,8 @@ FReply SSlider::OnTouchEnded(const FGeometry& MyGeometry, const FPointerEvent& I
 		SetCursor(CachedCursor);
 		OnMouseCaptureEnd.ExecuteIfBound();
 
+		CommitValue(PositionToValue(MyGeometry, InTouchEvent.GetScreenSpacePosition()));
+
 		// Release capture for controller/keyboard when switching to mouse.
 		ResetControllerState();
 
diff --git a/Engine/Source/Runtime/Slate/Public/Framework/Application/SlateApplication.h b/Engine/Source/Runtime/Slate/Public/Framework/Application/SlateApplication.h
--- a/Engine/Source/Runtime/Slate/Public/Framework/Application/SlateApplication.h
+++ b/Engine/Source/Runtime/Slate/Public/Framework/Application/SlateApplication.h
@@ -1218,6 +1218,7 @@ public:
 
 	/** @return true if the difference between the ScreenSpaceOrigin and the ScreenSpacePosition is larger than the trigger distance for dragging in Slate. */
 	bool HasTraveledFarEnoughToTriggerDrag(const FPointerEvent& PointerEvent, const FVector2D ScreenSpaceOrigin) const;
+	bool HasTraveledFarEnoughToTriggerDrag(const FPointerEvent& PointerEvent, const FVector2D ScreenSpaceOrigin, EOrientation Orientation) const;
 
 	/** Set the size of the deadzone for dragging in screen pixels */
 	void SetDragTriggerDistance( float ScreenPixels );
diff --git a/Engine/Source/Runtime/Slate/Public/Widgets/Input/SSlider.h b/Engine/Source/Runtime/Slate/Public/Widgets/Input/SSlider.h
--- a/Engine/Source/Runtime/Slate/Public/Widgets/Input/SSlider.h
+++ b/Engine/Source/Runtime/Slate/Public/Widgets/Input/SSlider.h
@@ -226,6 +226,9 @@ protected:
 	// Holds the initial cursor in case a custom cursor has been specified, so we can restore it after dragging the slider
 	EMouseCursor::Type CachedCursor;
 
+	/** The location in screenspace the slider was pressed by a touch */
+	FVector2D PressedScreenSpaceTouchDownPosition = FVector2D(0, 0);
+
 	/** Holds the amount to adjust the value by when using a controller or keyboard */
 	TAttribute<float> StepSize;
 
