Slate - If you change the scrolloffset of a scrollbox that needs to invalidate layout.

diff --git a/Engine/Source/Runtime/Slate/Private/Widgets/Layout/SScrollBox.cpp b/Engine/Source/Runtime/Slate/Private/Widgets/Layout/SScrollBox.cpp
--- a/Engine/Source/Runtime/Slate/Private/Widgets/Layout/SScrollBox.cpp
+++ b/Engine/Source/Runtime/Slate/Private/Widgets/Layout/SScrollBox.cpp
@@ -355,6 +355,8 @@ void SScrollBox::SetScrollOffset( float NewScrollOffset )
 {
 	DesiredScrollOffset = NewScrollOffset;
 	bScrollToEnd = false;
+
+	Invalidate(EInvalidateWidget::Layout);
 }
 
 void SScrollBox::ScrollToStart()
