Modified SListView to always select items on item selection when in single selection or singletoggle selection modes.

#rnx

diff --git a/Engine/Source/Runtime/Slate/Public/Widgets/Views/SListView.h b/Engine/Source/Runtime/Slate/Public/Widgets/Views/SListView.h
--- a/Engine/Source/Runtime/Slate/Public/Widgets/Views/SListView.h
+++ b/Engine/Source/Runtime/Slate/Public/Widgets/Views/SListView.h
@@ -740,8 +740,8 @@ public:
 			SelectedItems.Remove( TheItem );
 		}
 
-		// Only move the selector item and range selection start if the user directed this change in selection.
-		if( bWasUserDirected )
+		// Move the selector item and range selection start if the user directed this change in selection or if the list view is single selection
+		if( bWasUserDirected || SelectionMode.Get() == ESelectionMode::Single || SelectionMode.Get() == ESelectionMode::SingleToggle )
 		{
 			SelectorItem = TheItem;
 			RangeSelectionStart = TheItem;
