Slate - Fixing the SSlider to return handled on touch down, so it doesn't fallthrough to handling it in the mouse handler.

diff --git a/Engine/Source/Runtime/Slate/Private/Widgets/Input/SSlider.cpp b/Engine/Source/Runtime/Slate/Private/Widgets/Input/SSlider.cpp
--- a/Engine/Source/Runtime/Slate/Private/Widgets/Input/SSlider.cpp
+++ b/Engine/Source/Runtime/Slate/Private/Widgets/Input/SSlider.cpp
@@ -319,6 +319,8 @@ FReply SSlider::OnTouchStarted(const FGeometry& MyGeometry, const FPointerEvent&
 		ResetControllerState();
 
 		PressedScreenSpaceTouchDownPosition = InTouchEvent.GetScreenSpacePosition();
+
+		return FReply::Handled();
 	}
 
 	return FReply::Unhandled();
