Making ClipToBounds the default value for SlistView, STileview and STreeview. Making StableView pass over clipping to the SlistPanel when it changes.

diff --git a/Engine/Source/Runtime/Slate/Private/Widgets/Views/STableViewBase.cpp b/Engine/Source/Runtime/Slate/Private/Widgets/Views/STableViewBase.cpp
--- a/Engine/Source/Runtime/Slate/Private/Widgets/Views/STableViewBase.cpp
+++ b/Engine/Source/Runtime/Slate/Private/Widgets/Views/STableViewBase.cpp
@@ -62,6 +62,7 @@ void STableViewBase::ConstructChildren( const TAttribute<float>& InItemWidth, co
 	UE_CLOG(InScrollOrientation != Orientation, LogSlate, Error, TEXT("STableViewBase does not support horizontal scrolling when displaying a header row"));
 
 	ItemsPanel = SNew(SListPanel)
+		.Clipping(GetClipping())
 		.ItemWidth(InItemWidth)
 		.ItemHeight(InItemHeight)
 		.NumDesiredItems(this, &STableViewBase::GetNumItemsBeingObserved)
@@ -971,6 +972,11 @@ TSharedRef<class SWidget> STableViewBase::GetScrollWidget()
 	return SharedThis(this);
 }
 
+void STableViewBase::OnClippingChanged()
+{
+	ItemsPanel->SetClipping(GetClipping());
+}
+
 bool STableViewBase::CanUseInertialScroll( float ScrollAmount ) const
 {
 	const auto CurrentOverscroll = Overscroll.GetOverscroll(GetTickSpaceGeometry());
diff --git a/Engine/Source/Runtime/Slate/Public/Widgets/Views/SListView.h b/Engine/Source/Runtime/Slate/Public/Widgets/Views/SListView.h
--- a/Engine/Source/Runtime/Slate/Public/Widgets/Views/SListView.h
+++ b/Engine/Source/Runtime/Slate/Public/Widgets/Views/SListView.h
@@ -97,6 +97,7 @@ public:
 		, _OnItemToString_Debug()
 		, _OnEnteredBadState()
 		{
+			_Clipping = EWidgetClipping::ClipToBounds;
 		}
 
 		SLATE_EVENT( FOnGenerateRow, OnGenerateRow )
diff --git a/Engine/Source/Runtime/Slate/Public/Widgets/Views/STableViewBase.h b/Engine/Source/Runtime/Slate/Public/Widgets/Views/STableViewBase.h
--- a/Engine/Source/Runtime/Slate/Public/Widgets/Views/STableViewBase.h
+++ b/Engine/Source/Runtime/Slate/Public/Widgets/Views/STableViewBase.h
@@ -407,6 +407,9 @@ protected:
 	/** The layout and scroll orientation of the list */
 	EOrientation Orientation = Orient_Vertical;
 
+	/** Passing over the clipping to SListPanel */
+	virtual void OnClippingChanged() override;
+
 protected:
 
 	/** Check whether the current state of the table warrants inertial scroll by the specified amount */
diff --git a/Engine/Source/Runtime/Slate/Public/Widgets/Views/STileView.h b/Engine/Source/Runtime/Slate/Public/Widgets/Views/STileView.h
--- a/Engine/Source/Runtime/Slate/Public/Widgets/Views/STileView.h
+++ b/Engine/Source/Runtime/Slate/Public/Widgets/Views/STileView.h
@@ -67,6 +67,7 @@ public:
 		, _OnEnteredBadState()
 		, _WrapHorizontalNavigation(true)
 		{
+			_Clipping = EWidgetClipping::ClipToBounds;
 		}
 
 		SLATE_EVENT( FOnGenerateRow, OnGenerateTile )
diff --git a/Engine/Source/Runtime/Slate/Public/Widgets/Views/STreeView.h b/Engine/Source/Runtime/Slate/Public/Widgets/Views/STreeView.h
--- a/Engine/Source/Runtime/Slate/Public/Widgets/Views/STreeView.h
+++ b/Engine/Source/Runtime/Slate/Public/Widgets/Views/STreeView.h
@@ -138,7 +138,7 @@ public:
 		, _AllowInvisibleItemSelection(false)
 		, _HighlightParentNodesForSelection(false)
 		{
-			//_Clipping = EWidgetClipping::ClipToBounds;
+			_Clipping = EWidgetClipping::ClipToBounds;
 		}
 
 		SLATE_EVENT( FOnGenerateRow, OnGenerateRow )
diff --git a/Engine/Source/Runtime/UMG/Private/Components/ListViewBase.cpp b/Engine/Source/Runtime/UMG/Private/Components/ListViewBase.cpp
--- a/Engine/Source/Runtime/UMG/Private/Components/ListViewBase.cpp
+++ b/Engine/Source/Runtime/UMG/Private/Components/ListViewBase.cpp
@@ -14,6 +14,7 @@ UListViewBase::UListViewBase(const FObjectInitializer& ObjectInitializer)
 	, EntryWidgetPool(*this)
 {
 	bIsVariable = true;
+	Clipping = EWidgetClipping::ClipToBounds;
 }
 
 #if WITH_EDITOR
