Pottential fix for compile errors in Mac.

diff --git a/Engine/Source/Runtime/Slate/Public/Widgets/Views/SListView.h b/Engine/Source/Runtime/Slate/Public/Widgets/Views/SListView.h
--- a/Engine/Source/Runtime/Slate/Public/Widgets/Views/SListView.h
+++ b/Engine/Source/Runtime/Slate/Public/Widgets/Views/SListView.h
@@ -97,7 +97,7 @@ public:
 		, _OnItemToString_Debug()
 		, _OnEnteredBadState()
 		{
-			_Clipping = EWidgetClipping::ClipToBounds;
+			this->_Clipping = EWidgetClipping::ClipToBounds;
 		}
 
 		SLATE_EVENT( FOnGenerateRow, OnGenerateRow )
diff --git a/Engine/Source/Runtime/Slate/Public/Widgets/Views/STileView.h b/Engine/Source/Runtime/Slate/Public/Widgets/Views/STileView.h
--- a/Engine/Source/Runtime/Slate/Public/Widgets/Views/STileView.h
+++ b/Engine/Source/Runtime/Slate/Public/Widgets/Views/STileView.h
@@ -67,7 +67,7 @@ public:
 		, _OnEnteredBadState()
 		, _WrapHorizontalNavigation(true)
 		{
-			_Clipping = EWidgetClipping::ClipToBounds;
+			this->_Clipping = EWidgetClipping::ClipToBounds;
 		}
 
 		SLATE_EVENT( FOnGenerateRow, OnGenerateTile )
diff --git a/Engine/Source/Runtime/Slate/Public/Widgets/Views/STreeView.h b/Engine/Source/Runtime/Slate/Public/Widgets/Views/STreeView.h
--- a/Engine/Source/Runtime/Slate/Public/Widgets/Views/STreeView.h
+++ b/Engine/Source/Runtime/Slate/Public/Widgets/Views/STreeView.h
@@ -138,7 +138,7 @@ public:
 		, _AllowInvisibleItemSelection(false)
 		, _HighlightParentNodesForSelection(false)
 		{
-			_Clipping = EWidgetClipping::ClipToBounds;
+			this->_Clipping = EWidgetClipping::ClipToBounds;
 		}
 
 		SLATE_EVENT( FOnGenerateRow, OnGenerateRow )
