Add system to blacklist/whitelist ui commands

diff --git a/Engine/Source/Runtime/Slate/Private/Framework/Commands/InputBindingManager.cpp b/Engine/Source/Runtime/Slate/Private/Framework/Commands/InputBindingManager.cpp
--- a/Engine/Source/Runtime/Slate/Private/Framework/Commands/InputBindingManager.cpp
+++ b/Engine/Source/Runtime/Slate/Private/Framework/Commands/InputBindingManager.cpp
@@ -11,7 +11,7 @@
 #include "Serialization/JsonSerializer.h"
 #include "SlateGlobals.h"
 #include "Misc/RemoteConfigIni.h"
-
+#include "HAL/IConsoleManager.h"
 
 /* FUserDefinedChords helper class
  *****************************************************************************/
@@ -536,6 +536,68 @@ void FInputBindingManager::RemoveInputCommand(const TSharedRef<FBindingContext>&
 	}
 }
 
+bool FInputBindingManager::CommandPassesFilter(const FName InBindingContext, const FName InCommandName) const
+{
+	if (const FCommandFilterForContext* CommandFilterForContext = CommandFiltersByContext.Find(InBindingContext))
+	{
+		if (CommandFilterForContext->BlacklistedCommands.Contains(InCommandName))
+		{
+			return false;
+		}
+		else if (!CommandFilterForContext->WhitelistedCommands.Contains(InCommandName) && CommandFilterForContext->WhitelistedCommands.Num() > 0)
+		{
+			return false;
+		}
+	}
+
+	return true;
+}
+
+void FInputBindingManager::AddCommandFilter(const FName InOwnerName, const FName InBindingContext, const FName InCommandName, const ECommandFilterType FilterType)
+{
+	if (FilterType == ECommandFilterType::Blacklist)
+	{
+		CommandFiltersByContext.FindOrAdd(InBindingContext).BlacklistedCommands.FindOrAdd(InCommandName).OwnerNames.AddUnique(InOwnerName);
+	}
+	else if (FilterType == ECommandFilterType::Whitelist)
+	{
+		CommandFiltersByContext.FindOrAdd(InBindingContext).WhitelistedCommands.FindOrAdd(InCommandName).OwnerNames.AddUnique(InOwnerName);
+	}
+}
+
+void FInputBindingManager::UnregisterCommandFilterOwner(const FName InOwnerName)
+{
+	for (auto CommandFiltersByContextIt = CommandFiltersByContext.CreateIterator(); ++CommandFiltersByContextIt; CommandFiltersByContextIt)
+	{
+		FCommandFilterForContext& CommandFilterForContext = CommandFiltersByContextIt->Value;
+
+		for (auto CommandIt = CommandFilterForContext.BlacklistedCommands.CreateIterator(); ++CommandIt; CommandIt)
+		{
+			FCommandFilterOwners& CommandFilterOwners = CommandIt->Value;
+			CommandFilterOwners.OwnerNames.Remove(InOwnerName);
+			if (CommandFilterOwners.OwnerNames.Num() == 0)
+			{
+				CommandIt.RemoveCurrent();
+			}
+		}
+
+		for (auto CommandIt = CommandFilterForContext.WhitelistedCommands.CreateIterator(); ++CommandIt; CommandIt)
+		{
+			FCommandFilterOwners& CommandFilterOwners = CommandIt->Value;
+			CommandFilterOwners.OwnerNames.Remove(InOwnerName);
+			if (CommandFilterOwners.OwnerNames.Num() == 0)
+			{
+				CommandIt.RemoveCurrent();
+			}
+		}
+
+		if (CommandFilterForContext.BlacklistedCommands.Num() == 0 && CommandFilterForContext.WhitelistedCommands.Num() == 0)
+		{
+			CommandFiltersByContextIt.RemoveCurrent();
+		}
+	}
+}
+
 const TSharedPtr<FUICommandInfo> FInputBindingManager::FindCommandInContext( const FName InBindingContext, const FInputChord& InChord, bool bCheckDefault ) const
 {
 	const FContextEntry& ContextEntry = ContextMap.FindRef( InBindingContext );
@@ -568,6 +630,11 @@ const TSharedPtr<FUICommandInfo> FInputBindingManager::FindCommandInContext( con
 		}
 	}
 
+	if (FoundCommand.IsValid() && !CommandPassesFilter(InBindingContext, FoundCommand->GetCommandName()))
+	{
+		return TSharedPtr<FUICommandInfo>();
+	}
+
 	return FoundCommand;
 }
 
@@ -651,3 +718,90 @@ void FInputBindingManager::RemoveContextByName( const FName& InContextName )
 {
 	ContextMap.Remove(InContextName);
 }
+
+void FInputBindingManager::PrintAllInputCommands(bool bBoundOnly)
+{
+#if WITH_EDITOR
+
+	UE_LOG(LogSlate, Log, TEXT("UIContext List  ---------------"));
+	TArray<FName> ContextNames;
+	{	
+		for (TMap< FName, FContextEntry >::TConstIterator It(ContextMap); It; ++It)
+		{
+			ContextNames.Add(It.Key());
+		}
+
+		Algo::Sort(ContextNames, [](const FName& A, const FName& B) { return A.Compare(B) < 0; });
+
+		if (!bBoundOnly)
+		{
+			for (const FName& ContextName : ContextNames)
+			{
+				UE_LOG(LogSlate, Log, TEXT("%s"), *ContextName.ToString());
+			}
+		}
+	}
+
+	UE_LOG(LogSlate, Log, TEXT("UICommand List ---------------"));
+	for (const FName& ContextName : ContextNames)
+	{
+		TArray<FString> Commands;
+		const FContextEntry& ContextEntry = ContextMap.FindChecked(ContextName);
+		const FCommandInfoMap& InfoMap = ContextEntry.CommandInfoMap;
+		for (FCommandInfoMap::TConstIterator CommandInfoIt(InfoMap); CommandInfoIt; ++CommandInfoIt)
+		{
+			const FName CommandName = CommandInfoIt.Key();
+			if (bBoundOnly)
+			{
+				TSharedPtr<FUICommandInfo> CommandInfo = CommandInfoIt.Value();
+				if (CommandInfo.IsValid())
+				{
+					const TSharedRef<const FInputChord> FirstValidChord = CommandInfo->GetFirstValidChord();
+					if (FirstValidChord->IsValidChord())
+					{
+						Commands.Add(CommandName.ToString() + TEXT("\t[") + FirstValidChord->GetInputText().ToString() + TEXT("]"));
+					}
+				}
+			}
+			else
+			{
+				Commands.Add(CommandName.ToString());
+			}
+		}
+
+		Algo::Sort(Commands);
+
+		for (const FString& CommandName : Commands)
+		{
+			UE_LOG(LogSlate, Log, TEXT("%s.%s"), *ContextName.ToString(), *CommandName);
+		}
+	}
+#endif WITH_EDITOR
+}
+
+FInputBindingManager::FInputBindingManager()
+{
+#if WITH_EDITOR
+
+	static FAutoConsoleCommand CVarListKeyboundCommands = FAutoConsoleCommand(
+		TEXT("Slate.Commands.ListBound"),
+		TEXT(""),
+		FConsoleCommandDelegate::CreateLambda([]()
+		{
+			FInputBindingManager::Get().PrintAllInputCommands(true);
+		})
+	);
+
+	static FAutoConsoleCommand CVarListAllCommands = FAutoConsoleCommand(
+		TEXT("Slate.Commands.ListAll"),
+		TEXT(""),
+		FConsoleCommandDelegate::CreateLambda([]()
+		{
+			FInputBindingManager::Get().PrintAllInputCommands(false);
+		})
+	);
+
+#endif // WITH_EDITOR
+}
+
+
diff --git a/Engine/Source/Runtime/Slate/Public/Framework/Commands/InputBindingManager.h b/Engine/Source/Runtime/Slate/Public/Framework/Commands/InputBindingManager.h
--- a/Engine/Source/Runtime/Slate/Public/Framework/Commands/InputBindingManager.h
+++ b/Engine/Source/Runtime/Slate/Public/Framework/Commands/InputBindingManager.h
@@ -142,12 +142,41 @@ public:
 	*/
 	bool GetUserDefinedChord(const FName InBindingContext, const FName InCommandName, const EMultipleKeyBindingIndex InChordIndex, FInputChord& OutUserDefinedChord);
 
+	// Type of filter
+	enum class ECommandFilterType
+	{
+		None,
+		Blacklist,
+		Whitelist
+	};
+
+	/**
+	* Filters a command by name
+	*
+	* @param InOwnerName		The name of the owner adding filter so it can be unregistered later
+	* @param InBindingContext	The context in which the command is active
+	* @param InCommandName		The name of the command
+	* @param FilterType			Type of filter to apply
+	*/
+	void AddCommandFilter(const FName InOwnerName, const FName InBindingContext, const FName InCommandName, const ECommandFilterType FilterType);
+
+	/** Removes all filters associated with an owner */
+	void UnregisterCommandFilterOwner(const FName InOwnerName);
+
+	/**
+	* Command passes blacklist and whitelist filters
+	*
+	* @param InBindingContext	The context in which the command is active
+	* @param InCommandName		The name of the command
+	*/
+	bool CommandPassesFilter(const FName InBindingContext, const FName InCommandName) const;
+
 private:
 
 	/**
 	 * Hidden default constructor.
 	 */
-	FInputBindingManager() { }
+	FInputBindingManager();
 
 	/**
 	 *	Checks a binding context for duplicate chords 
@@ -162,6 +191,11 @@ private:
 	 */
 	void GetAllChildContexts( const FName InBindingContext, TArray<FName>& AllChildren ) const;
 
+	/**
+	 * Displays list of commands.
+	 */
+	void PrintAllInputCommands(bool bBoundOnly);
+
 private:
 
 	struct FContextEntry
@@ -180,6 +214,20 @@ private:
 		TSharedPtr< FBindingContext > BindingContext;
 	};
 
+	struct FCommandFilterOwners
+	{
+		TArray<FName> OwnerNames;
+	};
+
+	struct FCommandFilterForContext
+	{
+		TMap<FName, FCommandFilterOwners> BlacklistedCommands;
+		TMap<FName, FCommandFilterOwners> WhitelistedCommands;
+	};
+
+	/** Mapping of command filters by context */
+	TMap<FName, FCommandFilterForContext> CommandFiltersByContext;
+
 	/** A mapping of context name to the associated entry map */
 	TMap< FName, FContextEntry > ContextMap;
 	
