Fix compile linux
#rnx

diff --git a/Engine/Source/Runtime/Slate/Private/Framework/Commands/InputBindingManager.cpp b/Engine/Source/Runtime/Slate/Private/Framework/Commands/InputBindingManager.cpp
--- a/Engine/Source/Runtime/Slate/Private/Framework/Commands/InputBindingManager.cpp
+++ b/Engine/Source/Runtime/Slate/Private/Framework/Commands/InputBindingManager.cpp
@@ -567,7 +567,7 @@ void FInputBindingManager::AddCommandFilter(const FName InOwnerName, const FName
 
 void FInputBindingManager::UnregisterCommandFilterOwner(const FName InOwnerName)
 {
-	for (auto CommandFiltersByContextIt = CommandFiltersByContext.CreateIterator(); ++CommandFiltersByContextIt; CommandFiltersByContextIt)
+	for (auto CommandFiltersByContextIt = CommandFiltersByContext.CreateIterator(); CommandFiltersByContextIt; ++CommandFiltersByContextIt)
 	{
 		FCommandFilterForContext& CommandFilterForContext = CommandFiltersByContextIt->Value;
 
@@ -581,7 +581,7 @@ void FInputBindingManager::UnregisterCommandFilterOwner(const FName InOwnerName)
 			}
 		}
 
-		for (auto CommandIt = CommandFilterForContext.WhitelistedCommands.CreateIterator(); ++CommandIt; CommandIt)
+		for (auto CommandIt = CommandFilterForContext.WhitelistedCommands.CreateIterator(); CommandIt; ++CommandIt)
 		{
 			FCommandFilterOwners& CommandFilterOwners = CommandIt->Value;
 			CommandFilterOwners.OwnerNames.Remove(InOwnerName);
@@ -776,7 +776,7 @@ void FInputBindingManager::PrintAllInputCommands(bool bBoundOnly)
 			UE_LOG(LogSlate, Log, TEXT("%s.%s"), *ContextName.ToString(), *CommandName);
 		}
 	}
-#endif WITH_EDITOR
+#endif // WITH_EDITOR
 }
 
 FInputBindingManager::FInputBindingManager()
