Fix compile linux
#rnx

diff --git a/Engine/Source/Runtime/Slate/Private/Framework/Commands/InputBindingManager.cpp b/Engine/Source/Runtime/Slate/Private/Framework/Commands/InputBindingManager.cpp
--- a/Engine/Source/Runtime/Slate/Private/Framework/Commands/InputBindingManager.cpp
+++ b/Engine/Source/Runtime/Slate/Private/Framework/Commands/InputBindingManager.cpp
@@ -571,7 +571,7 @@ void FInputBindingManager::UnregisterCommandFilterOwner(const FName InOwnerName)
 	{
 		FCommandFilterForContext& CommandFilterForContext = CommandFiltersByContextIt->Value;
 
-		for (auto CommandIt = CommandFilterForContext.BlacklistedCommands.CreateIterator(); ++CommandIt; CommandIt)
+		for (auto CommandIt = CommandFilterForContext.BlacklistedCommands.CreateIterator(); CommandIt; ++CommandIt)
 		{
 			FCommandFilterOwners& CommandFilterOwners = CommandIt->Value;
 			CommandFilterOwners.OwnerNames.Remove(InOwnerName);
