Merging Dev-Social into Release-11.30

diff --git a/Engine/Source/Runtime/Engine/Classes/Components/InputComponent.h b/Engine/Source/Runtime/Engine/Classes/Components/InputComponent.h
--- a/Engine/Source/Runtime/Engine/Classes/Components/InputComponent.h
+++ b/Engine/Source/Runtime/Engine/Classes/Components/InputComponent.h
@@ -169,6 +169,42 @@ struct FInputActionUnifiedDelegate
 		return false;
 	}
 
+	/** Returns the UObject bound to either the native or dynamic delegate (if either is a UFunction or UObject delegate) */
+	inline const UObject* GetUObject() const
+	{
+		switch (BoundDelegateType)
+		{
+		case EBoundDelegate::Delegate:
+			return FuncDelegate->GetUObject();
+
+		case EBoundDelegate::DelegateWithKey:
+			return FuncDelegateWithKey->GetUObject();
+
+		case EBoundDelegate::DynamicDelegate:
+			return FuncDynDelegate->GetUObject();
+		}
+
+		return nullptr;
+	}
+
+	/** Returns the object bound to either the native or dynamic delegate as a raw, untyped pointer. If you're looking for a bound UObject, prefer GetUObject(). */
+	inline const void* GetObject() const
+	{
+		switch (BoundDelegateType)
+		{
+		case EBoundDelegate::Delegate:
+			return FuncDelegate->GetObjectForTimerManager();
+
+		case EBoundDelegate::DelegateWithKey:
+			return FuncDelegateWithKey->GetObjectForTimerManager();
+
+		case EBoundDelegate::DynamicDelegate:
+			return FuncDynDelegate->GetUObject();
+		}
+
+		return nullptr;
+	}
+
 	/** Binds a native delegate and unbinds any bound dynamic delegate */
 	template< class UserClass >
 	inline void BindDelegate(UserClass* Object, typename FInputActionHandlerSignature::template TUObjectMethodDelegate< UserClass >::FMethodPtr Func)
@@ -801,6 +837,7 @@ public:
 	 * @see AddActionBinding, ClearActionBindings, GetActionBinding, GetNumActionBindings
 	 */
 	void RemoveActionBinding( const int32 BindingIndex );
+	void RemoveActionBinding(FName ActionName, EInputEvent KeyEvent);
 
 	/**
 	 * Removes the action binding at the specified handle.
diff --git a/Engine/Source/Runtime/Engine/Private/Components/InputComponent.cpp b/Engine/Source/Runtime/Engine/Private/Components/InputComponent.cpp
--- a/Engine/Source/Runtime/Engine/Private/Components/InputComponent.cpp
+++ b/Engine/Source/Runtime/Engine/Private/Components/InputComponent.cpp
@@ -317,6 +317,19 @@ void UInputComponent::RemoveActionBinding(const FInputActionBinding &BindingToRe
 	}
 }
 
+void UInputComponent::RemoveActionBinding(FName InActionName, EInputEvent InKeyEvent)
+{
+	for (int32 ActionIdx = 0; ActionIdx < ActionBindings.Num(); ++ActionIdx)
+	{
+		const TSharedPtr<FInputActionBinding>& Binding = ActionBindings[ActionIdx];
+		if (Binding->GetActionName() == InActionName && Binding->KeyEvent == InKeyEvent)
+		{
+			RemoveActionBinding(ActionIdx);
+			break;
+		}
+	}
+}
+
 void UInputComponent::ClearBindingValues()
 {
 	for (FInputAxisBinding& AxisBinding : AxisBindings)
diff --git a/Engine/Source/Runtime/Slate/Private/Framework/Application/SlateApplication.cpp b/Engine/Source/Runtime/Slate/Private/Framework/Application/SlateApplication.cpp
--- a/Engine/Source/Runtime/Slate/Private/Framework/Application/SlateApplication.cpp
+++ b/Engine/Source/Runtime/Slate/Private/Framework/Application/SlateApplication.cpp
@@ -6512,79 +6512,94 @@ void FSlateApplication::InputPreProcessorsHelper::Tick(const float DeltaTime, FS
 {
 	TGuardValue<bool> IteratingGuard(bIsIteratingPreProcessors, true);
 
-	for (int32 ProcessorIndex = 0; ProcessorIndex < InputPreProcessorList.Num(); ProcessorIndex++)
+	for (const TSharedPtr<IInputProcessor>& Preprocessor : InputPreProcessorList)
 	{
-		TSharedPtr<IInputProcessor> InputPreProcessor = InputPreProcessorList[ProcessorIndex];
-		InputPreProcessor->Tick(DeltaTime, SlateApp, Cursor);
+		if (Preprocessor)
+		{
+			Preprocessor->Tick(DeltaTime, SlateApp, Cursor);
+		}
 	}
 }
 
 bool FSlateApplication::InputPreProcessorsHelper::HandleKeyDownEvent(FSlateApplication& SlateApp, const FKeyEvent& InKeyEvent)
 {
-	return PreProcessInput([&SlateApp, &InKeyEvent](TSharedPtr<IInputProcessor> Processor) { return Processor->HandleKeyDownEvent(SlateApp, InKeyEvent); });
+	return PreProcessInput([&SlateApp, &InKeyEvent](IInputProcessor& Processor) { return Processor.HandleKeyDownEvent(SlateApp, InKeyEvent); });
 }
 
 bool FSlateApplication::InputPreProcessorsHelper::HandleKeyUpEvent(FSlateApplication& SlateApp, const FKeyEvent& InKeyEvent)
 {
-	return PreProcessInput([&SlateApp, &InKeyEvent](TSharedPtr<IInputProcessor> Processor) { return Processor->HandleKeyUpEvent(SlateApp, InKeyEvent); });
+	return PreProcessInput([&SlateApp, &InKeyEvent](IInputProcessor& Processor) { return Processor.HandleKeyUpEvent(SlateApp, InKeyEvent); });
 }
 
 bool FSlateApplication::InputPreProcessorsHelper::HandleAnalogInputEvent(FSlateApplication& SlateApp, const FAnalogInputEvent& InAnalogInputEvent)
 {
-	return PreProcessInput([&SlateApp, &InAnalogInputEvent](TSharedPtr<IInputProcessor> Processor) { return Processor->HandleAnalogInputEvent(SlateApp, InAnalogInputEvent); });
+	return PreProcessInput([&SlateApp, &InAnalogInputEvent](IInputProcessor& Processor) { return Processor.HandleAnalogInputEvent(SlateApp, InAnalogInputEvent); });
 }
 
 bool FSlateApplication::InputPreProcessorsHelper::HandleMouseMoveEvent(FSlateApplication& SlateApp, const FPointerEvent& MouseEvent)
 {
-	return PreProcessInput([&SlateApp, &MouseEvent](TSharedPtr<IInputProcessor> Processor) { return Processor->HandleMouseMoveEvent(SlateApp, MouseEvent); });
+	return PreProcessInput([&SlateApp, &MouseEvent](IInputProcessor& Processor) { return Processor.HandleMouseMoveEvent(SlateApp, MouseEvent); });
 }
 
 bool FSlateApplication::InputPreProcessorsHelper::HandleMouseButtonDownEvent(FSlateApplication& SlateApp, const FPointerEvent& MouseEvent)
 {
-	return PreProcessInput([&SlateApp, &MouseEvent](TSharedPtr<IInputProcessor> Processor) { return Processor->HandleMouseButtonDownEvent(SlateApp, MouseEvent); });
+	return PreProcessInput([&SlateApp, &MouseEvent](IInputProcessor& Processor) { return Processor.HandleMouseButtonDownEvent(SlateApp, MouseEvent); });
 }
 
 bool FSlateApplication::InputPreProcessorsHelper::HandleMouseButtonUpEvent(FSlateApplication& SlateApp, const FPointerEvent& MouseEvent)
 {
-	return PreProcessInput([&SlateApp, &MouseEvent](TSharedPtr<IInputProcessor> Processor) { return Processor->HandleMouseButtonUpEvent(SlateApp, MouseEvent); });
-
+	return PreProcessInput([&SlateApp, &MouseEvent](IInputProcessor& Processor) { return Processor.HandleMouseButtonUpEvent(SlateApp, MouseEvent); });
 }
 
 bool FSlateApplication::InputPreProcessorsHelper::HandleMouseButtonDoubleClickEvent(FSlateApplication& SlateApp, const FPointerEvent& MouseEvent)
 {
-	return PreProcessInput([&SlateApp, &MouseEvent](TSharedPtr<IInputProcessor> Processor) { return Processor->HandleMouseButtonDoubleClickEvent(SlateApp, MouseEvent); });
+	return PreProcessInput([&SlateApp, &MouseEvent](IInputProcessor& Processor) { return Processor.HandleMouseButtonDoubleClickEvent(SlateApp, MouseEvent); });
 }
 
 bool FSlateApplication::InputPreProcessorsHelper::HandleMouseWheelOrGestureEvent(FSlateApplication& SlateApp, const FPointerEvent& WheelEvent, const FPointerEvent* GestureEvent)
 {
-	return PreProcessInput([&SlateApp, &WheelEvent, &GestureEvent](TSharedPtr<IInputProcessor> Processor) { return Processor->HandleMouseWheelOrGestureEvent(SlateApp, WheelEvent, GestureEvent); });
+	return PreProcessInput([&SlateApp, &WheelEvent, &GestureEvent](IInputProcessor& Processor) { return Processor.HandleMouseWheelOrGestureEvent(SlateApp, WheelEvent, GestureEvent); });
 }
 
 bool FSlateApplication::InputPreProcessorsHelper::HandleMotionDetectedEvent(FSlateApplication& SlateApp, const FMotionEvent& MotionEvent)
 {
-	return PreProcessInput([&SlateApp, &MotionEvent](TSharedPtr<IInputProcessor> Processor) { return Processor->HandleMotionDetectedEvent(SlateApp, MotionEvent); });
+	return PreProcessInput([&SlateApp, &MotionEvent](IInputProcessor& Processor) { return Processor.HandleMotionDetectedEvent(SlateApp, MotionEvent); });
 }
 
 bool FSlateApplication::InputPreProcessorsHelper::Add(TSharedPtr<IInputProcessor> InputProcessor, const int32 Index /*= INDEX_NONE*/)
 {
-	bool bResult = false;
-	if (Index == INDEX_NONE)
+	const bool bAlreadyInList = InputPreProcessorList.Contains(InputProcessor);
+	if (!bAlreadyInList)
 	{
-		InputPreProcessorList.AddUnique(InputProcessor);
-		bResult = true;
+		if (!bIsIteratingPreProcessors)
+		{
+			AddInternal(InputProcessor, Index);
+		}
+		else
+		{
+			ProcessorsPendingAddition.Add(InputProcessor, Index);
+		}
 	}
-	else if (InputPreProcessorList.Find(InputProcessor) == INDEX_NONE)
+
+	ProcessorsPendingRemoval.Remove(InputProcessor);
+
+	return !bAlreadyInList;
+}
+
+void FSlateApplication::InputPreProcessorsHelper::AddInternal(TSharedPtr<IInputProcessor> InputProcessor, const int32 Index)
+{
+	if (Index == INDEX_NONE)
 	{
-		InputPreProcessorList.Insert(InputProcessor, Index);
-		bResult = true;
+		InputPreProcessorList.Add(InputProcessor);
 	}
-
-	if (bResult)
+	else
 	{
-		ProcessorsPendingRemoval.Remove(InputProcessor);
+		if (Index >= InputPreProcessorList.Num())
+		{
+			InputPreProcessorList.SetNum(Index + 1);
+		}
+		InputPreProcessorList.Insert(InputProcessor, Index);
 	}
-
-	return bResult;
 }
 
 void FSlateApplication::InputPreProcessorsHelper::Remove(TSharedPtr<IInputProcessor> InputProcessor)
@@ -6597,6 +6612,8 @@ void FSlateApplication::InputPreProcessorsHelper::Remove(TSharedPtr<IInputProces
 	{
 		InputPreProcessorList.Remove(InputProcessor);
 	}
+
+	ProcessorsPendingAddition.Remove(InputProcessor);
 }
 
 void FSlateApplication::InputPreProcessorsHelper::RemoveAll()
@@ -6609,6 +6626,8 @@ void FSlateApplication::InputPreProcessorsHelper::RemoveAll()
 	{
 		InputPreProcessorList.Reset();
 	}
+
+	ProcessorsPendingAddition.Reset();
 }
 
 
@@ -6617,14 +6636,14 @@ int32 FSlateApplication::InputPreProcessorsHelper::Find(TSharedPtr<IInputProcess
 	return InputPreProcessorList.Find(InputProcessor);
 }
 
-bool FSlateApplication::InputPreProcessorsHelper::PreProcessInput(TFunctionRef<bool(TSharedPtr<IInputProcessor>)> ToRun)
+bool FSlateApplication::InputPreProcessorsHelper::PreProcessInput(TFunctionRef<bool(IInputProcessor&)> InputProcessFunc)
 {
 	TGuardValue<bool> IteratingGuard(bIsIteratingPreProcessors, true);
 
 	bool bShouldExit = false;
-	for (TSharedPtr<IInputProcessor> InputPreProcessor : InputPreProcessorList)
+	for (const TSharedPtr<IInputProcessor>& InputPreProcessor : InputPreProcessorList)
 	{
-		if (ToRun(InputPreProcessor))
+		if (InputPreProcessor && InputProcessFunc(*InputPreProcessor))
 		{
 			bShouldExit = true;
 			break;
@@ -6634,5 +6653,11 @@ bool FSlateApplication::InputPreProcessorsHelper::PreProcessInput(TFunctionRef<b
 	InputPreProcessorList.RemoveAll([this](const TSharedPtr<IInputProcessor> Processor) { return ProcessorsPendingRemoval.Contains(Processor); });
 	ProcessorsPendingRemoval.Reset();
 
+	for (TPair<TSharedPtr<IInputProcessor>, int32>& ProcessorIndexPair : ProcessorsPendingAddition)
+	{
+		AddInternal(ProcessorIndexPair.Key, ProcessorIndexPair.Value);
+	}
+	ProcessorsPendingAddition.Reset();
+
 	return bShouldExit;
 }
\ No newline at end of file
diff --git a/Engine/Source/Runtime/Slate/Public/Framework/Application/SlateApplication.h b/Engine/Source/Runtime/Slate/Public/Framework/Application/SlateApplication.h
--- a/Engine/Source/Runtime/Slate/Public/Framework/Application/SlateApplication.h
+++ b/Engine/Source/Runtime/Slate/Public/Framework/Application/SlateApplication.h
@@ -1831,10 +1831,10 @@ private:
 		 */
 		int32 Find(TSharedPtr<IInputProcessor> InputProcessor) const;
 
-
 	private:
+		bool PreProcessInput(TFunctionRef<bool(IInputProcessor&)> InputProcessFunc);
 
-		bool PreProcessInput(TFunctionRef<bool(TSharedPtr<IInputProcessor>)> ToRun);
+		void AddInternal(TSharedPtr<IInputProcessor> InputProcessor, const int32 Index);
 
 		/** The list of input pre-processors. */
 		TArray<TSharedPtr<IInputProcessor>> InputPreProcessorList;
@@ -1844,6 +1844,9 @@ private:
 
 		/** A list of pre-processors to remove if we are iterating them while removal is requested. */
 		TArray<TSharedPtr<IInputProcessor>> ProcessorsPendingRemoval;
+
+		/** A list of pre-processors to add if we are iterating them while addition is requested. */
+		TMap<TSharedPtr<IInputProcessor>, int32> ProcessorsPendingAddition;
 	};
 
 	/** A list of input pre-processors, gets an opportunity to parse input before anything else. */
diff --git a/Engine/Source/Runtime/Slate/Public/Widgets/Layout/SWidgetSwitcher.h b/Engine/Source/Runtime/Slate/Public/Widgets/Layout/SWidgetSwitcher.h
--- a/Engine/Source/Runtime/Slate/Public/Widgets/Layout/SWidgetSwitcher.h
+++ b/Engine/Source/Runtime/Slate/Public/Widgets/Layout/SWidgetSwitcher.h
@@ -156,6 +156,8 @@ protected:
 	virtual bool ComputeVolatility() const override { return WidgetIndex.IsBound(); }
 	const FSlot* GetActiveSlot() const;
 
+	TPanelChildren<FSlot>& GetTypedChildren() { return AllChildren; }
+
 private:
 
 	/** Holds the desired widget index */
diff --git a/Engine/Source/Runtime/UMG/Public/Blueprint/UserWidgetPool.h b/Engine/Source/Runtime/UMG/Public/Blueprint/UserWidgetPool.h
--- a/Engine/Source/Runtime/UMG/Public/Blueprint/UserWidgetPool.h
+++ b/Engine/Source/Runtime/UMG/Public/Blueprint/UserWidgetPool.h
@@ -91,7 +91,7 @@ private:
 	template <typename UserWidgetT = UUserWidget>
 	UserWidgetT* AddActiveWidgetInternal(TSubclassOf<UserWidgetT> WidgetClass, WidgetConstructFunc ConstructWidgetFunc)
 	{
-		if (!ensure(IsInitialized()))
+		if (!ensure(IsInitialized()) || !WidgetClass)
 		{
 			return nullptr;
 		}
diff --git a/Engine/Source/Runtime/UMG/Public/Components/ListView.h b/Engine/Source/Runtime/UMG/Public/Components/ListView.h
--- a/Engine/Source/Runtime/UMG/Public/Components/ListView.h
+++ b/Engine/Source/Runtime/UMG/Public/Components/ListView.h
@@ -47,6 +47,7 @@ public:
 	}
 
 	ESelectionMode::Type GetSelectionMode() const { return SelectionMode; }
+	EOrientation GetOrientation() const { return Orientation; }
 
 	template <typename RowWidgetT = UUserWidget>
 	RowWidgetT* GetEntryWidgetFromItem(const UObject* Item) const
@@ -158,7 +159,7 @@ protected:
 	 * Vertical will scroll vertically and arrange tiles into rows.
 	 * Horizontal will scroll horizontally and arrange tiles into columns.
 	 */
-	UPROPERTY(EditAnywhere, Category = ListView)
+	UPROPERTY(EditAnywhere, BlueprintReadOnly, Category = ListView)
 	TEnumAsByte<EOrientation> Orientation;
 
 	UPROPERTY(EditAnywhere, BlueprintReadOnly, Category = ListView)
