Fix extenders on menu sometimes not working
#rnx

diff --git a/Engine/Source/Runtime/Slate/Private/Framework/MultiBox/MultiBoxBuilder.cpp b/Engine/Source/Runtime/Slate/Private/Framework/MultiBox/MultiBoxBuilder.cpp
--- a/Engine/Source/Runtime/Slate/Private/Framework/MultiBox/MultiBoxBuilder.cpp
+++ b/Engine/Source/Runtime/Slate/Private/Framework/MultiBox/MultiBoxBuilder.cpp
@@ -25,6 +25,7 @@ FMultiBoxBuilder::FMultiBoxBuilder( const EMultiBoxType InType, FMultiBoxCustomi
 	, CommandListStack()
 	, TutorialHighlightName(InTutorialHighlightName)
 	, MenuName(InMenuName)
+	, bExtendersEnabled(true)
 {
 	CommandListStack.Push( InCommandList );
 	ExtenderStack.Push(InExtender);
