Analog cursor now checks if input is relevant and does not process it if it is not.

diff --git a/Engine/Source/Runtime/Slate/Private/Framework/Application/AnalogCursor.cpp b/Engine/Source/Runtime/Slate/Private/Framework/Application/AnalogCursor.cpp
--- a/Engine/Source/Runtime/Slate/Private/Framework/Application/AnalogCursor.cpp
+++ b/Engine/Source/Runtime/Slate/Private/Framework/Application/AnalogCursor.cpp
@@ -7,6 +7,7 @@
 #include "Layout/ArrangedChildren.h"
 #include "Layout/WidgetPath.h"
 #include "Framework/Application/SlateApplication.h"
+#include "Framework/Application/SlateUser.h"
 
 
 FAnalogCursor::FAnalogCursor()
@@ -22,50 +23,52 @@ FAnalogCursor::FAnalogCursor()
 	AnalogValues[ static_cast< uint8 >( EAnalogStick::Right ) ] = FVector2D::ZeroVector;
 }
 
-void FAnalogCursor::Tick(const float DeltaTime, FSlateApplication& SlateApp, TSharedRef<ICursor> Cursor)
+void FAnalogCursor::Tick(const float DeltaTime, FSlateApplication& SlateApp, TSharedRef<ICursor>)
 {
-	const FVector2D OldPosition = Cursor->GetPosition();
+	if (TSharedPtr<FSlateUser> SlateUser = SlateApp.GetUser(GetUserControllerId()))
+	{
+		const FVector2D OldPosition = SlateUser->GetCursorPosition();
 
-	float SpeedMult = 1.0f; // Used to do a speed multiplication before adding the delta to the position to make widgets sticky
-	FVector2D AdjAnalogVals = GetAnalogValue(EAnalogStick::Left); // A copy of the analog values so I can modify them based being over a widget, not currently doing this
+		float SpeedMult = 1.0f; // Used to do a speed multiplication before adding the delta to the position to make widgets sticky
+		FVector2D AdjAnalogVals = GetAnalogValue(EAnalogStick::Left); // A copy of the analog values so I can modify them based being over a widget, not currently doing this
 
-	// Adjust analog values according to dead zone
-	const float AnalogValsSize = AdjAnalogVals.Size();
+		// Adjust analog values according to dead zone
+		const float AnalogValsSize = AdjAnalogVals.Size();
 
-	if (AnalogValsSize > 0.0f)
-	{
-		const float TargetSize = FMath::Max(AnalogValsSize - DeadZone, 0.0f) / (1.0f - DeadZone);
-		AdjAnalogVals /= AnalogValsSize;
-		AdjAnalogVals *= TargetSize;
-	}
+		if (AnalogValsSize > 0.0f)
+		{
+			const float TargetSize = FMath::Max(AnalogValsSize - DeadZone, 0.0f) / (1.0f - DeadZone);
+			AdjAnalogVals /= AnalogValsSize;
+			AdjAnalogVals *= TargetSize;
+		}
 
 
-	// Check if there is a sticky widget beneath the cursor
-	FWidgetPath WidgetPath = SlateApp.LocateWindowUnderMouse(OldPosition, SlateApp.GetInteractiveTopLevelWindows());
-	if (WidgetPath.IsValid())
-	{
-		const FArrangedChildren::FArrangedWidgetArray& AllArrangedWidgets = WidgetPath.Widgets.GetInternalArray();
-		for(const FArrangedWidget& ArrangedWidget : AllArrangedWidgets)
+		// Check if there is a sticky widget beneath the cursor
+		FWidgetPath WidgetPath = SlateApp.LocateWindowUnderMouse(OldPosition, SlateApp.GetInteractiveTopLevelWindows());
+		if (WidgetPath.IsValid())
 		{
-			TSharedRef<SWidget> Widget = ArrangedWidget.Widget;
-			if (Widget->IsInteractable())
+			const FArrangedChildren::FArrangedWidgetArray& AllArrangedWidgets = WidgetPath.Widgets.GetInternalArray();
+			for (const FArrangedWidget& ArrangedWidget : AllArrangedWidgets)
 			{
-				SpeedMult = StickySlowdown;
-				//FVector2D Adjustment = WidgetsAndCursors.Last().Geometry.Position - OldPosition; // example of calculating distance from cursor to widget center
-				break;
+				TSharedRef<SWidget> Widget = ArrangedWidget.Widget;
+				if (Widget->IsInteractable())
+				{
+					SpeedMult = StickySlowdown;
+					//FVector2D Adjustment = WidgetsAndCursors.Last().Geometry.Position - OldPosition; // example of calculating distance from cursor to widget center
+					break;
+				}
 			}
 		}
-	}
 
-	switch(Mode)
-	{
+		switch (Mode)
+		{
 		case AnalogCursorMode::Accelerated:
 		{
 			// Generate Min and Max for X to clamp the speed, this gives us instant direction change when crossing the axis
 			float CurrentMinSpeedX = 0.0f;
 			float CurrentMaxSpeedX = 0.0f;
 			if (AdjAnalogVals.X > 0.0f)
-			{ 
+			{
 				CurrentMaxSpeedX = AdjAnalogVals.X * MaxSpeed;
 			}
 			else
@@ -104,53 +107,60 @@ void FAnalogCursor::Tick(const float DeltaTime, FSlateApplication& SlateApp, TSh
 			CurrentSpeed = AdjAnalogVals * MaxSpeed;
 
 			break;
-	}
+		}
 
-	CurrentOffset += CurrentSpeed * DeltaTime * SpeedMult;
-	const FVector2D NewPosition = OldPosition + CurrentOffset;
+		CurrentOffset += CurrentSpeed * DeltaTime * SpeedMult;
+		const FVector2D NewPosition = OldPosition + CurrentOffset;
 
-	// save the remaining sub-pixel offset 
-	CurrentOffset.X = FGenericPlatformMath::Frac(NewPosition.X);
-	CurrentOffset.Y = FGenericPlatformMath::Frac(NewPosition.Y);
+		// save the remaining sub-pixel offset 
+		CurrentOffset.X = FGenericPlatformMath::Frac(NewPosition.X);
+		CurrentOffset.Y = FGenericPlatformMath::Frac(NewPosition.Y);
 
-	// update the cursor position
-	UpdateCursorPosition(SlateApp, Cursor, NewPosition);
+		// update the cursor position
+		UpdateCursorPosition(SlateApp, SlateUser.ToSharedRef(), NewPosition);
+	}
 }
 
 bool FAnalogCursor::HandleKeyDownEvent(FSlateApplication& SlateApp, const FKeyEvent& InKeyEvent)
 {
-	FKey Key = InKeyEvent.GetKey();
-
-	// Consume the sticks input so it doesn't effect other things
-	if (Key == EKeys::Gamepad_LeftStick_Right || 
-		Key == EKeys::Gamepad_LeftStick_Left ||
-		Key == EKeys::Gamepad_LeftStick_Up ||
-		Key == EKeys::Gamepad_LeftStick_Down)
+	if (IsRelevantInput(InKeyEvent))
 	{
-		return true;
-	}
+		FKey Key = InKeyEvent.GetKey();
+		// Consume the sticks input so it doesn't effect other things
+		if (Key == EKeys::Gamepad_LeftStick_Right ||
+			Key == EKeys::Gamepad_LeftStick_Left ||
+			Key == EKeys::Gamepad_LeftStick_Up ||
+			Key == EKeys::Gamepad_LeftStick_Down)
+		{
+			return true;
+		}
 
-	// Bottom face button is a click
-	if (Key == EKeys::Virtual_Accept)
-	{
-		if ( !InKeyEvent.IsRepeat() )
+		// Bottom face button is a click
+		if (Key == EKeys::Virtual_Accept)
 		{
-			FPointerEvent MouseEvent(
-				InKeyEvent.GetUserIndex(),
-				SlateApp.CursorPointerIndex,
-				SlateApp.GetCursorPos(),
-				SlateApp.GetLastCursorPos(),
-				SlateApp.GetPressedMouseButtons(),
-				EKeys::LeftMouseButton,
-				0,
-				SlateApp.GetPlatformApplication()->GetModifierKeys()
-				);
+			if (!InKeyEvent.IsRepeat())
+			{
+				if (TSharedPtr<FSlateUser> SlateUser = SlateApp.GetUser(InKeyEvent))
+				{
+					const bool bIsPrimaryUser = FSlateApplication::CursorUserIndex == SlateUser->GetUserIndex();
+					FPointerEvent MouseEvent(
+						SlateUser->GetUserIndex(),
+						FSlateApplication::CursorPointerIndex,
+						SlateUser->GetCursorPosition(),
+						SlateUser->GetPreviousCursorPosition(),
+						bIsPrimaryUser ? SlateApp.GetPressedMouseButtons() : TSet<FKey>(),
+						EKeys::LeftMouseButton,
+						0,
+						bIsPrimaryUser ? SlateApp.GetModifierKeys() : FModifierKeysState()
+					);
+
+					TSharedPtr<FGenericWindow> GenWindow;
+					return SlateApp.ProcessMouseButtonDownEvent(GenWindow, MouseEvent);
+				}
+			}
 
-			TSharedPtr<FGenericWindow> GenWindow;
-			return SlateApp.ProcessMouseButtonDownEvent(GenWindow, MouseEvent);
+			return true;
 		}
-
-		return true;
 	}
 
 	return false;
@@ -158,67 +168,75 @@ bool FAnalogCursor::HandleKeyDownEvent(FSlateApplication& SlateApp, const FKeyEv
 
 bool FAnalogCursor::HandleKeyUpEvent(FSlateApplication& SlateApp, const FKeyEvent& InKeyEvent)
 {
-	FKey Key = InKeyEvent.GetKey();
-
-	// Consume the sticks input so it doesn't effect other things
-	if (Key == EKeys::Gamepad_LeftStick_Right ||
-		Key == EKeys::Gamepad_LeftStick_Left ||
-		Key == EKeys::Gamepad_LeftStick_Up ||
-		Key == EKeys::Gamepad_LeftStick_Down)
+	if (IsRelevantInput(InKeyEvent))
 	{
-		return true;
-	}
+		FKey Key = InKeyEvent.GetKey();
 
-	// Bottom face button is a click
-	if (Key == EKeys::Virtual_Accept && !InKeyEvent.IsRepeat())
-	{
-		FPointerEvent MouseEvent(
-			InKeyEvent.GetUserIndex(),
-			SlateApp.CursorPointerIndex,
-			SlateApp.GetCursorPos(),
-			SlateApp.GetLastCursorPos(),
-			SlateApp.GetPressedMouseButtons(),
-			EKeys::LeftMouseButton,
-			0,
-			SlateApp.GetPlatformApplication()->GetModifierKeys()
-			);
+		// Consume the sticks input so it doesn't effect other things
+		if (Key == EKeys::Gamepad_LeftStick_Right ||
+			Key == EKeys::Gamepad_LeftStick_Left ||
+			Key == EKeys::Gamepad_LeftStick_Up ||
+			Key == EKeys::Gamepad_LeftStick_Down)
+		{
+			return true;
+		}
 
-		return SlateApp.ProcessMouseButtonUpEvent(MouseEvent);
+		// Bottom face button is a click
+		if (Key == EKeys::Virtual_Accept && !InKeyEvent.IsRepeat())
+		{
+			if (TSharedPtr<FSlateUser> SlateUser = SlateApp.GetUser(InKeyEvent))
+			{
+				const bool bIsPrimaryUser = FSlateApplication::CursorUserIndex == SlateUser->GetUserIndex();
+				FPointerEvent MouseEvent(
+					SlateUser->GetUserIndex(),
+					FSlateApplication::CursorPointerIndex,
+					SlateUser->GetCursorPosition(),
+					SlateUser->GetPreviousCursorPosition(),
+					bIsPrimaryUser ? SlateApp.GetPressedMouseButtons() : TSet<FKey>(),
+					EKeys::LeftMouseButton,
+					0,
+					bIsPrimaryUser ? SlateApp.GetModifierKeys() : FModifierKeysState()
+				);
+
+				return SlateApp.ProcessMouseButtonUpEvent(MouseEvent);
+			}
+		}
 	}
-	
 	return false;
 }
 
 bool FAnalogCursor::HandleAnalogInputEvent(FSlateApplication& SlateApp, const FAnalogInputEvent& InAnalogInputEvent)
 {
-	FKey Key = InAnalogInputEvent.GetKey();
-	float AnalogValue = InAnalogInputEvent.GetAnalogValue();
-
-	if (Key == EKeys::Gamepad_LeftX)
+	if (IsRelevantInput(InAnalogInputEvent))
 	{
-		FVector2D& Value = GetAnalogValue( EAnalogStick::Left );
-		Value.X = AnalogValue;
-	}
-	else if (Key == EKeys::Gamepad_LeftY)
-	{
-		FVector2D& Value = GetAnalogValue( EAnalogStick::Left );
-		Value.Y = -AnalogValue;
-	}
-	else if ( Key == EKeys::Gamepad_RightX )
-	{
-		FVector2D& Value = GetAnalogValue( EAnalogStick::Right );
-		Value.X = AnalogValue;
-	}
-	else if ( Key == EKeys::Gamepad_RightY )
-	{
-		FVector2D& Value = GetAnalogValue( EAnalogStick::Right );
-		Value.Y = -AnalogValue;
-	}
-	else 
-	{
-		return false;
-	}
+		FKey Key = InAnalogInputEvent.GetKey();
+		float AnalogValue = InAnalogInputEvent.GetAnalogValue();
 
+		if (Key == EKeys::Gamepad_LeftX)
+		{
+			FVector2D& Value = GetAnalogValue(EAnalogStick::Left);
+			Value.X = AnalogValue;
+		}
+		else if (Key == EKeys::Gamepad_LeftY)
+		{
+			FVector2D& Value = GetAnalogValue(EAnalogStick::Left);
+			Value.Y = -AnalogValue;
+		}
+		else if (Key == EKeys::Gamepad_RightX)
+		{
+			FVector2D& Value = GetAnalogValue(EAnalogStick::Right);
+			Value.X = AnalogValue;
+		}
+		else if (Key == EKeys::Gamepad_RightY)
+		{
+			FVector2D& Value = GetAnalogValue(EAnalogStick::Right);
+			Value.Y = -AnalogValue;
+		}
+		else
+		{
+			return false;
+		}
+	}
 	return true;
 }
 
@@ -254,6 +272,11 @@ void FAnalogCursor::SetMode(AnalogCursorMode::Type NewMode)
 	CurrentSpeed = FVector2D::ZeroVector;
 }
 
+bool FAnalogCursor::IsRelevantInput(const FInputEvent& InputEvent) const
+{
+	return GetUserControllerId() == InputEvent.GetUserIndex();
+}
+
 void FAnalogCursor::ClearAnalogValues()
 {
 	AnalogValues[static_cast<uint8>(EAnalogStick::Left)] = FVector2D::ZeroVector;
@@ -274,22 +297,57 @@ void FAnalogCursor::UpdateCursorPosition(FSlateApplication& SlateApp, TSharedRef
 	{
 		//put the cursor in the correct spot
 		Cursor->SetPosition(NewIntPosX, NewIntPosY);
-	
+
 		// Since the cursor may have been locked and its location clamped, get the actual new position
 		const FVector2D UpdatedPosition = Cursor->GetPosition();
+		if (TSharedPtr<FSlateUser> SlateUser = SlateApp.GetUser(GetUserControllerId()))
+		{
+			//create a new mouse event
+			const bool bIsPrimaryUser = FSlateApplication::CursorUserIndex == SlateUser->GetUserIndex();
+			FPointerEvent MouseEvent(
+				SlateApp.CursorPointerIndex,
+				UpdatedPosition,
+				OldPosition,
+				bIsPrimaryUser ? SlateApp.GetPressedMouseButtons() : TSet<FKey>(),
+				EKeys::Invalid,
+				0,
+				bIsPrimaryUser ? SlateApp.GetModifierKeys() : FModifierKeysState()
+			);
+			//process the event
+			SlateApp.ProcessMouseMoveEvent(MouseEvent);
+		}
+	}
+}
+
+void FAnalogCursor::UpdateCursorPosition(FSlateApplication& SlateApp, TSharedRef<FSlateUser> SlateUser, const FVector2D& NewPosition, bool bForce)
+{
+	//grab the old position
+	const FVector2D OldPosition = SlateUser->GetCursorPosition();
 
+	//make sure we are actually moving
+	int32 NewIntPosX = NewPosition.X;
+	int32 NewIntPosY = NewPosition.Y;
+	int32 OldIntPosX = OldPosition.X;
+	int32 OldIntPosY = OldPosition.Y;
+	if (bForce || OldIntPosX != NewIntPosX || OldIntPosY != NewIntPosY)
+	{
+		//put the cursor in the correct spot
+		SlateUser->SetCursorPosition(NewIntPosX, NewIntPosY);
+	
+		// Since the cursor may have been locked and its location clamped, get the actual new position
+		const FVector2D UpdatedPosition = SlateUser->GetCursorPosition();
 		//create a new mouse event
+		const bool bIsPrimaryUser = FSlateApplication::CursorUserIndex == SlateUser->GetUserIndex();
 		FPointerEvent MouseEvent(
-			SlateApp.CursorPointerIndex,
+			SlateUser->GetUserIndex(),
 			UpdatedPosition,
 			OldPosition,
-			SlateApp.GetPressedMouseButtons(),
+			bIsPrimaryUser ? SlateApp.GetPressedMouseButtons() : TSet<FKey>(),
 			EKeys::Invalid,
 			0,
-			SlateApp.GetPlatformApplication()->GetModifierKeys()
-			);
+			bIsPrimaryUser ? SlateApp.GetModifierKeys() : FModifierKeysState()
+		);
 
-		
 		//process the event
 		SlateApp.ProcessMouseMoveEvent(MouseEvent);
 	}
diff --git a/Engine/Source/Runtime/Slate/Public/Framework/Application/AnalogCursor.h b/Engine/Source/Runtime/Slate/Public/Framework/Application/AnalogCursor.h
--- a/Engine/Source/Runtime/Slate/Public/Framework/Application/AnalogCursor.h
+++ b/Engine/Source/Runtime/Slate/Public/Framework/Application/AnalogCursor.h
@@ -7,6 +7,8 @@
 #include "Framework/Application/IInputProcessor.h"
 
 class FSlateApplication;
+class FSlateUser;
+struct FInputEvent;
 struct FAnalogInputEvent;
 struct FKeyEvent;
 struct FPointerEvent;
@@ -46,6 +48,8 @@ public:
 	virtual bool HandleAnalogInputEvent(FSlateApplication& SlateApp, const FAnalogInputEvent& InAnalogInputEvent) override;
 	virtual bool HandleMouseMoveEvent(FSlateApplication& SlateApp, const FPointerEvent& MouseEvent) override;
 
+	virtual int32 GetUserControllerId() const { return 0; };
+
 	void SetAcceleration(float NewAcceleration);
 	void SetMaxSpeed(float NewMaxSpeed);
 	void SetStickySlowdown(float NewStickySlowdown);
@@ -54,6 +58,8 @@ public:
 
 protected:
 
+	virtual bool IsRelevantInput(const FInputEvent& InputEvent) const;
+
 	/** Getter */
 	FORCEINLINE const FVector2D& GetAnalogValues( EAnalogStick Stick = EAnalogStick::Left ) const
 	{
@@ -64,8 +70,11 @@ protected:
 	void ClearAnalogValues();
 
 	/** Handles updating the cursor position and processing a Mouse Move Event */
+	UE_DEPRECATED(4.24, "FAnalogCursor now updates cursor position based on user, not the hardware cursor specifically.")
 	virtual void UpdateCursorPosition(FSlateApplication& SlateApp, TSharedRef<ICursor> Cursor, const FVector2D& NewPosition, bool bForce = false);
 
+	virtual void UpdateCursorPosition(FSlateApplication& SlateApp, TSharedRef<FSlateUser> SlateUser, const FVector2D& NewPosition, bool bForce = false);
+
 	/** Current speed of the cursor */
 	FVector2D CurrentSpeed;
 
