- Fixed the signature of FPointerEvent::GetCursorDelta() to return a const ref to the member instead of copy
- Added FSlateUser::IsWidgetInFocusPath() and added some comments explaining the three different functions there for checking focused-ness of a widget

diff --git a/Engine/Source/Runtime/Slate/Private/Framework/Application/SlateUser.cpp b/Engine/Source/Runtime/Slate/Private/Framework/Application/SlateUser.cpp
--- a/Engine/Source/Runtime/Slate/Private/Framework/Application/SlateUser.cpp
+++ b/Engine/Source/Runtime/Slate/Private/Framework/Application/SlateUser.cpp
@@ -130,6 +130,11 @@ bool FSlateUser::HasFocusedDescendants(TSharedRef<const SWidget> Widget) const
 	return WeakFocusPath.IsValid() && WeakFocusPath.GetLastWidget().Pin() != Widget && WeakFocusPath.ContainsWidget(Widget);
 }
 
+bool FSlateUser::IsWidgetInFocusPath(TSharedPtr<const SWidget> Widget) const
+{
+	return Widget && WeakFocusPath.IsValid() && WeakFocusPath.ContainsWidget(Widget.ToSharedRef());
+}
+
 bool FSlateUser::SetFocus(const TSharedRef<SWidget>& WidgetToFocus, EFocusCause ReasonFocusIsChanging)
 {
 	return FSlateApplication::Get().SetUserFocus(UserIndex, WidgetToFocus, ReasonFocusIsChanging);
diff --git a/Engine/Source/Runtime/Slate/Public/Framework/Application/SlateUser.h b/Engine/Source/Runtime/Slate/Public/Framework/Application/SlateUser.h
--- a/Engine/Source/Runtime/Slate/Public/Framework/Application/SlateUser.h
+++ b/Engine/Source/Runtime/Slate/Public/Framework/Application/SlateUser.h
@@ -45,12 +45,19 @@ public:
 	FORCEINLINE bool IsVirtualUser() const { return !Cursor.IsValid(); }
 
 	TSharedPtr<SWidget> GetFocusedWidget() const;
-	TOptional<EFocusCause> HasFocus(TSharedPtr<const SWidget> Widget) const;
 	bool ShouldShowFocus(TSharedPtr<const SWidget> Widget) const;
-	bool HasFocusedDescendants(TSharedRef<const SWidget> Widget) const;
 	bool SetFocus(const TSharedRef<SWidget>& WidgetToFocus, EFocusCause ReasonFocusIsChanging = EFocusCause::SetDirectly);
 	void ClearFocus(EFocusCause ReasonFocusIsChanging = EFocusCause::SetDirectly);
 
+	/** Returns the cause for which the provided widget was focused, or nothing if the given widget is not the current focus target. */
+	TOptional<EFocusCause> HasFocus(TSharedPtr<const SWidget> Widget) const;
+
+	/** Returns true if the given widget is in the focus path, but is not the focused widget itself. */
+	bool HasFocusedDescendants(TSharedRef<const SWidget> Widget) const;
+
+	/** Returns true if the given widget is anywhere in the focus path, including the focused widget itself. */
+	bool IsWidgetInFocusPath(TSharedPtr<const SWidget> Widget) const;
+	
 	bool HasAnyCapture() const;
 	bool HasCursorCapture() const;
 	bool HasCapture(uint32 PointerIndex) const;
diff --git a/Engine/Source/Runtime/SlateCore/Public/Input/Events.h b/Engine/Source/Runtime/SlateCore/Public/Input/Events.h
--- a/Engine/Source/Runtime/SlateCore/Public/Input/Events.h
+++ b/Engine/Source/Runtime/SlateCore/Public/Input/Events.h
@@ -816,7 +816,7 @@ public:
 	const FVector2D& GetLastScreenSpacePosition() const { return LastScreenSpacePosition; }
 
 	/** Returns the distance the mouse traveled since the last event was handled. */
-	FVector2D GetCursorDelta() const { return CursorDelta; }
+	const FVector2D& GetCursorDelta() const { return CursorDelta; }
 
 	/** Mouse buttons that are currently pressed */
 	bool IsMouseButtonDown( FKey MouseButton ) const { return PressedButtons.Contains( MouseButton ); }
