Fix crash that can occur when a menu is opened

diff --git a/Engine/Source/Runtime/Slate/Private/Framework/Application/MenuStack.cpp b/Engine/Source/Runtime/Slate/Private/Framework/Application/MenuStack.cpp
--- a/Engine/Source/Runtime/Slate/Private/Framework/Application/MenuStack.cpp
+++ b/Engine/Source/Runtime/Slate/Private/Framework/Application/MenuStack.cpp
@@ -267,7 +267,8 @@ TSharedRef<IMenu> FMenuStack::PushInternal(const TSharedPtr<IMenu>& InParentMenu
 
 	// Post-push stage
 	//   Updates the stack and content map member variables
-	PostPush(InParentMenu, OutMenu, ShouldThrottle);
+	const bool bInInsertAfterDismiss = ActiveMethod.GetPopupMethod() == EPopupMethod::CreateNewWindow;
+	PostPush(InParentMenu, OutMenu, ShouldThrottle, bInInsertAfterDismiss);
 
 	PendingNewMenu.Reset();
 
@@ -474,7 +475,7 @@ TSharedRef<FMenuBase> FMenuStack::PushPopup(TSharedPtr<IMenu> InParentMenu, cons
 	return Menu;
 }
 
-void FMenuStack::PostPush(TSharedPtr<IMenu> InParentMenu, TSharedRef<FMenuBase> InMenu, EShouldThrottle ShouldThrottle )
+void FMenuStack::PostPush(TSharedPtr<IMenu> InParentMenu, TSharedRef<FMenuBase> InMenu, EShouldThrottle ShouldThrottle, bool bInInsertAfterDismiss)
 {
 	// Determine at which index we insert this new menu in the stack
 	int32 InsertIndex = 0;
@@ -486,23 +487,39 @@ void FMenuStack::PostPush(TSharedPtr<IMenu> InParentMenu, TSharedRef<FMenuBase>
 		InsertIndex = ParentIndex + 1;
 	}
 
-	// Insert before dismissing others to stop the stack accidentally emptying briefly mid-push and reseting some state
-	Stack.Insert(InMenu, InsertIndex);
-	CachedContentMap.Add(InMenu->GetContent(), InMenu);
+	// Do original behavior of insert before dismiss
+	// Note: This will often crash because DismissFrom() may trigger FMenuStack::OnWindowActivated() then DismissAll() that empties Stack
+	int32 RemovingAtIndex = InsertIndex;
+	if (!bInInsertAfterDismiss)
+	{
+		Stack.Insert(InMenu, InsertIndex);
+		CachedContentMap.Add(InMenu->GetContent(), InMenu);
+		RemovingAtIndex = InsertIndex + 1;
+	}
 
 	// Dismiss menus after the insert point
-	if (Stack.Num() > InsertIndex + 1)
+	if (Stack.Num() > RemovingAtIndex)
 	{
-		DismissFrom(Stack[InsertIndex + 1]);
+		// Note: DismissFrom() may trigger FMenuStack::OnWindowActivated() then DismissAll() that empties Stack
+		DismissFrom(Stack[RemovingAtIndex]);
 
 		// tidy the stack data now (it will happen via callbacks from the dismissed menus but that might be delayed)
-		for (int32 StackIndex = Stack.Num() - 1; StackIndex > InsertIndex; --StackIndex)
+		for (int32 StackIndex = Stack.Num() - 1; StackIndex >= RemovingAtIndex; --StackIndex)
 		{
 			CachedContentMap.Remove(Stack[StackIndex]->GetContent());
 			Stack.RemoveAt(StackIndex);
 		}
 	}
 
+	// Note: DismissFrom() above may trigger FMenuStack::OnWindowActivated() then DismissAll() that empties Stack
+	// Insert menu after the dismiss when possible to avoid menu being deleted
+	if (bInInsertAfterDismiss)
+	{
+		check(InsertIndex == Stack.Num());
+		Stack.Add(InMenu);
+		CachedContentMap.Add(InMenu->GetContent(), InMenu);
+	}
+
 	// When a new menu is pushed, if we are not already in responsive mode for Slate UI, enter it now
 	// to ensure the menu is responsive in low FPS situations
 	if (!ThrottleHandle.IsValid() && ShouldThrottle == EShouldThrottle::Yes)
diff --git a/Engine/Source/Runtime/Slate/Public/Framework/Application/MenuStack.h b/Engine/Source/Runtime/Slate/Public/Framework/Application/MenuStack.h
--- a/Engine/Source/Runtime/Slate/Public/Framework/Application/MenuStack.h
+++ b/Engine/Source/Runtime/Slate/Public/Framework/Application/MenuStack.h
@@ -334,8 +334,9 @@ private:
 	 * @param InParentMenu         The parent menu for this menu
 	 * @param InMenu               A newly created menu
 	 * @param ShouldThrottle       Should pushing this menu enable throttling of the engine for a more responsive UI
+	 * @param bInInsertAfterDismiss	Avoid unwanted delete due to dismissing another menu causing all menus in stack to be dismissed
 	 */
-	void PostPush(TSharedPtr<IMenu> InParentMenu, TSharedRef<FMenuBase> InMenu, EShouldThrottle ShouldThrottle);
+	void PostPush(TSharedPtr<IMenu> InParentMenu, TSharedRef<FMenuBase> InMenu, EShouldThrottle ShouldThrottle, bool bInInsertAfterDismiss=false);
 
 	/** The popup method currently used by the whole stack. It can only use one at a time */
 	FPopupMethodReply ActiveMethod;
