PR #6374: Add ScrollWhenFocusChanges to scrollbox (Contributed by [REDACTED_PERSON])

diff --git a/Engine/Source/Runtime/Slate/Private/Widgets/Layout/SScrollBox.cpp b/Engine/Source/Runtime/Slate/Private/Widgets/Layout/SScrollBox.cpp
--- a/Engine/Source/Runtime/Slate/Private/Widgets/Layout/SScrollBox.cpp
+++ b/Engine/Source/Runtime/Slate/Private/Widgets/Layout/SScrollBox.cpp
@@ -131,6 +131,7 @@ void SScrollBox::Construct( const FArguments& InArgs )
 	WheelScrollMultiplier = InArgs._WheelScrollMultiplier;
 	NavigationScrollPadding = InArgs._NavigationScrollPadding;
 	NavigationDestination = InArgs._NavigationDestination;
+	ScrollWhenFocusChanges = InArgs._ScrollWhenFocusChanges;
 	bTouchPanningCapture = false;
 	bVolatilityAlwaysInvalidatesPrepass = true;
 
@@ -998,6 +999,17 @@ FNavigationReply SScrollBox::OnNavigation(const FGeometry& MyGeometry, const FNa
 	return SCompoundWidget::OnNavigation(MyGeometry, InNavigationEvent);
 }
 
+void SScrollBox::OnFocusChanging(const FWeakWidgetPath& PreviousFocusPath, const FWidgetPath& NewWidgetPath, const FFocusEvent& InFocusEvent)
+{
+	if (ScrollWhenFocusChanges != EScrollWhenFocusChanges::NoScroll)
+	{
+		if (NewWidgetPath.IsValid() && NewWidgetPath.ContainsWidget(SharedThis(this)))
+		{
+			ScrollDescendantIntoView(NewWidgetPath.GetLastWidget(), ScrollWhenFocusChanges == EScrollWhenFocusChanges::AnimatedScroll ? true : false, NavigationDestination, NavigationScrollPadding);
+		}
+	}
+}
+
 TSharedPtr<SWidget> SScrollBox::GetKeyboardFocusableWidget(TSharedPtr<SWidget> InWidget)
 {
 	if (InWidget->SupportsKeyboardFocus() && EVisibility::DoesVisibilityPassFilter(InWidget->GetVisibility(), EVisibility::Visible))
diff --git a/Engine/Source/Runtime/Slate/Public/Widgets/Layout/SScrollBox.h b/Engine/Source/Runtime/Slate/Public/Widgets/Layout/SScrollBox.h
--- a/Engine/Source/Runtime/Slate/Public/Widgets/Layout/SScrollBox.h
+++ b/Engine/Source/Runtime/Slate/Public/Widgets/Layout/SScrollBox.h
@@ -47,6 +47,20 @@ enum class EDescendantScrollDestination : uint8
 	Center,
 };
 
+/** Set behavior when user focus changes inside this scroll box */
+UENUM(BlueprintType)
+enum class EScrollWhenFocusChanges : uint8
+{
+	/** Don't automatically scroll, navigation or child widget will handle this */
+	NoScroll,
+
+	/** Instantly scroll using NavigationDestination rule */
+	InstantScroll,
+
+	/** Use animation to scroll using NavigationDestination rule */
+	AnimatedScroll,
+};
+
 /** SScrollBox can scroll through an arbitrary number of widgets. */
 class SLATE_API SScrollBox : public SCompoundWidget
 {
@@ -96,6 +110,7 @@ public:
 		, _WheelScrollMultiplier(1.f)
 		, _NavigationDestination(EDescendantScrollDestination::IntoView)
 		, _NavigationScrollPadding(0.0f)
+		, _ScrollWhenFocusChanges(EScrollWhenFocusChanges::NoScroll)
 		, _OnUserScrolled()
 		, _ConsumeMouseWheel(EConsumeMouseWheel::WhenScrollingPossible)
 		{
@@ -140,6 +155,8 @@ public:
 		 */
 		SLATE_ARGUMENT(float, NavigationScrollPadding);
 
+		SLATE_ARGUMENT(EScrollWhenFocusChanges, ScrollWhenFocusChanges);
+
 		/** Called when the button is clicked */
 		SLATE_EVENT(FOnUserScrolled, OnUserScrolled)
 
@@ -237,6 +254,7 @@ public:
 	virtual FReply OnTouchEnded(const FGeometry& MyGeometry, const FPointerEvent& InTouchEvent) override;
 	virtual void OnMouseCaptureLost(const FCaptureLostEvent& CaptureLostEvent) override;
 	virtual FNavigationReply OnNavigation(const FGeometry& MyGeometry, const FNavigationEvent& InNavigationEvent) override;
+	virtual void OnFocusChanging(const FWeakWidgetPath& PreviousFocusPath, const FWidgetPath& NewWidgetPath, const FFocusEvent& InFocusEvent) override;
 	// End of SWidget interface
 
 protected:
@@ -342,11 +360,12 @@ protected:
 	 */
 	float NavigationScrollPadding;
 
-	/**
-	 * 
-	 */
+	/** Sets where to scroll a widget to when using explicit navigation or if ScrollWhenFocusChanges is enabled */
 	EDescendantScrollDestination NavigationDestination;
 
+	/** Scroll behavior when user focus is given to a child widget */
+	EScrollWhenFocusChanges ScrollWhenFocusChanges;
+
 	/**	The current position of the software cursor */
 	FVector2D SoftwareCursorPosition;
 
diff --git a/Engine/Source/Runtime/UMG/Private/Components/ScrollBox.cpp b/Engine/Source/Runtime/UMG/Private/Components/ScrollBox.cpp
--- a/Engine/Source/Runtime/UMG/Private/Components/ScrollBox.cpp
+++ b/Engine/Source/Runtime/UMG/Private/Components/ScrollBox.cpp
@@ -25,6 +25,7 @@ UScrollBox::UScrollBox(const FObjectInitializer& ObjectInitializer)
 	, AllowOverscroll(true)
 	, NavigationDestination(EDescendantScrollDestination::IntoView)
 	, NavigationScrollPadding(0.0f)
+	, ScrollWhenFocusChanges(EScrollWhenFocusChanges::NoScroll)
 {
 	bIsVariable = false;
 
@@ -98,6 +99,7 @@ TSharedRef<SWidget> UScrollBox::RebuildWidget()
 		.ConsumeMouseWheel(ConsumeMouseWheel)
 		.NavigationDestination(NavigationDestination)
 		.NavigationScrollPadding(NavigationScrollPadding)
+		.ScrollWhenFocusChanges(ScrollWhenFocusChanges)
 		.AnimateWheelScrolling(bAnimateWheelScrolling)
 		.WheelScrollMultiplier(WheelScrollMultiplier)
 		.OnUserScrolled(BIND_UOBJECT_DELEGATE(FOnUserScrolled, SlateHandleUserScrolled));
diff --git a/Engine/Source/Runtime/UMG/Public/Components/ScrollBox.h b/Engine/Source/Runtime/UMG/Public/Components/ScrollBox.h
--- a/Engine/Source/Runtime/UMG/Public/Components/ScrollBox.h
+++ b/Engine/Source/Runtime/UMG/Public/Components/ScrollBox.h
@@ -72,7 +72,7 @@ public:
 	UPROPERTY(EditAnywhere, BlueprintReadOnly, Category = "Scroll")
 	bool bAnimateWheelScrolling = false;
 
-	/**  */
+	/** Sets where to scroll a widget to when using explicit navigation or if ScrollWhenFocusChanges is enabled */
 	UPROPERTY(EditAnywhere, BlueprintReadOnly, Category="Scroll")
 	EDescendantScrollDestination NavigationDestination;
 
@@ -82,6 +82,10 @@ public:
 	 */
 	UPROPERTY(EditAnywhere, BlueprintReadOnly, Category="Scroll")
 	float NavigationScrollPadding;
+
+	/** Scroll behavior when user focus is given to a child widget */
+	UPROPERTY(EditAnywhere, BlueprintReadOnly, Category="Scroll", meta=(DisplayName="Scroll When Focus Changes"))
+	EScrollWhenFocusChanges ScrollWhenFocusChanges;
 	
 	/** Option to disable right-click-drag scrolling */
 	UPROPERTY(EditAnywhere, BlueprintReadOnly, Category = "Scroll")
