Update Tool menus to use generic blacklist filter

diff --git a/Engine/Source/Developer/ToolMenus/Private/ToolMenus.cpp b/Engine/Source/Developer/ToolMenus/Private/ToolMenus.cpp
--- a/Engine/Source/Developer/ToolMenus/Private/ToolMenus.cpp
+++ b/Engine/Source/Developer/ToolMenus/Private/ToolMenus.cpp
@@ -597,15 +597,15 @@ void UToolMenus::ApplyCustomization(UToolMenu* GeneratedMenu)
 		}
 	}
 
-	// Hide items based on whitelist
-	if (CustomizedMenu.WhitelistEntries.Num() > 0)
+	// Hide items based on blacklist
+	if (CustomizedMenu.BlacklistFilter.HasFiltering())
 	{
 		for (int32 SectionIndex = 0; SectionIndex < NewSections.Num(); ++SectionIndex)
 		{
 			FToolMenuSection& Section = NewSections[SectionIndex];
 			for (int32 i = 0; i < Section.Blocks.Num(); ++i)
 			{
-				if (!CustomizedMenu.IsEntryWhitelisted(Section.Blocks[i].Name))
+				if (!CustomizedMenu.BlacklistFilter.PassesFilter(Section.Blocks[i].Name))
 				{
 					Section.Blocks.RemoveAt(i);
 					--i;
diff --git a/Engine/Source/Runtime/Slate/Private/Framework/MultiBox/ToolMenuCustomization.cpp b/Engine/Source/Runtime/Slate/Private/Framework/MultiBox/ToolMenuCustomization.cpp
--- a/Engine/Source/Runtime/Slate/Private/Framework/MultiBox/ToolMenuCustomization.cpp
+++ b/Engine/Source/Runtime/Slate/Private/Framework/MultiBox/ToolMenuCustomization.cpp
@@ -32,16 +32,6 @@ FCustomizedToolMenuEntry* FCustomizedToolMenu::AddEntry(const FName InEntryName)
 	return &Entries.FindOrAdd(InEntryName);
 }
 
-void FCustomizedToolMenu::AddWhitelistEntry(const FName InEntryName)
-{
-	WhitelistEntries.AddUnique(InEntryName);
-}
-
-bool FCustomizedToolMenu::IsEntryWhitelisted(const FName InEntryName) const
-{
-	return WhitelistEntries.Contains(InEntryName);
-}
-
 ECustomizedToolMenuVisibility FCustomizedToolMenu::GetEntryVisiblity(const FName InEntryName) const
 {
 	if (const FCustomizedToolMenuEntry* Found = FindEntry(InEntryName))
@@ -204,10 +194,7 @@ FCustomizedToolMenu FCustomizedToolMenuHierarchy::GenerateFlattened() const
 			}
 		}
 
-		for (const FName WhitelistEntry : Current->WhitelistEntries)
-		{
-			Result.AddWhitelistEntry(WhitelistEntry);
-		}
+		Result.BlacklistFilter.Append(Current->BlacklistFilter);
 	}
 
 	return Result;
diff --git a/Engine/Source/Runtime/Slate/Public/Framework/MultiBox/ToolMenuBase.h b/Engine/Source/Runtime/Slate/Public/Framework/MultiBox/ToolMenuBase.h
--- a/Engine/Source/Runtime/Slate/Public/Framework/MultiBox/ToolMenuBase.h
+++ b/Engine/Source/Runtime/Slate/Public/Framework/MultiBox/ToolMenuBase.h
@@ -3,6 +3,7 @@
 #pragma once
 
 #include "CoreMinimal.h"
+#include "Misc/NamedBlacklist.h"
 
 #include "ToolMenuBase.generated.h"
 
@@ -73,7 +74,7 @@ struct SLATE_API FCustomizedToolMenu
 	UPROPERTY()
 	TArray<FName> SectionOrder;
 
-	TArray<FName> WhitelistEntries;
+	FNamedBlacklist BlacklistFilter;
 
 	FCustomizedToolMenuEntry* FindEntry(const FName InEntryName);
 	const FCustomizedToolMenuEntry* FindEntry(const FName InEntryName) const;
@@ -81,8 +82,6 @@ struct SLATE_API FCustomizedToolMenu
 	ECustomizedToolMenuVisibility GetEntryVisiblity(const FName InSectionName) const;
 	bool IsEntryHidden(const FName InEntryName) const;
 	FName GetEntrySectionName(const FName InEntryName) const;
-	void AddWhitelistEntry(const FName InEntryName);
-	bool IsEntryWhitelisted(const FName InEntryName) const;
 
 	FCustomizedToolMenuSection* FindSection(const FName InSectionName);
 	const FCustomizedToolMenuSection* FindSection(const FName InSectionName) const;
