Font size multipliers will now apply to all parts of text
#rnx

diff --git a/Engine/Source/Runtime/Slate/Private/Framework/Text/RichTextLayoutMarshaller.cpp b/Engine/Source/Runtime/Slate/Private/Framework/Text/RichTextLayoutMarshaller.cpp
--- a/Engine/Source/Runtime/Slate/Private/Framework/Text/RichTextLayoutMarshaller.cpp
+++ b/Engine/Source/Runtime/Slate/Private/Framework/Text/RichTextLayoutMarshaller.cpp
@@ -85,6 +85,14 @@ void FRichTextLayoutMarshaller::GetText(FString& TargetString, const FTextLayout
 	Writer->Write(WriterLines, TargetString);
 }
 
+void FRichTextLayoutMarshaller::SetFontSizeMultiplier(const float NewFontSizeMultiplier)
+{
+	if (ensure(NewFontSizeMultiplier != 0.0f))
+	{
+		FontSizeMultiplier = NewFontSizeMultiplier;
+	}
+}
+
 FRichTextLayoutMarshaller::FRichTextLayoutMarshaller(TArray< TSharedRef< ITextDecorator > > InDecorators, const ISlateStyle* const InDecoratorStyleSet)
 	: Parser(FDefaultRichTextMarkupParser::GetStaticInstance())
 	, Writer(FDefaultRichTextMarkupWriter::Create())
@@ -170,7 +178,14 @@ void FRichTextLayoutMarshaller::AppendRunsForText(
 		ModelRange.EndIndex = InOutModelText->Len();
 
 		// Create run.
-		Run = FSlateTextRun::Create(RunInfo, InOutModelText, *TextBlockStyle, ModelRange);
+		TSharedPtr< FSlateTextRun > SlateTextRun = FSlateTextRun::Create(RunInfo, InOutModelText, *TextBlockStyle, ModelRange);
+
+		if (SlateTextRun)
+		{
+			// Apply the FontSizeMultiplier at the style use by the IRun
+			SlateTextRun->ApplyFontSizeMultiplierOnTextStyle(FontSizeMultiplier);
+		}
+		Run = SlateTextRun;
 
 		if (!TextBlockStyle->UnderlineBrush.GetResourceName().IsNone())
 		{
diff --git a/Engine/Source/Runtime/Slate/Private/Framework/Text/SlateTextRun.cpp b/Engine/Source/Runtime/Slate/Private/Framework/Text/SlateTextRun.cpp
--- a/Engine/Source/Runtime/Slate/Private/Framework/Text/SlateTextRun.cpp
+++ b/Engine/Source/Runtime/Slate/Private/Framework/Text/SlateTextRun.cpp
@@ -285,3 +285,11 @@ FSlateTextRun::FSlateTextRun( const FSlateTextRun& Run )
 {
 
 }
+
+void FSlateTextRun::ApplyFontSizeMultiplierOnTextStyle(float FontSizeMultiplier)
+{
+	if (FontSizeMultiplier != 0.0f)
+	{
+		Style.SetFontSize(Style.Font.Size * FontSizeMultiplier);
+	}
+}
diff --git a/Engine/Source/Runtime/Slate/Private/Widgets/Text/SRichTextBlock.cpp b/Engine/Source/Runtime/Slate/Private/Widgets/Text/SRichTextBlock.cpp
--- a/Engine/Source/Runtime/Slate/Private/Widgets/Text/SRichTextBlock.cpp
+++ b/Engine/Source/Runtime/Slate/Private/Widgets/Text/SRichTextBlock.cpp
@@ -194,4 +194,13 @@ bool SRichTextBlock::ComputeVolatility() const
 		|| MinDesiredWidth.IsBound();
 }
 
+void SRichTextBlock::SetFontSizeMultiplier(const float NewFontSizeMultiplier)
+{
+	if (Marshaller)
+	{
+		Marshaller->SetFontSizeMultiplier(NewFontSizeMultiplier);
+		Invalidate(EInvalidateWidget::Layout);
+	}
+}
+
 #endif //WITH_FANCY_TEXT
diff --git a/Engine/Source/Runtime/Slate/Public/Framework/Text/RichTextLayoutMarshaller.h b/Engine/Source/Runtime/Slate/Public/Framework/Text/RichTextLayoutMarshaller.h
--- a/Engine/Source/Runtime/Slate/Public/Framework/Text/RichTextLayoutMarshaller.h
+++ b/Engine/Source/Runtime/Slate/Public/Framework/Text/RichTextLayoutMarshaller.h
@@ -35,6 +35,9 @@ public:
 	virtual void SetText(const FString& SourceString, FTextLayout& TargetTextLayout) override;
 	virtual void GetText(FString& TargetString, const FTextLayout& SourceTextLayout) override;
 
+	/** Set the Multiplier applied on the font size of the text*/
+	void SetFontSizeMultiplier(const float NewFontSizeMultiplier);
+
 	/**
 	 * Append an inline decorator to this marshaller
 	 */
@@ -84,6 +87,9 @@ protected:
 	/** The style set used for looking up styles used by decorators */
 	const ISlateStyle* DecoratorStyleSet;
 
+	/** Multiplier applied on font */
+	float FontSizeMultiplier = 1.0f;
+
 };
 
 #endif //WITH_FANCY_TEXT
diff --git a/Engine/Source/Runtime/Slate/Public/Framework/Text/SlateTextRun.h b/Engine/Source/Runtime/Slate/Public/Framework/Text/SlateTextRun.h
--- a/Engine/Source/Runtime/Slate/Public/Framework/Text/SlateTextRun.h
+++ b/Engine/Source/Runtime/Slate/Public/Framework/Text/SlateTextRun.h
@@ -59,6 +59,8 @@ public:
 
 	virtual ERunAttributes GetRunAttributes() const override;
 
+	void ApplyFontSizeMultiplierOnTextStyle(float FontSizeMultiplier);
+
 protected:
 
 	FSlateTextRun( const FRunInfo& InRunInfo, const TSharedRef< const FString >& InText, const FTextBlockStyle& InStyle );
diff --git a/Engine/Source/Runtime/Slate/Public/Widgets/Text/SRichTextBlock.h b/Engine/Source/Runtime/Slate/Public/Widgets/Text/SRichTextBlock.h
--- a/Engine/Source/Runtime/Slate/Public/Widgets/Text/SRichTextBlock.h
+++ b/Engine/Source/Runtime/Slate/Public/Widgets/Text/SRichTextBlock.h
@@ -216,6 +216,10 @@ public:
 	 */
 	void Refresh();
 
+	/** set the multiplier affected at the font in the text */
+	void SetFontSizeMultiplier(const float NewFontSizeMultiplier);
+
+
 protected:
 	//~ SWidget interface
 	virtual FVector2D ComputeDesiredSize(float) const override;
