Fixed lifetime extension issue with pressed keys array

FPointerEvent keeps a reference to the data, so it can't be passed a temporary

#rnx

diff --git a/Engine/Source/Runtime/Slate/Private/Framework/Application/SlateUser.cpp b/Engine/Source/Runtime/Slate/Private/Framework/Application/SlateUser.cpp
--- a/Engine/Source/Runtime/Slate/Private/Framework/Application/SlateUser.cpp
+++ b/Engine/Source/Runtime/Slate/Private/Framework/Application/SlateUser.cpp
@@ -672,14 +672,15 @@ void FSlateUser::QueryCursor()
 			const bool bHasHardwareCursor = SlateApp.GetPlatformCursor() == Cursor;
 			const FVector2D CurrentCursorPosition = GetCursorPosition();
 			const FVector2D LastCursorPosition = GetPreviousCursorPosition();
-
+			
+			const TSet<FKey> EmptySet;
 			const FPointerEvent CursorEvent(
 				SlateApp.GetUserIndexForMouse(),
 				FSlateApplication::CursorPointerIndex,
 				CurrentCursorPosition,
 				LastCursorPosition,
 				CurrentCursorPosition - LastCursorPosition,
-				bHasHardwareCursor ? SlateApp.GetPressedMouseButtons() : TSet<FKey>(),
+				bHasHardwareCursor ? SlateApp.GetPressedMouseButtons() : EmptySet,
 				bHasHardwareCursor ? SlateApp.GetModifierKeys() : FModifierKeysState());
 
 			FWidgetPath WidgetsToQueryForCursor;
