If we are using dev tools make sure to depend on settings

diff --git a/Engine/Source/Runtime/Slate/Slate.Build.cs b/Engine/Source/Runtime/Slate/Slate.Build.cs
--- a/Engine/Source/Runtime/Slate/Slate.Build.cs
+++ b/Engine/Source/Runtime/Slate/Slate.Build.cs
@@ -67,5 +67,10 @@ public class Slate : ModuleRules
 				RuntimeDependencies.Add("$(ProjectDir)/Content/SlateDebug/...", StagedFileType.DebugNonUFS);
 			}
 		}
+
+		if (Target.bBuildDeveloperTools)
+        {
+            DynamicallyLoadedModuleNames.Add("Settings");
+        }
 	}
 }
