Changes in Slate Controller-User mapping:

- AnalogCursor: Renamed GetUserControllerId to GetOwnerUserIndex (either way, children are responsible for returning something other than 0)
- UWidget: Run ControllerId through SlateApplication::GetUserIndexForController before using it as a user index. In the default implementation of ISlateInputManager, this will always be the same number, but engine users can override the controllerid-slateId mapping
- SlateApplication: don't handle controller input if mapping doesn't return a value. The default mapping is unaffected because it simply returns the given ControllerId; this gives specialized mappings more control.
- SlateApplication: Pass more information to ISlateInputManager when mapping controller input to slate user. The default mapping's behavior is unaffected.

diff --git a/Engine/Source/Runtime/Slate/Private/Framework/Application/AnalogCursor.cpp b/Engine/Source/Runtime/Slate/Private/Framework/Application/AnalogCursor.cpp
--- a/Engine/Source/Runtime/Slate/Private/Framework/Application/AnalogCursor.cpp
+++ b/Engine/Source/Runtime/Slate/Private/Framework/Application/AnalogCursor.cpp
@@ -25,7 +25,7 @@ FAnalogCursor::FAnalogCursor()
 
 void FAnalogCursor::Tick(const float DeltaTime, FSlateApplication& SlateApp, TSharedRef<ICursor>)
 {
-	if (TSharedPtr<FSlateUser> SlateUser = SlateApp.GetUser(GetUserControllerId()))
+	if (TSharedPtr<FSlateUser> SlateUser = SlateApp.GetUser(GetOwnerUserIndex()))
 	{
 		const FVector2D OldPosition = SlateUser->GetCursorPosition();
 
@@ -274,7 +274,7 @@ void FAnalogCursor::SetMode(AnalogCursorMode::Type NewMode)
 
 bool FAnalogCursor::IsRelevantInput(const FInputEvent& InputEvent) const
 {
-	return GetUserControllerId() == InputEvent.GetUserIndex();
+	return GetOwnerUserIndex() == InputEvent.GetUserIndex();
 }
 
 void FAnalogCursor::ClearAnalogValues()
@@ -300,7 +300,7 @@ void FAnalogCursor::UpdateCursorPosition(FSlateApplication& SlateApp, TSharedRef
 
 		// Since the cursor may have been locked and its location clamped, get the actual new position
 		const FVector2D UpdatedPosition = Cursor->GetPosition();
-		if (TSharedPtr<FSlateUser> SlateUser = SlateApp.GetUser(GetUserControllerId()))
+		if (TSharedPtr<FSlateUser> SlateUser = SlateApp.GetUser(GetOwnerUserIndex()))
 		{
 			//create a new mouse event
 			const bool bIsPrimaryUser = FSlateApplication::CursorUserIndex == SlateUser->GetUserIndex();
diff --git a/Engine/Source/Runtime/Slate/Private/Framework/Application/SlateApplication.cpp b/Engine/Source/Runtime/Slate/Private/Framework/Application/SlateApplication.cpp
--- a/Engine/Source/Runtime/Slate/Private/Framework/Application/SlateApplication.cpp
+++ b/Engine/Source/Runtime/Slate/Private/Framework/Application/SlateApplication.cpp
@@ -2236,6 +2236,11 @@ int32 FSlateApplication::GetUserIndexForKeyboard() const
 	return InputManager->GetUserIndexForKeyboard();
 }
  
+TOptional<int32> FSlateApplication::GetUserIndexForController(int32 ControllerId, FKey InKey) const
+{
+	return InputManager->GetUserIndexForController(ControllerId, InKey);
+}
+
 int32 FSlateApplication::GetUserIndexForController(int32 ControllerId) const
 {
 	return InputManager->GetUserIndexForController(ControllerId);
@@ -5542,9 +5547,13 @@ bool FSlateApplication::OnControllerAnalog( FGamepadKeyNames::Type KeyName, int3
 	FKey Key(KeyName);
 	check(Key.IsValid());
 
-	int32 UserIndex = GetUserIndexForController(ControllerId);
-	
-	FAnalogInputEvent AnalogInputEvent(Key, PlatformApplication->GetModifierKeys(), UserIndex, false, 0, 0, AnalogValue);
+	TOptional<int32> UserIndex = GetUserIndexForController(ControllerId, Key);
+	if (!UserIndex.IsSet())
+	{
+		return false;
+	}
+
+	FAnalogInputEvent AnalogInputEvent(Key, PlatformApplication->GetModifierKeys(), UserIndex.GetValue(), false, 0, 0, AnalogValue);
 
 	return ProcessAnalogInputEvent(AnalogInputEvent);
 }
@@ -5554,9 +5563,13 @@ bool FSlateApplication::OnControllerButtonPressed(FGamepadKeyNames::Type KeyName
 	FKey Key(KeyName);
 	check(Key.IsValid());
 
-	int32 UserIndex = GetUserIndexForController(ControllerId);
+	TOptional<int32> UserIndex = GetUserIndexForController(ControllerId, Key);
+	if (!UserIndex.IsSet())
+	{
+		return false;
+	}
 
-	FKeyEvent KeyEvent(Key, PlatformApplication->GetModifierKeys(), UserIndex, IsRepeat, 0, 0);
+	FKeyEvent KeyEvent(Key, PlatformApplication->GetModifierKeys(), UserIndex.GetValue(), IsRepeat, 0, 0);
 
 	return ProcessKeyDownEvent(KeyEvent);
 }
@@ -5566,9 +5579,13 @@ bool FSlateApplication::OnControllerButtonReleased(FGamepadKeyNames::Type KeyNam
 	FKey Key(KeyName);
 	check(Key.IsValid());
 
-	int32 UserIndex = GetUserIndexForController(ControllerId);
+	TOptional<int32> UserIndex = GetUserIndexForController(ControllerId, Key);
+	if (!UserIndex.IsSet())
+	{
+		return false;
+	}
 
-	FKeyEvent KeyEvent(Key, PlatformApplication->GetModifierKeys(), UserIndex, IsRepeat, 0, 0);
+	FKeyEvent KeyEvent(Key, PlatformApplication->GetModifierKeys(), UserIndex.GetValue(), IsRepeat, 0, 0);
 
 	return ProcessKeyUpEvent(KeyEvent);
 }
diff --git a/Engine/Source/Runtime/Slate/Public/Framework/Application/AnalogCursor.h b/Engine/Source/Runtime/Slate/Public/Framework/Application/AnalogCursor.h
--- a/Engine/Source/Runtime/Slate/Public/Framework/Application/AnalogCursor.h
+++ b/Engine/Source/Runtime/Slate/Public/Framework/Application/AnalogCursor.h
@@ -48,7 +48,7 @@ public:
 	virtual bool HandleAnalogInputEvent(FSlateApplication& SlateApp, const FAnalogInputEvent& InAnalogInputEvent) override;
 	virtual bool HandleMouseMoveEvent(FSlateApplication& SlateApp, const FPointerEvent& MouseEvent) override;
 
-	virtual int32 GetUserControllerId() const { return 0; };
+	virtual int32 GetOwnerUserIndex() const { return 0; };
 
 	void SetAcceleration(float NewAcceleration);
 	void SetMaxSpeed(float NewMaxSpeed);
diff --git a/Engine/Source/Runtime/Slate/Public/Framework/Application/SlateApplication.h b/Engine/Source/Runtime/Slate/Public/Framework/Application/SlateApplication.h
--- a/Engine/Source/Runtime/Slate/Public/Framework/Application/SlateApplication.h
+++ b/Engine/Source/Runtime/Slate/Public/Framework/Application/SlateApplication.h
@@ -117,7 +117,8 @@ class SLATE_API ISlateInputManager
 public:
 	virtual int32 GetUserIndexForMouse() const = 0;
 	virtual int32 GetUserIndexForKeyboard() const = 0;
-	virtual int32 GetUserIndexForController(int32 ControllerId) const = 0;
+	virtual int32 GetUserIndexForController(int32 ControllerId) const { return ControllerId; }
+	virtual TOptional<int32> GetUserIndexForController(int32 ControllerId, FKey InKey) const = 0;
 };
 
 class SLATE_API FSlateDefaultInputMapping : public ISlateInputManager
@@ -126,6 +127,7 @@ public:
 	virtual int32 GetUserIndexForMouse() const override { return 0; }
 	virtual int32 GetUserIndexForKeyboard() const override { return 0; }
 	virtual int32 GetUserIndexForController(int32 ControllerId) const override { return ControllerId; }
+	virtual TOptional<int32> GetUserIndexForController(int32 ControllerId, FKey InKey) const override { return GetUserIndexForController(ControllerId); }
 };
 
 enum class ESlateTickType : uint8
@@ -1450,6 +1452,7 @@ public:
 
 	/** @return int user index that this controller is mapped to. */
 	int32 GetUserIndexForController(int32 ControllerId) const;
+	TOptional<int32> GetUserIndexForController(int32 ControllerId, FKey InKey) const;
 
 	/** Establishes the input mapping object used to map input sources to SlateUser indices */
 	void SetInputManager(TSharedRef<ISlateInputManager> InputManager);
diff --git a/Engine/Source/Runtime/UMG/Private/Components/Widget.cpp b/Engine/Source/Runtime/UMG/Private/Components/Widget.cpp
--- a/Engine/Source/Runtime/UMG/Private/Components/Widget.cpp
+++ b/Engine/Source/Runtime/UMG/Private/Components/Widget.cpp
@@ -486,8 +486,7 @@ bool UWidget::HasUserFocus(APlayerController* PlayerController) const
 
 		if ( ULocalPlayer* LocalPlayer = Context.GetLocalPlayer() )
 		{
-			// HACK: We use the controller Id as the local player index for focusing widgets in Slate.
-			int32 UserIndex = LocalPlayer->GetControllerId();
+			int32 UserIndex = FSlateApplication::Get().GetUserIndexForController(LocalPlayer->GetControllerId());
 
 			TOptional<EFocusCause> FocusCause = SafeWidget->HasUserFocus(UserIndex);
 			return FocusCause.IsSet();
@@ -534,8 +533,7 @@ bool UWidget::HasUserFocusedDescendants(APlayerController* PlayerController) con
 
 		if ( ULocalPlayer* LocalPlayer = Context.GetLocalPlayer() )
 		{
-			// HACK: We use the controller Id as the local player index for focusing widgets in Slate.
-			int32 UserIndex = LocalPlayer->GetControllerId();
+			int32 UserIndex = FSlateApplication::Get().GetUserIndexForController(LocalPlayer->GetControllerId());
 
 			return SafeWidget->HasUserFocusedDescendants(UserIndex);
 		}
@@ -586,10 +584,9 @@ void UWidget::SetUserFocus(APlayerController* PlayerController)
 
 		if ( ULocalPlayer* LocalPlayer = PlayerController->GetLocalPlayer() )
 		{
-			// HACK: We use the controller Id as the local player index for focusing widgets in Slate.
-			int32 UserIndex = LocalPlayer->GetControllerId();
+			int32 UserIndex = FSlateApplication::Get().GetUserIndexForController(LocalPlayer->GetControllerId());
 
-			if ( !FSlateApplication::Get().SetUserFocus(UserIndex, SafeWidget) )
+			if (UserIndex >= 0 && !FSlateApplication::Get().SetUserFocus(UserIndex, SafeWidget) )
 			{
 				LocalPlayer->GetSlateOperations().SetUserFocus(SafeWidget.ToSharedRef());
 			}
