Revert to old style iterator for DrawWindowAndChildren, preventing potential dereference of stale pointer when slate windows are destroyed at odd times
#rnx

diff --git a/Engine/Source/Runtime/Slate/Private/Framework/Application/SlateApplication.cpp b/Engine/Source/Runtime/Slate/Private/Framework/Application/SlateApplication.cpp
--- a/Engine/Source/Runtime/Slate/Private/Framework/Application/SlateApplication.cpp
+++ b/Engine/Source/Runtime/Slate/Private/Framework/Application/SlateApplication.cpp
@@ -1228,8 +1228,13 @@ void FSlateApplication::PrivateDrawWindows( TSharedPtr<SWindow> DrawOnlyThisWind
 		else
 		{
 			// Draw all windows
-			for(TSharedRef<SWindow>& CurrentWindow : SlateWindows )
+			// Use of an old-style iterator is intentional here, as SlateWindows 
+			// array may be mutated by user logic in draw calls. The iterator 
+			// prevents us from reading off the end and only keeps an index 
+			// internally:
+			for( TArray< TSharedRef<SWindow> >::TConstIterator CurrentWindowIt( SlateWindows ); CurrentWindowIt; ++CurrentWindowIt )
 			{
+				TSharedRef<SWindow> CurrentWindow = *CurrentWindowIt;
 				// Only draw visible windows or in off-screen rendering mode
 				if (bRenderOffScreen || CurrentWindow->IsVisible() )
 				{
