Enabled per-user hit testing in Slate to avoid cases where two users hovered over incompabile UI (UI designed for a specific user) can still interact with it.

diff --git a/Engine/Source/Editor/UMGEditor/Private/Designer/SDesignerView.cpp b/Engine/Source/Editor/UMGEditor/Private/Designer/SDesignerView.cpp
--- a/Engine/Source/Editor/UMGEditor/Private/Designer/SDesignerView.cpp
+++ b/Engine/Source/Editor/UMGEditor/Private/Designer/SDesignerView.cpp
@@ -1660,7 +1660,7 @@ bool SDesignerView::FindWidgetUnderCursor(const FGeometry& MyGeometry, const FPo
 	//@TODO UMG Make it so you can request dropable widgets only, to find the first parentable.
 
 	// Query the hit test grid we create for the design surface, and determine what widgets we hit.
-	TArray<FWidgetAndPointer> BubblePath = DesignerHittestGrid->GetBubblePath(MouseEvent.GetScreenSpacePosition(), 0.0f, true);
+	TArray<FWidgetAndPointer> BubblePath = DesignerHittestGrid->GetBubblePath(MouseEvent.GetScreenSpacePosition(), 0.0f, true, INDEX_NONE);
 
 	HitResult.Widget = FWidgetReference();
 	HitResult.NamedSlot = NAME_None;
diff --git a/Engine/Source/Runtime/Engine/Private/Slate/SGameLayerManager.cpp b/Engine/Source/Runtime/Engine/Private/Slate/SGameLayerManager.cpp
--- a/Engine/Source/Runtime/Engine/Private/Slate/SGameLayerManager.cpp
+++ b/Engine/Source/Runtime/Engine/Private/Slate/SGameLayerManager.cpp
@@ -17,9 +17,24 @@
 #include "Widgets/Layout/SWindowTitleBarArea.h"
 #include "DebugCanvas.h"
 #include "Types/InvisibleToWidgetReflectorMetaData.h"
+#include "Framework/Application/SlateApplication.h"
+
 
 /* SGameLayerManager interface
  *****************************************************************************/
+static void HandlePerUserHitTestingToggled(IConsoleVariable* CVar)
+{
+	FSlateApplication::Get().InvalidateAllWidgets(false);
+}
+
+static int32 bEnablePerUserHitTesting = true;
+
+static FAutoConsoleVariableRef CVarEnablePerUserHitTesting(
+	TEXT("Slate.AllowPerUserHitTesting"),
+	bEnablePerUserHitTesting,
+	TEXT("Toggles between widgets mapping to a user id and requring a matching user id from an input event or allowing all users to interact with widget"),
+	FConsoleVariableDelegate::CreateStatic(&HandlePerUserHitTestingToggled)
+);
 
 SGameLayerManager::SGameLayerManager()
 :	DefaultWindowTitleBarHeight(64.0f)
@@ -33,6 +48,12 @@ void SGameLayerManager::Construct(const SGameLayerManager::FArguments& InArgs)
 {
 	SceneViewport = InArgs._SceneViewport;
 
+	// In PIE we should default to per user hit testing being off because developers will need the mouse and keyboard to work for all players
+	if (GIsEditor)
+	{
+		bEnablePerUserHitTesting = false;
+	}
+
 	TSharedRef<SDPIScaler> DPIScaler =
 		SNew(SDPIScaler)
 		.DPIScale(this, &SGameLayerManager::GetGameViewportDPIScale)
@@ -357,6 +378,43 @@ void SGameLayerManager::UpdateLayout()
 	}
 }
 
+class SPlayerLayer : public SOverlay
+{
+	SLATE_BEGIN_ARGS(SPlayerLayer)
+	{
+		_Visibility = EVisibility::SelfHitTestInvisible;
+	}
+	SLATE_END_ARGS()
+
+public:
+	void Construct(const FArguments& InArgs, ULocalPlayer* InOwningPlayer)
+	{
+		OwningPlayer = InOwningPlayer;
+
+		SOverlay::Construct(
+			SOverlay::FArguments()
+			.Clipping(EWidgetClipping::ClipToBoundsAlways));
+	}
+
+	virtual int32 OnPaint(const FPaintArgs& Args, const FGeometry& AllottedGeometry, const FSlateRect& MyCullingRect, FSlateWindowElementList& OutDrawElements, int32 LayerId, const FWidgetStyle& InWidgetStyle, bool bParentEnabled) const override
+	{
+		const int32 UserIndex = bEnablePerUserHitTesting && OwningPlayer.IsValid() ? FSlateApplication::Get().GetUserIndexForController(OwningPlayer->GetControllerId()) : INDEX_NONE;
+
+		// Set user index for all widgets beneath this layer to the index of the player that owns this layer
+		const int32 OldUserIndex = Args.GetHittestGrid().GetUserIndex();
+		Args.GetHittestGrid().SetUserIndex(UserIndex);
+
+		const int32 OutgoingLayer = SOverlay::OnPaint(Args, AllottedGeometry, MyCullingRect, OutDrawElements, LayerId, InWidgetStyle, bParentEnabled);
+
+		// Restore whatever index was set before
+		Args.GetHittestGrid().SetUserIndex(OldUserIndex);
+
+		return OutgoingLayer;
+	}
+private:
+	TWeakObjectPtr<ULocalPlayer> OwningPlayer;
+};
+
 TSharedPtr<SGameLayerManager::FPlayerLayer> SGameLayerManager::FindOrCreatePlayerLayer(ULocalPlayer* LocalPlayer)
 {
 	TSharedPtr<FPlayerLayer>* PlayerLayerPtr = PlayerLayers.Find(LocalPlayer);
@@ -375,9 +433,8 @@ TSharedPtr<SGameLayerManager::FPlayerLayer> SGameLayerManager::FindOrCreatePlaye
 		TSharedPtr<FPlayerLayer> NewLayer = MakeShareable(new FPlayerLayer());
 
 		// Create a new overlay widget to house any widgets we want to display for the player.
-		NewLayer->Widget = SNew(SOverlay)
-			.AddMetaData(StopNavigation)
-			.Clipping(EWidgetClipping::ClipToBoundsAlways);
+		NewLayer->Widget = SNew(SPlayerLayer, LocalPlayer)
+			.AddMetaData(StopNavigation);
 		
 		// Add the overlay to the player canvas, which we'll update every frame to match
 		// the dimensions of the player's split screen rect.
diff --git a/Engine/Source/Runtime/Slate/Private/Framework/Application/AnalogCursor.cpp b/Engine/Source/Runtime/Slate/Private/Framework/Application/AnalogCursor.cpp
--- a/Engine/Source/Runtime/Slate/Private/Framework/Application/AnalogCursor.cpp
+++ b/Engine/Source/Runtime/Slate/Private/Framework/Application/AnalogCursor.cpp
@@ -44,7 +44,7 @@ void FAnalogCursor::Tick(const float DeltaTime, FSlateApplication& SlateApp, TSh
 
 
 		// Check if there is a sticky widget beneath the cursor
-		FWidgetPath WidgetPath = SlateApp.LocateWindowUnderMouse(OldPosition, SlateApp.GetInteractiveTopLevelWindows());
+		FWidgetPath WidgetPath = SlateApp.LocateWindowUnderMouse(OldPosition, SlateApp.GetInteractiveTopLevelWindows(), false, SlateUser->GetUserIndex());
 		if (WidgetPath.IsValid())
 		{
 			const FArrangedChildren::FArrangedWidgetArray& AllArrangedWidgets = WidgetPath.Widgets.GetInternalArray();
diff --git a/Engine/Source/Runtime/Slate/Private/Framework/Application/SlateApplication.cpp b/Engine/Source/Runtime/Slate/Private/Framework/Application/SlateApplication.cpp
--- a/Engine/Source/Runtime/Slate/Private/Framework/Application/SlateApplication.cpp
+++ b/Engine/Source/Runtime/Slate/Private/Framework/Application/SlateApplication.cpp
@@ -907,7 +907,7 @@ void FSlateApplication::SetCursorPos(const FVector2D& MouseCoordinate)
 	GetCursorUser()->SetCursorPosition(MouseCoordinate);
 }
 
-FWidgetPath FSlateApplication::LocateWindowUnderMouse( FVector2D ScreenspaceMouseCoordinate, const TArray< TSharedRef< SWindow > >& Windows, bool bIgnoreEnabledStatus )
+FWidgetPath FSlateApplication::LocateWindowUnderMouse( FVector2D ScreenspaceMouseCoordinate, const TArray< TSharedRef< SWindow > >& Windows, bool bIgnoreEnabledStatus, int32 UserIndex)
 {
 	// First, give the OS a chance to tell us which window to use, in case a child window is not guaranteed to stay on top of its parent window
 	TSharedPtr<FGenericWindow> NativeWindowUnderMouse = PlatformApplication->GetWindowUnderCursor();
@@ -916,7 +916,7 @@ FWidgetPath FSlateApplication::LocateWindowUnderMouse( FVector2D ScreenspaceMous
 		TSharedPtr<SWindow> Window = FSlateWindowHelper::FindWindowByPlatformWindow(Windows, NativeWindowUnderMouse.ToSharedRef());
 		if (Window.IsValid())
 		{
-			return LocateWidgetInWindow(ScreenspaceMouseCoordinate, Window.ToSharedRef(), bIgnoreEnabledStatus);
+			return LocateWidgetInWindow(ScreenspaceMouseCoordinate, Window.ToSharedRef(), bIgnoreEnabledStatus, UserIndex);
 		}
 	}
 
@@ -932,7 +932,7 @@ FWidgetPath FSlateApplication::LocateWindowUnderMouse( FVector2D ScreenspaceMous
 		}
 		
 		// Hittest the window's children first.
-		FWidgetPath ResultingPath = LocateWindowUnderMouse(ScreenspaceMouseCoordinate, Window->GetChildWindows(), bIgnoreEnabledStatus);
+		FWidgetPath ResultingPath = LocateWindowUnderMouse(ScreenspaceMouseCoordinate, Window->GetChildWindows(), bIgnoreEnabledStatus, UserIndex);
 		if (ResultingPath.IsValid())
 		{
 			return ResultingPath;
@@ -944,7 +944,7 @@ FWidgetPath FSlateApplication::LocateWindowUnderMouse( FVector2D ScreenspaceMous
 
 		if (!bPrevWindowWasModal)
 		{
-			FWidgetPath PathToLocatedWidget = LocateWidgetInWindow(ScreenspaceMouseCoordinate, Window, bIgnoreEnabledStatus);
+			FWidgetPath PathToLocatedWidget = LocateWidgetInWindow(ScreenspaceMouseCoordinate, Window, bIgnoreEnabledStatus, UserIndex);
 			if (PathToLocatedWidget.IsValid())
 			{
 				return PathToLocatedWidget;
@@ -1577,12 +1577,12 @@ void FSlateApplication::ThrottleApplicationBasedOnMouseMovement()
 	}
 }
 
-FWidgetPath FSlateApplication::LocateWidgetInWindow(FVector2D ScreenspaceMouseCoordinate, const TSharedRef<SWindow>& Window, bool bIgnoreEnabledStatus) const
+FWidgetPath FSlateApplication::LocateWidgetInWindow(FVector2D ScreenspaceMouseCoordinate, const TSharedRef<SWindow>& Window, bool bIgnoreEnabledStatus, int32 UserIndex) const
 {
 	const bool bAcceptsInput = Window->IsVisible() && (Window->AcceptsInput() || IsWindowHousingInteractiveTooltip(Window));
 	if (bAcceptsInput && Window->IsScreenspaceMouseWithin(ScreenspaceMouseCoordinate))
 	{
-		TArray<FWidgetAndPointer> WidgetsAndCursors = Window->GetHittestGrid().GetBubblePath(ScreenspaceMouseCoordinate, GetCursorRadius(), bIgnoreEnabledStatus);
+		TArray<FWidgetAndPointer> WidgetsAndCursors = Window->GetHittestGrid().GetBubblePath(ScreenspaceMouseCoordinate, GetCursorRadius(), bIgnoreEnabledStatus, UserIndex);
 		return FWidgetPath(MoveTemp(WidgetsAndCursors));
 	}
 	else
@@ -4462,7 +4462,7 @@ bool FSlateApplication::ProcessMouseButtonDownEvent( const TSharedPtr< FGenericW
 		}
 		else
 		{
-			FWidgetPath WidgetsUnderCursor = LocateWindowUnderMouse( MouseEvent.GetScreenSpacePosition(), GetInteractiveTopLevelWindows() );
+			FWidgetPath WidgetsUnderCursor = LocateWindowUnderMouse( MouseEvent.GetScreenSpacePosition(), GetInteractiveTopLevelWindows(), false, SlateUser->GetUserIndex());
 
 			PopupSupport.SendNotifications( WidgetsUnderCursor );
 
@@ -4668,7 +4668,7 @@ FReply FSlateApplication::RoutePointerUpEvent(const FWidgetPath& WidgetsUnderPoi
 	{
 		if (!LocalWidgetsUnderPointer.IsValid())
 		{
-			LocalWidgetsUnderPointer = LocateWindowUnderMouse(PointerEvent.GetScreenSpacePosition(), GetInteractiveTopLevelWindows());
+			LocalWidgetsUnderPointer = LocateWindowUnderMouse(PointerEvent.GetScreenSpacePosition(), GetInteractiveTopLevelWindows(), SlateUser->GetUserIndex());
 		}
 		
 		// Switch worlds widgets in the current path
@@ -5065,7 +5065,7 @@ bool FSlateApplication::ProcessMouseButtonDoubleClickEvent( const TSharedPtr< FG
 		return ProcessMouseButtonDownEvent(PlatformWindow, InMouseEvent);
 	}
 	
-	FWidgetPath WidgetsUnderCursor = LocateWindowUnderMouse(InMouseEvent.GetScreenSpacePosition(), GetInteractiveTopLevelWindows());
+	FWidgetPath WidgetsUnderCursor = LocateWindowUnderMouse(InMouseEvent.GetScreenSpacePosition(), GetInteractiveTopLevelWindows(), false, InMouseEvent.GetUserIndex());
 
 	FReply Reply = RoutePointerDoubleClickEvent( WidgetsUnderCursor, InMouseEvent );
 
@@ -5225,7 +5225,7 @@ bool FSlateApplication::ProcessMouseWheelOrGestureEvent( const FPointerEvent& In
 	
 	// NOTE: We intentionally don't reset LastUserInteractionTimeForThrottling here so that the UI can be responsive while scrolling
 
-	FWidgetPath EventPath = LocateWindowUnderMouse(InWheelEvent.GetScreenSpacePosition(), GetInteractiveTopLevelWindows());
+	FWidgetPath EventPath = LocateWindowUnderMouse(InWheelEvent.GetScreenSpacePosition(), GetInteractiveTopLevelWindows(), false, InWheelEvent.GetUserIndex());
 
 	return RouteMouseWheelOrGestureEvent(EventPath, InWheelEvent, InGestureEvent).IsEventHandled();
 }
@@ -5377,7 +5377,7 @@ bool FSlateApplication::ProcessMouseMoveEvent( const FPointerEvent& MouseEvent,
 	const bool bOverSlateWindow = !bIsSynthetic || PlatformApplication->IsCursorDirectlyOverSlateWindow();
 	
 	FWidgetPath WidgetsUnderCursor = bOverSlateWindow
-		? LocateWindowUnderMouse(MouseEvent.GetScreenSpacePosition(), GetInteractiveTopLevelWindows())
+		? LocateWindowUnderMouse(MouseEvent.GetScreenSpacePosition(), GetInteractiveTopLevelWindows(), false, MouseEvent.GetUserIndex())
 		: FWidgetPath();
 
 	bool bResult;
@@ -5403,7 +5403,7 @@ void FSlateApplication::NavigateToWidget(const uint32 UserIndex, const TSharedPt
 		FWidgetPath NavigationSourceWP;
 		if (NavigationSource == ENavigationSource::WidgetUnderCursor)
 		{
-			NavigationSourceWP = LocateWindowUnderMouse(GetOrCreateUser(UserIndex)->GetCursorPosition(), GetInteractiveTopLevelWindows());
+			NavigationSourceWP = LocateWindowUnderMouse(GetOrCreateUser(UserIndex)->GetCursorPosition(), GetInteractiveTopLevelWindows(), false, UserIndex);
 		}
 		else
 		{
@@ -5489,7 +5489,7 @@ bool FSlateApplication::AttemptNavigation(const FWidgetPath& NavigationSource, c
 			// Switch worlds for widgets in the current path 
 			FScopedSwitchWorldHack SwitchWorld(NavigationSource);
 
-			DestinationWidget = NavigationSource.GetDeepestWindow()->GetHittestGrid().FindNextFocusableWidget(FocusedArrangedWidget, NavigationType, NavigationReply, BoundaryWidget);
+			DestinationWidget = NavigationSource.GetDeepestWindow()->GetHittestGrid().FindNextFocusableWidget(FocusedArrangedWidget, NavigationType, NavigationReply, BoundaryWidget, NavigationEvent.GetUserIndex());
 
 #if WITH_SLATE_DEBUGGING
 			NavigationMethod = ESlateDebuggingNavigationMethod::HitTestGrid;
@@ -6321,7 +6321,7 @@ bool FSlateApplication::ProcessDragEnterEvent( TSharedRef<SWindow> WindowEntered
 {
 	SetLastUserInteractionTime(this->GetCurrentTime());
 	
-	FWidgetPath WidgetsUnderCursor = LocateWindowUnderMouse( DragDropEvent.GetScreenSpacePosition(), GetInteractiveTopLevelWindows() );
+	FWidgetPath WidgetsUnderCursor = LocateWindowUnderMouse( DragDropEvent.GetScreenSpacePosition(), GetInteractiveTopLevelWindows(), false, DragDropEvent.GetUserIndex());
 
 	// Switch worlds for widgets in the current path
 	FScopedSwitchWorldHack SwitchWorld( WidgetsUnderCursor );
@@ -6517,7 +6517,7 @@ void FSlateApplication::NavigateFromWidgetUnderCursor(const uint32 InUserIndex,
 {
 	if (InNavigationType != EUINavigation::Invalid)
 	{
-		FWidgetPath PathToLocatedWidget = LocateWidgetInWindow(GetOrCreateUser(InUserIndex)->GetCursorPosition(), InWindow, false);
+		FWidgetPath PathToLocatedWidget = LocateWidgetInWindow(GetOrCreateUser(InUserIndex)->GetCursorPosition(), InWindow, false, InUserIndex);
 		if (PathToLocatedWidget.IsValid())
 		{
 			TSharedPtr<SWidget> WidgetToNavFrom = PathToLocatedWidget.Widgets.Last().Widget;
diff --git a/Engine/Source/Runtime/Slate/Private/Framework/Application/SlateUser.cpp b/Engine/Source/Runtime/Slate/Private/Framework/Application/SlateUser.cpp
--- a/Engine/Source/Runtime/Slate/Private/Framework/Application/SlateUser.cpp
+++ b/Engine/Source/Runtime/Slate/Private/Framework/Application/SlateUser.cpp
@@ -702,7 +702,7 @@ void FSlateUser::QueryCursor()
 			}
 			else
 			{
-				WidgetsToQueryForCursor = SlateApp.LocateWindowUnderMouse(CurrentCursorPosition, SlateApp.GetInteractiveTopLevelWindows());
+				WidgetsToQueryForCursor = SlateApp.LocateWindowUnderMouse(CurrentCursorPosition, SlateApp.GetInteractiveTopLevelWindows(), false, UserIndex);
 			}
 
 			if (WidgetsToQueryForCursor.IsValid())
@@ -1078,7 +1078,7 @@ void FSlateUser::UpdateTooltip(const FMenuStack& MenuStack, bool bCanSpawnNewToo
 	if (bCheckForTooltipChanges)
 	{
 		// We're gonna check each widget under the cursor (including disabled widgets) until we find one with a tooltip to show
-		FWidgetPath WidgetsUnderCursor = SlateApp.LocateWindowUnderMouse(GetCursorPosition(), SlateApp.GetInteractiveTopLevelWindows(), /*bIgnoreEnabledStatus =*/true);
+		FWidgetPath WidgetsUnderCursor = SlateApp.LocateWindowUnderMouse(GetCursorPosition(), SlateApp.GetInteractiveTopLevelWindows(), /*bIgnoreEnabledStatus =*/true, UserIndex);
 		if (WidgetsUnderCursor.IsValid() && WidgetsUnderCursor.GetWindow() != TooltipWindowPtr.Pin())
 		{
 			WidgetsToQueryForTooltip = WidgetsUnderCursor;
diff --git a/Engine/Source/Runtime/Slate/Private/Widgets/SInvalidationPanel.cpp b/Engine/Source/Runtime/Slate/Private/Widgets/SInvalidationPanel.cpp
--- a/Engine/Source/Runtime/Slate/Private/Widgets/SInvalidationPanel.cpp
+++ b/Engine/Source/Runtime/Slate/Private/Widgets/SInvalidationPanel.cpp
@@ -194,6 +194,8 @@ int32 SInvalidationPanel::OnPaint( const FPaintArgs& Args, const FGeometry& Allo
 
 		FPaintArgs NewArgs = Args.WithNewHitTestGrid(HittestGrid);
 
+		// Copy the current user index into the new grid since nested hit test grids should inherit their parents user id
+		NewArgs.GetHittestGrid().SetUserIndex(Args.RootGrid.GetUserIndex());
 		check(!GSlateEnableGlobalInvalidation);
 
 		const bool bRequiresRecache = UpdateCachePrequisites(OutDrawElements, AllottedGeometry, MyCullingRect, LayerId);
diff --git a/Engine/Source/Runtime/Slate/Public/Framework/Application/SlateApplication.h b/Engine/Source/Runtime/Slate/Public/Framework/Application/SlateApplication.h
--- a/Engine/Source/Runtime/Slate/Public/Framework/Application/SlateApplication.h
+++ b/Engine/Source/Runtime/Slate/Public/Framework/Application/SlateApplication.h
@@ -968,7 +968,7 @@ protected:
 	/** Engages or disengages application throttling based on user behavior */
 	void ThrottleApplicationBasedOnMouseMovement();
 
-	virtual FWidgetPath LocateWidgetInWindow(FVector2D ScreenspaceMouseCoordinate, const TSharedRef<SWindow>& Window, bool bIgnoreEnabledStatus) const override;
+	virtual FWidgetPath LocateWidgetInWindow(FVector2D ScreenspaceMouseCoordinate, const TSharedRef<SWindow>& Window, bool bIgnoreEnabledStatus, int32 UserIndex) const override;
 
 	/**
 	 * Sets up any values that need to be based on the physical dimensions of the device.  
@@ -1320,7 +1320,7 @@ public:
 	virtual bool HasFocusedDescendants( const TSharedRef<const SWidget>& Widget ) const override;
 	virtual bool HasUserFocusedDescendants(const TSharedRef< const SWidget >& Widget, int32 UserIndex) const override;
 	virtual bool IsExternalUIOpened() override;
-	virtual FWidgetPath LocateWindowUnderMouse( FVector2D ScreenspaceMouseCoordinate, const TArray<TSharedRef<SWindow>>& Windows, bool bIgnoreEnabledStatus = false ) override;
+	virtual FWidgetPath LocateWindowUnderMouse( FVector2D ScreenspaceMouseCoordinate, const TArray<TSharedRef<SWindow>>& Windows, bool bIgnoreEnabledStatus = false, int32 UserIndex = INDEX_NONE) override;
 	virtual bool IsWindowHousingInteractiveTooltip(const TSharedRef<const SWindow>& WindowToTest) const override;
 	virtual TSharedRef<SImage> MakeImage( const TAttribute<const FSlateBrush*>& Image, const TAttribute<FSlateColor>& Color, const TAttribute<EVisibility>& Visibility ) const override;
 	virtual TSharedRef<SWidget> MakeWindowTitleBar( const TSharedRef<SWindow>& Window, const TSharedPtr<SWidget>& CenterContent, EHorizontalAlignment CenterContentAlignment, TSharedPtr<IWindowTitleBar>& OutTitleBar ) const override;
diff --git a/Engine/Source/Runtime/SlateCore/Private/Application/SlateApplicationBase.cpp b/Engine/Source/Runtime/SlateCore/Private/Application/SlateApplicationBase.cpp
--- a/Engine/Source/Runtime/SlateCore/Private/Application/SlateApplicationBase.cpp
+++ b/Engine/Source/Runtime/SlateCore/Private/Application/SlateApplicationBase.cpp
@@ -19,9 +19,9 @@ TSharedPtr<GenericApplication> FSlateApplicationBase::PlatformApplication = null
 const uint32 FSlateApplicationBase::CursorPointerIndex = ETouchIndex::CursorPointerIndex;
 const uint32 FSlateApplicationBase::CursorUserIndex = 0;
 
-FWidgetPath FHitTesting::LocateWidgetInWindow(FVector2D ScreenspaceMouseCoordinate, const TSharedRef<SWindow>& Window, bool bIgnoreEnabledStatus) const
+FWidgetPath FHitTesting::LocateWidgetInWindow(FVector2D ScreenspaceMouseCoordinate, const TSharedRef<SWindow>& Window, bool bIgnoreEnabledStatus, int32 UserIndex) const
 {
-	return SlateApp->LocateWidgetInWindow(ScreenspaceMouseCoordinate, Window, bIgnoreEnabledStatus);
+	return SlateApp->LocateWidgetInWindow(ScreenspaceMouseCoordinate, Window, bIgnoreEnabledStatus, UserIndex);
 }
 
 
diff --git a/Engine/Source/Runtime/SlateCore/Private/FastUpdate/WidgetProxy.cpp b/Engine/Source/Runtime/SlateCore/Private/FastUpdate/WidgetProxy.cpp
--- a/Engine/Source/Runtime/SlateCore/Private/FastUpdate/WidgetProxy.cpp
+++ b/Engine/Source/Runtime/SlateCore/Private/FastUpdate/WidgetProxy.cpp
@@ -6,6 +6,7 @@
 #include "Widgets/SWindow.h"
 #include "ProfilingDebugging/CsvProfiler.h"
 #include "Types/ReflectionMetadata.h"
+#include "Input/HittestGrid.h"
 
 const FSlateWidgetPersistentState FSlateWidgetPersistentState::NoState;
 
@@ -174,7 +175,10 @@ int32 FWidgetProxy::Repaint(const FPaintArgs& PaintArgs, int32 MyIndex, FSlateWi
 	{
 		OutDrawElements.GetClippingManager().PushClippingState(MyState.InitialClipState.GetValue());
 	}
+	
+	const int32 PrevUserIndex = PaintArgs.GetHittestGrid().GetUserIndex();
 
+	PaintArgs.GetHittestGrid().SetUserIndex(MyState.IncomingUserIndex);
 	GSlateFlowDirection = MyState.IncomingFlowDirection;
 	
 	FPaintArgs UpdatedArgs = PaintArgs.WithNewParent(MyState.PaintParent.Pin().Get());
@@ -195,6 +199,8 @@ int32 FWidgetProxy::Repaint(const FPaintArgs& PaintArgs, int32 MyIndex, FSlateWi
 	}
 	const int32 NewLayerId = Widget->Paint(UpdatedArgs, MyState.AllottedGeometry, MyState.CullingBounds, OutDrawElements, MyState.LayerId, MyState.WidgetStyle, MyState.bParentEnabled);
 
+	PaintArgs.GetHittestGrid().SetUserIndex(PrevUserIndex);
+
 	if (bNeedsNewClipState)
 	{
 		OutDrawElements.PopClip();
diff --git a/Engine/Source/Runtime/SlateCore/Private/Input/HittestGrid.cpp b/Engine/Source/Runtime/SlateCore/Private/Input/HittestGrid.cpp
--- a/Engine/Source/Runtime/SlateCore/Private/Input/HittestGrid.cpp
+++ b/Engine/Source/Runtime/SlateCore/Private/Input/HittestGrid.cpp
@@ -16,6 +16,12 @@ DECLARE_CYCLE_STAT(TEXT("HitTestGrid Clear"), STAT_SlateHTG_Clear, STATGROUP_Sla
 int32 SlateVerifyHitTestVisibility = 0;
 static FAutoConsoleVariableRef CVarSlateVerifyHitTestVisibility(TEXT("Slate.VerifyHitTestVisibility"), SlateVerifyHitTestVisibility, TEXT("Should we double check the visibility of widgets during hit testing, in case previously resolved hit tests that same frame may have changed state?"), ECVF_Default);
 
+constexpr bool IsCompatibleUserIndex(int32 RequestedUserIndex, int32 TestUserIndex)
+{
+	// INDEX_NONE means the user index is compatible with all
+	return RequestedUserIndex == INDEX_NONE || TestUserIndex == INDEX_NONE|| RequestedUserIndex == TestUserIndex;
+}
+
 //
 // Helper Functions
 //
@@ -158,6 +164,7 @@ FHittestGrid::FHittestGrid()
 	, NumCells(0, 0)
 	, GridOrigin(0, 0)
 	, GridSize(0, 0)
+	, CurrentUserIndex(INDEX_NONE)
 {
 }
 
@@ -165,7 +172,7 @@ FHittestGrid::~FHittestGrid()
 {
 }
 
-TArray<FWidgetAndPointer> FHittestGrid::GetBubblePath(FVector2D DesktopSpaceCoordinate, float CursorRadius, bool bIgnoreEnabledStatus)
+TArray<FWidgetAndPointer> FHittestGrid::GetBubblePath(FVector2D DesktopSpaceCoordinate, float CursorRadius, bool bIgnoreEnabledStatus, int32 UserIndex)
 {
 	checkSlow(IsInGameThread());
 
@@ -220,7 +227,8 @@ TArray<FWidgetAndPointer> FHittestGrid::GetBubblePath(FVector2D DesktopSpaceCoor
 			const FWidgetData& BestHitWidgetData = WidgetArray[BestHit.WidgetIndex];
 
 			const TSharedPtr<SWidget> FirstHitWidget = BestHitWidgetData.GetWidget();
-			if (FirstHitWidget.IsValid()) // Make Sure we landed on a valid widget
+			// Make Sure we landed on a valid widget
+			if (FirstHitWidget.IsValid() && IsCompatibleUserIndex(UserIndex, BestHitWidgetData.UserIndex))
 			{
 				TArray<FWidgetAndPointer> Path;
 
@@ -333,7 +341,7 @@ bool FHittestGrid::IsDescendantOf(const TSharedRef<SWidget> Parent, const FWidge
 }
 
 template<typename TCompareFunc, typename TSourceSideFunc, typename TDestSideFunc>
-TSharedPtr<SWidget> FHittestGrid::FindFocusableWidget(FSlateRect WidgetRect, const FSlateRect SweptRect, int32 AxisIndex, int32 Increment, const EUINavigation Direction, const FNavigationReply& NavigationReply, TCompareFunc CompareFunc, TSourceSideFunc SourceSideFunc, TDestSideFunc DestSideFunc)
+TSharedPtr<SWidget> FHittestGrid::FindFocusableWidget(FSlateRect WidgetRect, const FSlateRect SweptRect, int32 AxisIndex, int32 Increment, const EUINavigation Direction, const FNavigationReply& NavigationReply, TCompareFunc CompareFunc, TSourceSideFunc SourceSideFunc, TDestSideFunc DestSideFunc, int32 UserIndex)
 {
 	FIntPoint CurrentCellPoint = GetCellCoordinate(WidgetRect.GetCenter());
 
@@ -386,7 +394,7 @@ TSharedPtr<SWidget> FHittestGrid::FindFocusableWidget(FSlateRect WidgetRect, con
 
 				const FWidgetData& TestCandidate = WidgetArray[CurrentIndex];
 				const TSharedPtr<SWidget> TestWidget = TestCandidate.GetWidget();
-				if (TestWidget.IsValid())
+				if (TestWidget.IsValid() && IsCompatibleUserIndex(UserIndex, TestCandidate.UserIndex))
 				{
 					FGeometry TestCandidateGeo = TestWidget->GetPaintSpaceGeometry();
 					TestCandidateGeo.AppendTransform(FSlateLayoutTransform(-GridWindowOrigin));
@@ -485,7 +493,7 @@ TSharedPtr<SWidget> FHittestGrid::FindFocusableWidget(FSlateRect WidgetRect, con
 	return TSharedPtr<SWidget>();
 }
 
-TSharedPtr<SWidget> FHittestGrid::FindNextFocusableWidget(const FArrangedWidget& StartingWidget, const EUINavigation Direction, const FNavigationReply& NavigationReply, const FArrangedWidget& RuleWidget)
+TSharedPtr<SWidget> FHittestGrid::FindNextFocusableWidget(const FArrangedWidget& StartingWidget, const EUINavigation Direction, const FNavigationReply& NavigationReply, const FArrangedWidget& RuleWidget, int32 UserIndex)
 {
 	FGeometry StartingWidgetGeo = StartingWidget.Widget->GetPaintSpaceGeometry();
 	StartingWidgetGeo.AppendTransform(FSlateLayoutTransform(-GridWindowOrigin));
@@ -509,7 +517,7 @@ TSharedPtr<SWidget> FHittestGrid::FindNextFocusableWidget(const FArrangedWidget&
 		Widget = FindFocusableWidget(WidgetRect, SweptWidgetRect, 0, -1, Direction, NavigationReply,
 			[](float A, float B) { return A - 0.1f < B; }, // Compare function
 			[](FSlateRect SourceRect) { return SourceRect.Left; }, // Source side function
-			[](FSlateRect DestRect) { return DestRect.Right; }); // Dest side function
+			[](FSlateRect DestRect) { return DestRect.Right; }, UserIndex); // Dest side function
 		break;
 	case EUINavigation::Right:
 		SweptWidgetRect.Left = BoundingRuleRect.Left;
@@ -519,7 +527,7 @@ TSharedPtr<SWidget> FHittestGrid::FindNextFocusableWidget(const FArrangedWidget&
 		Widget = FindFocusableWidget(WidgetRect, SweptWidgetRect, 0, 1, Direction, NavigationReply,
 			[](float A, float B) { return A + 0.1f > B; }, // Compare function
 			[](FSlateRect SourceRect) { return SourceRect.Right; }, // Source side function
-			[](FSlateRect DestRect) { return DestRect.Left; }); // Dest side function
+			[](FSlateRect DestRect) { return DestRect.Left; }, UserIndex); // Dest side function
 		break;
 	case EUINavigation::Up:
 		SweptWidgetRect.Top = BoundingRuleRect.Top;
@@ -529,7 +537,7 @@ TSharedPtr<SWidget> FHittestGrid::FindNextFocusableWidget(const FArrangedWidget&
 		Widget = FindFocusableWidget(WidgetRect, SweptWidgetRect, 1, -1, Direction, NavigationReply,
 			[](float A, float B) { return A - 0.1f < B; }, // Compare function
 			[](FSlateRect SourceRect) { return SourceRect.Top; }, // Source side function
-			[](FSlateRect DestRect) { return DestRect.Bottom; }); // Dest side function
+			[](FSlateRect DestRect) { return DestRect.Bottom; }, UserIndex); // Dest side function
 		break;
 	case EUINavigation::Down:
 		SweptWidgetRect.Top = BoundingRuleRect.Top;
@@ -539,7 +547,7 @@ TSharedPtr<SWidget> FHittestGrid::FindNextFocusableWidget(const FArrangedWidget&
 		Widget = FindFocusableWidget(WidgetRect, SweptWidgetRect, 1, 1, Direction, NavigationReply,
 			[](float A, float B) { return A + 0.1f > B; }, // Compare function
 			[](FSlateRect SourceRect) { return SourceRect.Bottom; }, // Source side function
-			[](FSlateRect DestRect) { return DestRect.Top; }); // Dest side function
+			[](FSlateRect DestRect) { return DestRect.Top; }, UserIndex); // Dest side function
 		break;
 
 	default:
@@ -626,11 +634,29 @@ void FHittestGrid::LogGrid() const
 		UE_LOG(LogHittestDebug, Warning, TEXT("  [%d][%d][%d] => %s @ %s"),
 			It.GetIndex(),
 			CurWidgetData.PrimarySort,
-			CachedWidget.IsValid() ? *CachedWidget->ToString() : TEXT("Invalid WIdget"),
-			CachedWidget.IsValid() ? *CachedWidget->GetPaintSpaceGeometry().ToString() : TEXT("Invalid WIdget"));
+			CachedWidget.IsValid() ? *CachedWidget->ToString() : TEXT("Invalid Widget"),
+			CachedWidget.IsValid() ? *CachedWidget->GetPaintSpaceGeometry().ToString() : TEXT("Invalid Widget"));
 	}
 }
 
+constexpr FLinearColor GetUserColorFromIndex(int32 Index)
+{
+	switch (Index)
+	{
+	case 0:
+		return FLinearColor::Red;
+	case 1:
+		return FLinearColor::Green;
+	case 2:
+		return FLinearColor::Blue;
+	case 3:
+		return FLinearColor::Yellow;
+	default:
+		return FLinearColor::White;
+	}
+};
+
+
 void FHittestGrid::DisplayGrid(int32 InLayer, const FGeometry& AllottedGeometry, FSlateWindowElementList& WindowElementList) const
 {
 	static const FSlateBrush* WhiteBrush = FCoreStyle::Get().GetBrush(TEXT("FocusRectangle"));
@@ -646,7 +672,9 @@ void FHittestGrid::DisplayGrid(int32 InLayer, const FGeometry& AllottedGeometry,
 				WindowElementList,
 				InLayer,
 				CachedWidget->GetPaintSpaceGeometry().ToPaintGeometry(),
-				WhiteBrush
+				WhiteBrush,
+				ESlateDrawEffect::None,
+				GetUserColorFromIndex(CurWidgetData.UserIndex)
 			);
 		}
 	}
@@ -686,7 +714,7 @@ void FHittestGrid::AddWidget(const TSharedRef<SWidget>& InWidget, int32 InBatchP
 	const FIntPoint LowerRightCell = GetCellCoordinate(BoundingRect.GetBottomRight());
 
 	int32& WidgetIndex = WidgetMap.Add(&*InWidget);
-	FWidgetData Data(InWidget, UpperLeftCell, LowerRightCell, PrimarySort, InSecondarySort);
+	FWidgetData Data(InWidget, UpperLeftCell, LowerRightCell, PrimarySort, InSecondarySort, CurrentUserIndex);
 	WidgetIndex = WidgetArray.Add(Data); // Bleh why doesn't Sparse Array have an emplace.
 
 	for (int32 XIndex = UpperLeftCell.X; XIndex <= LowerRightCell.X; ++XIndex)
diff --git a/Engine/Source/Runtime/SlateCore/Private/Widgets/Accessibility/SlateCoreAccessibleWidgets.cpp b/Engine/Source/Runtime/SlateCore/Private/Widgets/Accessibility/SlateCoreAccessibleWidgets.cpp
--- a/Engine/Source/Runtime/SlateCore/Private/Widgets/Accessibility/SlateCoreAccessibleWidgets.cpp
+++ b/Engine/Source/Runtime/SlateCore/Private/Widgets/Accessibility/SlateCoreAccessibleWidgets.cpp
@@ -270,7 +270,7 @@ TSharedPtr<IAccessibleWidget> FSlateAccessibleWindow::GetChildAtPosition(int32 X
 		if (UseHitTestGrid)
 		{
 			TSharedPtr<SWindow> SlateWindow = StaticCastSharedPtr<SWindow>(Widget.Pin());
-			TArray<FWidgetAndPointer> Hits = SlateWindow->GetHittestGrid().GetBubblePath(FVector2D(X, Y), 0.0f, false);
+			TArray<FWidgetAndPointer> Hits = SlateWindow->GetHittestGrid().GetBubblePath(FVector2D(X, Y), 0.0f, false, INDEX_NONE);
 			TSharedPtr<SWidget> LastAccessibleWidget = nullptr;
 			for (int32 i = 0; i < Hits.Num(); ++i)
 			{
diff --git a/Engine/Source/Runtime/SlateCore/Private/Widgets/SWidget.cpp b/Engine/Source/Runtime/SlateCore/Private/Widgets/SWidget.cpp
--- a/Engine/Source/Runtime/SlateCore/Private/Widgets/SWidget.cpp
+++ b/Engine/Source/Runtime/SlateCore/Private/Widgets/SWidget.cpp
@@ -1267,6 +1267,7 @@ int32 SWidget::Paint(const FPaintArgs& Args, const FGeometry& AllottedGeometry,
 	PersistentState.DesktopGeometry = DesktopSpaceGeometry;
 	PersistentState.WidgetStyle = InWidgetStyle;
 	PersistentState.CullingBounds = MyCullingRect;
+	PersistentState.IncomingUserIndex = Args.GetHittestGrid().GetUserIndex();
 	PersistentState.IncomingFlowDirection = GSlateFlowDirection;
 
 	FPaintArgs UpdatedArgs = Args.WithNewParent(this);
diff --git a/Engine/Source/Runtime/SlateCore/Private/Widgets/SWindow.cpp b/Engine/Source/Runtime/SlateCore/Private/Widgets/SWindow.cpp
--- a/Engine/Source/Runtime/SlateCore/Private/Widgets/SWindow.cpp
+++ b/Engine/Source/Runtime/SlateCore/Private/Widgets/SWindow.cpp
@@ -1926,7 +1926,7 @@ EWindowZone::Type SWindow::GetCurrentWindowZone(FVector2D LocalMousePosition)
 		if (InZone == EWindowZone::ClientArea)
 		{
 			// Hittest to see if the widget under the mouse should be treated as a title bar (i.e. should move the window)
-			FWidgetPath HitTestResults = FSlateApplicationBase::Get().GetHitTesting().LocateWidgetInWindow(FSlateApplicationBase::Get().GetCursorPos(), SharedThis(this), false);
+			FWidgetPath HitTestResults = FSlateApplicationBase::Get().GetHitTesting().LocateWidgetInWindow(FSlateApplicationBase::Get().GetCursorPos(), SharedThis(this), false, INDEX_NONE);
 			if( HitTestResults.Widgets.Num() > 0 )
 			{
 				const EWindowZone::Type ZoneOverride = HitTestResults.Widgets.Last().Widget->GetWindowZoneOverride();
diff --git a/Engine/Source/Runtime/SlateCore/Public/Application/SlateApplicationBase.h b/Engine/Source/Runtime/SlateCore/Public/Application/SlateApplicationBase.h
--- a/Engine/Source/Runtime/SlateCore/Public/Application/SlateApplicationBase.h
+++ b/Engine/Source/Runtime/SlateCore/Public/Application/SlateApplicationBase.h
@@ -58,7 +58,7 @@ private:
 private:	
 	FSlateApplicationBase* SlateApp;
 	// @see FSlateApplicationBase::LocateWidgetInWindow
-	FWidgetPath LocateWidgetInWindow(FVector2D ScreenspaceMouseCoordinate, const TSharedRef<SWindow>& Window, bool bIgnoreEnabledStatus) const;
+	FWidgetPath LocateWidgetInWindow(FVector2D ScreenspaceMouseCoordinate, const TSharedRef<SWindow>& Window, bool bIgnoreEnabledStatus, int32 UserIndex) const;
 };
 
 /**
@@ -337,7 +337,7 @@ public:
 	 *
 	 * @return The path to the widget.
 	 */
-	virtual FWidgetPath LocateWindowUnderMouse( FVector2D ScreenspaceMouseCoordinate, const TArray< TSharedRef<SWindow > >& Windows, bool bIgnoreEnabledStatus = false ) = 0;
+	virtual FWidgetPath LocateWindowUnderMouse( FVector2D ScreenspaceMouseCoordinate, const TArray< TSharedRef<SWindow > >& Windows, bool bIgnoreEnabledStatus = false, int32 UserIndex = INDEX_NONE) = 0;
 
 	/**
 	 * Calculates the tooltip window position.
@@ -544,7 +544,7 @@ protected:
 	virtual bool ShowUserFocus(const TSharedPtr<const SWidget> Widget) const = 0;
 
 	/** Given a window, locate a widget under the cursor in it; returns an invalid path if cursor is not over this window. */
-	virtual FWidgetPath LocateWidgetInWindow(FVector2D ScreenspaceMouseCoordinate, const TSharedRef<SWindow>& Window, bool bIgnoreEnabledStatus) const = 0;
+	virtual FWidgetPath LocateWidgetInWindow(FVector2D ScreenspaceMouseCoordinate, const TSharedRef<SWindow>& Window, bool bIgnoreEnabledStatus, int32 UserIndex) const = 0;
 
 #if WITH_EDITOR
 	void UpdateCustomSafeZone(const FMargin& NewSafeZoneRatio, bool bShouldRecacheMetrics) 
diff --git a/Engine/Source/Runtime/SlateCore/Public/FastUpdate/WidgetProxy.h b/Engine/Source/Runtime/SlateCore/Public/FastUpdate/WidgetProxy.h
--- a/Engine/Source/Runtime/SlateCore/Public/FastUpdate/WidgetProxy.h
+++ b/Engine/Source/Runtime/SlateCore/Public/FastUpdate/WidgetProxy.h
@@ -138,6 +138,7 @@ struct FSlateWidgetPersistentState
 		: CachedElementListNode(nullptr)
 		, LayerId(0)
 		, OutgoingLayerId(0)
+		, IncomingUserIndex(INDEX_NONE)
 		, IncomingFlowDirection(EFlowDirection::LeftToRight)
 		, bParentEnabled(true)
 		, bInheritedHittestability(false)
@@ -153,6 +154,7 @@ struct FSlateWidgetPersistentState
 	/** Starting layer id for drawing children **/
 	int32 LayerId;
 	int32 OutgoingLayerId;
+	int8 IncomingUserIndex;
 	EFlowDirection IncomingFlowDirection;
 	uint8 bParentEnabled : 1;
 	uint8 bInheritedHittestability : 1;
diff --git a/Engine/Source/Runtime/SlateCore/Public/Input/HittestGrid.h b/Engine/Source/Runtime/SlateCore/Public/Input/HittestGrid.h
--- a/Engine/Source/Runtime/SlateCore/Public/Input/HittestGrid.h
+++ b/Engine/Source/Runtime/SlateCore/Public/Input/HittestGrid.h
@@ -34,7 +34,7 @@ public:
 	 * Given a Slate Units coordinate in virtual desktop space, perform a hittest
 	 * and return the path along which the corresponding event would be bubbled.
 	 */
-	TArray<FWidgetAndPointer> GetBubblePath(FVector2D DesktopSpaceCoordinate, float CursorRadius, bool bIgnoreEnabledStatus);
+	TArray<FWidgetAndPointer> GetBubblePath(FVector2D DesktopSpaceCoordinate, float CursorRadius, bool bIgnoreEnabledStatus, int32 UserIndex);
 
 	/**
 	 * Set the position and size of the hittest area in desktop coordinates
@@ -58,6 +58,12 @@ public:
 	/** Insert custom hit test data for a widget already in the grid */
 	void InsertCustomHitTestPath(const TSharedRef<SWidget> InWidget, TSharedRef<ICustomHitTestPath> CustomHitTestPath);
 
+	/** Sets the current slate user index that should be associated with any added widgets */
+	void SetUserIndex(int32 UserIndex) { CurrentUserIndex = UserIndex; }
+
+	/** Gets current slate user index that should be associated with any added widgets */
+	int32 GetUserIndex() const { return CurrentUserIndex; }
+
 	/**
 	 * Finds the next focusable widget by searching through the hit test grid
 	 *
@@ -66,7 +72,7 @@ public:
 	 * @param NavigationReply The Navigation Reply to specify a boundary rule for the search.
 	 * @param RuleWidget      The Widget that is applying the boundary rule, used to get the bounds of the Rule.
 	 */
-	TSharedPtr<SWidget> FindNextFocusableWidget(const FArrangedWidget& StartingWidget, const EUINavigation Direction, const FNavigationReply& NavigationReply, const FArrangedWidget& RuleWidget);
+	TSharedPtr<SWidget> FindNextFocusableWidget(const FArrangedWidget& StartingWidget, const EUINavigation Direction, const FNavigationReply& NavigationReply, const FArrangedWidget& RuleWidget, int32 UserIndex);
 
 	void AppendGrid(FHittestGrid& OtherGrid);
 
@@ -87,12 +93,13 @@ private:
 	 */
 	struct FWidgetData
 	{
-		FWidgetData(TSharedRef<SWidget> InWidget, const FIntPoint& InUpperLeftCell, const FIntPoint& InLowerRightCell, int64 InPrimarySort, int32 InSecondarySort)
+		FWidgetData(TSharedRef<SWidget> InWidget, const FIntPoint& InUpperLeftCell, const FIntPoint& InLowerRightCell, int64 InPrimarySort, int32 InSecondarySort, int32 InUserIndex)
 			: WeakWidget(InWidget)
 			, UpperLeftCell(InUpperLeftCell)
 			, LowerRightCell(InLowerRightCell)
 			, PrimarySort(InPrimarySort)
 			, SecondarySort(InSecondarySort)
+			, UserIndex(InUserIndex)
 		{}
 		TWeakPtr<SWidget> WeakWidget;
 		TWeakPtr<ICustomHitTestPath> CustomPath;
@@ -100,6 +107,7 @@ private:
 		FIntPoint LowerRightCell;
 		int64 PrimarySort;
 		int32 SecondarySort;
+		int32 UserIndex;
 
 		TSharedPtr<SWidget> GetWidget() const { return WeakWidget.Pin(); }
 	};
@@ -148,7 +156,7 @@ private:
 
 	/** Utility function for searching for the next focusable widget. */
 	template<typename TCompareFunc, typename TSourceSideFunc, typename TDestSideFunc>
-	TSharedPtr<SWidget> FindFocusableWidget(const FSlateRect WidgetRect, const FSlateRect SweptRect, int32 AxisIndex, int32 Increment, const EUINavigation Direction, const FNavigationReply& NavigationReply, TCompareFunc CompareFunc, TSourceSideFunc SourceSideFunc, TDestSideFunc DestSideFunc);
+	TSharedPtr<SWidget> FindFocusableWidget(const FSlateRect WidgetRect, const FSlateRect SweptRect, int32 AxisIndex, int32 Increment, const EUINavigation Direction, const FNavigationReply& NavigationReply, TCompareFunc CompareFunc, TSourceSideFunc SourceSideFunc, TDestSideFunc DestSideFunc, int32 UserIndex);
 
 	/** Constrains a float position into the grid coordinate. */
 	FIntPoint GetCellCoordinate(FVector2D Position);
@@ -188,5 +196,6 @@ private:
 	/** The Size of the current grid. */
 	FVector2D GridSize;
 
-	//friend class SWidgetReflector;
+	/** The current slate user index that should be associated with any added widgets */
+	int32 CurrentUserIndex;
 };
diff --git a/Engine/Source/Runtime/UMG/Private/Components/WidgetComponent.cpp b/Engine/Source/Runtime/UMG/Private/Components/WidgetComponent.cpp
--- a/Engine/Source/Runtime/UMG/Private/Components/WidgetComponent.cpp
+++ b/Engine/Source/Runtime/UMG/Private/Components/WidgetComponent.cpp
@@ -1851,7 +1851,9 @@ TArray<FWidgetAndPointer> UWidgetComponent::GetHitWidgetPath(FVector2D WidgetSpa
 	TArray<FWidgetAndPointer> ArrangedWidgets;
 	if ( SlateWindow.IsValid() )
 	{
-		ArrangedWidgets = SlateWindow->GetHittestGrid().GetBubblePath( LocalHitLocation, CursorRadius, bIgnoreEnabledStatus );
+		// @todo slate - widget components would need to be associated with a user for this to be anthing valid
+		const int32 UserIndex = INDEX_NONE;
+		ArrangedWidgets = SlateWindow->GetHittestGrid().GetBubblePath( LocalHitLocation, CursorRadius, bIgnoreEnabledStatus, UserIndex);
 
 		for( FWidgetAndPointer& ArrangedWidget : ArrangedWidgets )
 		{
diff --git a/Engine/Source/Runtime/UMG/Private/Slate/SRetainerWidget.cpp b/Engine/Source/Runtime/UMG/Private/Slate/SRetainerWidget.cpp
--- a/Engine/Source/Runtime/UMG/Private/Slate/SRetainerWidget.cpp
+++ b/Engine/Source/Runtime/UMG/Private/Slate/SRetainerWidget.cpp
@@ -478,6 +478,9 @@ int32 SRetainerWidget::OnPaint(const FPaintArgs& Args, const FGeometry& Allotted
 		}
 		FPaintArgs NewArgs = Args.WithNewHitTestGrid(HittestGrid);
 
+		// Copy the current user index into the new grid since nested hittest grids should inherit their parents user id
+		NewArgs.GetHittestGrid().SetUserIndex(Args.RootGrid.GetUserIndex());
+
 		FSlateInvalidationContext Context(OutDrawElements, InWidgetStyle);
 		Context.bParentEnabled = bParentEnabled;
 		Context.bAllowFastPathUpdate = true;
