Fix VS2019 warning

diff --git a/Engine/Source/Runtime/Slate/Private/Framework/Application/SlateApplication.cpp b/Engine/Source/Runtime/Slate/Private/Framework/Application/SlateApplication.cpp
--- a/Engine/Source/Runtime/Slate/Private/Framework/Application/SlateApplication.cpp
+++ b/Engine/Source/Runtime/Slate/Private/Framework/Application/SlateApplication.cpp
@@ -4668,7 +4668,7 @@ FReply FSlateApplication::RoutePointerUpEvent(const FWidgetPath& WidgetsUnderPoi
 	{
 		if (!LocalWidgetsUnderPointer.IsValid())
 		{
-			LocalWidgetsUnderPointer = LocateWindowUnderMouse(PointerEvent.GetScreenSpacePosition(), GetInteractiveTopLevelWindows(), SlateUser->GetUserIndex());
+			LocalWidgetsUnderPointer = LocateWindowUnderMouse(PointerEvent.GetScreenSpacePosition(), GetInteractiveTopLevelWindows(), false, SlateUser->GetUserIndex());
 		}
 		
 		// Switch worlds widgets in the current path
