Adding option to FSlateUser to completely disable painting the user's cursor

diff --git a/Engine/Source/Runtime/Slate/Private/Framework/Application/SlateUser.cpp b/Engine/Source/Runtime/Slate/Private/Framework/Application/SlateUser.cpp
--- a/Engine/Source/Runtime/Slate/Private/Framework/Application/SlateUser.cpp
+++ b/Engine/Source/Runtime/Slate/Private/Framework/Application/SlateUser.cpp
@@ -331,6 +331,11 @@ TSharedPtr<SWidget> FSlateUser::GetPointerCaptor(uint32 PointerIndex) const
 	return WeakCaptorPath ? WeakCaptorPath->GetLastWidget().Pin() : nullptr;
 }
 
+void FSlateUser::SetCursorVisibility(bool bDrawCursor)
+{
+	bCanDrawCursor = bDrawCursor;
+}
+
 void FSlateUser::SetCursorPosition(int32 PosX, int32 PosY)
 {
 	SetPointerPosition(FSlateApplication::CursorPointerIndex, PosX, PosY);
@@ -552,7 +557,7 @@ void FSlateUser::DrawWindowlessDragDropContent(const TSharedRef<SWindow>& Window
 void FSlateUser::DrawCursor(const TSharedRef<SWindow>& WindowToDraw, FSlateWindowElementList& WindowElementList, int32& MaxLayerId)
 {
 	TSharedPtr<SWindow> CursorWindow = CursorWindowPtr.Pin();
-	if (CursorWindow && WindowToDraw == CursorWindow)
+	if (bCanDrawCursor && CursorWindow && WindowToDraw == CursorWindow)
 	{
 		if (TSharedPtr<SWidget> CursorWidget = CursorWidgetPtr.Pin())
 		{
diff --git a/Engine/Source/Runtime/Slate/Public/Framework/Application/SlateUser.h b/Engine/Source/Runtime/Slate/Public/Framework/Application/SlateUser.h
--- a/Engine/Source/Runtime/Slate/Public/Framework/Application/SlateUser.h
+++ b/Engine/Source/Runtime/Slate/Public/Framework/Application/SlateUser.h
@@ -84,6 +84,8 @@ public:
 	TSharedPtr<SWidget> GetCursorCaptor() const;
 	TSharedPtr<SWidget> GetPointerCaptor(uint32 PointerIndex) const;
 
+	void SetCursorVisibility(bool bDrawCursor);
+
 	void SetCursorPosition(int32 PosX, int32 PosY);
 	void SetCursorPosition(const FVector2D& NewCursorPos);
 	void SetPointerPosition(uint32 PointerIndex, int32 PosX, int32 PosY);
@@ -193,6 +195,9 @@ private:
 	/** The cursor this user is in control of. Guaranteed to be valid for all real users, absence implies this is a virtual user. */
 	TSharedPtr<ICursor> Cursor;
 
+	/** Whether this user is currently drawing the cursor onscreen each frame */
+	bool bCanDrawCursor = true;
+
 	/** The OS or actions taken by the user may require we refresh the current state of the cursor. */
 	bool bQueryCursorRequested = false;
 
