Hide hardware cursors as well when FSlateUser::bCanDrawCursor is false by setting the cursor type to EMouseCursor::None
- We skip any work to query the cursor for the user while they don't want the cursor shown
- If anyone does try to ProcessCursorReply(), it either does nothing or sets the cursor to None again

diff --git a/Engine/Source/Runtime/Slate/Private/Framework/Application/SlateUser.cpp b/Engine/Source/Runtime/Slate/Private/Framework/Application/SlateUser.cpp
--- a/Engine/Source/Runtime/Slate/Private/Framework/Application/SlateUser.cpp
+++ b/Engine/Source/Runtime/Slate/Private/Framework/Application/SlateUser.cpp
@@ -334,6 +334,15 @@ TSharedPtr<SWidget> FSlateUser::GetPointerCaptor(uint32 PointerIndex) const
 void FSlateUser::SetCursorVisibility(bool bDrawCursor)
 {
 	bCanDrawCursor = bDrawCursor;
+	
+	if (bCanDrawCursor)
+	{
+		RequestCursorQuery();
+	}
+	else
+	{
+		ProcessCursorReply(FCursorReply::Cursor(EMouseCursor::None));
+	}
 }
 
 void FSlateUser::SetCursorPosition(int32 PosX, int32 PosY)
@@ -659,7 +668,7 @@ void FSlateUser::QueryCursor()
 	bQueryCursorRequested = false;
 
 	// The slate loading widget thread is not allowed to execute this code (it's unsafe to read the hittest grid in another thread)
-	if (Cursor && IsInGameThread())
+	if (bCanDrawCursor && Cursor && IsInGameThread())
 	{
 		SCOPE_CYCLE_COUNTER(STAT_SlateQueryCursor);
 
@@ -814,20 +823,27 @@ void FSlateUser::ProcessCursorReply(const FCursorReply& CursorReply)
 {
 	if (Cursor && CursorReply.IsEventHandled())
 	{
-		CursorWidgetPtr = CursorReply.GetCursorWidget();
-		if (CursorReply.GetCursorWidget().IsValid())
+		if (bCanDrawCursor)
 		{
-			CursorReply.GetCursorWidget()->SetVisibility(EVisibility::HitTestInvisible);
-			CursorWindowPtr = CursorReply.GetCursorWindow();
-			if (!FSlateApplication::Get().IsFakingTouchEvents())
+			CursorWidgetPtr = CursorReply.GetCursorWidget();
+			if (CursorReply.GetCursorWidget().IsValid())
+			{
+				CursorReply.GetCursorWidget()->SetVisibility(EVisibility::HitTestInvisible);
+				CursorWindowPtr = CursorReply.GetCursorWindow();
+				if (!FSlateApplication::Get().IsFakingTouchEvents())
+				{
+					Cursor->SetType(EMouseCursor::Custom);
+				}
+			}
+			else
 			{
-				Cursor->SetType(EMouseCursor::Custom);
+				CursorWindowPtr.Reset();
+				Cursor->SetType(CursorReply.GetCursorType());
 			}
 		}
 		else
 		{
-			CursorWindowPtr.Reset();
-			Cursor->SetType(CursorReply.GetCursorType());
+			Cursor->SetType(EMouseCursor::None);
 		}
 	}
 	else
@@ -934,7 +950,7 @@ void FSlateUser::NotifyPointerMoveComplete(const FPointerEvent& PointerEvent, co
 				CursorReply = FCursorReply::Cursor(EMouseCursor::Default);
 			}
 
-			FSlateApplication::Get().ProcessCursorReply(CursorReply);
+			ProcessCursorReply(CursorReply);
 		}
 		
 	}
