Adding TableViewMetadata object to SListView, STileView, and STreeView
- Currently just exposes a way to get all the currently selected rows, but more functionality may be desired
- Useful when in a scenario where you only have a raw SWidget, but need to behave differently if that widget is actually a list view

diff --git a/Engine/Source/Runtime/Slate/Public/Framework/Views/ITypedTableView.h b/Engine/Source/Runtime/Slate/Public/Framework/Views/ITypedTableView.h
--- a/Engine/Source/Runtime/Slate/Public/Framework/Views/ITypedTableView.h
+++ b/Engine/Source/Runtime/Slate/Public/Framework/Views/ITypedTableView.h
@@ -206,6 +206,18 @@ public:
 	/** Is the list pending a refresh? */
 	virtual bool Private_IsPendingRefresh() const = 0;
 
+	/** @return All currently selected items in the table view */
+	virtual TArray<ItemType> GetSelectedItems() const = 0;
+
+	/**
+	 * Find a widget for this item if it has already been constructed.
+	 *
+	 * @param InItem  The item for which to find the widget.
+	 *
+	 * @return A pointer to the corresponding widget if it exists; otherwise nullptr.
+	*/
+	virtual TSharedPtr<ITableRow> WidgetFromItem(const ItemType& InItem) const = 0;
+
 	/** @return Is the tableview a tree or a list? */
 	virtual ETableViewMode::Type GetTableViewMode() const = 0;
 	
diff --git a/Engine/Source/Runtime/Slate/Public/Framework/Views/TableViewMetadata.h b/Engine/Source/Runtime/Slate/Public/Framework/Views/TableViewMetadata.h
new file mode 100644
--- /dev/null
+++ b/Engine/Source/Runtime/Slate/Public/Framework/Views/TableViewMetadata.h
@@ -0,0 +1,46 @@
+// Copyright 1998-2019 Epic Games, Inc. All Rights Reserved.
+
+#pragma once
+
+#include "CoreMinimal.h"
+#include "Types/ISlateMetaData.h"
+#include "ITypedTableView.h"
+
+class ITableViewMetadata : public ISlateMetaData
+{
+public:
+	SLATE_METADATA_TYPE(ITableViewMetadata, ISlateMetaData);
+
+	/** Returns all currently constructed rows that represent selected items */
+	virtual TArray<TSharedPtr<ITableRow>> GatherSelectedRows() const = 0;
+};
+
+template <typename ItemType>
+class TTableViewMetadata : public ITableViewMetadata
+{
+public:
+	SLATE_METADATA_TYPE(TTableViewMetadata, ITableViewMetadata);
+
+	TTableViewMetadata(TSharedRef<ITypedTableView<ItemType>> InOwnerTableView)
+		: OwnerTableView(InOwnerTableView)
+	{}
+
+	virtual TArray<TSharedPtr<ITableRow>> GatherSelectedRows() const override
+	{
+		TArray<TSharedPtr<ITableRow>> SelectedRows;
+		if (TSharedPtr<ITypedTableView<ItemType>> PinnedTableView = OwnerTableView.Pin())
+		{
+			for (const ItemType& Item : PinnedTableView->GetSelectedItems())
+			{
+				if (TSharedPtr<ITableRow> TableRow = PinnedTableView->WidgetFromItem(Item))
+				{
+					SelectedRows.Add(TableRow);
+				}
+			}
+		}
+		return SelectedRows;
+	}
+
+private:
+	TWeakPtr<ITypedTableView<ItemType>> OwnerTableView;
+};
\ No newline at end of file
diff --git a/Engine/Source/Runtime/Slate/Public/Widgets/Views/SListView.h b/Engine/Source/Runtime/Slate/Public/Widgets/Views/SListView.h
--- a/Engine/Source/Runtime/Slate/Public/Widgets/Views/SListView.h
+++ b/Engine/Source/Runtime/Slate/Public/Widgets/Views/SListView.h
@@ -12,6 +12,7 @@
 #include "Framework/SlateDelegates.h"
 #include "Widgets/Text/STextBlock.h"
 #include "Framework/Views/ITypedTableView.h"
+#include "Framework/Views/TableViewMetadata.h"
 #include "Widgets/Views/STableViewBase.h"
 #include "Framework/Views/TableViewTypeTraits.h"
 #include "Widgets/Views/STableRow.h"
@@ -241,7 +242,7 @@ public:
 				ScrollBar->SetDragFocusCause(InArgs._ScrollbarDragFocusCause);
 				ScrollBar->SetUserVisibility(InArgs._ScrollbarVisibility);
 			}
-
+			this->AddMetadata(MakeShared<TTableViewMetadata<ItemType>>(SharedThis(this)));
 		}
 	}
 
@@ -1349,7 +1350,7 @@ public:
 	 *
 	 * @return	List of selected item indices (in no particular order)
 	 */
-	TArray< ItemType > GetSelectedItems() const
+	virtual TArray< ItemType > GetSelectedItems() const override
 	{
 		TArray< ItemType > SelectedItemArray;
 		SelectedItemArray.Empty( SelectedItems.Num() );
@@ -1360,7 +1361,7 @@ public:
 		return SelectedItemArray;
 	}
 
-	int32 GetSelectedItems(TArray< ItemType >&SelectedItemArray) const
+	int32 GetSelectedItems(TArray< ItemType >& SelectedItemArray) const
 	{
 		SelectedItemArray.Empty(SelectedItems.Num());
 		for (typename TItemSet::TConstIterator SelectedItemIt(SelectedItems); SelectedItemIt; ++SelectedItemIt)
@@ -1498,7 +1499,7 @@ public:
 	 *
 	 * @return A pointer to the corresponding widget if it exists; otherwise nullptr.
 	*/
-	TSharedPtr<ITableRow> WidgetFromItem( const ItemType& InItem )
+	virtual TSharedPtr<ITableRow> WidgetFromItem( const ItemType& InItem ) const override
 	{
 		return WidgetGenerator.GetWidgetForItem(InItem);
 	}
diff --git a/Engine/Source/Runtime/Slate/Public/Widgets/Views/STileView.h b/Engine/Source/Runtime/Slate/Public/Widgets/Views/STileView.h
--- a/Engine/Source/Runtime/Slate/Public/Widgets/Views/STileView.h
+++ b/Engine/Source/Runtime/Slate/Public/Widgets/Views/STileView.h
@@ -207,6 +207,7 @@ public:
 				this->ScrollBar->SetDragFocusCause(InArgs._ScrollbarDragFocusCause);
 				this->ScrollBar->SetUserVisibility(InArgs._ScrollbarVisibility);
 			}
+			this->AddMetadata(MakeShared<TTableViewMetadata<ItemType>>(SharedThis(this)));
 		}
 	}
 
diff --git a/Engine/Source/Runtime/Slate/Public/Widgets/Views/STreeView.h b/Engine/Source/Runtime/Slate/Public/Widgets/Views/STreeView.h
--- a/Engine/Source/Runtime/Slate/Public/Widgets/Views/STreeView.h
+++ b/Engine/Source/Runtime/Slate/Public/Widgets/Views/STreeView.h
@@ -284,6 +284,7 @@ public:
 			{
 				this->ScrollBar->SetDragFocusCause(InArgs._ScrollbarDragFocusCause);
 			}
+			this->AddMetadata(MakeShared<TTableViewMetadata<ItemType>>(SharedThis(this)));
 		}
 	}
 
