Clang compile fixes

diff --git a/Engine/Source/Runtime/Slate/Public/Widgets/Views/SListView.h b/Engine/Source/Runtime/Slate/Public/Widgets/Views/SListView.h
--- a/Engine/Source/Runtime/Slate/Public/Widgets/Views/SListView.h
+++ b/Engine/Source/Runtime/Slate/Public/Widgets/Views/SListView.h
@@ -242,7 +242,7 @@ public:
 				ScrollBar->SetDragFocusCause(InArgs._ScrollbarDragFocusCause);
 				ScrollBar->SetUserVisibility(InArgs._ScrollbarVisibility);
 			}
-			this->AddMetadata(MakeShared<TTableViewMetadata<ItemType>>(SharedThis(this)));
+			this->AddMetadata(MakeShared<TTableViewMetadata<ItemType>>(this->SharedThis(this)));
 		}
 	}
 
diff --git a/Engine/Source/Runtime/Slate/Public/Widgets/Views/STileView.h b/Engine/Source/Runtime/Slate/Public/Widgets/Views/STileView.h
--- a/Engine/Source/Runtime/Slate/Public/Widgets/Views/STileView.h
+++ b/Engine/Source/Runtime/Slate/Public/Widgets/Views/STileView.h
@@ -207,7 +207,7 @@ public:
 				this->ScrollBar->SetDragFocusCause(InArgs._ScrollbarDragFocusCause);
 				this->ScrollBar->SetUserVisibility(InArgs._ScrollbarVisibility);
 			}
-			this->AddMetadata(MakeShared<TTableViewMetadata<ItemType>>(SharedThis(this)));
+			this->AddMetadata(MakeShared<TTableViewMetadata<ItemType>>(this->SharedThis(this)));
 		}
 	}
 
diff --git a/Engine/Source/Runtime/Slate/Public/Widgets/Views/STreeView.h b/Engine/Source/Runtime/Slate/Public/Widgets/Views/STreeView.h
--- a/Engine/Source/Runtime/Slate/Public/Widgets/Views/STreeView.h
+++ b/Engine/Source/Runtime/Slate/Public/Widgets/Views/STreeView.h
@@ -284,7 +284,7 @@ public:
 			{
 				this->ScrollBar->SetDragFocusCause(InArgs._ScrollbarDragFocusCause);
 			}
-			this->AddMetadata(MakeShared<TTableViewMetadata<ItemType>>(SharedThis(this)));
+			this->AddMetadata(MakeShared<TTableViewMetadata<ItemType>>(this->SharedThis(this)));
 		}
 	}
 
