fix issue font size multiplier on text with Decorator
#rnx

diff --git a/Engine/Source/Runtime/Slate/Private/Framework/Text/RichTextLayoutMarshaller.cpp b/Engine/Source/Runtime/Slate/Private/Framework/Text/RichTextLayoutMarshaller.cpp
--- a/Engine/Source/Runtime/Slate/Private/Framework/Text/RichTextLayoutMarshaller.cpp
+++ b/Engine/Source/Runtime/Slate/Private/Framework/Text/RichTextLayoutMarshaller.cpp
@@ -8,6 +8,7 @@
 #include "Framework/Text/RichTextMarkupProcessing.h"
 #include "Framework/Text/SlateTextLayout.h"
 #include "Framework/Text/SlateTextUnderlineLineHighlighter.h"
+#include "Templates/SharedPointer.h"
 
 #if WITH_FANCY_TEXT
 
@@ -178,14 +179,7 @@ void FRichTextLayoutMarshaller::AppendRunsForText(
 		ModelRange.EndIndex = InOutModelText->Len();
 
 		// Create run.
-		TSharedPtr< FSlateTextRun > SlateTextRun = FSlateTextRun::Create(RunInfo, InOutModelText, *TextBlockStyle, ModelRange);
-
-		if (SlateTextRun)
-		{
-			// Apply the FontSizeMultiplier at the style use by the IRun
-			SlateTextRun->ApplyFontSizeMultiplierOnTextStyle(FontSizeMultiplier);
-		}
-		Run = SlateTextRun;
+		Run = FSlateTextRun::Create(RunInfo, InOutModelText, *TextBlockStyle, ModelRange);
 
 		if (!TextBlockStyle->UnderlineBrush.GetResourceName().IsNone())
 		{
@@ -212,6 +206,15 @@ void FRichTextLayoutMarshaller::AppendRunsForText(
 		}
 	}
 
+	if (Run.IsValid())
+	{
+		TSharedRef<FSlateTextRun> SlateTextRun = StaticCastSharedRef<FSlateTextRun>(Run.ToSharedRef());
+		if (SlateTextRun.Get().DoesSharedInstanceExist())
+		{
+			SlateTextRun.Get().ApplyFontSizeMultiplierOnTextStyle(FontSizeMultiplier);
+		}
+	}
+
 	Runs.Add(Run.ToSharedRef());
 }
 
