Defaulting SListview, STileView and STreeView to ClipToBounds on Construct as well because the parent's value for Clipping is overwritten by base args defined on children (Ex: SSCSEditorDragDropTree )

diff --git a/Engine/Source/Runtime/Slate/Public/Widgets/Views/SListView.h b/Engine/Source/Runtime/Slate/Public/Widgets/Views/SListView.h
--- a/Engine/Source/Runtime/Slate/Public/Widgets/Views/SListView.h
+++ b/Engine/Source/Runtime/Slate/Public/Widgets/Views/SListView.h
@@ -173,6 +173,8 @@ public:
 	 */
 	void Construct(const typename SListView<ItemType>::FArguments& InArgs)
 	{
+		this->Clipping = InArgs._Clipping;
+
 		this->OnGenerateRow = InArgs._OnGenerateRow;
 		this->OnEntryInitialized = InArgs._OnEntryInitialized;
 		this->OnRowReleased = InArgs._OnRowReleased;
diff --git a/Engine/Source/Runtime/Slate/Public/Widgets/Views/STileView.h b/Engine/Source/Runtime/Slate/Public/Widgets/Views/STileView.h
--- a/Engine/Source/Runtime/Slate/Public/Widgets/Views/STileView.h
+++ b/Engine/Source/Runtime/Slate/Public/Widgets/Views/STileView.h
@@ -140,6 +140,8 @@ public:
 	 */
 	void Construct( const typename STileView<ItemType>::FArguments& InArgs )
 	{
+		this->Clipping = InArgs._Clipping;
+
 		this->OnGenerateRow = InArgs._OnGenerateTile;
 		this->OnRowReleased = InArgs._OnTileReleased;
 		this->OnItemScrolledIntoView = InArgs._OnItemScrolledIntoView;
diff --git a/Engine/Source/Runtime/Slate/Public/Widgets/Views/STreeView.h b/Engine/Source/Runtime/Slate/Public/Widgets/Views/STreeView.h
--- a/Engine/Source/Runtime/Slate/Public/Widgets/Views/STreeView.h
+++ b/Engine/Source/Runtime/Slate/Public/Widgets/Views/STreeView.h
@@ -212,6 +212,8 @@ public:
 	 */
 	void Construct( const FArguments& InArgs )
 	{
+		this->Clipping = InArgs._Clipping;
+
 		this->OnGenerateRow = InArgs._OnGenerateRow;
 		this->OnRowReleased = InArgs._OnRowReleased;
 		this->OnItemScrolledIntoView = InArgs._OnItemScrolledIntoView;
