Fixed SSplitter hover not painting correctly when using slots with hidden children

diff --git a/Engine/Source/Runtime/Slate/Private/Widgets/Layout/SSplitter.cpp b/Engine/Source/Runtime/Slate/Private/Widgets/Layout/SSplitter.cpp
--- a/Engine/Source/Runtime/Slate/Private/Widgets/Layout/SSplitter.cpp
+++ b/Engine/Source/Runtime/Slate/Private/Widgets/Layout/SSplitter.cpp
@@ -166,6 +166,17 @@ int32 SSplitter::OnPaint( const FPaintArgs& Args, const FGeometry& AllottedGeome
 	// Draw the splitter above any children
 	MaxLayerId += 1;
 
+	// HoveredHandleIndex is an index into our Children array, but some of those may be hidden and missing from ArrangedChildren so we need to update the index for the loop below
+	int32 ArrangedHoveredHandleIndex = HoveredHandleIndex;
+	if (ArrangedHoveredHandleIndex != INDEX_NONE && ArrangedChildren.Num() != Children.Num())
+	{
+		const TSharedRef<SWidget>& HoveredChildToFind = Children[HoveredHandleIndex].GetWidget();
+		ArrangedHoveredHandleIndex = ArrangedChildren.IndexOfByPredicate([&HoveredChildToFind](const FArrangedWidget& InArrangedWidget)
+		{
+			return InArrangedWidget.Widget == HoveredChildToFind;
+		});
+	}
+
 	for( int32 ChildIndex = 0; ChildIndex < ArrangedChildren.Num() - 1; ++ChildIndex )
 	{
 		const FGeometry& GeometryAfterSplitter = ArrangedChildren[ FMath::Clamp(ChildIndex + 1, 0, ArrangedChildren.Num()-1) ].Geometry;
@@ -186,7 +197,7 @@ int32 SSplitter::OnPaint( const FPaintArgs& Args, const FGeometry& AllottedGeome
 			HandlePosition.Set( 0, -PhysicalSplitterHandleSize/*-(HalfHitDetectionSplitterHandleSize + HalfPhysicalSplitterHandleSize) */);
 		}
 
-		if (HoveredHandleIndex != ChildIndex)
+		if (ArrangedHoveredHandleIndex != ChildIndex)
 		{
 			FSlateDrawElement::MakeBox(
 				OutDrawElements,
