Fix issues with password fields not ignoring word breaks
Fix bug where Shift+Delete would not cut text

diff --git a/Engine/Source/Runtime/Slate/Private/Widgets/Text/SlateEditableTextLayout.cpp b/Engine/Source/Runtime/Slate/Private/Widgets/Text/SlateEditableTextLayout.cpp
--- a/Engine/Source/Runtime/Slate/Private/Widgets/Text/SlateEditableTextLayout.cpp
+++ b/Engine/Source/Runtime/Slate/Private/Widgets/Text/SlateEditableTextLayout.cpp
@@ -856,25 +856,43 @@ FReply FSlateEditableTextLayout::HandleKeyDown(const FKeyEvent& InKeyEvent)
 
 	if (Key == EKeys::Left)
 	{
-		Reply = BoolToReply(MoveCursor(FMoveCursor::Cardinal(
-			// Ctrl moves a whole word instead of one character.	
-			InKeyEvent.IsControlDown() ? ECursorMoveGranularity::Word : ECursorMoveGranularity::Character,
-			// Move left
-			FIntPoint(-1, 0),
-			// Shift selects text.	
-			InKeyEvent.IsShiftDown() ? ECursorAction::SelectText : ECursorAction::MoveCursor
+		if (OwnerWidget->IsTextPassword() && InKeyEvent.IsControlDown())
+		{
+			// If the text is sensitive, we should not clue the user in to where word breaks are
+			JumpTo(ETextLocation::BeginningOfLine, InKeyEvent.IsShiftDown() ? ECursorAction::SelectText : ECursorAction::MoveCursor);
+			Reply = FReply::Handled();
+		}
+		else
+		{
+			Reply = BoolToReply(MoveCursor(FMoveCursor::Cardinal(
+				// Ctrl moves a whole word instead of one character.	
+				InKeyEvent.IsControlDown() ? ECursorMoveGranularity::Word : ECursorMoveGranularity::Character,
+				// Move left
+				FIntPoint(-1, 0),
+				// Shift selects text.	
+				InKeyEvent.IsShiftDown() ? ECursorAction::SelectText : ECursorAction::MoveCursor
 			)));
+		}
 	}
 	else if (Key == EKeys::Right)
 	{
-		Reply = BoolToReply(MoveCursor(FMoveCursor::Cardinal(
-			// Ctrl moves a whole word instead of one character.	
-			InKeyEvent.IsControlDown() ? ECursorMoveGranularity::Word : ECursorMoveGranularity::Character,
-			// Move right
-			FIntPoint(+1, 0),
-			// Shift selects text.	
-			InKeyEvent.IsShiftDown() ? ECursorAction::SelectText : ECursorAction::MoveCursor
+		if (OwnerWidget->IsTextPassword() && InKeyEvent.IsControlDown())
+		{
+			// If the text is sensitive, we should not clue the user in to where word breaks are
+			JumpTo(ETextLocation::EndOfLine, InKeyEvent.IsShiftDown() ? ECursorAction::SelectText : ECursorAction::MoveCursor);
+			Reply = FReply::Handled();
+		}
+		else
+		{
+			Reply = BoolToReply(MoveCursor(FMoveCursor::Cardinal(
+				// Ctrl moves a whole word instead of one character.	
+				InKeyEvent.IsControlDown() ? ECursorMoveGranularity::Word : ECursorMoveGranularity::Character,
+				// Move right
+				FIntPoint(+1, 0),
+				// Shift selects text.	
+				InKeyEvent.IsShiftDown() ? ECursorAction::SelectText : ECursorAction::MoveCursor
 			)));
+		}
 	}
 	else if (Key == EKeys::Up)
 	{
@@ -934,24 +952,6 @@ FReply FSlateEditableTextLayout::HandleKeyDown(const FKeyEvent& InKeyEvent)
 		HandleCarriageReturn();
 		Reply = FReply::Handled();
 	}
-	else if (Key == EKeys::Delete && !OwnerWidget->IsTextReadOnly())
-	{
-		// @Todo: Slate keybindings support more than one set of keys. 
-		// Delete to next word boundary (Ctrl+Delete)
-		if (InKeyEvent.IsControlDown() && !InKeyEvent.IsAltDown() && !InKeyEvent.IsShiftDown())
-		{
-			MoveCursor(FMoveCursor::Cardinal(
-				ECursorMoveGranularity::Word,
-				// Move right
-				FIntPoint(+1, 0),
-				// selects text.	
-				ECursorAction::SelectText
-				));
-		}
-
-		FScopedEditableTextTransaction TextTransaction(*this);
-		Reply = BoolToReply(HandleDelete());
-	}
 	else if (Key == EKeys::Tab && OwnerWidget->CanTypeCharacter(TEXT('\t')))
 	{
 		Reply = FReply::Handled();
@@ -969,7 +969,33 @@ FReply FSlateEditableTextLayout::HandleKeyDown(const FKeyEvent& InKeyEvent)
 		CutSelectedTextToClipboard();
 		Reply = FReply::Handled();
 	}
+	// This must come after the Cut hotkey or else Cut is unreachable
+	else if (Key == EKeys::Delete && !OwnerWidget->IsTextReadOnly())
+	{
+		// @Todo: Slate keybindings support more than one set of keys. 
+		// Delete to next word boundary (Ctrl+Delete)
+		if (InKeyEvent.IsControlDown() && !InKeyEvent.IsAltDown() && !InKeyEvent.IsShiftDown())
+		{
+			if (OwnerWidget->IsTextPassword())
+			{
+				// If the text is sensitive, we should not clue the user in to where word breaks are
+				JumpTo(ETextLocation::EndOfLine, ECursorAction::SelectText);
+			}
+			else
+			{
+				MoveCursor(FMoveCursor::Cardinal(
+					ECursorMoveGranularity::Word,
+					// Move right
+					FIntPoint(+1, 0),
+					// selects text.	
+					ECursorAction::SelectText
+				));
+			}
+		}
 
+		FScopedEditableTextTransaction TextTransaction(*this);
+		Reply = BoolToReply(HandleDelete());
+	}
 	// @Todo: Slate keybindings support more than one set of keys. 
 	// Alternate key for copy (Ctrl+Insert) 
 	else if (Key == EKeys::Insert && InKeyEvent.IsControlDown() && CanExecuteCopy())
@@ -1013,12 +1039,20 @@ FReply FSlateEditableTextLayout::HandleKeyDown(const FKeyEvent& InKeyEvent)
 	{
 		FScopedEditableTextTransaction TextTransaction(*this);
 
-		MoveCursor(FMoveCursor::Cardinal(
-			ECursorMoveGranularity::Word,
-			// Move left
-			FIntPoint(-1, 0),
-			ECursorAction::SelectText
+		if (OwnerWidget->IsTextPassword())
+		{
+			// If the text is sensitive, we should not clue the user in to where word breaks are
+			JumpTo(ETextLocation::BeginningOfLine, ECursorAction::SelectText);
+		}
+		else
+		{
+			MoveCursor(FMoveCursor::Cardinal(
+				ECursorMoveGranularity::Word,
+				// Move left
+				FIntPoint(-1, 0),
+				ECursorAction::SelectText
 			));
+		}
 		Reply = BoolToReply(HandleBackspace());
 	}
 
