Fix bad merge for SSlider

diff --git a/Engine/Source/Runtime/Slate/Private/Widgets/Input/SSlider.cpp b/Engine/Source/Runtime/Slate/Private/Widgets/Input/SSlider.cpp
--- a/Engine/Source/Runtime/Slate/Private/Widgets/Input/SSlider.cpp
+++ b/Engine/Source/Runtime/Slate/Private/Widgets/Input/SSlider.cpp
@@ -18,6 +18,7 @@ void SSlider::Construct( const SSlider::FArguments& InDeclaration )
 	LockedAttribute = InDeclaration._Locked;
 	Orientation = InDeclaration._Orientation;
 	StepSize = InDeclaration._StepSize;
+	ValueAttribute = InDeclaration._Value;
 	MinValue = InDeclaration._MinValue;
 	MaxValue = InDeclaration._MaxValue;
 	SliderBarColor = InDeclaration._SliderBarColor;
@@ -50,7 +51,7 @@ int32 SSlider::OnPaint( const FPaintArgs& Args, const FGeometry& AllottedGeometr
 	const float Indentation = IndentHandle.Get() ? HandleSize.X : 0.0f;
 
 	const float SliderLength = AllottedWidth - (Indentation + HandleSize.X);
-	const float SliderPercent = GetNormalizedValue();
+	const float SliderPercent = FMath::Clamp(GetNormalizedValue(), 0.0f, 1.0f);
 	const float SliderHandleOffset = SliderPercent * SliderLength;
 	const float SliderY = 0.5f * AllottedHeight;
 
