Correction to TextLayout to correctly move a Non TextRun when inserting before it.

diff --git a/Engine/Source/Runtime/Slate/Private/Framework/Text/TextLayout.cpp b/Engine/Source/Runtime/Slate/Private/Framework/Text/TextLayout.cpp
--- a/Engine/Source/Runtime/Slate/Private/Framework/Text/TextLayout.cpp
+++ b/Engine/Source/Runtime/Slate/Private/Framework/Text/TextLayout.cpp
@@ -1664,7 +1664,8 @@ bool FTextLayout::InsertAt(const FTextLocation& Location, const FString& Text)
 				{
 					// Insert the new text run to the left of the non-text run
 					TSharedRef<IRun> NewTextRun = CreateDefaultTextRun( LineModel.Text, FTextRange( RunRange.BeginIndex, RunRange.BeginIndex + Text.Len() ) );
-					RunModel.SetTextRange( FTextRange( RunRange.BeginIndex + 1, RunRange.EndIndex + Text.Len() ) );
+					// Move the Non-Text Run to right to free space for the new run.
+					RunModel.SetTextRange( FTextRange( RunRange.BeginIndex + Text.Len(), RunRange.EndIndex + Text.Len() ) );
 					LineModel.Runs.Insert( NewTextRun, RunIndex++ );
 				}
 				else
