Fix bad SSlider merge again

diff --git a/Engine/Source/Runtime/Slate/Private/Widgets/Input/SSlider.cpp b/Engine/Source/Runtime/Slate/Private/Widgets/Input/SSlider.cpp
--- a/Engine/Source/Runtime/Slate/Private/Widgets/Input/SSlider.cpp
+++ b/Engine/Source/Runtime/Slate/Private/Widgets/Input/SSlider.cpp
@@ -51,7 +51,7 @@ int32 SSlider::OnPaint( const FPaintArgs& Args, const FGeometry& AllottedGeometr
 	const float Indentation = IndentHandle.Get() ? HandleSize.X : 0.0f;
 
 	// We clamp to make sure that the slider cannot go out of the slider Length.
-	const float SliderPercent = FMath::Clamp(ValueAttribute.Get(), 0.0f, 1.0f); 
+	const float SliderPercent = FMath::Clamp(GetNormalizedValue(), 0.0f, 1.0f); 
 	const float SliderLength = AllottedWidth - (Indentation + HandleSize.X);
 	const float SliderHandleOffset = SliderPercent * SliderLength;
 	const float SliderY = 0.5f * AllottedHeight;
