Added an ini option to require focus in order to process gamepad input

diff --git a/Engine/Config/BaseEngine.ini b/Engine/Config/BaseEngine.ini
--- a/Engine/Config/BaseEngine.ini
+++ b/Engine/Config/BaseEngine.ini
@@ -1798,6 +1798,9 @@ NumPreallocatedVertices=50000
 [MobileSlateUI]
 bTouchFallbackToMouse=true
 
+[GamepadControl]
+bRequireFocusForGamepadInput=false
+
 [Nadzorca]
 StagingDir=~/LinuxServer/
 RunningDir=../../../
diff --git a/Engine/Source/Runtime/Slate/Private/Framework/Application/SlateApplication.cpp b/Engine/Source/Runtime/Slate/Private/Framework/Application/SlateApplication.cpp
--- a/Engine/Source/Runtime/Slate/Private/Framework/Application/SlateApplication.cpp
+++ b/Engine/Source/Runtime/Slate/Private/Framework/Application/SlateApplication.cpp
@@ -1039,6 +1039,7 @@ FSlateApplication::FSlateApplication()
 	{
 		GConfig->GetBool(TEXT("MobileSlateUI"), TEXT("bTouchFallbackToMouse"), bTouchFallbackToMouse, GEngineIni);
 		GConfig->GetBool(TEXT("CursorControl"), TEXT("bAllowSoftwareCursor"), bSoftwareCursorAvailable, GEngineIni);
+		GConfig->GetBool(TEXT("GamepadControl"), TEXT("bRequireFocusForGamepadInput"), bRequireFocusForGamepadInput, GEngineIni);
 	}
 
 	bRenderOffScreen = FParse::Param(FCommandLine::Get(), TEXT("RenderOffScreen"));
@@ -1605,7 +1606,7 @@ void FSlateApplication::PrivateDrawWindows( TSharedPtr<SWindow> DrawOnlyThisWind
 
 void FSlateApplication::PollGameDeviceState()
 {
-	if( ActiveModalWindows.Num() == 0 && !GIntraFrameDebuggingGameThread )
+	if( ActiveModalWindows.Num() == 0 && !GIntraFrameDebuggingGameThread && (!bRequireFocusForGamepadInput || IsActive()))
 	{
 		// Don't poll when a modal window open or intra frame debugging is happening
 		PlatformApplication->PollGameDeviceState( GetDeltaTime() );
diff --git a/Engine/Source/Runtime/Slate/Public/Framework/Application/SlateApplication.h b/Engine/Source/Runtime/Slate/Public/Framework/Application/SlateApplication.h
--- a/Engine/Source/Runtime/Slate/Public/Framework/Application/SlateApplication.h
+++ b/Engine/Source/Runtime/Slate/Public/Framework/Application/SlateApplication.h
@@ -2103,6 +2103,9 @@ private:
 	/** .ini controlled option to allow or disallow software cursor rendering */
 	bool bSoftwareCursorAvailable;
 
+	/** .ini controlled option to filter out gamepad input when the engine window isn't active */
+	bool bRequireFocusForGamepadInput;
+
 	/** The OS or actions taken by the user may require we refresh the current state of the cursor. */
 	bool bQueryCursorRequested;
 
