Fixed issue where SpinBox widget would clamp initial values to the slider min/max instead of the true min/max

diff --git a/Engine/Source/Runtime/Slate/Public/Widgets/Input/SSpinBox.h b/Engine/Source/Runtime/Slate/Public/Widgets/Input/SSpinBox.h
--- a/Engine/Source/Runtime/Slate/Public/Widgets/Input/SSpinBox.h
+++ b/Engine/Source/Runtime/Slate/Public/Widgets/Input/SSpinBox.h
@@ -689,7 +689,7 @@ public:
 	void SetValue(const TAttribute<NumericType>& InValueAttribute) 
 	{
 		ValueAttribute = InValueAttribute; 
-		CommitValue(InValueAttribute.Get(), ECommitMethod::CommittedViaSpin, ETextCommit::Default);
+		CommitValue(InValueAttribute.Get(), ECommitMethod::CommittedViaTypeIn, ETextCommit::Default);
 	}
 
 	/** See the MinValue attribute */
