Expand accessibility support in UMG

diff --git a/Engine/Source/Runtime/Slate/Private/Widgets/Accessibility/SlateAccessibleWidgets.cpp b/Engine/Source/Runtime/Slate/Private/Widgets/Accessibility/SlateAccessibleWidgets.cpp
--- a/Engine/Source/Runtime/Slate/Private/Widgets/Accessibility/SlateAccessibleWidgets.cpp
+++ b/Engine/Source/Runtime/Slate/Private/Widgets/Accessibility/SlateAccessibleWidgets.cpp
@@ -6,6 +6,7 @@
 #include "Widgets/Input/SButton.h"
 #include "Widgets/Input/SCheckBox.h"
 #include "Widgets/Input/SEditableText.h"
+#include "Widgets/Input/SEditableTextBox.h"
 #include "Widgets/Input/SSlider.h"
 #include "Widgets/Text/STextBlock.h"
 
@@ -37,22 +38,14 @@ bool FSlateAccessibleCheckBox::GetCheckedState() const {
 	return false;
 }
 
-FString FSlateAccessibleEditableText::GetWidgetName() const
-{
-	return FSlateAccessibleWidget::GetWidgetName().Left(50);
-}
-
 const FString& FSlateAccessibleEditableText::GetText() const
 {
 	if (Widget.IsValid())
 	{
 		return StaticCastSharedPtr<SEditableText>(Widget.Pin())->GetText().ToString();
 	}
-	else
-	{
-		static FString EmptyString;
-		return EmptyString;
-	}
+	static FString EmptyString;
+	return EmptyString;
 }
 
 bool FSlateAccessibleEditableText::IsPassword() const
@@ -86,6 +79,47 @@ void FSlateAccessibleEditableText::SetValue(const FString& Value)
 	}
 }
 
+const FString& FSlateAccessibleEditableTextBox::GetText() const
+{
+	if (Widget.IsValid())
+	{
+		return StaticCastSharedPtr<SEditableTextBox>(Widget.Pin())->GetText().ToString();
+	}
+	static FString EmptyString;
+	return EmptyString;
+}
+
+bool FSlateAccessibleEditableTextBox::IsPassword() const
+{
+	if (Widget.IsValid())
+	{
+		return StaticCastSharedPtr<SEditableTextBox>(Widget.Pin())->IsPassword();
+	}
+	return true;
+}
+
+bool FSlateAccessibleEditableTextBox::IsReadOnly() const
+{
+	if (Widget.IsValid())
+	{
+		return StaticCastSharedPtr<SEditableTextBox>(Widget.Pin())->IsReadOnly();
+	}
+	return true;
+}
+
+FString FSlateAccessibleEditableTextBox::GetValue() const
+{
+	return GetText();
+}
+
+void FSlateAccessibleEditableTextBox::SetValue(const FString& Value)
+{
+	if (Widget.IsValid())
+	{
+		StaticCastSharedPtr<SEditableTextBox>(Widget.Pin())->SetText(FText::FromString(Value));
+	}
+}
+
 bool FSlateAccessibleSlider::IsReadOnly() const
 {
 	if (Widget.IsValid())
@@ -139,11 +173,6 @@ void FSlateAccessibleSlider::SetValue(const FString& Value)
 	}
 }
 
-FString FSlateAccessibleTextBlock::GetWidgetName() const
-{
-	return FSlateAccessibleWidget::GetWidgetName().Left(50);
-}
-
 const FString& FSlateAccessibleTextBlock::GetText() const
 {
 	if (Widget.IsValid())
diff --git a/Engine/Source/Runtime/Slate/Private/Widgets/Input/SEditableTextBox.cpp b/Engine/Source/Runtime/Slate/Private/Widgets/Input/SEditableTextBox.cpp
--- a/Engine/Source/Runtime/Slate/Private/Widgets/Input/SEditableTextBox.cpp
+++ b/Engine/Source/Runtime/Slate/Private/Widgets/Input/SEditableTextBox.cpp
@@ -5,6 +5,10 @@
 #include "Widgets/Layout/SBox.h"
 #include "Widgets/Notifications/SPopUpErrorText.h"
 
+#if WITH_ACCESSIBILITY
+#include "Widgets/Accessibility/SlateAccessibleWidgets.h"
+#endif
+
 /**
  * Construct this widget
  *
@@ -80,6 +84,17 @@ void SEditableTextBox::Construct( const FArguments& InArgs )
 		];
 	}
 
+#if WITH_ACCESSIBILITY
+	// These functions need to be called again because SWidget::Construct is called before SEditableText is created.
+	if (!InArgs._AccessibleParams.AccessibleText.IsSet())
+	{
+		SetDefaultAccessibleText(EAccessibleType::Main);
+	}
+	if (!InArgs._AccessibleParams.AccessibleSummaryText.IsSet())
+	{
+		SetDefaultAccessibleText(EAccessibleType::Summary);
+	}
+#endif
 }
 
 void SEditableTextBox::SetStyle(const FEditableTextBoxStyle* InStyle)
@@ -396,3 +411,21 @@ void SEditableTextBox::SetVirtualKeyboardDismissAction(TAttribute<EVirtualKeyboa
 {
 	EditableText->SetVirtualKeyboardDismissAction(InVirtualKeyboardDismissAction);
 }
+
+#if WITH_ACCESSIBILITY
+TSharedPtr<FSlateAccessibleWidget> SEditableTextBox::CreateAccessibleWidget()
+{
+	return MakeShareable<FSlateAccessibleWidget>(new FSlateAccessibleEditableTextBox(SharedThis(this)));
+}
+
+void SEditableTextBox::SetDefaultAccessibleText(EAccessibleType AccessibleType)
+{
+	// The parent Construct() function will call this before EditableText exists,
+	// so we need a guard here to ignore that function call.
+	if (EditableText.IsValid())
+	{
+		TAttribute<FText>& Text = (AccessibleType == EAccessibleType::Main) ? AccessibleData.AccessibleText : AccessibleData.AccessibleSummaryText;
+		Text.Bind(EditableText.Get(), &SEditableText::GetHintText);
+	}
+}
+#endif
diff --git a/Engine/Source/Runtime/Slate/Public/Widgets/Accessibility/SlateAccessibleWidgets.h b/Engine/Source/Runtime/Slate/Public/Widgets/Accessibility/SlateAccessibleWidgets.h
--- a/Engine/Source/Runtime/Slate/Public/Widgets/Accessibility/SlateAccessibleWidgets.h
+++ b/Engine/Source/Runtime/Slate/Public/Widgets/Accessibility/SlateAccessibleWidgets.h
@@ -53,7 +53,33 @@ public:
 	FSlateAccessibleEditableText(TWeakPtr<SWidget> InWidget) : FSlateAccessibleWidget(InWidget, EAccessibleWidgetType::TextEdit) {}
 
 	// IAccessibleWidget
-	virtual FString GetWidgetName() const override;
+	virtual IAccessibleText* AsText() override { return this; }
+	virtual IAccessibleProperty* AsProperty() override { return this; }
+	// ~
+
+	// IAccessibleText
+	virtual const FString& GetText() const override;
+	// ~
+
+	// IAccessibleProperty
+	virtual bool IsReadOnly() const override;
+	virtual bool IsPassword() const override;
+	virtual FString GetValue() const override;
+	virtual void SetValue(const FString& Value) override;
+	// ~
+};
+// ~
+
+// SEditableTextBox
+class SLATE_API FSlateAccessibleEditableTextBox
+	: public FSlateAccessibleWidget
+	, public IAccessibleText
+	, public IAccessibleProperty
+{
+public:
+	FSlateAccessibleEditableTextBox(TWeakPtr<SWidget> InWidget) : FSlateAccessibleWidget(InWidget, EAccessibleWidgetType::TextEdit) {}
+
+	// IAccessibleWidget
 	virtual IAccessibleText* AsText() override { return this; }
 	virtual IAccessibleProperty* AsProperty() override { return this; }
 	// ~
@@ -123,7 +149,6 @@ public:
 	FSlateAccessibleTextBlock(TWeakPtr<SWidget> InWidget) : FSlateAccessibleWidget(InWidget, EAccessibleWidgetType::Text) {}
 
 	// IAccessibleWidget
-	virtual FString GetWidgetName() const override;
 	//virtual IAccessibleText* AsText() override { return this; }
 	// ~
 
diff --git a/Engine/Source/Runtime/Slate/Public/Widgets/Input/SEditableTextBox.h b/Engine/Source/Runtime/Slate/Public/Widgets/Input/SEditableTextBox.h
--- a/Engine/Source/Runtime/Slate/Public/Widgets/Input/SEditableTextBox.h
+++ b/Engine/Source/Runtime/Slate/Public/Widgets/Input/SEditableTextBox.h
@@ -52,7 +52,9 @@ public:
 		, _VirtualKeyboardOptions(FVirtualKeyboardOptions())
 		, _VirtualKeyboardTrigger(EVirtualKeyboardTrigger::OnFocusByPointer)
 		, _VirtualKeyboardDismissAction(EVirtualKeyboardDismissAction::TextChangeOnDismiss)
-		{}
+		{
+			_AccessibleParams = FAccessibleWidgetData(EAccessibleBehavior::Auto, EAccessibleBehavior::Auto, false);
+		}
 
 		/** The styling of the textbox */
 		SLATE_STYLE_ARGUMENT( FEditableTextBoxStyle, Style )
@@ -187,10 +189,14 @@ public:
 
 	/** See the IsReadOnly attribute */
 	void SetIsReadOnly( TAttribute< bool > InIsReadOnly );
+
+	bool IsReadOnly() const { return EditableText->IsTextReadOnly(); }
 	
 	/** See the IsPassword attribute */
 	void SetIsPassword( TAttribute< bool > InIsPassword );
 
+	bool IsPassword() const { return EditableText->IsTextPassword(); }
+
 	/** See the AllowContextMenu attribute */
 	void SetAllowContextMenu(TAttribute< bool > InAllowContextMenu);
 
@@ -336,6 +342,11 @@ public:
 	virtual FReply OnKeyDown(const FGeometry& MyGeometry, const FKeyEvent& InKeyEvent) override;
 
 protected:
+#if WITH_ACCESSIBILITY
+	virtual TSharedPtr<FSlateAccessibleWidget> CreateAccessibleWidget() override;
+	virtual void SetDefaultAccessibleText(EAccessibleType AccessibleType = EAccessibleType::Main) override;
+#endif
+
 	const FEditableTextBoxStyle* Style;
 
 	/** Box widget that adds padding around the editable text */
diff --git a/Engine/Source/Runtime/UMG/Private/Components/Button.cpp b/Engine/Source/Runtime/UMG/Private/Components/Button.cpp
--- a/Engine/Source/Runtime/UMG/Private/Components/Button.cpp
+++ b/Engine/Source/Runtime/UMG/Private/Components/Button.cpp
@@ -214,6 +214,13 @@ void UButton::SlateHandleUnhovered()
 	OnUnhovered.Broadcast();
 }
 
+#if WITH_ACCESSIBILITY
+TSharedPtr<SWidget> UButton::GetAccessibleWidget() const
+{
+	return MyButton;
+}
+#endif
+
 #if WITH_EDITOR
 
 const FText UButton::GetPaletteCategory()
diff --git a/Engine/Source/Runtime/UMG/Private/Components/CheckBox.cpp b/Engine/Source/Runtime/UMG/Private/Components/CheckBox.cpp
--- a/Engine/Source/Runtime/UMG/Private/Components/CheckBox.cpp
+++ b/Engine/Source/Runtime/UMG/Private/Components/CheckBox.cpp
@@ -222,6 +222,13 @@ void UCheckBox::PostLoad()
 	}
 }
 
+#if WITH_ACCESSIBILITY
+TSharedPtr<SWidget> UCheckBox::GetAccessibleWidget() const
+{
+	return MyCheckbox;
+}
+#endif
+
 #if WITH_EDITOR
 
 const FText UCheckBox::GetPaletteCategory()
diff --git a/Engine/Source/Runtime/UMG/Private/Components/EditableText.cpp b/Engine/Source/Runtime/UMG/Private/Components/EditableText.cpp
--- a/Engine/Source/Runtime/UMG/Private/Components/EditableText.cpp
+++ b/Engine/Source/Runtime/UMG/Private/Components/EditableText.cpp
@@ -38,6 +38,9 @@ UEditableText::UEditableText(const FObjectInitializer& ObjectInitializer)
 	AllowContextMenu = true;
 	VirtualKeyboardDismissAction = EVirtualKeyboardDismissAction::TextChangeOnDismiss;
 	Clipping = EWidgetClipping::ClipToBounds;
+
+	AccessibleBehavior = ESlateAccessibleBehavior::Auto;
+	bCanChildrenBeAccessible = false;
 }
 
 void UEditableText::ReleaseSlateResources(bool bReleaseChildren)
@@ -200,6 +203,13 @@ void UEditableText::PostLoad()
 	}
 }
 
+#if WITH_ACCESSIBILITY
+TSharedPtr<SWidget> UEditableText::GetAccessibleWidget() const
+{
+	return MyEditableText;
+}
+#endif
+
 #if WITH_EDITOR
 
 const FText UEditableText::GetPaletteCategory()
diff --git a/Engine/Source/Runtime/UMG/Private/Components/EditableTextBox.cpp b/Engine/Source/Runtime/UMG/Private/Components/EditableTextBox.cpp
--- a/Engine/Source/Runtime/UMG/Private/Components/EditableTextBox.cpp
+++ b/Engine/Source/Runtime/UMG/Private/Components/EditableTextBox.cpp
@@ -39,6 +39,9 @@ UEditableTextBox::UEditableTextBox(const FObjectInitializer& ObjectInitializer)
 	// HACK: THIS SHOULD NOT COME FROM CORESTYLE AND SHOULD INSTEAD BY DEFINED BY ENGINE TEXTURES/PROJECT SETTINGS
 	static const FEditableTextBoxStyle StaticNormalEditableTextBox = FCoreStyle::Get().GetWidgetStyle< FEditableTextBoxStyle >("NormalEditableTextBox");
 	WidgetStyle = StaticNormalEditableTextBox;
+
+	AccessibleBehavior = ESlateAccessibleBehavior::Auto;
+	bCanChildrenBeAccessible = false;
 }
 
 void UEditableTextBox::ReleaseSlateResources(bool bReleaseChildren)
@@ -227,6 +230,13 @@ void UEditableTextBox::PostLoad()
 	}
 }
 
+#if WITH_ACCESSIBILITY
+TSharedPtr<SWidget> UEditableTextBox::GetAccessibleWidget() const
+{
+	return MyEditableTextBlock;
+}
+#endif
+
 #if WITH_EDITOR
 
 const FText UEditableTextBox::GetPaletteCategory()
diff --git a/Engine/Source/Runtime/UMG/Private/Components/Image.cpp b/Engine/Source/Runtime/UMG/Private/Components/Image.cpp
--- a/Engine/Source/Runtime/UMG/Private/Components/Image.cpp
+++ b/Engine/Source/Runtime/UMG/Private/Components/Image.cpp
@@ -395,6 +395,13 @@ FReply UImage::HandleMouseButtonDown(const FGeometry& Geometry, const FPointerEv
 	return FReply::Unhandled();
 }
 
+#if WITH_ACCESSIBILITY
+TSharedPtr<SWidget> UImage::GetAccessibleWidget() const
+{
+	return MyImage;
+}
+#endif
+
 #if WITH_EDITOR
 
 const FText UImage::GetPaletteCategory()
diff --git a/Engine/Source/Runtime/UMG/Private/Components/Slider.cpp b/Engine/Source/Runtime/UMG/Private/Components/Slider.cpp
--- a/Engine/Source/Runtime/UMG/Private/Components/Slider.cpp
+++ b/Engine/Source/Runtime/UMG/Private/Components/Slider.cpp
@@ -196,6 +196,12 @@ void USlider::SetSliderBarColor(FLinearColor InValue)
 	}
 }
 
+#if WITH_ACCESSIBILITY
+TSharedPtr<SWidget> USlider::GetAccessibleWidget() const
+{
+	return MySlider;
+}
+#endif
 
 #if WITH_EDITOR
 
diff --git a/Engine/Source/Runtime/UMG/Private/Components/TextBlock.cpp b/Engine/Source/Runtime/UMG/Private/Components/TextBlock.cpp
--- a/Engine/Source/Runtime/UMG/Private/Components/TextBlock.cpp
+++ b/Engine/Source/Runtime/UMG/Private/Components/TextBlock.cpp
@@ -223,6 +223,13 @@ EVisibility UTextBlock::GetTextWarningImageVisibility() const
 	return Text.IsCultureInvariant() ? EVisibility::Visible : EVisibility::Collapsed;
 }
 
+#if WITH_ACCESSIBILITY
+TSharedPtr<SWidget> UTextBlock::GetAccessibleWidget() const
+{
+	return MyTextBlock;
+}
+#endif
+
 void UTextBlock::OnBindingChanged(const FName& Property)
 {
 	Super::OnBindingChanged(Property);
diff --git a/Engine/Source/Runtime/UMG/Private/Components/Widget.cpp b/Engine/Source/Runtime/UMG/Private/Components/Widget.cpp
--- a/Engine/Source/Runtime/UMG/Private/Components/Widget.cpp
+++ b/Engine/Source/Runtime/UMG/Private/Components/Widget.cpp
@@ -1167,9 +1167,13 @@ void UWidget::SynchronizeProperties()
 	}
 
 #if WITH_ACCESSIBILITY
-	SafeWidget->SetAccessibleBehavior((EAccessibleBehavior)AccessibleBehavior, PROPERTY_BINDING(FText, AccessibleText), EAccessibleType::Main);
-	SafeWidget->SetAccessibleBehavior((EAccessibleBehavior)AccessibleSummaryBehavior, PROPERTY_BINDING(FText, AccessibleSummaryText), EAccessibleType::Summary);
-	SafeWidget->SetCanChildrenBeAccessible(bCanChildrenBeAccessible);
+	TSharedPtr<SWidget> AccessibleWidget = GetAccessibleWidget();
+	if (AccessibleWidget.IsValid())
+	{
+		AccessibleWidget->SetAccessibleBehavior((EAccessibleBehavior)AccessibleBehavior, PROPERTY_BINDING(FText, AccessibleText), EAccessibleType::Main);
+		AccessibleWidget->SetAccessibleBehavior((EAccessibleBehavior)AccessibleSummaryBehavior, PROPERTY_BINDING(FText, AccessibleSummaryText), EAccessibleType::Summary);
+		AccessibleWidget->SetCanChildrenBeAccessible(bCanChildrenBeAccessible);
+	}
 #endif
 
 #if !UE_BUILD_SHIPPING
@@ -1177,6 +1181,13 @@ void UWidget::SynchronizeProperties()
 #endif
 }
 
+#if WITH_ACCESSIBILITY
+TSharedPtr<SWidget> UWidget::GetAccessibleWidget() const
+{
+	return GetCachedWidget();
+}
+#endif
+
 UObject* UWidget::GetSourceAssetOrClass() const
 {
 	UObject* SourceAsset = nullptr;
diff --git a/Engine/Source/Runtime/UMG/Public/Components/Button.h b/Engine/Source/Runtime/UMG/Public/Components/Button.h
--- a/Engine/Source/Runtime/UMG/Public/Components/Button.h
+++ b/Engine/Source/Runtime/UMG/Public/Components/Button.h
@@ -158,6 +158,10 @@ protected:
 #endif
 	//~ End UWidget Interface
 
+#if WITH_ACCESSIBILITY
+	virtual TSharedPtr<SWidget> GetAccessibleWidget() const override;
+#endif
+
 protected:
 	/** Cached pointer to the underlying slate button owned by this UWidget */
 	TSharedPtr<SButton> MyButton;
diff --git a/Engine/Source/Runtime/UMG/Public/Components/CheckBox.h b/Engine/Source/Runtime/UMG/Public/Components/CheckBox.h
--- a/Engine/Source/Runtime/UMG/Public/Components/CheckBox.h
+++ b/Engine/Source/Runtime/UMG/Public/Components/CheckBox.h
@@ -163,6 +163,10 @@ protected:
 	//~ End UWidget Interface
 
 	void SlateOnCheckStateChangedCallback(ECheckBoxState NewState);
+
+#if WITH_ACCESSIBILITY
+	virtual TSharedPtr<SWidget> GetAccessibleWidget() const override;
+#endif
 	
 protected:
 	TSharedPtr<SCheckBox> MyCheckbox;
diff --git a/Engine/Source/Runtime/UMG/Public/Components/EditableText.h b/Engine/Source/Runtime/UMG/Public/Components/EditableText.h
--- a/Engine/Source/Runtime/UMG/Public/Components/EditableText.h
+++ b/Engine/Source/Runtime/UMG/Public/Components/EditableText.h
@@ -196,6 +196,10 @@ protected:
 	void HandleOnTextChanged(const FText& Text);
 	void HandleOnTextCommitted(const FText& Text, ETextCommit::Type CommitMethod);
 
+#if WITH_ACCESSIBILITY
+	virtual TSharedPtr<SWidget> GetAccessibleWidget() const override;
+#endif
+
 protected:
 	TSharedPtr<SEditableText> MyEditableText;
 
diff --git a/Engine/Source/Runtime/UMG/Public/Components/EditableTextBox.h b/Engine/Source/Runtime/UMG/Public/Components/EditableTextBox.h
--- a/Engine/Source/Runtime/UMG/Public/Components/EditableTextBox.h
+++ b/Engine/Source/Runtime/UMG/Public/Components/EditableTextBox.h
@@ -200,6 +200,10 @@ protected:
 	virtual void HandleOnTextChanged(const FText& Text);
 	virtual void HandleOnTextCommitted(const FText& Text, ETextCommit::Type CommitMethod);
 
+#if WITH_ACCESSIBILITY
+	virtual TSharedPtr<SWidget> GetAccessibleWidget() const override;
+#endif
+
 protected:
 	TSharedPtr<SEditableTextBox> MyEditableTextBlock;
 
diff --git a/Engine/Source/Runtime/UMG/Public/Components/Image.h b/Engine/Source/Runtime/UMG/Public/Components/Image.h
--- a/Engine/Source/Runtime/UMG/Public/Components/Image.h
+++ b/Engine/Source/Runtime/UMG/Public/Components/Image.h
@@ -166,6 +166,10 @@ protected:
 	//
 	FReply HandleMouseButtonDown(const FGeometry& Geometry, const FPointerEvent& MouseEvent);
 
+#if WITH_ACCESSIBILITY
+	virtual TSharedPtr<SWidget> GetAccessibleWidget() const override;
+#endif
+
 protected:
 	TSharedPtr<SImage> MyImage;
 
diff --git a/Engine/Source/Runtime/UMG/Public/Components/Slider.h b/Engine/Source/Runtime/UMG/Public/Components/Slider.h
--- a/Engine/Source/Runtime/UMG/Public/Components/Slider.h
+++ b/Engine/Source/Runtime/UMG/Public/Components/Slider.h
@@ -174,5 +174,9 @@ protected:
 	void HandleOnControllerCaptureBegin();
 	void HandleOnControllerCaptureEnd();
 
+#if WITH_ACCESSIBILITY
+	virtual TSharedPtr<SWidget> GetAccessibleWidget() const override;
+#endif
+
 	PROPERTY_BINDING_IMPLEMENTATION(float, Value);
 };
diff --git a/Engine/Source/Runtime/UMG/Public/Components/TextBlock.h b/Engine/Source/Runtime/UMG/Public/Components/TextBlock.h
--- a/Engine/Source/Runtime/UMG/Public/Components/TextBlock.h
+++ b/Engine/Source/Runtime/UMG/Public/Components/TextBlock.h
@@ -212,6 +212,10 @@ protected:
 
 	EVisibility GetTextWarningImageVisibility() const;
 
+#if WITH_ACCESSIBILITY
+	virtual TSharedPtr<SWidget> GetAccessibleWidget() const override;
+#endif
+
 protected:
 
 	TSharedPtr<STextBlock> MyTextBlock;
diff --git a/Engine/Source/Runtime/UMG/Public/Components/Widget.h b/Engine/Source/Runtime/UMG/Public/Components/Widget.h
--- a/Engine/Source/Runtime/UMG/Public/Components/Widget.h
+++ b/Engine/Source/Runtime/UMG/Public/Components/Widget.h
@@ -929,6 +929,11 @@ protected:
 
 	void SetNavigationRuleInternal(EUINavigation Direction, EUINavigationRule Rule, FName WidgetToFocus);
 
+#if WITH_ACCESSIBILITY
+	/** Gets the widget that accessibility properties should synchronize to. */
+	virtual TSharedPtr<SWidget> GetAccessibleWidget() const;
+#endif
+
 protected:
 	/** The underlying SWidget. */
 	TWeakPtr<SWidget> MyWidget;
