Add accessibility support for SWidgetSwitcher

diff --git a/Engine/Source/Runtime/Slate/Private/Widgets/Layout/SWidgetSwitcher.cpp b/Engine/Source/Runtime/Slate/Private/Widgets/Layout/SWidgetSwitcher.cpp
--- a/Engine/Source/Runtime/Slate/Private/Widgets/Layout/SWidgetSwitcher.cpp
+++ b/Engine/Source/Runtime/Slate/Private/Widgets/Layout/SWidgetSwitcher.cpp
@@ -2,6 +2,10 @@
 
 #include "Widgets/Layout/SWidgetSwitcher.h"
 #include "Layout/LayoutUtils.h"
+#if WITH_ACCESSIBILITY
+#include "Framework/Application/SlateApplication.h"
+#include "Widgets/Accessibility/SlateAccessibleMessageHandler.h"
+#endif
 
 SWidgetSwitcher::SWidgetSwitcher()
 	: WidgetIndex()
@@ -131,6 +135,18 @@ void SWidgetSwitcher::OnArrangeChildren(const FGeometry& AllottedGeometry, FArra
 	if (const FSlot* ActiveSlotPtr = GetActiveSlot())
 	{
 		ArrangeSingleChild(GSlateFlowDirection, AllottedGeometry, ArrangedChildren, *ActiveSlotPtr, ContentScale);
+
+#if WITH_ACCESSIBILITY
+		// This would go in SetActiveWidgetIndex were this not a TAttribute. WidgetIndex changing does not affect
+		// parent or visibility assignments, so we need a special custom notification to tell the accessibility
+		// tree to update itself.
+		if (LastActiveWidget != ActiveSlotPtr->GetWidget())
+		{
+			SWidgetSwitcher* UnconstThis = const_cast<SWidgetSwitcher*>(this);
+			UnconstThis->LastActiveWidget = ActiveSlotPtr->GetWidget();
+			FSlateApplication::Get().GetAccessibleMessageHandler()->OnWidgetChildrenChanged(SharedThis(UnconstThis));
+		}
+#endif
 	}
 }
 	
diff --git a/Engine/Source/Runtime/Slate/Public/Widgets/Layout/SWidgetSwitcher.h b/Engine/Source/Runtime/Slate/Public/Widgets/Layout/SWidgetSwitcher.h
--- a/Engine/Source/Runtime/Slate/Public/Widgets/Layout/SWidgetSwitcher.h
+++ b/Engine/Source/Runtime/Slate/Public/Widgets/Layout/SWidgetSwitcher.h
@@ -169,4 +169,9 @@ private:
 
 	/** Required to implement GetChildren() in a way that can dynamically return the currently active child. */
 	TOneDynamicChild<FSlot> OneDynamicChild;
+
+#if WITH_ACCESSIBILITY
+	/** Used to detect when WidgetIndex changes while bound. */
+	TWeakPtr<SWidget> LastActiveWidget;
+#endif
 };
