Correction to SListView to handle the case where and object in the ListItems is Garbage Collected while being used by the List.

diff --git a/Engine/Source/Runtime/Slate/Public/Framework/Views/TableViewTypeTraits.h b/Engine/Source/Runtime/Slate/Public/Framework/Views/TableViewTypeTraits.h
--- a/Engine/Source/Runtime/Slate/Public/Framework/Views/TableViewTypeTraits.h
+++ b/Engine/Source/Runtime/Slate/Public/Framework/Views/TableViewTypeTraits.h
@@ -95,7 +95,11 @@ template <typename T> struct TListTypeTraits< TSharedPtr<T, ESPMode::NotThreadSa
 public:
 	typedef TSharedPtr<T> NullableType;
 
-	static void AddReferencedObjects( FReferenceCollector&, TArray< TSharedPtr<T> >&, TSet< TSharedPtr<T> >&  )
+	template<typename U>
+	static void AddReferencedObjects( FReferenceCollector&, 
+		TArray< TSharedPtr<T> >&, 
+		TSet< TSharedPtr<T> >&, 
+		TMap< const U*, TSharedPtr<T> >& )
 	{
 	}
 
@@ -128,7 +132,11 @@ template <typename T> struct TListTypeTraits< TSharedPtr<T, ESPMode::ThreadSafe>
 public:
 	typedef TSharedPtr<T, ESPMode::ThreadSafe> NullableType;
 
-	static void AddReferencedObjects( FReferenceCollector&, TArray< TSharedPtr<T, ESPMode::ThreadSafe> >&, TSet< TSharedPtr<T, ESPMode::ThreadSafe> >&  )
+	template<typename U>
+	static void AddReferencedObjects( FReferenceCollector&, 
+		TArray< TSharedPtr<T, ESPMode::ThreadSafe> >&, 
+		TSet< TSharedPtr<T, ESPMode::ThreadSafe> >&, 
+		TMap< const U*, TSharedPtr<T, ESPMode::ThreadSafe> >& WidgetToItemMap)
 	{
 	}
 
@@ -161,7 +169,11 @@ template <typename T> struct TListTypeTraits< TSharedRef<T, ESPMode::NotThreadSa
 public:
 	typedef TSharedPtr<T> NullableType;
 
-	static void AddReferencedObjects( FReferenceCollector&, TArray< TSharedRef<T> >&, TSet< TSharedRef<T> >&  )
+	template<typename U>
+	static void AddReferencedObjects( FReferenceCollector&, 
+		TArray< TSharedRef<T> >&, 
+		TSet< TSharedRef<T> >&, 
+		TMap< const U*, TSharedRef<T> >& )
 	{
 	}
 
@@ -194,7 +206,11 @@ template <typename T> struct TListTypeTraits< TSharedRef<T, ESPMode::ThreadSafe>
 public:
 	typedef TSharedPtr<T, ESPMode::ThreadSafe> NullableType;
 
-	static void AddReferencedObjects( FReferenceCollector&, TArray< TSharedRef<T, ESPMode::ThreadSafe> >&, TSet< TSharedRef<T, ESPMode::ThreadSafe> >&  )
+	template<typename U>
+	static void AddReferencedObjects( FReferenceCollector&, 
+		TArray< TSharedRef<T, ESPMode::ThreadSafe> >&, 
+		TSet< TSharedRef<T, ESPMode::ThreadSafe> >&,
+		TMap< const U*, TSharedRef<T, ESPMode::ThreadSafe> >&)
 	{
 	}
 
@@ -230,7 +246,11 @@ template <typename T> struct TListTypeTraits< TWeakObjectPtr<T> >
 public:
 	typedef TWeakObjectPtr<T> NullableType;
 
-	static void AddReferencedObjects( FReferenceCollector&, TArray< TWeakObjectPtr<T> >&, TSet< TWeakObjectPtr<T> >&  )
+	template<typename U>
+	static void AddReferencedObjects( FReferenceCollector&, 
+		TArray< TWeakObjectPtr<T> >&,
+		TSet< TWeakObjectPtr<T> >&,
+		TMap< const U*, TWeakObjectPtr<T> >&)
 	{
 	}
 
@@ -269,10 +289,21 @@ struct TListTypeTraits<T*, typename TEnableIf<TPointerIsConvertibleFromTo<T, UOb
 public:
 	typedef T* NullableType;
 
-	static void AddReferencedObjects( FReferenceCollector& Collector, TArray<T*>& ItemsWithGeneratedWidgets, TSet<T*>& SelectedItems )
+	template<typename U>
+	static void AddReferencedObjects( FReferenceCollector& Collector, 
+		TArray<T*>& ItemsWithGeneratedWidgets, 
+		TSet<T*>& SelectedItems, 
+		TMap< const U*, T* >& WidgetToItemMap)
 	{
 		// Serialize generated items
 		Collector.AddReferencedObjects(ItemsWithGeneratedWidgets);
+		
+		// Serialize the map Value. We only do it for the WidgetToItemMap because we know that both maps are updated at the same time and contains the same objects
+		// Also, we cannot AddReferencedObject to the Keys of the ItemToWidgetMap or we end up with keys being set to 0 when the UObject is destroyed which generate an invalid id in the map.
+		for (auto& It : WidgetToItemMap)
+		{
+			Collector.AddReferencedObject(*(UObject**)&It.Value);
+		}
 
 		// Serialize the selected items
 		Collector.AddReferencedObjects(SelectedItems);
@@ -295,11 +326,22 @@ struct TListTypeTraits<const T*, typename TEnableIf<TPointerIsConvertibleFromTo<
 public:
 	typedef const T* NullableType;
 
-	static void AddReferencedObjects( FReferenceCollector& Collector, TArray<const T*>& ItemsWithGeneratedWidgets, TSet<const T*>& SelectedItems )
+	template<typename U>
+	static void AddReferencedObjects( FReferenceCollector& Collector, 
+		TArray<const T*>& ItemsWithGeneratedWidgets, 
+		TSet<const T*>& SelectedItems,
+		TMap< const U*, const T* >& WidgetToItemMap)
 	{
 		// Serialize generated items
 		Collector.AddReferencedObjects(ItemsWithGeneratedWidgets);
 
+		// Serialize the map Value. We only do it for the WidgetToItemMap because we know that both maps are updated at the same time and contains the same objects
+		// Also, we cannot AddReferencedObject to the Keys of the ItemToWidgetMap or we end up with keys being set to 0 when the UObject is destroyed which generate an invalid id in the map.
+		for (auto& It : WidgetToItemMap)
+		{
+			Collector.AddReferencedObject(*(UObject**)&It.Value);
+		}
+
 		// Serialize the selected items
 		Collector.AddReferencedObjects(SelectedItems);
 	}
diff --git a/Engine/Source/Runtime/Slate/Public/Widgets/Views/SListView.h b/Engine/Source/Runtime/Slate/Public/Widgets/Views/SListView.h
--- a/Engine/Source/Runtime/Slate/Public/Widgets/Views/SListView.h
+++ b/Engine/Source/Runtime/Slate/Public/Widgets/Views/SListView.h
@@ -503,6 +503,7 @@ private:
 		 */
 		void OnItemSeen( ItemType InItem, TSharedRef<ITableRow> InGeneratedWidget)
 		{
+			ensure(TListTypeTraits<ItemType>::IsPtrValid(InItem));
 			TSharedRef<ITableRow>* LookupResult = ItemToWidgetMap.Find( InItem );
 			const bool bWidgetIsNewlyGenerated = (LookupResult == nullptr);
 			if ( bWidgetIsNewlyGenerated )
@@ -571,6 +572,26 @@ private:
 						OwnerList->OnRowReleased.ExecuteIfBound(WidgetToCleanUp);
 					}
 				}
+				else if(!TListTypeTraits<ItemType>::IsPtrValid(ItemToBeCleanedUp))
+				{
+					// If we get here, it means we have an invalid object. We will need to remove that object from both maps.
+					// This may happen for example when ItemType is a UObject* and the object is garbage collected.
+					auto Widget = WidgetMapToItem.FindKey(ItemToBeCleanedUp);
+					if (Widget != nullptr)
+					{
+						for (auto WidgetItemPair = ItemToWidgetMap.CreateIterator(); WidgetItemPair; ++WidgetItemPair)
+						{
+							const ITableRow* Item = &(WidgetItemPair.Value().Get());
+							if (Item == *Widget)
+							{
+								WidgetItemPair.RemoveCurrent();
+								break;
+							}
+						}
+
+						WidgetMapToItem.Remove(*Widget);
+					}
+				}
 			}
 
 			ItemsToBeCleanedUp.Reset();
@@ -945,6 +966,12 @@ public:
 			{
 				const ItemType& CurItem = (*SourceItems)[ItemIndex];
 
+				// We do not generatie a new widget if the CurItem is an invalid object.
+				if (!TListTypeTraits<ItemType>::IsPtrValid(CurItem))
+				{
+					continue;
+				}
+
 				const float ItemHeight = GenerateWidgetForItem(CurItem, ItemIndex, StartIndex, LayoutScaleMultiplier);
 
 				const bool bIsFirstItem = ItemIndex == StartIndex;
@@ -1003,6 +1030,11 @@ public:
 				for( int32 ItemIndex = StartIndex-1; HeightGeneratedSoFar < MyGeometry.GetLocalSize().Y && ItemIndex >= 0; --ItemIndex )
 				{
 					const ItemType& CurItem = (*SourceItems)[ItemIndex];
+					// When the item is not valid, we do not generate a widget for it.
+					if (!TListTypeTraits<ItemType>::IsPtrValid(CurItem))
+					{
+						continue;
+					}
 
 					const float ItemHeight = GenerateWidgetForItem(CurItem, ItemIndex, StartIndex, LayoutScaleMultiplier);
 
@@ -1029,6 +1061,7 @@ public:
 
 	float GenerateWidgetForItem( const ItemType& CurItem, int32 ItemIndex, int32 StartIndex, float LayoutScaleMultiplier )
 	{
+		ensure(TListTypeTraits<ItemType>::IsPtrValid(CurItem));
 		// Find a previously generated Widget for this item, if one exists.
 		TSharedPtr<ITableRow> WidgetForItem = WidgetGenerator.GetWidgetForItem( CurItem );
 		if ( !WidgetForItem.IsValid() )
@@ -1316,7 +1349,7 @@ public:
 	 */
 	virtual void AddReferencedObjects( FReferenceCollector& Collector )
 	{
-		TListTypeTraits<ItemType>::AddReferencedObjects( Collector, WidgetGenerator.ItemsWithGeneratedWidgets, SelectedItems );
+		TListTypeTraits<ItemType>::AddReferencedObjects( Collector, WidgetGenerator.ItemsWithGeneratedWidgets, SelectedItems, WidgetGenerator.WidgetMapToItem );
 	}
 
 	/**
@@ -1533,6 +1566,13 @@ protected:
 				while( AbsScrollByAmount != 0 && ItemIndex < SourceItems->Num() && ItemIndex >= 0 )
 				{
 					const ItemType CurItem = (*SourceItems)[ ItemIndex ];
+					// If the CurItem is not valid, we do not generate a new widget for it, we skip it.
+					if (!TListTypeTraits<ItemType>::IsPtrValid(CurItem))
+					{
+						++ItemIndex;
+						continue;
+					}
+
 					TSharedPtr<ITableRow> RowWidget = WidgetGenerator.GetWidgetForItem( CurItem );
 					if ( !RowWidget.IsValid() )
 					{
