Add separate navigation configs for game vs editor

diff --git a/Engine/Source/Runtime/Slate/Private/Framework/Application/SlateApplication.cpp b/Engine/Source/Runtime/Slate/Private/Framework/Application/SlateApplication.cpp
--- a/Engine/Source/Runtime/Slate/Private/Framework/Application/SlateApplication.cpp
+++ b/Engine/Source/Runtime/Slate/Private/Framework/Application/SlateApplication.cpp
@@ -1032,6 +1032,9 @@ FSlateApplication::FSlateApplication()
 	, AppIcon( FCoreStyle::Get().GetBrush("DefaultAppIcon") )
 	, VirtualDesktopRect( 0,0,0,0 )
 	, NavigationConfig(MakeShared<FNavigationConfig>())
+#if WITH_EDITOR
+	, EditorNavigationConfig(MakeShared<FNavigationConfig>())
+#endif
 	, SimulateGestures(false, (int32)EGestureEvent::Count)
 	, ProcessingInput(0)
 	, InputManager(MakeShared<FSlateDefaultInputMapping>())
@@ -1065,6 +1068,9 @@ FSlateApplication::FSlateApplication()
 	RegisterUser(MakeShareable(new FSlateUser(0, false)));
 
 	NavigationConfig->OnRegister();
+#if WITH_EDITOR
+	EditorNavigationConfig->OnRegister();
+#endif
 
 	SimulateGestures[(int32)EGestureEvent::LongPress] = true;
 
@@ -2099,13 +2105,53 @@ bool FSlateApplication::CanDisplayWindows() const
 	return Renderer.IsValid() && Renderer->AreShadersInitialized();
 }
 
+#if WITH_EDITOR
+static bool IsFocusInViewport(const TSet<TWeakPtr<SViewport>> Viewports, const FWeakWidgetPath& FocusPath)
+{
+	if (Viewports.Num() > 0)
+	{
+		for (const TWeakPtr<SWidget> FocusWidget : FocusPath.Widgets)
+		{
+			for (const TWeakPtr<SViewport> Viewport : Viewports)
+			{
+				if (FocusWidget == Viewport)
+				{
+					return true;
+				}
+			}
+		}
+	}
+	return false;
+}
+#endif
+
 EUINavigation FSlateApplication::GetNavigationDirectionFromKey(const FKeyEvent& InKeyEvent) const
 {
+#if WITH_EDITOR
+	// Check if the focused widget is an editor widget or a PIE widget so we know which config to use.
+	if (const FSlateUser* User = GetUser(GetUserIndexForKeyboard()))
+	{
+		if (!IsFocusInViewport(AllGameViewports, User->GetWeakFocusPath()))
+		{
+			return EditorNavigationConfig->GetNavigationDirectionFromKey(InKeyEvent);
+		}
+	}
+#endif
 	return NavigationConfig->GetNavigationDirectionFromKey(InKeyEvent);
 }
 
 EUINavigation FSlateApplication::GetNavigationDirectionFromAnalog(const FAnalogInputEvent& InAnalogEvent)
 {
+#if WITH_EDITOR
+	// Check if the focused widget is an editor widget or a PIE widget so we know which config to use.
+	if (const FSlateUser* User = GetUser(GetUserIndexForKeyboard()))
+	{
+		if (!IsFocusInViewport(AllGameViewports, User->GetWeakFocusPath()))
+		{
+			return EditorNavigationConfig->GetNavigationDirectionFromAnalog(InAnalogEvent);
+		}
+	}
+#endif
 	return NavigationConfig->GetNavigationDirectionFromAnalog(InAnalogEvent);
 }
 
@@ -2491,6 +2537,10 @@ void FSlateApplication::InvalidateAllViewports()
 void FSlateApplication::RegisterGameViewport( TSharedRef<SViewport> InViewport )
 {
 	RegisterViewport(InViewport);
+
+#if WITH_EDITOR
+	AllGameViewports.Add(InViewport);
+#endif
 	
 	if (GameViewportWidget != InViewport)
 	{
@@ -2521,6 +2571,10 @@ void FSlateApplication::UnregisterGameViewport()
 	bIsFakingTouched = false;
 	bIsGameFakingTouch = false;
 
+#if WITH_EDITOR
+	AllGameViewports.Empty();
+#endif
+
 	if (GameViewportWidget.IsValid())
 	{
 		GameViewportWidget.Pin()->SetActive(false);
@@ -4794,6 +4848,9 @@ void FSlateApplication::UnregisterUser(int32 UserIndex)
 		Users[UserIndex].Reset();
 
 		NavigationConfig->OnUserRemoved(UserIndex);
+#if WITH_EDITOR
+		EditorNavigationConfig->OnUserRemoved(UserIndex);
+#endif
 	}
 }
 
diff --git a/Engine/Source/Runtime/Slate/Public/Framework/Application/SlateApplication.h b/Engine/Source/Runtime/Slate/Public/Framework/Application/SlateApplication.h
--- a/Engine/Source/Runtime/Slate/Public/Framework/Application/SlateApplication.h
+++ b/Engine/Source/Runtime/Slate/Public/Framework/Application/SlateApplication.h
@@ -1952,6 +1952,11 @@ private:
 	/** Pointer to the currently registered game viewport widget if any */
 	TWeakPtr<SViewport> GameViewportWidget;
 
+#if WITH_EDITOR
+	/** List of all registered game viewports since the last time UnregisterGameViewport was called. */
+	TSet<TWeakPtr<SViewport>> AllGameViewports;
+#endif
+
 	TSharedPtr<ISlateSoundDevice> SlateSoundDevice;
 
 	/** The current cached absolute real time, right before we tick widgets */
@@ -2148,6 +2153,12 @@ private:
 	/** This factory function creates a navigation config for each slate user. */
 	TSharedRef<FNavigationConfig> NavigationConfig;
 
+#if WITH_EDITOR
+	/** When PIE runs, the game's navigation config will overwrite the editor's navigation config.
+	    This separate config allows editor navigation to work even when PIE is running. */
+	TSharedRef<FNavigationConfig> EditorNavigationConfig;
+#endif
+
 	/** The simulated gestures Slate Application will be in charge of. */
 	TBitArray<FDefaultBitArrayAllocator> SimulateGestures;
 
