Expose a few menu UI related enums to blueprint

diff --git a/Engine/Source/Editor/UnrealEd/Private/SEditorViewportToolBarButton.cpp b/Engine/Source/Editor/UnrealEd/Private/SEditorViewportToolBarButton.cpp
--- a/Engine/Source/Editor/UnrealEd/Private/SEditorViewportToolBarButton.cpp
+++ b/Engine/Source/Editor/UnrealEd/Private/SEditorViewportToolBarButton.cpp
@@ -18,7 +18,7 @@ void SEditorViewportToolBarButton::Construct( const FArguments& Declaration )
 
 	bool bContentOverride = ContentSlotWidget != SNullWidget::NullWidget;
 
-	EUserInterfaceActionType::Type ButtonType = Declaration._ButtonType;
+	EUserInterfaceActionType ButtonType = Declaration._ButtonType;
 	
 	// The style of the image to show in the button
 	const FName ImageStyleName = Declaration._Image.Get();
diff --git a/Engine/Source/Editor/UnrealEd/Public/SEditorViewportToolBarButton.h b/Engine/Source/Editor/UnrealEd/Public/SEditorViewportToolBarButton.h
--- a/Engine/Source/Editor/UnrealEd/Public/SEditorViewportToolBarButton.h
+++ b/Engine/Source/Editor/UnrealEd/Public/SEditorViewportToolBarButton.h
@@ -27,7 +27,7 @@ public:
 		/** Called when the button is clicked */
 		SLATE_EVENT( FOnClicked, OnClicked )
 		/** The button type to use */
-		SLATE_ARGUMENT( EUserInterfaceActionType::Type, ButtonType )
+		SLATE_ARGUMENT( EUserInterfaceActionType, ButtonType )
 		/** Checked state of the button */
 		SLATE_ATTRIBUTE( bool, IsChecked )
 		/** Style name of an image to use. Simple two state images are supported.  An image can be different depending on checked/unchecked state */
diff --git a/Engine/Source/Runtime/Slate/Private/Framework/Commands/Commands.cpp b/Engine/Source/Runtime/Slate/Private/Framework/Commands/Commands.cpp
--- a/Engine/Source/Runtime/Slate/Private/Framework/Commands/Commands.cpp
+++ b/Engine/Source/Runtime/Slate/Private/Framework/Commands/Commands.cpp
@@ -5,7 +5,7 @@
 
 #define LOC_DEFINE_REGION
 
-void UI_COMMAND_Function( FBindingContext* This, TSharedPtr< FUICommandInfo >& OutCommand, const TCHAR* OutSubNamespace, const TCHAR* OutCommandName, const TCHAR* OutCommandNameUnderscoreTooltip, const ANSICHAR* DotOutCommandName, const TCHAR* FriendlyName, const TCHAR* InDescription, const EUserInterfaceActionType::Type CommandType, const FInputChord& InDefaultChord, const FInputChord& InAlternateDefaultChord)
+void UI_COMMAND_Function( FBindingContext* This, TSharedPtr< FUICommandInfo >& OutCommand, const TCHAR* OutSubNamespace, const TCHAR* OutCommandName, const TCHAR* OutCommandNameUnderscoreTooltip, const ANSICHAR* DotOutCommandName, const TCHAR* FriendlyName, const TCHAR* InDescription, const EUserInterfaceActionType CommandType, const FInputChord& InDefaultChord, const FInputChord& InAlternateDefaultChord)
 {
 	static const FString UICommandsStr(TEXT("UICommands"));
 	const FString Namespace = OutSubNamespace && FCString::Strlen(OutSubNamespace) > 0 ?  UICommandsStr + TEXT(".") + OutSubNamespace : UICommandsStr;
diff --git a/Engine/Source/Runtime/Slate/Private/Framework/Commands/UICommandInfo.cpp b/Engine/Source/Runtime/Slate/Private/Framework/Commands/UICommandInfo.cpp
--- a/Engine/Source/Runtime/Slate/Private/Framework/Commands/UICommandInfo.cpp
+++ b/Engine/Source/Runtime/Slate/Private/Framework/Commands/UICommandInfo.cpp
@@ -36,7 +36,7 @@ FUICommandInfoDecl& FUICommandInfoDecl::DefaultChord( const FInputChord& InDefau
 	Info->DefaultChords[static_cast<uint8>(InChordIndex)] = InDefaultChord;
 	return *this;
 }
-FUICommandInfoDecl& FUICommandInfoDecl::UserInterfaceType( EUserInterfaceActionType::Type InType )
+FUICommandInfoDecl& FUICommandInfoDecl::UserInterfaceType( EUserInterfaceActionType InType )
 {
 	Info->UserInterfaceType = InType;
 	return *this;
@@ -75,7 +75,7 @@ const FText FUICommandInfo::GetInputText() const
 }
 
 
-void FUICommandInfo::MakeCommandInfo( const TSharedRef<class FBindingContext>& InContext, TSharedPtr< FUICommandInfo >& OutCommand, const FName InCommandName, const FText& InCommandLabel, const FText& InCommandDesc, const FSlateIcon& InIcon, const EUserInterfaceActionType::Type InUserInterfaceType, const FInputChord& InDefaultChord, const FInputChord& InAlternateDefaultChord)
+void FUICommandInfo::MakeCommandInfo( const TSharedRef<class FBindingContext>& InContext, TSharedPtr< FUICommandInfo >& OutCommand, const FName InCommandName, const FText& InCommandLabel, const FText& InCommandDesc, const FSlateIcon& InIcon, const EUserInterfaceActionType InUserInterfaceType, const FInputChord& InDefaultChord, const FInputChord& InAlternateDefaultChord)
 {
 	ensureMsgf( !InCommandLabel.IsEmpty(), TEXT("Command labels cannot be empty") );
 
diff --git a/Engine/Source/Runtime/Slate/Private/Framework/MultiBox/Mac/MacMenu.cpp b/Engine/Source/Runtime/Slate/Private/Framework/MultiBox/Mac/MacMenu.cpp
--- a/Engine/Source/Runtime/Slate/Private/Framework/MultiBox/Mac/MacMenu.cpp
+++ b/Engine/Source/Runtime/Slate/Private/Framework/MultiBox/Mac/MacMenu.cpp
@@ -13,7 +13,7 @@
 struct FMacMenuItemState
 {
 	TSharedPtr<const FMenuEntryBlock> Block;
-	EMultiBlockType::Type Type;
+	EMultiBlockType Type;
 	NSString* Title;
 	NSString* KeyEquivalent;
 	uint32 KeyModifiers;
diff --git a/Engine/Source/Runtime/Slate/Private/Framework/MultiBox/MultiBox.cpp b/Engine/Source/Runtime/Slate/Private/Framework/MultiBox/MultiBox.cpp
--- a/Engine/Source/Runtime/Slate/Private/Framework/MultiBox/MultiBox.cpp
+++ b/Engine/Source/Runtime/Slate/Private/Framework/MultiBox/MultiBox.cpp
@@ -172,7 +172,7 @@ bool FMultiBlock::GetSearchable() const
  * @param	InType	Type of MultiBox
  * @param	bInShouldCloseWindowAfterMenuSelection	Sets whether or not the window that contains this multibox should be destroyed after the user clicks on a menu item in this box
  */
-FMultiBox::FMultiBox( const EMultiBoxType::Type InType, FMultiBoxCustomization InCustomization, const bool bInShouldCloseWindowAfterMenuSelection )
+FMultiBox::FMultiBox( const EMultiBoxType InType, FMultiBoxCustomization InCustomization, const bool bInShouldCloseWindowAfterMenuSelection )
 	: CustomizationData( new FMultiBoxCustomizationData( InCustomization.GetCustomizationName() ) )
 	, CommandLists()
 	, Blocks()
@@ -187,7 +187,7 @@ FMultiBox::~FMultiBox()
 {
 }
 
-TSharedRef<FMultiBox> FMultiBox::Create( const EMultiBoxType::Type InType, FMultiBoxCustomization InCustomization, const bool bInShouldCloseWindowAfterMenuSelection )
+TSharedRef<FMultiBox> FMultiBox::Create( const EMultiBoxType InType, FMultiBoxCustomization InCustomization, const bool bInShouldCloseWindowAfterMenuSelection )
 {
 	TSharedRef<FMultiBox> NewBox = MakeShareable( new FMultiBox( InType, InCustomization, bInShouldCloseWindowAfterMenuSelection ) );
 
diff --git a/Engine/Source/Runtime/Slate/Private/Framework/MultiBox/MultiBoxBuilder.cpp b/Engine/Source/Runtime/Slate/Private/Framework/MultiBox/MultiBoxBuilder.cpp
--- a/Engine/Source/Runtime/Slate/Private/Framework/MultiBox/MultiBoxBuilder.cpp
+++ b/Engine/Source/Runtime/Slate/Private/Framework/MultiBox/MultiBoxBuilder.cpp
@@ -16,7 +16,7 @@
 #include "Framework/MultiBox/SWidgetBlock.h"
 #include "Framework/MultiBox/SGroupMarkerBlock.h"
 
-FMultiBoxBuilder::FMultiBoxBuilder( const EMultiBoxType::Type InType, FMultiBoxCustomization InCustomization, const bool bInShouldCloseWindowAfterMenuSelection, const TSharedPtr< const FUICommandList >& InCommandList, TSharedPtr<FExtender> InExtender, FName InTutorialHighlightName )
+FMultiBoxBuilder::FMultiBoxBuilder( const EMultiBoxType InType, FMultiBoxCustomization InCustomization, const bool bInShouldCloseWindowAfterMenuSelection, const TSharedPtr< const FUICommandList >& InCommandList, TSharedPtr<FExtender> InExtender, FName InTutorialHighlightName )
 	: MultiBox( FMultiBox::Create( InType, InCustomization, bInShouldCloseWindowAfterMenuSelection ) )
 	, CommandListStack()
 	, TutorialHighlightName(InTutorialHighlightName)
@@ -118,7 +118,7 @@ static FName GenerateTutorialIdentfierName(FName InContainerName, FName InElemen
 	}
 }
 
-FBaseMenuBuilder::FBaseMenuBuilder( const EMultiBoxType::Type InType, const bool bInShouldCloseWindowAfterMenuSelection, TSharedPtr< const FUICommandList > InCommandList, bool bInCloseSelfOnly, TSharedPtr<FExtender> InExtender, const ISlateStyle* InStyleSet, FName InTutorialHighlightName )
+FBaseMenuBuilder::FBaseMenuBuilder( const EMultiBoxType InType, const bool bInShouldCloseWindowAfterMenuSelection, TSharedPtr< const FUICommandList > InCommandList, bool bInCloseSelfOnly, TSharedPtr<FExtender> InExtender, const ISlateStyle* InStyleSet, FName InTutorialHighlightName )
 	: FMultiBoxBuilder( InType, FMultiBoxCustomization::None, bInShouldCloseWindowAfterMenuSelection, InCommandList, InExtender, InTutorialHighlightName )
 	, bCloseSelfOnly( bInCloseSelfOnly )
 {
@@ -140,7 +140,7 @@ void FBaseMenuBuilder::AddMenuEntry( const TSharedPtr< const FUICommandInfo > In
 	ApplyHook(InExtensionHook, EExtensionHook::After);
 }
 
-void FBaseMenuBuilder::AddMenuEntry( const TAttribute<FText>& InLabel, const TAttribute<FText>& InToolTip, const FSlateIcon& InIcon, const FUIAction& InAction, FName InExtensionHook, const EUserInterfaceActionType::Type UserInterfaceActionType, FName InTutorialHighlightName )
+void FBaseMenuBuilder::AddMenuEntry( const TAttribute<FText>& InLabel, const TAttribute<FText>& InToolTip, const FSlateIcon& InIcon, const FUIAction& InAction, FName InExtensionHook, const EUserInterfaceActionType UserInterfaceActionType, FName InTutorialHighlightName )
 {
 	ApplySectionBeginning();
 
@@ -153,7 +153,7 @@ void FBaseMenuBuilder::AddMenuEntry( const TAttribute<FText>& InLabel, const TAt
 	ApplyHook(InExtensionHook, EExtensionHook::After);
 }
 
-void FBaseMenuBuilder::AddMenuEntry( const FUIAction& UIAction, const TSharedRef< SWidget > Contents, const FName& InExtensionHook, const TAttribute<FText>& InToolTip, const EUserInterfaceActionType::Type UserInterfaceActionType, FName InTutorialHighlightName )
+void FBaseMenuBuilder::AddMenuEntry( const FUIAction& UIAction, const TSharedRef< SWidget > Contents, const FName& InExtensionHook, const TAttribute<FText>& InToolTip, const EUserInterfaceActionType UserInterfaceActionType, FName InTutorialHighlightName )
 {
 	ApplySectionBeginning();
 
@@ -219,7 +219,7 @@ void FMenuBuilder::AddMenuSeparator(FName InExtensionHook)
 	ApplyHook(InExtensionHook, EExtensionHook::After);
 }
 
-void FMenuBuilder::AddSubMenu( const TAttribute<FText>& InMenuLabel, const TAttribute<FText>& InToolTip, const FNewMenuDelegate& InSubMenu, const FUIAction& InUIAction, FName InExtensionHook, const EUserInterfaceActionType::Type InUserInterfaceActionType, const bool bInOpenSubMenuOnClick, const FSlateIcon& InIcon, const bool bInShouldCloseWindowAfterMenuSelection /*= true*/ )
+void FMenuBuilder::AddSubMenu( const TAttribute<FText>& InMenuLabel, const TAttribute<FText>& InToolTip, const FNewMenuDelegate& InSubMenu, const FUIAction& InUIAction, FName InExtensionHook, const EUserInterfaceActionType InUserInterfaceActionType, const bool bInOpenSubMenuOnClick, const FSlateIcon& InIcon, const bool bInShouldCloseWindowAfterMenuSelection /*= true*/ )
 {
 	ApplySectionBeginning();
 
@@ -376,7 +376,7 @@ void FToolBarBuilder::AddToolBarButton(const TSharedPtr< const FUICommandInfo >
 	ApplyHook(InExtensionHook, EExtensionHook::After);
 }
 
-void FToolBarBuilder::AddToolBarButton(const FUIAction& InAction, FName InExtensionHook, const TAttribute<FText>& InLabelOverride, const TAttribute<FText>& InToolTipOverride, const TAttribute<FSlateIcon>& InIconOverride, const EUserInterfaceActionType::Type UserInterfaceActionType, FName InTutorialHighlightName )
+void FToolBarBuilder::AddToolBarButton(const FUIAction& InAction, FName InExtensionHook, const TAttribute<FText>& InLabelOverride, const TAttribute<FText>& InToolTipOverride, const TAttribute<FSlateIcon>& InIconOverride, const EUserInterfaceActionType UserInterfaceActionType, FName InTutorialHighlightName )
 {
 	ApplySectionBeginning();
 
@@ -528,7 +528,7 @@ void FButtonRowBuilder::AddButton( const TSharedPtr< const FUICommandInfo > InCo
 	MultiBox->AddMultiBlock( NewButtonRowBlock );
 }
 
-void FButtonRowBuilder::AddButton( const FText& InLabel, const FText& InToolTip, const FUIAction& UIAction, const FSlateIcon& InIcon, const EUserInterfaceActionType::Type UserInterfaceActionType )
+void FButtonRowBuilder::AddButton( const FText& InLabel, const FText& InToolTip, const FUIAction& UIAction, const FSlateIcon& InIcon, const EUserInterfaceActionType UserInterfaceActionType )
 {
 	ApplySectionBeginning();
 
diff --git a/Engine/Source/Runtime/Slate/Private/Framework/MultiBox/SButtonRowBlock.cpp b/Engine/Source/Runtime/Slate/Private/Framework/MultiBox/SButtonRowBlock.cpp
--- a/Engine/Source/Runtime/Slate/Private/Framework/MultiBox/SButtonRowBlock.cpp
+++ b/Engine/Source/Runtime/Slate/Private/Framework/MultiBox/SButtonRowBlock.cpp
@@ -18,7 +18,7 @@ FButtonRowBlock::FButtonRowBlock( const TSharedPtr< const FUICommandInfo > InCom
 {
 }
 
-FButtonRowBlock::FButtonRowBlock( const TAttribute<FText>& InLabel, const TAttribute<FText>& InToolTip, const FSlateIcon& InIcon, const FUIAction& UIAction, const EUserInterfaceActionType::Type InUserInterfaceActionType )
+FButtonRowBlock::FButtonRowBlock( const TAttribute<FText>& InLabel, const TAttribute<FText>& InToolTip, const FSlateIcon& InIcon, const FUIAction& UIAction, const EUserInterfaceActionType InUserInterfaceActionType )
 	: FMultiBlock( UIAction, NAME_None, EMultiBlockType::ButtonRow )
 	, LabelOverride( InLabel )
 	, ToolTipOverride( InToolTip )
@@ -156,7 +156,7 @@ void SButtonRowBlock::BuildMultiBlockWidget(const ISlateStyle* StyleSet, const F
 
 
 	// What type of UI should we create for this block?
-	const EUserInterfaceActionType::Type UserInterfaceType = ButtonRowBlock->GetAction().IsValid() ? ButtonRowBlock->GetAction()->GetUserInterfaceType() : ButtonRowBlock->UserInterfaceActionTypeOverride;
+	const EUserInterfaceActionType UserInterfaceType = ButtonRowBlock->GetAction().IsValid() ? ButtonRowBlock->GetAction()->GetUserInterfaceType() : ButtonRowBlock->UserInterfaceActionTypeOverride;
 	if( UserInterfaceType == EUserInterfaceActionType::Button )
 	{
 		ChildSlot
diff --git a/Engine/Source/Runtime/Slate/Private/Framework/MultiBox/SButtonRowBlock.h b/Engine/Source/Runtime/Slate/Private/Framework/MultiBox/SButtonRowBlock.h
--- a/Engine/Source/Runtime/Slate/Private/Framework/MultiBox/SButtonRowBlock.h
+++ b/Engine/Source/Runtime/Slate/Private/Framework/MultiBox/SButtonRowBlock.h
@@ -44,7 +44,7 @@ public:
 	 * @param	UIAction				Action to execute when the button is clicked or when state should be checked
 	 * @param	UserInterfaceActionType	The style of button to display
 	 */
-	FButtonRowBlock( const TAttribute<FText>& InLabel, const TAttribute<FText>& InToolTip, const FSlateIcon& InIcon, const FUIAction& UIAction, const EUserInterfaceActionType::Type InUserInterfaceActionType );
+	FButtonRowBlock( const TAttribute<FText>& InLabel, const TAttribute<FText>& InToolTip, const FSlateIcon& InIcon, const FUIAction& UIAction, const EUserInterfaceActionType InUserInterfaceActionType );
 
 	/** FMultiBlock interface */
 	virtual bool HasIcon() const override;
@@ -73,7 +73,7 @@ private:
 
 	/** In the case where a command is not bound, the user interface action type to use.  If a command is bound, we
 	    simply use the action type associated with that command. */
-	EUserInterfaceActionType::Type UserInterfaceActionTypeOverride;
+	EUserInterfaceActionType UserInterfaceActionTypeOverride;
 };
 
 
diff --git a/Engine/Source/Runtime/Slate/Private/Framework/MultiBox/SMenuEntryBlock.cpp b/Engine/Source/Runtime/Slate/Private/Framework/MultiBox/SMenuEntryBlock.cpp
--- a/Engine/Source/Runtime/Slate/Private/Framework/MultiBox/SMenuEntryBlock.cpp
+++ b/Engine/Source/Runtime/Slate/Private/Framework/MultiBox/SMenuEntryBlock.cpp
@@ -27,7 +27,7 @@ FMenuEntryBlock::FMenuEntryBlock( const FName& InExtensionHook, const TSharedPtr
 }
 
 
-FMenuEntryBlock::FMenuEntryBlock( const FName& InExtensionHook, const TAttribute<FText>& InLabel, const TAttribute<FText>& InToolTip, const FNewMenuDelegate& InEntryBuilder, TSharedPtr<FExtender> InExtender, bool bInSubMenu, bool bInSubMenuOnClick, const FSlateIcon& InIcon, const FUIAction& InUIAction, const EUserInterfaceActionType::Type InUserInterfaceActionType, bool bInCloseSelfOnly, bool bInShouldCloseWindowAfterMenuSelection)
+FMenuEntryBlock::FMenuEntryBlock( const FName& InExtensionHook, const TAttribute<FText>& InLabel, const TAttribute<FText>& InToolTip, const FNewMenuDelegate& InEntryBuilder, TSharedPtr<FExtender> InExtender, bool bInSubMenu, bool bInSubMenuOnClick, const FSlateIcon& InIcon, const FUIAction& InUIAction, const EUserInterfaceActionType InUserInterfaceActionType, bool bInCloseSelfOnly, bool bInShouldCloseWindowAfterMenuSelection)
 	: FMultiBlock( InUIAction, InExtensionHook, EMultiBlockType::MenuEntry )
 	, LabelOverride( InLabel )
 	, ToolTipOverride( InToolTip )
@@ -43,7 +43,7 @@ FMenuEntryBlock::FMenuEntryBlock( const FName& InExtensionHook, const TAttribute
 }
 
 
-FMenuEntryBlock::FMenuEntryBlock( const FName& InExtensionHook, const TAttribute<FText>& InLabel, const TAttribute<FText>& InToolTip, const FSlateIcon& InIcon, const FUIAction& UIAction, const EUserInterfaceActionType::Type InUserInterfaceActionType, bool bInCloseSelfOnly, bool bInShouldCloseWindowAfterMenuSelection)
+FMenuEntryBlock::FMenuEntryBlock( const FName& InExtensionHook, const TAttribute<FText>& InLabel, const TAttribute<FText>& InToolTip, const FSlateIcon& InIcon, const FUIAction& UIAction, const EUserInterfaceActionType InUserInterfaceActionType, bool bInCloseSelfOnly, bool bInShouldCloseWindowAfterMenuSelection)
 	: FMultiBlock( UIAction, InExtensionHook, EMultiBlockType::MenuEntry )
 	, LabelOverride( InLabel )
 	, ToolTipOverride( InToolTip )
@@ -104,7 +104,7 @@ FMenuEntryBlock::FMenuEntryBlock( const FName& InExtensionHook, const TAttribute
 {
 }
 
-FMenuEntryBlock::FMenuEntryBlock( const FName& InExtensionHook, const FUIAction& UIAction, const TSharedRef< SWidget > Contents, const TAttribute<FText>& InToolTip, const EUserInterfaceActionType::Type InUserInterfaceActionType, bool bInCloseSelfOnly, bool bInShouldCloseWindowAfterMenuSelection)
+FMenuEntryBlock::FMenuEntryBlock( const FName& InExtensionHook, const FUIAction& UIAction, const TSharedRef< SWidget > Contents, const TAttribute<FText>& InToolTip, const EUserInterfaceActionType InUserInterfaceActionType, bool bInCloseSelfOnly, bool bInShouldCloseWindowAfterMenuSelection)
 	: FMultiBlock( UIAction, InExtensionHook, EMultiBlockType::MenuEntry )
 	, ToolTipOverride( InToolTip )
 	, EntryWidget( Contents )
@@ -464,7 +464,7 @@ TSharedRef< SWidget > SMenuEntryBlock::BuildMenuEntryWidget( const FMenuEntryBui
 	}
 
 	// What type of UI should we create for this block?
-	EUserInterfaceActionType::Type UserInterfaceType = MenuEntryBlock->UserInterfaceActionType;
+	EUserInterfaceActionType UserInterfaceType = MenuEntryBlock->UserInterfaceActionType;
 	if ( UICommand.IsValid() )
 	{
 		// If we have a UICommand, then this is specified in the command.
@@ -714,7 +714,7 @@ TSharedRef< SWidget> SMenuEntryBlock::BuildSubMenuWidget( const FMenuEntryBuildP
 	}
 
 	// What type of UI should we create for this block?
-	EUserInterfaceActionType::Type UserInterfaceType = MenuEntryBlock->UserInterfaceActionType;
+	EUserInterfaceActionType UserInterfaceType = MenuEntryBlock->UserInterfaceActionType;
 	if ( UICommand.IsValid() )
 	{
 		// If we have a UICommand, then this is specified in the command.
diff --git a/Engine/Source/Runtime/Slate/Private/Framework/MultiBox/SMenuEntryBlock.h b/Engine/Source/Runtime/Slate/Private/Framework/MultiBox/SMenuEntryBlock.h
--- a/Engine/Source/Runtime/Slate/Private/Framework/MultiBox/SMenuEntryBlock.h
+++ b/Engine/Source/Runtime/Slate/Private/Framework/MultiBox/SMenuEntryBlock.h
@@ -101,11 +101,11 @@ public:
 	 * @param	bInShouldCloseWindowAfterMenuSelection	In the case of a submenu, whether it should close after an item is selected
 	 * @param	bInInvertLabelOnHover	Whether to invert the label text's color on hover
 	 */
-	FMenuEntryBlock( const FName& InExtensionHook, const TAttribute<FText>& InLabel, const TAttribute<FText>& InToolTip, const FNewMenuDelegate& InEntryBuilder, TSharedPtr<FExtender> InExtender, bool bInSubMenu, bool bInSubMenuOnClick, const FSlateIcon& InIcon, const FUIAction& InUIAction, const EUserInterfaceActionType::Type InUserInterfaceActionType, bool bInCloseSelfOnly, bool bInShouldCloseWindowAfterMenuSelection = true);
+	FMenuEntryBlock( const FName& InExtensionHook, const TAttribute<FText>& InLabel, const TAttribute<FText>& InToolTip, const FNewMenuDelegate& InEntryBuilder, TSharedPtr<FExtender> InExtender, bool bInSubMenu, bool bInSubMenuOnClick, const FSlateIcon& InIcon, const FUIAction& InUIAction, const EUserInterfaceActionType InUserInterfaceActionType, bool bInCloseSelfOnly, bool bInShouldCloseWindowAfterMenuSelection = true);
 
-	FMenuEntryBlock( const FName& InExtensionHook, const TAttribute<FText>& InLabel, const TAttribute<FText>& InToolTip, const FSlateIcon& InIcon, const FUIAction& InUIAction, const EUserInterfaceActionType::Type InUserInterfaceActionType, bool bInCloseSelfOnly, bool bInShouldCloseWindowAfterMenuSelection = true);
+	FMenuEntryBlock( const FName& InExtensionHook, const TAttribute<FText>& InLabel, const TAttribute<FText>& InToolTip, const FSlateIcon& InIcon, const FUIAction& InUIAction, const EUserInterfaceActionType InUserInterfaceActionType, bool bInCloseSelfOnly, bool bInShouldCloseWindowAfterMenuSelection = true);
 
-	FMenuEntryBlock( const FName& InExtensionHook, const FUIAction& UIAction, const TSharedRef< SWidget > Contents, const TAttribute<FText>& InToolTip, const EUserInterfaceActionType::Type InUserInterfaceActionType, bool bInCloseSelfOnly, bool bInShouldCloseWindowAfterMenuSelection = true);
+	FMenuEntryBlock( const FName& InExtensionHook, const FUIAction& UIAction, const TSharedRef< SWidget > Contents, const TAttribute<FText>& InToolTip, const EUserInterfaceActionType InUserInterfaceActionType, bool bInCloseSelfOnly, bool bInShouldCloseWindowAfterMenuSelection = true);
 
 	FMenuEntryBlock( const FName& InExtensionHook, const TSharedRef< SWidget > Contents, const FNewMenuDelegate& InEntryBuilder, TSharedPtr<FExtender> InExtender, bool bInSubMenu, bool bInSubMenuOnClick, TSharedPtr< const FUICommandList > InCommandList, bool bInCloseSelfOnly, bool bInShouldCloseWindowAfterMenuSelection = true);
 
@@ -154,7 +154,7 @@ private:
 
 	/** In the case where a command is not bound, the user interface action type to use.  If a command is bound, we
 	    simply use the action type associated with that command. */
-	EUserInterfaceActionType::Type UserInterfaceActionType;
+	EUserInterfaceActionType UserInterfaceActionType;
 
 	/** True if the menu should close itself and all its children or the entire open menu stack */
 	bool bCloseSelfOnly;
diff --git a/Engine/Source/Runtime/Slate/Private/Framework/MultiBox/SToolBarButtonBlock.cpp b/Engine/Source/Runtime/Slate/Private/Framework/MultiBox/SToolBarButtonBlock.cpp
--- a/Engine/Source/Runtime/Slate/Private/Framework/MultiBox/SToolBarButtonBlock.cpp
+++ b/Engine/Source/Runtime/Slate/Private/Framework/MultiBox/SToolBarButtonBlock.cpp
@@ -22,7 +22,7 @@ FToolBarButtonBlock::FToolBarButtonBlock( const TSharedPtr< const FUICommandInfo
 {
 }
 
-FToolBarButtonBlock::FToolBarButtonBlock( const TAttribute<FText>& InLabel, const TAttribute<FText>& InToolTip, const TAttribute<FSlateIcon>& InIcon, const FUIAction& InUIAction, const EUserInterfaceActionType::Type InUserInterfaceActionType )
+FToolBarButtonBlock::FToolBarButtonBlock( const TAttribute<FText>& InLabel, const TAttribute<FText>& InToolTip, const TAttribute<FSlateIcon>& InIcon, const FUIAction& InUIAction, const EUserInterfaceActionType InUserInterfaceActionType )
 	: FMultiBlock( InUIAction )
 	, LabelOverride( InLabel )
 	, ToolTipOverride( InToolTip )
@@ -211,7 +211,7 @@ void SToolBarButtonBlock::BuildMultiBlockWidget(const ISlateStyle* StyleSet, con
 	EMultiBlockLocation::Type BlockLocation = GetMultiBlockLocation();
 	
 	// What type of UI should we create for this block?
-	EUserInterfaceActionType::Type UserInterfaceType = ToolBarButtonBlock->UserInterfaceActionType;
+	EUserInterfaceActionType UserInterfaceType = ToolBarButtonBlock->UserInterfaceActionType;
 	if ( Action.IsValid() )
 	{
 		// If we have a UICommand, then this is specified in the command.
diff --git a/Engine/Source/Runtime/Slate/Public/Framework/Commands/Commands.h b/Engine/Source/Runtime/Slate/Public/Framework/Commands/Commands.h
--- a/Engine/Source/Runtime/Slate/Public/Framework/Commands/Commands.h
+++ b/Engine/Source/Runtime/Slate/Public/Framework/Commands/Commands.h
@@ -12,7 +12,7 @@
 
 #define LOC_DEFINE_REGION
 
-SLATE_API void UI_COMMAND_Function(FBindingContext* This, TSharedPtr< FUICommandInfo >& OutCommand, const TCHAR* OutSubNamespace, const TCHAR* OutCommandName, const TCHAR* OutCommandNameUnderscoreTooltip, const ANSICHAR* DotOutCommandName, const TCHAR* FriendlyName, const TCHAR* InDescription, const EUserInterfaceActionType::Type CommandType, const FInputChord& InDefaultChord, const FInputChord& InAlternateDefaultChord = FInputChord());
+SLATE_API void UI_COMMAND_Function(FBindingContext* This, TSharedPtr< FUICommandInfo >& OutCommand, const TCHAR* OutSubNamespace, const TCHAR* OutCommandName, const TCHAR* OutCommandNameUnderscoreTooltip, const ANSICHAR* DotOutCommandName, const TCHAR* FriendlyName, const TCHAR* InDescription, const EUserInterfaceActionType CommandType, const FInputChord& InDefaultChord, const FInputChord& InAlternateDefaultChord = FInputChord());
 
 // This macro requires LOCTEXT_NAMESPACE to be defined. If you don't want the command to be placed under a sub namespace, provide "" as the namespace.
 #define UI_COMMAND_EXT( BindingContext, OutUICommandInfo, CommandIdName, FriendlyName, InDescription, CommandType, InDefaultChord, ... ) \
diff --git a/Engine/Source/Runtime/Slate/Public/Framework/Commands/UICommandInfo.h b/Engine/Source/Runtime/Slate/Public/Framework/Commands/UICommandInfo.h
--- a/Engine/Source/Runtime/Slate/Public/Framework/Commands/UICommandInfo.h
+++ b/Engine/Source/Runtime/Slate/Public/Framework/Commands/UICommandInfo.h
@@ -8,33 +8,32 @@
 #include "Layout/Visibility.h"
 #include "Textures/SlateIcon.h"
 #include "Framework/Commands/InputChord.h"
+#include "UICommandInfo.generated.h"
 
 class FBindingContext;
 class FUICommandInfo;
 
 /** Types of user interfaces that can be associated with a user interface action */
-namespace EUserInterfaceActionType
+UENUM(BlueprintType)
+enum class EUserInterfaceActionType : uint8
 {
-	enum Type
-	{
-		/** An action which should not be associated with a user interface action */
-		None,
+	/** An action which should not be associated with a user interface action */
+	None,
 
-		/** Momentary buttons or menu items.  These support enable state, and execute a delegate when clicked. */
-		Button,
+	/** Momentary buttons or menu items.  These support enable state, and execute a delegate when clicked. */
+	Button,
 
-		/** Toggleable buttons or menu items that store on/off state.  These support enable state, and execute a delegate when toggled. */
-		ToggleButton,
+	/** Toggleable buttons or menu items that store on/off state.  These support enable state, and execute a delegate when toggled. */
+	ToggleButton,
 		
-		/** Radio buttons are similar to toggle buttons in that they are for menu items that store on/off state.  However they should be used to indicate that menu items in a group can only be in one state */
-		RadioButton,
+	/** Radio buttons are similar to toggle buttons in that they are for menu items that store on/off state.  However they should be used to indicate that menu items in a group can only be in one state */
+	RadioButton,
 
-		/** Similar to Button but will display a readonly checkbox next to the item. */
-		Check,
+	/** Similar to Button but will display a readonly checkbox next to the item. */
+	Check,
 
-		/** Similar to Button but has the checkbox area collapsed */
-		CollapsedButton
-	};
+	/** Similar to Button but has the checkbox area collapsed */
+	CollapsedButton
 };
 
 UENUM()
@@ -58,7 +57,7 @@ class SLATE_API FUICommandInfoDecl
 public:
 
 	FUICommandInfoDecl& DefaultChord( const FInputChord& InDefaultChord, const EMultipleKeyBindingIndex InChordIndex = EMultipleKeyBindingIndex::Primary);
-	FUICommandInfoDecl& UserInterfaceType( EUserInterfaceActionType::Type InType );
+	FUICommandInfoDecl& UserInterfaceType( EUserInterfaceActionType InType );
 	FUICommandInfoDecl& Icon( const FSlateIcon& InIcon );
 	FUICommandInfoDecl& Description( const FText& InDesc );
 
@@ -223,7 +222,7 @@ public:
 	const FInputChord& GetDefaultChord(const EMultipleKeyBindingIndex InChordIndex) const { return DefaultChords[static_cast<uint8>(InChordIndex)]; }
 
 	/** Utility function to make an FUICommandInfo */
-	static void MakeCommandInfo( const TSharedRef<class FBindingContext>& InContext, TSharedPtr< FUICommandInfo >& OutCommand, const FName InCommandName, const FText& InCommandLabel, const FText& InCommandDesc, const FSlateIcon& InIcon, const EUserInterfaceActionType::Type InUserInterfaceType, const FInputChord& InDefaultChord, const FInputChord& InAlternateDefaultChord = FInputChord());
+	static void MakeCommandInfo( const TSharedRef<class FBindingContext>& InContext, TSharedPtr< FUICommandInfo >& OutCommand, const FName InCommandName, const FText& InCommandLabel, const FText& InCommandDesc, const FSlateIcon& InIcon, const EUserInterfaceActionType InUserInterfaceType, const FInputChord& InDefaultChord, const FInputChord& InAlternateDefaultChord = FInputChord());
 
 	/** Utility function to unregister an FUICommandInfo */
 	static void UnregisterCommandInfo(const TSharedRef<class FBindingContext>& InContext, const TSharedRef<FUICommandInfo>& InCommand);
@@ -238,7 +237,7 @@ public:
 	const FSlateIcon& GetIcon() const { return Icon; }
 
 	/** @return The type of command this is.  Used to determine what UI to create for it */
-	EUserInterfaceActionType::Type GetUserInterfaceType() const { return UserInterfaceType; }
+	EUserInterfaceActionType GetUserInterfaceType() const { return UserInterfaceType; }
 	
 	/** @return The name of the command */
 	FName GetCommandName() const { return CommandName; }
@@ -286,5 +285,5 @@ private:
 	FName BindingContext;
 
 	/** The type of user interface to associated with this action */
-	EUserInterfaceActionType::Type UserInterfaceType;
+	EUserInterfaceActionType UserInterfaceType;
 };
diff --git a/Engine/Source/Runtime/Slate/Public/Framework/MultiBox/MultiBox.h b/Engine/Source/Runtime/Slate/Public/Framework/MultiBox/MultiBox.h
--- a/Engine/Source/Runtime/Slate/Public/Framework/MultiBox/MultiBox.h
+++ b/Engine/Source/Runtime/Slate/Public/Framework/MultiBox/MultiBox.h
@@ -60,7 +60,7 @@ public:
 	 * @param InCommand		The command info that describes what action to take when this block is activated
 	 * @param InCommandList	The list of mappings from command info to delegates so we can find the delegates to process for the provided action
 	 */
-	FMultiBlock( const TSharedPtr< const FUICommandInfo > InCommand, TSharedPtr< const FUICommandList > InCommandList, FName InExtensionHook = NAME_None, EMultiBlockType::Type InType = EMultiBlockType::None )
+	FMultiBlock( const TSharedPtr< const FUICommandInfo > InCommand, TSharedPtr< const FUICommandList > InCommandList, FName InExtensionHook = NAME_None, EMultiBlockType InType = EMultiBlockType::None )
 		: Action( InCommand )
 		, ActionList( InCommandList )
 		, ExtensionHook( InExtensionHook )
@@ -74,7 +74,7 @@ public:
 	 *
 	 * @InAction UI action delegates that should be used in place of UI commands (dynamic menu items)
 	 */
-	FMultiBlock( const FUIAction& InAction,  FName InExtensionHook = NAME_None, EMultiBlockType::Type InType = EMultiBlockType::None )
+	FMultiBlock( const FUIAction& InAction,  FName InExtensionHook = NAME_None, EMultiBlockType InType = EMultiBlockType::None )
 		: DirectActions( InAction )
 		, ExtensionHook( InExtensionHook )
 		, Type( InType )
@@ -142,7 +142,7 @@ public:
 	 *
 	 * @return	The MultiBlock's type
 	 */
-	const EMultiBlockType::Type GetType() const
+	const EMultiBlockType GetType() const
 	{
 		return Type;
 	}
@@ -189,7 +189,7 @@ private:
 	FName ExtensionHook;
 
 	/** Type of MultiBlock */
-	EMultiBlockType::Type Type;
+	EMultiBlockType Type;
 
 	/** Name to identify a widget for tutorials */
 	FName TutorialHighlightName;
@@ -214,7 +214,7 @@ public:
 	/**
 	 * Creates a new multibox instance
 	 */
-	static TSharedRef<FMultiBox> Create( const EMultiBoxType::Type InType,  FMultiBoxCustomization InCustomization, const bool bInShouldCloseWindowAfterMenuSelection );
+	static TSharedRef<FMultiBox> Create( const EMultiBoxType InType,  FMultiBoxCustomization InCustomization, const bool bInShouldCloseWindowAfterMenuSelection );
 
 
 	/**
@@ -222,7 +222,7 @@ public:
 	 *
 	 * @return	The MultiBox's type
 	 */
-	const EMultiBoxType::Type GetType() const
+	const EMultiBoxType GetType() const
 	{
 		return Type;
 	}
@@ -319,7 +319,7 @@ private:
 	 * @param	InType	Type of MultiBox
 	 * @param	bInShouldCloseWindowAfterMenuSelection	Sets whether or not the window that contains this multibox should be destroyed after the user clicks on a menu item in this box
 	 */
-	FMultiBox( const EMultiBoxType::Type InType,  FMultiBoxCustomization InCustomization, const bool bInShouldCloseWindowAfterMenuSelection );
+	FMultiBox( const EMultiBoxType InType,  FMultiBoxCustomization InCustomization, const bool bInShouldCloseWindowAfterMenuSelection );
 
 	/**
 	 * @return true if this box can be customized by a user
@@ -348,7 +348,7 @@ private:
 	FName StyleName;
 
 	/** Type of MultiBox */
-	EMultiBoxType::Type Type;
+	EMultiBoxType Type;
 
 	/** True if window that owns any widgets created from this multibox should be closed automatically after the user commits to a menu choice */
 	bool bShouldCloseWindowAfterMenuSelection;
diff --git a/Engine/Source/Runtime/Slate/Public/Framework/MultiBox/MultiBoxBuilder.h b/Engine/Source/Runtime/Slate/Public/Framework/MultiBox/MultiBoxBuilder.h
--- a/Engine/Source/Runtime/Slate/Public/Framework/MultiBox/MultiBoxBuilder.h
+++ b/Engine/Source/Runtime/Slate/Public/Framework/MultiBox/MultiBoxBuilder.h
@@ -42,7 +42,7 @@ public:
 	 * @param	InCommandList	The action list that maps command infos to delegates that should be called for each command associated with a multiblock widget.  This can be modified after the MultiBox is created by calling the PushCommandList() and PopCommandList() methods.
 	 * @param	InTutorialHighlightName	Optional name to identify this widget and highlight during tutorials
 	 */
-	FMultiBoxBuilder( const EMultiBoxType::Type InType, FMultiBoxCustomization InCustomization, const bool bInShouldCloseWindowAfterMenuSelection, const TSharedPtr< const FUICommandList >& InCommandList, TSharedPtr<FExtender> InExtender = TSharedPtr<FExtender>(), FName InTutorialHighlightName = NAME_None );
+	FMultiBoxBuilder( const EMultiBoxType InType, FMultiBoxCustomization InCustomization, const bool bInShouldCloseWindowAfterMenuSelection, const TSharedPtr< const FUICommandList >& InCommandList, TSharedPtr<FExtender> InExtender = TSharedPtr<FExtender>(), FName InTutorialHighlightName = NAME_None );
 
 	virtual ~FMultiBoxBuilder() {}
 
@@ -161,7 +161,7 @@ public:
 	 * @param	bInCloseSelfOnly	True if clicking on a menu entry closes itself only and its children but not the entire stack 
 	 * @param	InTutorialHighlightName	Optional name to identify this widget and highlight during tutorials
 	 */
-	FBaseMenuBuilder( const EMultiBoxType::Type InType, const bool bInShouldCloseWindowAfterMenuSelection, TSharedPtr< const FUICommandList > InCommandList, bool bInCloseSelfOnly, TSharedPtr<FExtender> InExtender = TSharedPtr<FExtender>(), const ISlateStyle* InStyleSet = &FCoreStyle::Get(), FName InTutorialHighlightName = NAME_None );
+	FBaseMenuBuilder( const EMultiBoxType InType, const bool bInShouldCloseWindowAfterMenuSelection, TSharedPtr< const FUICommandList > InCommandList, bool bInCloseSelfOnly, TSharedPtr<FExtender> InExtender = TSharedPtr<FExtender>(), const ISlateStyle* InStyleSet = &FCoreStyle::Get(), FName InTutorialHighlightName = NAME_None );
 
 	/**
 	 * Adds a menu entry
@@ -186,9 +186,9 @@ public:
 	 * @param	UserInterfaceActionType	Type of interface action
 	 * @param	InTutorialHighlightName	Optional name to identify this widget and highlight during tutorials
 	 */
-	void AddMenuEntry( const TAttribute<FText>& InLabel, const TAttribute<FText>& InToolTip, const FSlateIcon& InIcon, const FUIAction& UIAction, FName InExtensionHook = NAME_None, const EUserInterfaceActionType::Type UserInterfaceActionType = EUserInterfaceActionType::Button, FName InTutorialHighlightName = NAME_None );
+	void AddMenuEntry( const TAttribute<FText>& InLabel, const TAttribute<FText>& InToolTip, const FSlateIcon& InIcon, const FUIAction& UIAction, FName InExtensionHook = NAME_None, const EUserInterfaceActionType UserInterfaceActionType = EUserInterfaceActionType::Button, FName InTutorialHighlightName = NAME_None );
 
-	void AddMenuEntry( const FUIAction& UIAction, const TSharedRef< SWidget > Contents, const FName& InExtensionHook = NAME_None, const TAttribute<FText>& InToolTip = TAttribute<FText>(), const EUserInterfaceActionType::Type UserInterfaceActionType = EUserInterfaceActionType::Button, FName InTutorialHighlightName = NAME_None );
+	void AddMenuEntry( const FUIAction& UIAction, const TSharedRef< SWidget > Contents, const FName& InExtensionHook = NAME_None, const TAttribute<FText>& InToolTip = TAttribute<FText>(), const EUserInterfaceActionType UserInterfaceActionType = EUserInterfaceActionType::Button, FName InTutorialHighlightName = NAME_None );
 
 protected:
 	/** True if clicking on a menu entry closes itself only and its children and not the entire stack */
@@ -261,7 +261,7 @@ public:
 	 * @param	InIcon				The icon to use
 	 * @param	bInShouldCloseWindowAfterMenuSelection	Whether the submenu should close after an item is selected
 	 */
-	void AddSubMenu( const TAttribute<FText>& InMenuLabel, const TAttribute<FText>& InToolTip, const FNewMenuDelegate& InSubMenu, const FUIAction& InUIAction, FName InExtensionHook, const EUserInterfaceActionType::Type InUserInterfaceActionType, const bool bInOpenSubMenuOnClick = false, const FSlateIcon& InIcon = FSlateIcon(), const bool bInShouldCloseWindowAfterMenuSelection = true );
+	void AddSubMenu( const TAttribute<FText>& InMenuLabel, const TAttribute<FText>& InToolTip, const FNewMenuDelegate& InSubMenu, const FUIAction& InUIAction, FName InExtensionHook, const EUserInterfaceActionType InUserInterfaceActionType, const bool bInOpenSubMenuOnClick = false, const FSlateIcon& InIcon = FSlateIcon(), const bool bInShouldCloseWindowAfterMenuSelection = true );
 
 	void AddSubMenu( const TAttribute<FText>& InMenuLabel, const TAttribute<FText>& InToolTip, const FNewMenuDelegate& InSubMenu, const bool bInOpenSubMenuOnClick = false, const FSlateIcon& InIcon = FSlateIcon(), const bool bInShouldCloseWindowAfterMenuSelection = true );
 
@@ -409,7 +409,7 @@ public:
 	 * @param	UserInterfaceActionType	Type of interface action
 	 * @param	InTutorialHighlightName	Name to identify this widget and highlight during tutorials
 	 */
-	void AddToolBarButton(const FUIAction& InAction, FName InExtensionHook = NAME_None, const TAttribute<FText>& InLabelOverride = TAttribute<FText>(), const TAttribute<FText>& InToolTipOverride = TAttribute<FText>(), const TAttribute<FSlateIcon>& InIconOverride = TAttribute<FSlateIcon>(), const EUserInterfaceActionType::Type UserInterfaceActionType = EUserInterfaceActionType::Button, FName InTutorialHighlightName = NAME_None );
+	void AddToolBarButton(const FUIAction& InAction, FName InExtensionHook = NAME_None, const TAttribute<FText>& InLabelOverride = TAttribute<FText>(), const TAttribute<FText>& InToolTipOverride = TAttribute<FText>(), const TAttribute<FSlateIcon>& InIconOverride = TAttribute<FSlateIcon>(), const EUserInterfaceActionType UserInterfaceActionType = EUserInterfaceActionType::Button, FName InTutorialHighlightName = NAME_None );
 
 	/**
 	 * Adds a combo button
@@ -518,7 +518,7 @@ public:
 	 * @param	InIcon					The icon for the button
 	 * @param	UserInterfaceActionType	The style of button to display
 	 */
-	void AddButton( const FText& InLabel, const FText& InToolTip, const FUIAction& UIAction, const FSlateIcon& InIcon = FSlateIcon(), const EUserInterfaceActionType::Type UserInterfaceActionType = EUserInterfaceActionType::Button );
+	void AddButton( const FText& InLabel, const FText& InToolTip, const FUIAction& UIAction, const FSlateIcon& InIcon = FSlateIcon(), const EUserInterfaceActionType UserInterfaceActionType = EUserInterfaceActionType::Button );
 	
 protected:
 	/** FMultiBoxBuilder interface */
diff --git a/Engine/Source/Runtime/Slate/Public/Framework/MultiBox/MultiBoxDefs.h b/Engine/Source/Runtime/Slate/Public/Framework/MultiBox/MultiBoxDefs.h
--- a/Engine/Source/Runtime/Slate/Public/Framework/MultiBox/MultiBoxDefs.h
+++ b/Engine/Source/Runtime/Slate/Public/Framework/MultiBox/MultiBoxDefs.h
@@ -6,6 +6,7 @@
 #include "Misc/Attribute.h"
 #include "Styling/ISlateStyle.h"
 #include "Framework/Commands/UICommandInfo.h"
+#include "MultiBoxDefs.generated.h"
 
 class SToolTip;
 class SWidget;
@@ -13,49 +14,47 @@ class SWidget;
 /**
  * Types of MultiBoxes
  */
-namespace EMultiBoxType
+UENUM(BlueprintType)
+enum class EMultiBoxType : uint8
 {
-	enum Type
-	{
-		/** Horizontal menu bar */
-		MenuBar,
+	/** Horizontal menu bar */
+	MenuBar,
 
-		/** Horizontal tool bar */
-		ToolBar,
+	/** Horizontal tool bar */
+	ToolBar,
 
-		/** Vertical tool bar */
-		VerticalToolBar,
+	/** Vertical tool bar */
+	VerticalToolBar,
 
-		/** Vertical menu (pull-down menu, or context menu) */
-		Menu,
+	/** Vertical menu (pull-down menu, or context menu) */
+	Menu,
 
-		/** Buttons arranged in rows, with a maximum number of buttons per row, like a toolbar but can have multiple rows*/
-		ButtonRow,
+	/** Buttons arranged in rows, with a maximum number of buttons per row, like a toolbar but can have multiple rows*/
+	ButtonRow,
+
+	/** Horizontal menu bar used as a tool bar */
+	ToolMenuBar
+};
 
-		/** Horizontal menu bar used as a tool bar */
-		ToolMenuBar
-	};
-}
 
 /**
  * Types of MultiBlocks
  */
-namespace EMultiBlockType
+UENUM(BlueprintType)
+enum class EMultiBlockType : uint8
 {
-	enum Type
-	{
-		None = 0,
-		ButtonRow,
-		EditableText,
-		Heading,
-		MenuEntry,
-		MenuSeparator,
-		ToolBarButton,
-		ToolBarComboButton,
-		ToolBarSeparator,
-		Widget,
-	};
-}
+	None = 0,
+	ButtonRow,
+	EditableText,
+	Heading,
+	MenuEntry,
+	MenuSeparator,
+	ToolBarButton,
+	ToolBarComboButton,
+	ToolBarSeparator,
+	Widget,
+};
+
 
 class SLATE_API FMultiBoxSettings
 {
diff --git a/Engine/Source/Runtime/Slate/Public/Framework/MultiBox/SToolBarButtonBlock.h b/Engine/Source/Runtime/Slate/Public/Framework/MultiBox/SToolBarButtonBlock.h
--- a/Engine/Source/Runtime/Slate/Public/Framework/MultiBox/SToolBarButtonBlock.h
+++ b/Engine/Source/Runtime/Slate/Public/Framework/MultiBox/SToolBarButtonBlock.h
@@ -41,7 +41,7 @@ public:
 	 * @param	InUIAction			UI action to take when this menu item is clicked as well as to determine if the menu entry can be executed or appears "checked"
 	 * @param	InUserInterfaceActionType	Type of interface action
 	 */
-	FToolBarButtonBlock( const TAttribute<FText>& InLabel, const TAttribute<FText>& InToolTip, const TAttribute<FSlateIcon>& InIcon, const FUIAction& InUIAction, const EUserInterfaceActionType::Type InUserInterfaceActionType );
+	FToolBarButtonBlock( const TAttribute<FText>& InLabel, const TAttribute<FText>& InToolTip, const TAttribute<FSlateIcon>& InIcon, const FUIAction& InUIAction, const EUserInterfaceActionType InUserInterfaceActionType );
 
 	void SetLabelVisibility( EVisibility InLabelVisibility ) { LabelVisibility = InLabelVisibility ; }
 
@@ -77,7 +77,7 @@ private:
 	
 	/** In the case where a command is not bound, the user interface action type to use.  If a command is bound, we
 		simply use the action type associated with that command. */
-	EUserInterfaceActionType::Type UserInterfaceActionType;
+	EUserInterfaceActionType UserInterfaceActionType;
 
 	/** Whether ToolBar will have Focusable buttons */
 	bool bIsFocusable;
