Fix shadowing warning

diff --git a/Engine/Source/Runtime/Slate/Private/Framework/Docking/TabManager.cpp b/Engine/Source/Runtime/Slate/Private/Framework/Docking/TabManager.cpp
--- a/Engine/Source/Runtime/Slate/Private/Framework/Docking/TabManager.cpp
+++ b/Engine/Source/Runtime/Slate/Private/Framework/Docking/TabManager.cpp
@@ -952,7 +952,7 @@ void FTabManager::PopulateTabSpawnerMenu_Helper( FMenuBuilder& PopulateMe, FPopu
 	}
 }
 
-void FTabManager::MakeSpawnerMenuEntry( FMenuBuilder &PopulateMe, const TSharedPtr<FTabSpawnerEntry> &SpawnerNode ) 
+void FTabManager::MakeSpawnerMenuEntry( FMenuBuilder &PopulateMe, const TSharedPtr<FTabSpawnerEntry> &InSpawnerNode ) 
 {
 	auto CanExecuteMenuEntry = [](TWeakPtr<FTabSpawnerEntry> SpawnerNode) -> bool
 	{
@@ -965,16 +965,16 @@ void FTabManager::MakeSpawnerMenuEntry( FMenuBuilder &PopulateMe, const TSharedP
 		return false;
 	};
 
-	if ( SpawnerNode->MenuType.Get() != ETabSpawnerMenuType::Hidden )
+	if (InSpawnerNode->MenuType.Get() != ETabSpawnerMenuType::Hidden )
 	{
 		PopulateMe.AddMenuEntry(
-			SpawnerNode->GetDisplayName().IsEmpty() ? FText::FromName( SpawnerNode->TabType ) : SpawnerNode->GetDisplayName(),
-			SpawnerNode->GetTooltipText(),
-			SpawnerNode->GetIcon(),
+			InSpawnerNode->GetDisplayName().IsEmpty() ? FText::FromName(InSpawnerNode->TabType ) : InSpawnerNode->GetDisplayName(),
+			InSpawnerNode->GetTooltipText(),
+			InSpawnerNode->GetIcon(),
 			FUIAction(
-			FExecuteAction::CreateSP(SharedThis(this), &FTabManager::InvokeTabForMenu, SpawnerNode->TabType),
-			FCanExecuteAction::CreateStatic(CanExecuteMenuEntry, TWeakPtr<FTabSpawnerEntry>(SpawnerNode)),
-			FIsActionChecked::CreateSP(SpawnerNode.ToSharedRef(), &FTabSpawnerEntry::IsSoleTabInstanceSpawned)
+			FExecuteAction::CreateSP(SharedThis(this), &FTabManager::InvokeTabForMenu, InSpawnerNode->TabType),
+			FCanExecuteAction::CreateStatic(CanExecuteMenuEntry, TWeakPtr<FTabSpawnerEntry>(InSpawnerNode)),
+			FIsActionChecked::CreateSP(InSpawnerNode.ToSharedRef(), &FTabSpawnerEntry::IsSoleTabInstanceSpawned)
 			),
 			NAME_None,
 			EUserInterfaceActionType::Check
diff --git a/Engine/Source/Runtime/Slate/Public/Framework/Docking/TabManager.h b/Engine/Source/Runtime/Slate/Public/Framework/Docking/TabManager.h
--- a/Engine/Source/Runtime/Slate/Public/Framework/Docking/TabManager.h
+++ b/Engine/Source/Runtime/Slate/Public/Framework/Docking/TabManager.h
@@ -778,7 +778,7 @@ class SLATE_API FTabManager : public TSharedFromThis<FTabManager>
 			
 		void PopulateTabSpawnerMenu_Helper( FMenuBuilder& PopulateMe, struct FPopulateTabSpawnerMenu_Args Args );
 
-		void MakeSpawnerMenuEntry( FMenuBuilder &PopulateMe, const TSharedPtr<FTabSpawnerEntry> &SpawnerNode );
+		void MakeSpawnerMenuEntry( FMenuBuilder &PopulateMe, const TSharedPtr<FTabSpawnerEntry> &InSpawnerNode );
 
 		TSharedPtr<SDockTab> InvokeTab_Internal( const FTabId& TabId );
 
