PR #5902: Re-order grid panel slots when slot layer changes.

diff --git a/Engine/Source/Runtime/Slate/Private/Widgets/Layout/SGridPanel.cpp b/Engine/Source/Runtime/Slate/Private/Widgets/Layout/SGridPanel.cpp
--- a/Engine/Source/Runtime/Slate/Private/Widgets/Layout/SGridPanel.cpp
+++ b/Engine/Source/Runtime/Slate/Private/Widgets/Layout/SGridPanel.cpp
@@ -343,7 +343,7 @@ SGridPanel::FSlot& SGridPanel::InsertSlot( SGridPanel::FSlot* InSlot )
 	return *InSlot;
 }
 
-void SGridPanel::NotifySlotChanged(SGridPanel::FSlot* InSlot)
+void SGridPanel::NotifySlotChanged(SGridPanel::FSlot* InSlot, bool bSlotLayerChanged /*= false*/)
 {
 	// Keep the size of the grid up to date.
 	// We need an extra cell at the end for easily figuring out the size across any number of cells
@@ -360,6 +360,14 @@ void SGridPanel::NotifySlotChanged(SGridPanel::FSlot* InSlot)
 	{
 		Rows.AddZeroed( NumRowsRequiredForThisSlot - Rows.Num() );
 	}
+
+	if (bSlotLayerChanged)
+	{
+		Slots.Sort([](const FSlot& LHS, const FSlot& RHS)
+		{
+			return LHS.LayerParam < RHS.LayerParam;
+		});
+	}
 }
 
 void SGridPanel::ComputeDesiredCellSizes( TArray<float>& OutColumns, TArray<float>& OutRows ) const
diff --git a/Engine/Source/Runtime/Slate/Public/Widgets/Layout/SGridPanel.h b/Engine/Source/Runtime/Slate/Public/Widgets/Layout/SGridPanel.h
--- a/Engine/Source/Runtime/Slate/Public/Widgets/Layout/SGridPanel.h
+++ b/Engine/Source/Runtime/Slate/Public/Widgets/Layout/SGridPanel.h
@@ -72,40 +72,60 @@ public:
 			/**  */
 			FSlot& Column(int32 Column)
 			{
-				ColumnParam = FMath::Max(0, Column);
-				NotifySlotChanged();
+				Column = FMath::Max(0, Column);
+				if (Column != ColumnParam)
+				{
+					ColumnParam = Column;
+					NotifySlotChanged();
+				}
 				return *this;
 			}
 
 			/** How many columns this slot spans over */
 			FSlot& ColumnSpan( int32 ColumnSpan )
 			{
-				ColumnSpanParam = FMath::Max(1,ColumnSpan);
-				NotifySlotChanged();
+				ColumnSpan = FMath::Max(1,ColumnSpan);
+				if (ColumnSpan != ColumnSpanParam)
+				{
+					ColumnSpanParam = ColumnSpan;
+					NotifySlotChanged();
+				}
 				return *this;
 			}
 
 			/**  */
 			FSlot& Row(int32 Row)
 			{
-				RowParam = FMath::Max(0, Row);
-				NotifySlotChanged();
+				Row = FMath::Max(0, Row);
+				if (Row != RowParam)
+				{
+					RowParam = Row;
+					NotifySlotChanged();
+				}
 				return *this;
 			}
 
 			/** How many rows this this slot spans over */
 			FSlot& RowSpan( int32 RowSpan )
 			{
-				RowSpanParam = FMath::Max(1, RowSpan);
-				NotifySlotChanged();
+				RowSpan = FMath::Max(1, RowSpan);
+				if (RowSpan != RowSpanParam)
+				{
+					RowSpanParam = RowSpan;
+					NotifySlotChanged();
+				}
 				return *this;
 			}
 
 			/** Positive values offset this cell to be hit-tested and drawn on top of others. Default is 0; i.e. no offset. */
 			FSlot& Layer(int32 Layer)
 			{
-				LayerParam = Layer;
-				NotifySlotChanged();
+				if (Layer != LayerParam)
+				{
+					LayerParam = Layer;
+					const bool bSlotLayerChanged = true;
+					NotifySlotChanged(bSlotLayerChanged);
+				}
 				return *this;
 			}
 
@@ -117,11 +137,11 @@ public:
 			}
 
 			/** Notify that the slot was changed */
-			FORCEINLINE void NotifySlotChanged()
+			FORCEINLINE void NotifySlotChanged(bool bSlotLayerChanged = false)
 			{
 				if ( Panel.IsValid() )
 				{
-					Panel.Pin()->NotifySlotChanged(this);
+					Panel.Pin()->NotifySlotChanged(this, bSlotLayerChanged);
 				}
 			}
 	};
@@ -251,9 +271,10 @@ private:
 	/** 
 	 * Callback used to resize our internal arrays if a slot (or slot span) is changed. 
 	 *
-	 * @param The slot that has just changed.
+	 * @param InSlot The slot that has just changed.
+	 * @param bSlotLayerChanged Whether the slot layer changed.
 	 */
-	void NotifySlotChanged( FSlot* InSlot );
+	void NotifySlotChanged(FSlot* InSlot, bool bSlotLayerChanged = false);
 
 private:
 
