Submenus can now have extension points in more situations
#rnx

diff --git a/Engine/Source/Runtime/Slate/Private/Framework/MultiBox/MultiBoxBuilder.cpp b/Engine/Source/Runtime/Slate/Private/Framework/MultiBox/MultiBoxBuilder.cpp
--- a/Engine/Source/Runtime/Slate/Private/Framework/MultiBox/MultiBoxBuilder.cpp
+++ b/Engine/Source/Runtime/Slate/Private/Framework/MultiBox/MultiBoxBuilder.cpp
@@ -229,12 +229,12 @@ void FMenuBuilder::AddSubMenu( const TAttribute<FText>& InMenuLabel, const TAttr
 	MultiBox->AddMultiBlock( NewMenuEntryBlock );
 }
 
-void FMenuBuilder::AddSubMenu( const TAttribute<FText>& InMenuLabel, const TAttribute<FText>& InToolTip, const FNewMenuDelegate& InSubMenu, const bool bInOpenSubMenuOnClick /*= false*/, const FSlateIcon& InIcon /*= FSlateIcon()*/, const bool bInShouldCloseWindowAfterMenuSelection /*= true*/ )
+void FMenuBuilder::AddSubMenu( const TAttribute<FText>& InMenuLabel, const TAttribute<FText>& InToolTip, const FNewMenuDelegate& InSubMenu, const bool bInOpenSubMenuOnClick /*= false*/, const FSlateIcon& InIcon /*= FSlateIcon()*/, const bool bInShouldCloseWindowAfterMenuSelection /*= true*/, FName InExtensionHook /*=NAME_None*/)
 {
 	ApplySectionBeginning();
 
 	const bool bIsSubMenu = true;
-	TSharedRef< FMenuEntryBlock > NewMenuEntryBlock( new FMenuEntryBlock( NAME_None, InMenuLabel, InToolTip, InSubMenu, ExtenderStack.Top(), bIsSubMenu, bInOpenSubMenuOnClick, CommandListStack.Last(), bCloseSelfOnly, InIcon, bInShouldCloseWindowAfterMenuSelection ) );
+	TSharedRef< FMenuEntryBlock > NewMenuEntryBlock( new FMenuEntryBlock( InExtensionHook, InMenuLabel, InToolTip, InSubMenu, ExtenderStack.Top(), bIsSubMenu, bInOpenSubMenuOnClick, CommandListStack.Last(), bCloseSelfOnly, InIcon, bInShouldCloseWindowAfterMenuSelection ) );
 
 	MultiBox->AddMultiBlock( NewMenuEntryBlock );
 }
diff --git a/Engine/Source/Runtime/Slate/Public/Framework/MultiBox/MultiBoxBuilder.h b/Engine/Source/Runtime/Slate/Public/Framework/MultiBox/MultiBoxBuilder.h
--- a/Engine/Source/Runtime/Slate/Public/Framework/MultiBox/MultiBoxBuilder.h
+++ b/Engine/Source/Runtime/Slate/Public/Framework/MultiBox/MultiBoxBuilder.h
@@ -263,7 +263,7 @@ public:
 	 */
 	void AddSubMenu( const TAttribute<FText>& InMenuLabel, const TAttribute<FText>& InToolTip, const FNewMenuDelegate& InSubMenu, const FUIAction& InUIAction, FName InExtensionHook, const EUserInterfaceActionType InUserInterfaceActionType, const bool bInOpenSubMenuOnClick = false, const FSlateIcon& InIcon = FSlateIcon(), const bool bInShouldCloseWindowAfterMenuSelection = true );
 
-	void AddSubMenu( const TAttribute<FText>& InMenuLabel, const TAttribute<FText>& InToolTip, const FNewMenuDelegate& InSubMenu, const bool bInOpenSubMenuOnClick = false, const FSlateIcon& InIcon = FSlateIcon(), const bool bInShouldCloseWindowAfterMenuSelection = true );
+	void AddSubMenu( const TAttribute<FText>& InMenuLabel, const TAttribute<FText>& InToolTip, const FNewMenuDelegate& InSubMenu, const bool bInOpenSubMenuOnClick = false, const FSlateIcon& InIcon = FSlateIcon(), const bool bInShouldCloseWindowAfterMenuSelection = true, FName InExtensionHook = NAME_None );
 
 	void AddSubMenu( const TSharedRef< SWidget > Contents, const FNewMenuDelegate& InSubMenu, const bool bInOpenSubMenuOnClick = false, const bool bInShouldCloseWindowAfterMenuSelection = true );
 
