Do not crash when toolbar button has empty label or tool tip
#rnx

diff --git a/Engine/Source/Runtime/Slate/Private/Framework/MultiBox/SToolBarButtonBlock.cpp b/Engine/Source/Runtime/Slate/Private/Framework/MultiBox/SToolBarButtonBlock.cpp
--- a/Engine/Source/Runtime/Slate/Private/Framework/MultiBox/SToolBarButtonBlock.cpp
+++ b/Engine/Source/Runtime/Slate/Private/Framework/MultiBox/SToolBarButtonBlock.cpp
@@ -132,6 +132,8 @@ void SToolBarButtonBlock::BuildMultiBlockWidget(const ISlateStyle* StyleSet, con
 	
 	TSharedRef< const FToolBarButtonBlock > ToolBarButtonBlock = StaticCastSharedRef< const FToolBarButtonBlock >( MultiBlock.ToSharedRef() );
 
+	TSharedPtr< const FUICommandInfo > UICommand = ToolBarButtonBlock->GetAction();
+
 	// Allow the block to override the action's label and tool tip string, if desired
 	TAttribute<FText> ActualLabel;
 	if (ToolBarButtonBlock->LabelOverride.IsSet())
@@ -140,7 +142,7 @@ void SToolBarButtonBlock::BuildMultiBlockWidget(const ISlateStyle* StyleSet, con
 	}
 	else
 	{
-		ActualLabel = ToolBarButtonBlock->GetAction()->GetLabel();
+		ActualLabel = UICommand.IsValid() ? UICommand->GetLabel() : FText::GetEmpty();
 	}
 
 	// Add this widget to the search list of the multibox
@@ -154,7 +156,7 @@ void SToolBarButtonBlock::BuildMultiBlockWidget(const ISlateStyle* StyleSet, con
 	}
 	else
 	{
-		ActualToolTip = ToolBarButtonBlock->GetAction()->GetDescription();
+		ActualToolTip = UICommand.IsValid() ? UICommand->GetDescription() : FText::GetEmpty();
 	}
 
 	// If a key is bound to the command, append it to the tooltip text.
