Fix issue where invisible but keyboard-focusable widgets in a scroll box could block navigation

diff --git a/Engine/Source/Runtime/Slate/Private/Widgets/Layout/SScrollBox.cpp b/Engine/Source/Runtime/Slate/Private/Widgets/Layout/SScrollBox.cpp
--- a/Engine/Source/Runtime/Slate/Private/Widgets/Layout/SScrollBox.cpp
+++ b/Engine/Source/Runtime/Slate/Private/Widgets/Layout/SScrollBox.cpp
@@ -973,7 +973,7 @@ FNavigationReply SScrollBox::OnNavigation(const FGeometry& MyGeometry, const FNa
 			for ( int32 ChildIndex = FocusedChildIndex + FocusedChildDirection; ChildIndex >= 0 && ChildIndex < Children.Num(); ChildIndex += FocusedChildDirection )
 			{
 				TSharedPtr<SWidget> PossiblyFocusableChild = GetKeyboardFocusableWidget(Children[ChildIndex].GetWidget());
-				if ( PossiblyFocusableChild.IsValid() )
+				if ( PossiblyFocusableChild.IsValid() && EVisibility::DoesVisibilityPassFilter(PossiblyFocusableChild->GetVisibility(), EVisibility::Visible))
 				{
 					NextFocusableChild = PossiblyFocusableChild;
 					break;
