Ensure scroll box focus logic can find focuable, nested widgets that are visible

diff --git a/Engine/Source/Runtime/Slate/Private/Widgets/Layout/SScrollBox.cpp b/Engine/Source/Runtime/Slate/Private/Widgets/Layout/SScrollBox.cpp
--- a/Engine/Source/Runtime/Slate/Private/Widgets/Layout/SScrollBox.cpp
+++ b/Engine/Source/Runtime/Slate/Private/Widgets/Layout/SScrollBox.cpp
@@ -973,7 +973,7 @@ FNavigationReply SScrollBox::OnNavigation(const FGeometry& MyGeometry, const FNa
 			for ( int32 ChildIndex = FocusedChildIndex + FocusedChildDirection; ChildIndex >= 0 && ChildIndex < Children.Num(); ChildIndex += FocusedChildDirection )
 			{
 				TSharedPtr<SWidget> PossiblyFocusableChild = GetKeyboardFocusableWidget(Children[ChildIndex].GetWidget());
-				if ( PossiblyFocusableChild.IsValid() && EVisibility::DoesVisibilityPassFilter(PossiblyFocusableChild->GetVisibility(), EVisibility::Visible))
+				if ( PossiblyFocusableChild.IsValid() )
 				{
 					NextFocusableChild = PossiblyFocusableChild;
 					break;
@@ -994,7 +994,7 @@ FNavigationReply SScrollBox::OnNavigation(const FGeometry& MyGeometry, const FNa
 
 TSharedPtr<SWidget> SScrollBox::GetKeyboardFocusableWidget(TSharedPtr<SWidget> InWidget)
 {
-	if (InWidget->SupportsKeyboardFocus())
+	if (InWidget->SupportsKeyboardFocus() && EVisibility::DoesVisibilityPassFilter(InWidget->GetVisibility(), EVisibility::Visible))
 	{
 		return InWidget;
 	}
@@ -1005,7 +1005,7 @@ TSharedPtr<SWidget> SScrollBox::GetKeyboardFocusableWidget(TSharedPtr<SWidget> I
 		{
 			TSharedPtr<SWidget> ChildWidget = Children->GetChildAt(i);
 			TSharedPtr<SWidget> FoucusableWidget = GetKeyboardFocusableWidget(ChildWidget);
-			if (FoucusableWidget.IsValid())
+			if (FoucusableWidget.IsValid() && EVisibility::DoesVisibilityPassFilter(FoucusableWidget->GetVisibility(), EVisibility::Visible))
 			{
 				return FoucusableWidget;
 			}
