Fixed slate combo buttons not focusing the opened widget. Allows search and keyboard navigation.

diff --git a/Engine/Source/Runtime/Slate/Private/Widgets/Input/SComboButton.cpp b/Engine/Source/Runtime/Slate/Private/Widgets/Input/SComboButton.cpp
--- a/Engine/Source/Runtime/Slate/Private/Widgets/Input/SComboButton.cpp
+++ b/Engine/Source/Runtime/Slate/Private/Widgets/Input/SComboButton.cpp
@@ -95,6 +95,15 @@ FReply SComboButton::OnButtonClicked()
 	if (bIsFocusable)
 	{
 		TSharedPtr<SWidget> WidgetToFocus = WidgetToFocusPtr.Pin();
+		if (!WidgetToFocus.IsValid())
+		{
+			// no explicitly focused widget, try to focus the content that is a child of the border
+			if (WrappedContent.IsValid() && WrappedContent->GetChildren()->Num() == 1)
+			{
+				WidgetToFocus = WrappedContent->GetChildren()->GetChildAt(0);
+			}
+		}
+
 		if (!WidgetToFocus.IsValid())
 		{
 			// no explicitly focused widget, try to focus the content
