PR #4561: Hide breadcrumb arrow when no child path available (Contributed by [REDACTED_PERSON])

diff --git a/Engine/Source/Editor/ContentBrowser/Private/SContentBrowser.cpp b/Engine/Source/Editor/ContentBrowser/Private/SContentBrowser.cpp
--- a/Engine/Source/Editor/ContentBrowser/Private/SContentBrowser.cpp
+++ b/Engine/Source/Editor/ContentBrowser/Private/SContentBrowser.cpp
@@ -391,6 +391,7 @@ void SContentBrowser::Construct( const FArguments& InArgs, const FName& InInstan
 						.ShowLeadingDelimiter(false)
 						.InvertTextColorOnHover(false)
 						.OnCrumbClicked(this, &SContentBrowser::OnPathClicked)
+						.HasCrumbMenuContent(this, &SContentBrowser::OnHasCrumbDelimiterContent)
 						.GetCrumbMenuContent(this, &SContentBrowser::OnGetCrumbDelimiterContent)
 						.AddMetaData<FTagMetaData>(FTagMetaData(TEXT("ContentBrowserPath")))
 					]
@@ -1756,6 +1757,61 @@ void SContentBrowser::OnPathMenuItemClicked(FString ClickedPath)
 	OnPathClicked( ClickedPath );
 }
 
+bool SContentBrowser::OnHasCrumbDelimiterContent(const FString& CrumbData) const
+{
+	const FSourcesData SourcesData = AssetViewPtr->GetSourcesData();
+	if (SourcesData.HasCollections())
+	{
+		TOptional<FCollectionNameType> CollectionClicked;
+		{
+			FString CollectionName;
+			FString CollectionTypeString;
+			if (CrumbData.Split(TEXT("?"), &CollectionName, &CollectionTypeString))
+			{
+				const int32 CollectionType = FCString::Atoi(*CollectionTypeString);
+				if (CollectionType >= 0 && CollectionType < ECollectionShareType::CST_All)
+				{
+					CollectionClicked = FCollectionNameType(FName(*CollectionName), ECollectionShareType::Type(CollectionType));
+				}
+			}
+		}
+
+		TArray<FCollectionNameType> ChildCollections;
+		if (CollectionClicked.IsSet())
+		{
+			FCollectionManagerModule& CollectionManagerModule = FCollectionManagerModule::GetModule();
+			CollectionManagerModule.Get().GetChildCollections(CollectionClicked->Name, CollectionClicked->Type, ChildCollections);
+		}
+
+		return (ChildCollections.Num() > 0);
+	}
+	else if (SourcesData.HasPackagePaths())
+	{
+		TArray<FString> SubPaths;
+		const bool bRecurse = false;
+		if (ContentBrowserUtils::IsClassPath(CrumbData))
+		{
+			TSharedRef<FNativeClassHierarchy> NativeClassHierarchy = FContentBrowserSingleton::Get().GetNativeClassHierarchy();
+
+			FNativeClassHierarchyFilter ClassFilter;
+			ClassFilter.ClassPaths.Add(FName(*CrumbData));
+			ClassFilter.bRecursivePaths = bRecurse;
+
+			NativeClassHierarchy->GetMatchingFolders(ClassFilter, SubPaths);
+		}
+		else
+		{
+			FAssetRegistryModule& AssetRegistryModule = FModuleManager::LoadModuleChecked<FAssetRegistryModule>(TEXT("AssetRegistry"));
+			IAssetRegistry& AssetRegistry = AssetRegistryModule.Get();
+
+			AssetRegistry.GetSubPaths(CrumbData, SubPaths, bRecurse);
+		}
+
+		return (SubPaths.Num() > 0);
+	}
+	return false;
+}
+
 TSharedPtr<SWidget> SContentBrowser::OnGetCrumbDelimiterContent(const FString& CrumbData) const
 {
 	FSourcesData SourcesData = AssetViewPtr->GetSourcesData();
diff --git a/Engine/Source/Editor/ContentBrowser/Private/SContentBrowser.h b/Engine/Source/Editor/ContentBrowser/Private/SContentBrowser.h
--- a/Engine/Source/Editor/ContentBrowser/Private/SContentBrowser.h
+++ b/Engine/Source/Editor/ContentBrowser/Private/SContentBrowser.h
@@ -223,6 +223,9 @@ private:
 	/** Called when item in the path delimiter arrow menu is clicked */
 	void OnPathMenuItemClicked(FString ClickedPath);
 
+	/** Called to query whether the crumb menu would contain any content (see also OnGetCrumbDelimiterContent) */
+	bool OnHasCrumbDelimiterContent(const FString& CrumbData) const;
+
 	/** 
 	 * Populates the delimiter arrow with a menu of directories under the current directory that can be navigated to
 	 * 
diff --git a/Engine/Source/Runtime/Slate/Public/Widgets/Navigation/SBreadcrumbTrail.h b/Engine/Source/Runtime/Slate/Public/Widgets/Navigation/SBreadcrumbTrail.h
--- a/Engine/Source/Runtime/Slate/Public/Widgets/Navigation/SBreadcrumbTrail.h
+++ b/Engine/Source/Runtime/Slate/Public/Widgets/Navigation/SBreadcrumbTrail.h
@@ -57,12 +57,14 @@ public:
 	/** Callback for when a crumb in the trail has been clicked */
 	DECLARE_DELEGATE_OneParam( FOnCrumbClicked, const ItemType& /*CrumbData*/ );
 
+	/** Callback for checking whether the menu to be displayed when clicking on a crumb's delimiter arrow would have any content */
+	DECLARE_DELEGATE_RetVal_OneParam( bool, FHasCrumbMenuContent, const ItemType& /*CrumbData*/ );
+
 	/** Callback for getting the menu content to be displayed when clicking on a crumb's delimiter arrow */
 	DECLARE_DELEGATE_RetVal_OneParam( TSharedPtr< SWidget >, FGetCrumbMenuContent, const ItemType& /*CrumbData*/ );
 
 	DECLARE_DELEGATE_RetVal_TwoParams( FSlateColor, FOnGetCrumbColor, int32 /*CrumbId*/, bool /*bInvert*/ );
 
-
 	SLATE_BEGIN_ARGS( SBreadcrumbTrail )
 		: _InvertTextColorOnHover(true)
 		, _ButtonStyle( &FCoreStyle::Get().GetWidgetStyle< FButtonStyle >( "BreadcrumbButton" ) )
@@ -71,6 +73,7 @@ public:
 		, _DelimiterImage( FCoreStyle::Get().GetBrush("BreadcrumbTrail.Delimiter") )
 		, _ShowLeadingDelimiter(false)
 		, _PersistentBreadcrumbs(false)
+		, _HasCrumbMenuContent()
 		, _GetCrumbMenuContent()
 		, _OnGetCrumbColor()
 	    {}
@@ -105,6 +108,8 @@ public:
 		/** If true, do not remove breadcrumbs when clicking */
 		SLATE_ARGUMENT( bool, PersistentBreadcrumbs )
 
+		SLATE_EVENT( FHasCrumbMenuContent, HasCrumbMenuContent )
+
 		SLATE_EVENT( FGetCrumbMenuContent, GetCrumbMenuContent )
 
 		SLATE_EVENT( FOnGetCrumbColor, OnGetCrumbColor )
@@ -124,10 +129,10 @@ public:
 		OnCrumbPopped = InArgs._OnCrumbPopped;
 		OnCrumbClicked = InArgs._OnCrumbClicked;
 		bHasStaticBreadcrumbs = InArgs._PersistentBreadcrumbs;
+		HasCrumbMenuContentCallback = InArgs._HasCrumbMenuContent;
 		GetCrumbMenuContentCallback = InArgs._GetCrumbMenuContent;
 		OnGetCrumbColor = InArgs._OnGetCrumbColor;
 
-
 		NextValidCrumbID = 0;
 
 		ChildSlot
@@ -176,8 +181,9 @@ public:
 			DelimiterContent = SNew(SButton)
 				.VAlign(EVerticalAlignment::VAlign_Center)
 				.ButtonStyle(ButtonStyle)
-				.OnClicked( this, &SBreadcrumbTrail::OnCrumbDelimiterClicked, NextValidCrumbID )
-				.ContentPadding( FMargin(5, 0) )
+				.Visibility(this, &SBreadcrumbTrail::GetCrumbDelimiterVisibility, NextValidCrumbID)
+				.OnClicked(this, &SBreadcrumbTrail::OnCrumbDelimiterClicked, NextValidCrumbID)
+				.ContentPadding(FMargin(5, 0))
 				[
 					SNew(SImage)
 					.Image(DelimiterImage)
@@ -293,6 +299,15 @@ public:
 
 private:
 
+	EVisibility GetCrumbDelimiterVisibility(int32 CrumbID) const
+	{
+		if (HasCrumbs() && CrumbList.Last().CrumbID == CrumbID && HasCrumbMenuContentCallback.IsBound() && !HasCrumbMenuContentCallback.Execute(CrumbList.Last().CrumbData))
+		{
+			return EVisibility::Collapsed;
+		}
+		return EVisibility::Visible;
+	}
+
 	FReply OnCrumbDelimiterClicked( int32 CrumbID )
 	{
 		FReply Reply = FReply::Unhandled();
@@ -337,13 +352,10 @@ private:
 	/** Handler to determine the visibility of the arrows between crumbs */
 	EVisibility GetDelimiterVisibility(int32 CrumbID) const
 	{
-		if ( HasCrumbs() )
+		if ( HasCrumbs() && CrumbList.Last().CrumbID == CrumbID )
 		{
-			if (CrumbList.Last().CrumbID == CrumbID)
-			{
-				// Collapse the last delimiter
-				return EVisibility::Collapsed;
-			}
+			// Collapse the last delimiter
+			return EVisibility::Collapsed;
 		}
 
 		return EVisibility::Visible;
@@ -487,6 +499,9 @@ private:
 	/** Delegate to invoke when selection changes. */
 	FOnCrumbClicked OnCrumbClicked;
 
+	/** Delegate to invoke to check whether the crumb's menu would have any content */
+	FHasCrumbMenuContent HasCrumbMenuContentCallback;
+
 	/** Delegate to invoke to retrieve the content for a crumb's menu */
 	FGetCrumbMenuContent GetCrumbMenuContentCallback;
 
