PR #4308: Can always use Select All shortcut in EditableText (Contributed by [REDACTED_PERSON])

diff --git a/Engine/Source/Runtime/Slate/Private/Widgets/Text/SlateEditableTextLayout.cpp b/Engine/Source/Runtime/Slate/Private/Widgets/Text/SlateEditableTextLayout.cpp
--- a/Engine/Source/Runtime/Slate/Private/Widgets/Text/SlateEditableTextLayout.cpp
+++ b/Engine/Source/Runtime/Slate/Private/Widgets/Text/SlateEditableTextLayout.cpp
@@ -147,12 +147,9 @@ FSlateEditableTextLayout::FSlateEditableTextLayout(ISlateEditableTextWidget& InO
 
 	auto ExecuteDeleteAction = [this]()
 	{
-		if (CanExecuteDelete())
-		{
-			BeginEditTransation();
-			DeleteSelectedText();
-			EndEditTransaction();
-		}
+		BeginEditTransation();
+		DeleteSelectedText();
+		EndEditTransaction();
 	};
 
 	UICommandList = MakeShareable(new FUICommandList());
@@ -1653,27 +1650,7 @@ bool FSlateEditableTextLayout::IsTextSelectedAt(const FVector2D& InLocalPosition
 
 bool FSlateEditableTextLayout::CanExecuteSelectAll() const
 {
-	bool bCanExecute = true;
-
-	// Can't select all if string is empty
-	if (TextLayout->IsEmpty())
-	{
-		bCanExecute = false;
-	}
-
-	const TArray< FTextLayout::FLineModel >& Lines = TextLayout->GetLineModels();
-	const int32 NumberOfLines = Lines.Num();
-
-	// Can't select all if entire string is already selected
-	const FTextLocation CursorInteractionPosition = CursorInfo.GetCursorInteractionLocation();
-	if (SelectionStart.IsSet() &&
-		SelectionStart.GetValue() == FTextLocation(0, 0) &&
-		CursorInteractionPosition == FTextLocation(NumberOfLines - 1, Lines[NumberOfLines - 1].Text->Len()))
-	{
-		bCanExecute = false;
-	}
-
-	return bCanExecute;
+	return true;
 }
 
 void FSlateEditableTextLayout::SelectAllText()
