Fix for sccissor rect not initializing it's slot.

diff --git a/Engine/Source/Runtime/Slate/Private/Widgets/Layout/SScissorRectBox.cpp b/Engine/Source/Runtime/Slate/Private/Widgets/Layout/SScissorRectBox.cpp
--- a/Engine/Source/Runtime/Slate/Private/Widgets/Layout/SScissorRectBox.cpp
+++ b/Engine/Source/Runtime/Slate/Private/Widgets/Layout/SScissorRectBox.cpp
@@ -11,6 +11,12 @@ namespace Slate
 	static FAutoConsoleVariable OverrideScissorRect( TEXT( "Slate.OverrideScissorRect" ), 1, TEXT( "Whether to allow Slate to apply a scissor clip to UI elements to prevent certain artifacts.  You might need to disable this if you are drawing UI over a mirror backbuffer while in VR." ) );
 }
 
+SScissorRectBox::SScissorRectBox()
+	:ChildSlot(this)
+{
+
+}
+
 void SScissorRectBox::Construct(const FArguments& InArgs)
 {
 	SetCanTick(false);
diff --git a/Engine/Source/Runtime/Slate/Public/Widgets/Layout/SScissorRectBox.h b/Engine/Source/Runtime/Slate/Public/Widgets/Layout/SScissorRectBox.h
--- a/Engine/Source/Runtime/Slate/Public/Widgets/Layout/SScissorRectBox.h
+++ b/Engine/Source/Runtime/Slate/Public/Widgets/Layout/SScissorRectBox.h
@@ -46,7 +46,9 @@ public:
 	}
 		SLATE_DEFAULT_SLOT(FArguments, Content)
 	SLATE_END_ARGS()
-
+	
+	SScissorRectBox();
+	
 	void Construct( const FArguments& InArgs );
 
 	/**
