This solves the bug of the layout border remaining visible when the dragged operation is cancelled by pressing Alt+Tab or command (i.e., Windows) key.

diff --git a/Engine/Source/Runtime/Slate/Private/Framework/Application/SlateApplication.cpp b/Engine/Source/Runtime/Slate/Private/Framework/Application/SlateApplication.cpp
--- a/Engine/Source/Runtime/Slate/Private/Framework/Application/SlateApplication.cpp
+++ b/Engine/Source/Runtime/Slate/Private/Framework/Application/SlateApplication.cpp
@@ -7290,7 +7290,7 @@ void FSlateApplication::ProcessApplicationActivationEvent(bool InAppActivated)
 		// If we have a slate-only drag-drop occurring, stop the drag drop.
 		if (IsDragDropping() && !DragDropContent->IsExternalOperation())
 		{
-			ResetDragDropState();
+			CancelDragDrop();
 		}
 
 		// Clear the pressed buttons when we deactivate the application, the button state can no longer be trusted.
