Fixed issue where child tabs would grow when destroyed/recreated in FGlobalTabmanager::OnTabRelocated due to DPI scale being applied redundantly

diff --git a/Engine/Source/Runtime/Slate/Private/Framework/Docking/TabManager.cpp b/Engine/Source/Runtime/Slate/Private/Framework/Docking/TabManager.cpp
--- a/Engine/Source/Runtime/Slate/Private/Framework/Docking/TabManager.cpp
+++ b/Engine/Source/Runtime/Slate/Private/Framework/Docking/TabManager.cpp
@@ -2149,6 +2149,7 @@ void FGlobalTabmanager::OnTabRelocated( const TSharedRef<SDockTab>& RelocatedTab
 					.SupportsMinimize(false)
 					.SupportsMaximize(false)
 					.CreateTitleBar(false)
+					.AdjustInitialSizeAndPositionForDPIScale(false)
 					[
 						ChildDockArea
 					];
