PR #6087: Add GetNavigationActionFromKey function that handles multiple users and modifiers properly (Contributed by [REDACTED_PERSON])

diff --git a/Engine/Plugins/Runtime/Synthesis/Source/Synthesis/Private/UI/SSynth2DSlider.cpp b/Engine/Plugins/Runtime/Synthesis/Source/Synthesis/Private/UI/SSynth2DSlider.cpp
--- a/Engine/Plugins/Runtime/Synthesis/Source/Synthesis/Private/UI/SSynth2DSlider.cpp
+++ b/Engine/Plugins/Runtime/Synthesis/Source/Synthesis/Private/UI/SSynth2DSlider.cpp
@@ -119,14 +119,14 @@ void SSynth2DSlider::ResetControllerState()
 FReply SSynth2DSlider::OnKeyDown(const FGeometry& MyGeometry, const FKeyEvent& InKeyEvent)
 {
 	FReply Reply = FReply::Unhandled();
-	const FKey KeyPressed = InKeyEvent.GetKey();
+	const EUINavigation NavDirection = FSlateApplication::Get().GetNavigationDirectionFromKey(InKeyEvent);
 
 	if (IsInteractable())
 	{
 		// The controller's bottom face button must be pressed once to begin manipulating the slider's value.
 		// Navigation away from the widget is prevented until the button has been pressed again or focus is lost.
 		// The value can be manipulated by using the game pad's directional arrows ( relative to slider orientation ).
-		if (FSlateApplication::Get().GetNavigationActionForKey(KeyPressed) == EUINavigationAction::Accept)
+		if (FSlateApplication::Get().GetNavigationActionFromKey(InKeyEvent) == EUINavigationAction::Accept)
 		{
 			if (bControllerInputCaptured == false)
 			{
@@ -149,22 +149,22 @@ FReply SSynth2DSlider::OnKeyDown(const FGeometry& MyGeometry, const FKeyEvent& I
 
 			if (Orientation == EOrientation::Orient_Horizontal)
 			{
-				if (KeyPressed == EKeys::Left || KeyPressed == EKeys::Gamepad_DPad_Left || KeyPressed == EKeys::Gamepad_LeftStick_Left)
+				if (NavDirection == EUINavigation::Left)
 				{
 					NewValueX -= StepSize.Get();
 				}
-				else if (KeyPressed == EKeys::Right || KeyPressed == EKeys::Gamepad_DPad_Right || KeyPressed == EKeys::Gamepad_LeftStick_Right)
+				else if (NavDirection == EUINavigation::Right)
 				{
 					NewValueX += StepSize.Get();
 				}
 			}
 			else
 			{
-				if (KeyPressed == EKeys::Down || KeyPressed == EKeys::Gamepad_DPad_Down || KeyPressed == EKeys::Gamepad_LeftStick_Down)
+				if (NavDirection == EUINavigation::Down)
 				{
 					NewValueY -= StepSize.Get();
 				}
-				else if (KeyPressed == EKeys::Up || KeyPressed == EKeys::Gamepad_DPad_Up || KeyPressed == EKeys::Gamepad_LeftStick_Up)
+				else if (NavDirection == EUINavigation::Up)
 				{
 					NewValueY += StepSize.Get();
 				}
diff --git a/Engine/Plugins/Runtime/Synthesis/Source/Synthesis/Private/UI/SSynthKnob.cpp b/Engine/Plugins/Runtime/Synthesis/Source/Synthesis/Private/UI/SSynthKnob.cpp
--- a/Engine/Plugins/Runtime/Synthesis/Source/Synthesis/Private/UI/SSynthKnob.cpp
+++ b/Engine/Plugins/Runtime/Synthesis/Source/Synthesis/Private/UI/SSynthKnob.cpp
@@ -112,9 +112,9 @@ void SSynthKnob::ResetControllerState()
 FReply SSynthKnob::OnKeyDown(const FGeometry& MyGeometry, const FKeyEvent& InKeyEvent)
 {
 	FReply Reply = FReply::Unhandled();
-	const FKey KeyPressed = InKeyEvent.GetKey();
+	const EUINavigation NavDirection = FSlateApplication::Get().GetNavigationDirectionFromKey(InKeyEvent);
 
-	if (KeyPressed == FineTuneKey)
+	if (InKeyEvent.GetKey() == FineTuneKey)
 	{
 		bIsFineTune = true;
 	}
@@ -124,7 +124,7 @@ FReply SSynthKnob::OnKeyDown(const FGeometry& MyGeometry, const FKeyEvent& InKey
 		// The controller's bottom face button must be pressed once to begin manipulating the slider's value.
 		// Navigation away from the widget is prevented until the button has been pressed again or focus is lost.
 		// The value can be manipulated by using the game pad's directional arrows ( relative to slider orientation ).
-		if (FSlateApplication::Get().GetNavigationActionForKey(KeyPressed) == EUINavigationAction::Accept)
+		if (FSlateApplication::Get().GetNavigationActionFromKey(InKeyEvent) == EUINavigationAction::Accept)
 		{
 			if (bControllerInputCaptured == false)
 			{
@@ -143,11 +143,11 @@ FReply SSynthKnob::OnKeyDown(const FGeometry& MyGeometry, const FKeyEvent& InKey
 		if (bControllerInputCaptured)
 		{
 			float NewValue = ValueAttribute.Get();
-			if (KeyPressed == EKeys::Down || KeyPressed == EKeys::Gamepad_DPad_Down || KeyPressed == EKeys::Gamepad_LeftStick_Down)
+			if (NavDirection == EUINavigation::Down)
 			{
 				NewValue -= StepSize.Get();
 			}
-			else if (KeyPressed == EKeys::Up || KeyPressed == EKeys::Gamepad_DPad_Up || KeyPressed == EKeys::Gamepad_LeftStick_Up)
+			else if (NavDirection == EUINavigation::Up)
 			{
 				NewValue += StepSize.Get();
 			}
diff --git a/Engine/Source/Runtime/Engine/Classes/Kismet/KismetInputLibrary.h b/Engine/Source/Runtime/Engine/Classes/Kismet/KismetInputLibrary.h
--- a/Engine/Source/Runtime/Engine/Classes/Kismet/KismetInputLibrary.h
+++ b/Engine/Source/Runtime/Engine/Classes/Kismet/KismetInputLibrary.h
@@ -90,10 +90,13 @@ class ENGINE_API UKismetInputLibrary : public UBlueprintFunctionLibrary
 	UFUNCTION(BlueprintPure, meta = (DisplayName = "Is Valid"), Category = "Utilities|Key")
 	static bool Key_IsValid(const FKey& Key);
 
-	/** Returns the navigation action corresponding to this key, or Invalid if not found */
-	UFUNCTION(BlueprintPure, meta = (DisplayName = "Get Key Navigation Action"), Category = "Utilities|Key")
+	UFUNCTION(BlueprintPure, meta = (DisplayName = "Get Key Navigation Action", DeprecatedFunction, DeprecationMessage = "Use Get Key Event Navigation Action instead"), Category = "Utilities|Key")
 	static EUINavigationAction Key_GetNavigationAction(const FKey& InKey);
 
+	/** Returns the navigation action corresponding to this key, or Invalid if not found */
+	UFUNCTION(BlueprintPure, meta = (DisplayName = "Get Key Event Navigation Action"), Category = "Utilities|KeyEvent")
+	static EUINavigationAction Key_GetNavigationActionFromKey(const FKeyEvent& InKeyEvent);
+
 	/** Returns the navigation action corresponding to this key, or Invalid if not found */
 	UFUNCTION(BlueprintPure, meta = (DisplayName = "Get Key Event Navigation Direction"), Category = "Utilities|KeyEvent")
 	static EUINavigation Key_GetNavigationDirectionFromKey(const FKeyEvent& InKeyEvent);
diff --git a/Engine/Source/Runtime/Engine/Private/KismetInputLibrary.cpp b/Engine/Source/Runtime/Engine/Private/KismetInputLibrary.cpp
--- a/Engine/Source/Runtime/Engine/Private/KismetInputLibrary.cpp
+++ b/Engine/Source/Runtime/Engine/Private/KismetInputLibrary.cpp
@@ -70,11 +70,23 @@ bool UKismetInputLibrary::Key_IsValid(const FKey& Key)
 
 EUINavigationAction UKismetInputLibrary::Key_GetNavigationAction(const FKey& InKey)
 {
+PRAGMA_DISABLE_DEPRECATION_WARNINGS
 	if (FSlateApplication::IsInitialized())
 	{
 		return FSlateApplication::Get().GetNavigationActionForKey(InKey);
 	}
 
+	return EUINavigationAction::Invalid;
+PRAGMA_ENABLE_DEPRECATION_WARNINGS
+}
+
+EUINavigationAction UKismetInputLibrary::Key_GetNavigationActionFromKey(const FKeyEvent& InKeyEvent)
+{
+	if (FSlateApplication::IsInitialized())
+	{
+		return FSlateApplication::Get().GetNavigationActionFromKey(InKeyEvent);
+	}
+
 	return EUINavigationAction::Invalid;
 }
 
diff --git a/Engine/Source/Runtime/GameMenuBuilder/Private/SGameMenuPageWidget.cpp b/Engine/Source/Runtime/GameMenuBuilder/Private/SGameMenuPageWidget.cpp
--- a/Engine/Source/Runtime/GameMenuBuilder/Private/SGameMenuPageWidget.cpp
+++ b/Engine/Source/Runtime/GameMenuBuilder/Private/SGameMenuPageWidget.cpp
@@ -942,12 +942,12 @@ FReply SGameMenuPageWidget::OnKeyDown(const FGeometry& MyGeometry, const FKeyEve
 	if ((CurrentMenu.IsValid() == true) && !bControlsLocked)
 	{
 		bool bNavigationLocked = PendingMainMenu.IsValid() || PendingSubMenu.IsValid();
-		const FKey Key = InKeyEvent.GetKey();
-		EUINavigationAction NavAction = FSlateApplication::Get().GetNavigationActionForKey(Key);
+		const EUINavigation NavDirection = FSlateApplication::Get().GetNavigationDirectionFromKey(InKeyEvent);
+		const EUINavigationAction NavAction = FSlateApplication::Get().GetNavigationActionFromKey(InKeyEvent);
 
 		if (bNavigationLocked == false)
 		{
-			if (Key == EKeys::Up || Key == EKeys::Gamepad_DPad_Up || Key == EKeys::Gamepad_LeftStick_Up)
+			if (NavDirection == EUINavigation::Up)
 			{
 				if (SelectedIndex > 0)
 				{
@@ -955,7 +955,7 @@ FReply SGameMenuPageWidget::OnKeyDown(const FGeometry& MyGeometry, const FKeyEve
 				}
 				Result = FReply::Handled();
 			}
-			else if (Key == EKeys::Down || Key == EKeys::Gamepad_DPad_Down || Key == EKeys::Gamepad_LeftStick_Down)
+			else if (NavDirection == EUINavigation::Down)
 			{
 				if (SelectedIndex + 1 < CurrentMenu->NumItems())
 				{
@@ -963,12 +963,12 @@ FReply SGameMenuPageWidget::OnKeyDown(const FGeometry& MyGeometry, const FKeyEve
 				}
 				Result = FReply::Handled();
 			}
-			else if (Key == EKeys::Left || Key == EKeys::Gamepad_DPad_Left || Key == EKeys::Gamepad_LeftStick_Left)
+			else if (NavDirection == EUINavigation::Left)
 			{
 				ChangeOption(-1);
 				Result = FReply::Handled();
 			}
-			else if (Key == EKeys::Right || Key == EKeys::Gamepad_DPad_Right || Key == EKeys::Gamepad_LeftStick_Right)
+			else if (NavDirection == EUINavigation::Right)
 			{
 				ChangeOption(1);
 				Result = FReply::Handled();
@@ -979,7 +979,7 @@ FReply SGameMenuPageWidget::OnKeyDown(const FGeometry& MyGeometry, const FKeyEve
 			ConfirmMenuItem();
 			Result = FReply::Handled();
 		} 
-		else if (NavAction == EUINavigationAction::Back || Key == EKeys::Gamepad_Special_Left)
+		else if (NavAction == EUINavigationAction::Back || InKeyEvent.GetKey() == EKeys::Gamepad_Special_Left)
 		{
 			MenuGoBack(true);
 			Result = FReply::Handled();
diff --git a/Engine/Source/Runtime/Slate/Private/Framework/Application/NavigationConfig.cpp b/Engine/Source/Runtime/Slate/Private/Framework/Application/NavigationConfig.cpp
--- a/Engine/Source/Runtime/Slate/Private/Framework/Application/NavigationConfig.cpp
+++ b/Engine/Source/Runtime/Slate/Private/Framework/Application/NavigationConfig.cpp
@@ -147,6 +147,14 @@ float FNavigationConfig::GetRepeatRateForPressure(float InPressure, int32 InRepe
 	return RepeatRate;
 }
 
+EUINavigationAction FNavigationConfig::GetNavigationActionFromKey(const FKeyEvent& InKeyEvent) const
+{
+PRAGMA_DISABLE_DEPRECATION_WARNINGS
+	// Call raw key version for back compatibility, subclasses should override this function
+	return GetNavigationActionForKey(InKeyEvent.GetKey());
+PRAGMA_ENABLE_DEPRECATION_WARNINGS
+}
+
 EUINavigationAction FNavigationConfig::GetNavigationActionForKey(const FKey& InKey) const
 {
 	if (InKey == EKeys::Enter || InKey == EKeys::SpaceBar || InKey == EKeys::Virtual_Accept)
diff --git a/Engine/Source/Runtime/Slate/Private/Framework/Application/SlateApplication.cpp b/Engine/Source/Runtime/Slate/Private/Framework/Application/SlateApplication.cpp
--- a/Engine/Source/Runtime/Slate/Private/Framework/Application/SlateApplication.cpp
+++ b/Engine/Source/Runtime/Slate/Private/Framework/Application/SlateApplication.cpp
@@ -2170,9 +2170,16 @@ EUINavigation FSlateApplication::GetNavigationDirectionFromAnalog(const FAnalogI
 	return NavigationConfig->GetNavigationDirectionFromAnalog(InAnalogEvent);
 }
 
+EUINavigationAction FSlateApplication::GetNavigationActionFromKey(const FKeyEvent& InKeyEvent) const
+{
+	return NavigationConfig->GetNavigationActionFromKey(InKeyEvent);
+}
+
 EUINavigationAction FSlateApplication::GetNavigationActionForKey(const FKey& InKey) const
 {
+PRAGMA_DISABLE_DEPRECATION_WARNINGS
 	return NavigationConfig->GetNavigationActionForKey(InKey);
+PRAGMA_ENABLE_DEPRECATION_WARNINGS
 }
 
 void FSlateApplication::AddModalWindow( TSharedRef<SWindow> InSlateWindow, const TSharedPtr<const SWidget> InParentWidget, bool bSlowTaskWindow )
diff --git a/Engine/Source/Runtime/Slate/Private/Widgets/Input/SButton.cpp b/Engine/Source/Runtime/Slate/Private/Widgets/Input/SButton.cpp
--- a/Engine/Source/Runtime/Slate/Private/Widgets/Input/SButton.cpp
+++ b/Engine/Source/Runtime/Slate/Private/Widgets/Input/SButton.cpp
@@ -164,7 +164,7 @@ void SButton::OnFocusLost( const FFocusEvent& InFocusEvent )
 FReply SButton::OnKeyDown( const FGeometry& MyGeometry, const FKeyEvent& InKeyEvent )
 {
 	FReply Reply = FReply::Unhandled();
-	if (IsEnabled() && FSlateApplication::Get().GetNavigationActionForKey(InKeyEvent.GetKey()) == EUINavigationAction::Accept)
+	if (IsEnabled() && FSlateApplication::Get().GetNavigationActionFromKey(InKeyEvent) == EUINavigationAction::Accept)
 	{
 		Press();
 
@@ -194,7 +194,7 @@ FReply SButton::OnKeyUp(const FGeometry& MyGeometry, const FKeyEvent& InKeyEvent
 {
 	FReply Reply = FReply::Unhandled();
 
-	if (IsEnabled() && FSlateApplication::Get().GetNavigationActionForKey(InKeyEvent.GetKey()) == EUINavigationAction::Accept)
+	if (IsEnabled() && FSlateApplication::Get().GetNavigationActionFromKey(InKeyEvent) == EUINavigationAction::Accept)
 	{
 		const bool bWasPressed = bIsPressed;
 
diff --git a/Engine/Source/Runtime/Slate/Private/Widgets/Input/SCheckBox.cpp b/Engine/Source/Runtime/Slate/Private/Widgets/Input/SCheckBox.cpp
--- a/Engine/Source/Runtime/Slate/Private/Widgets/Input/SCheckBox.cpp
+++ b/Engine/Source/Runtime/Slate/Private/Widgets/Input/SCheckBox.cpp
@@ -79,7 +79,7 @@ bool SCheckBox::SupportsKeyboardFocus() const
 
 FReply SCheckBox::OnKeyUp( const FGeometry& MyGeometry, const FKeyEvent& InKeyEvent )
 {
-	if (FSlateApplication::Get().GetNavigationActionForKey(InKeyEvent.GetKey()) == EUINavigationAction::Accept)
+	if (FSlateApplication::Get().GetNavigationActionFromKey(InKeyEvent) == EUINavigationAction::Accept)
 	{
 		ToggleCheckedState();
 
diff --git a/Engine/Source/Runtime/Slate/Private/Widgets/Input/SComboButton.cpp b/Engine/Source/Runtime/Slate/Private/Widgets/Input/SComboButton.cpp
--- a/Engine/Source/Runtime/Slate/Private/Widgets/Input/SComboButton.cpp
+++ b/Engine/Source/Runtime/Slate/Private/Widgets/Input/SComboButton.cpp
@@ -128,7 +128,7 @@ FReply SComboButton::OnButtonClicked()
 FReply SComboButton::OnKeyDown(const FGeometry& MyGeometry, const FKeyEvent& InKeyEvent)
 {
 	FReply Reply = FReply::Unhandled();
-	if (FSlateApplication::Get().GetNavigationActionForKey(InKeyEvent.GetKey()) == EUINavigationAction::Accept)
+	if (FSlateApplication::Get().GetNavigationActionFromKey(InKeyEvent) == EUINavigationAction::Accept)
 	{
 		// Handle menu open with controller.
 		Reply = OnButtonClicked();
diff --git a/Engine/Source/Runtime/Slate/Private/Widgets/Input/SSlider.cpp b/Engine/Source/Runtime/Slate/Private/Widgets/Input/SSlider.cpp
--- a/Engine/Source/Runtime/Slate/Private/Widgets/Input/SSlider.cpp
+++ b/Engine/Source/Runtime/Slate/Private/Widgets/Input/SSlider.cpp
@@ -161,35 +161,35 @@ void SSlider::ResetControllerState()
 FNavigationReply SSlider::OnNavigation(const FGeometry& MyGeometry, const FNavigationEvent& InNavigationEvent)
 {
 	FNavigationReply Reply = FNavigationReply::Escape();
-		if (bControllerInputCaptured || !bRequiresControllerLock)
+	if (bControllerInputCaptured || !bRequiresControllerLock)
+	{
+		float NewValue = ValueAttribute.Get();
+		if (Orientation == EOrientation::Orient_Horizontal)
 		{
-			float NewValue = ValueAttribute.Get();
-			if (Orientation == EOrientation::Orient_Horizontal)
-			{
 			if (InNavigationEvent.GetNavigationType() == EUINavigation::Left)
-				{
-					NewValue -= StepSize.Get();
+			{
+				NewValue -= StepSize.Get();
 				Reply = FNavigationReply::Stop();
-				}
+			}
 			else if (InNavigationEvent.GetNavigationType() == EUINavigation::Right)
-				{
-					NewValue += StepSize.Get();
+			{
+				NewValue += StepSize.Get();
 				Reply = FNavigationReply::Stop();
-				}
 			}
-			else
-			{
+		}
+		else
+		{
 			if (InNavigationEvent.GetNavigationType() == EUINavigation::Down)
-				{
-					NewValue -= StepSize.Get();
+			{
+				NewValue -= StepSize.Get();
 				Reply = FNavigationReply::Stop();
-				}
-			if (InNavigationEvent.GetNavigationType() == EUINavigation::Up)
-				{
-					NewValue += StepSize.Get();
+			}
+			else if (InNavigationEvent.GetNavigationType() == EUINavigation::Up)
+			{
+				NewValue += StepSize.Get();
 				Reply = FNavigationReply::Stop();
-				}
 			}
+		}
 		if (ValueAttribute.Get() != NewValue)
 		{
 			CommitValue(FMath::Clamp(NewValue, MinValue, MaxValue));
@@ -197,7 +197,7 @@ FNavigationReply SSlider::OnNavigation(const FGeometry& MyGeometry, const FNavig
 	}
 
 	if (Reply.GetBoundaryRule() != EUINavigationRule::Escape)
-			{
+	{
 		Reply = SLeafWidget::OnNavigation(MyGeometry, InNavigationEvent);
 	}
 
@@ -207,14 +207,13 @@ FNavigationReply SSlider::OnNavigation(const FGeometry& MyGeometry, const FNavig
 FReply SSlider::OnKeyDown(const FGeometry& MyGeometry, const FKeyEvent& InKeyEvent)
 {
 	FReply Reply = FReply::Unhandled();
-	const FKey KeyPressed = InKeyEvent.GetKey();
 
 	if (IsInteractable())
 	{
 		// The controller's bottom face button must be pressed once to begin manipulating the slider's value.
 		// Navigation away from the widget is prevented until the button has been pressed again or focus is lost.
 		// The value can be manipulated by using the game pad's directional arrows ( relative to slider orientation ).
-		if (FSlateApplication::Get().GetNavigationActionForKey(KeyPressed) == EUINavigationAction::Accept && bRequiresControllerLock)
+		if (FSlateApplication::Get().GetNavigationActionFromKey(InKeyEvent) == EUINavigationAction::Accept && bRequiresControllerLock)
 		{
 			if (bControllerInputCaptured == false)
 			{
diff --git a/Engine/Source/Runtime/Slate/Private/Widgets/Text/SlateEditableTextLayout.cpp b/Engine/Source/Runtime/Slate/Private/Widgets/Text/SlateEditableTextLayout.cpp
--- a/Engine/Source/Runtime/Slate/Private/Widgets/Text/SlateEditableTextLayout.cpp
+++ b/Engine/Source/Runtime/Slate/Private/Widgets/Text/SlateEditableTextLayout.cpp
@@ -1104,7 +1104,7 @@ FReply FSlateEditableTextLayout::HandleKeyDown(const FKeyEvent& InKeyEvent)
 
 FReply FSlateEditableTextLayout::HandleKeyUp(const FKeyEvent& InKeyEvent)
 {
-	if (FPlatformApplicationMisc::RequiresVirtualKeyboard() && FSlateApplication::Get().GetNavigationActionForKey(InKeyEvent.GetKey()) == EUINavigationAction::Accept)
+	if (FPlatformApplicationMisc::RequiresVirtualKeyboard() && FSlateApplication::Get().GetNavigationActionFromKey(InKeyEvent) == EUINavigationAction::Accept)
 	{
 		if (!OwnerWidget->IsTextReadOnly())
 		{
diff --git a/Engine/Source/Runtime/Slate/Public/Framework/Application/NavigationConfig.h b/Engine/Source/Runtime/Slate/Public/Framework/Application/NavigationConfig.h
--- a/Engine/Source/Runtime/Slate/Public/Framework/Application/NavigationConfig.h
+++ b/Engine/Source/Runtime/Slate/Public/Framework/Application/NavigationConfig.h
@@ -58,7 +58,10 @@ public:
 	/** Notified when navigation has caused a widget change to occur */
 	virtual void OnNavigationChangedFocus(TSharedPtr<SWidget> OldWidget, TSharedPtr<SWidget> NewWidget, FFocusEvent FocusEvent) {}
 
-	/** Returns the navigation action corresponding to this key, or Invalid if not found */
+	/** Returns the navigation action corresponding to a key event. This version will handle multiple users correctly */
+	virtual EUINavigationAction GetNavigationActionFromKey(const FKeyEvent& InKeyEvent) const;
+
+	UE_DEPRECATED(4.24, "GetNavigationActionForKey doesn't handle multiple users properly, use GetNavigationActionFromKey instead")
 	virtual EUINavigationAction GetNavigationActionForKey(const FKey& InKey) const;
 
 public:
diff --git a/Engine/Source/Runtime/Slate/Public/Framework/Application/SlateApplication.h b/Engine/Source/Runtime/Slate/Public/Framework/Application/SlateApplication.h
--- a/Engine/Source/Runtime/Slate/Public/Framework/Application/SlateApplication.h
+++ b/Engine/Source/Runtime/Slate/Public/Framework/Application/SlateApplication.h
@@ -370,10 +370,16 @@ public:
 	/** Returns true if this slate application is ready to display windows. */
 	bool CanDisplayWindows() const;
 
+	/** Returns navigation direction matching a key event, this is determined in the FNavigationConfig */
 	virtual EUINavigation GetNavigationDirectionFromKey(const FKeyEvent& InKeyEvent) const override;
+
+	/** Returns navigation direction matching an anlog event, this is determined in the FNavigationConfig */
 	virtual EUINavigation GetNavigationDirectionFromAnalog(const FAnalogInputEvent& InAnalogEvent) override;
 
-	/** Returns the navigation action corresponding to this key, or Invalid if not found */
+	/** Returns the navigation action corresponding to a key event. This version will handle multiple users correctly */
+	virtual EUINavigationAction GetNavigationActionFromKey(const FKeyEvent& InKeyEvent) const override;
+
+	UE_DEPRECATED(4.24, "GetNavigationActionForKey doesn't handle multiple users properly, use GetNavigationActionFromKey instead")
 	virtual EUINavigationAction GetNavigationActionForKey(const FKey& InKey) const override;
 
 	/**
diff --git a/Engine/Source/Runtime/Slate/Public/Widgets/Input/SComboBox.h b/Engine/Source/Runtime/Slate/Public/Widgets/Input/SComboBox.h
--- a/Engine/Source/Runtime/Slate/Public/Widgets/Input/SComboBox.h
+++ b/Engine/Source/Runtime/Slate/Public/Widgets/Input/SComboBox.h
@@ -284,37 +284,14 @@ public:
 
 protected:
 	/** Handle key presses that SListView ignores */
-	virtual FReply OnHandleKeyPressed(FKey KeyPressed)
-	{
-		EUINavigationAction NavAction = FSlateApplication::Get().GetNavigationActionForKey(KeyPressed);
-		if (NavAction == EUINavigationAction::Accept)
-		{
-			TArray<OptionType> SelectedItems = ComboListView->GetSelectedItems();
-			if (SelectedItems.Num() > 0)
-			{
-				ComboListView->SetSelection(SelectedItems[0]);
-			}
-			return FReply::Handled();
-		}
-		else if (NavAction == EUINavigationAction::Back)
-		{
-			this->SetIsOpen(false);
-			return FReply::Handled();
-		}
-		return FReply::Unhandled();
-	}
-
 	FReply OnKeyDown( const FGeometry& MyGeometry, const FKeyEvent& InKeyEvent ) override
 	{
-
-		const FKey KeyPressed = InKeyEvent.GetKey();
-
 		if (IsInteractable())
 		{
-
+			const EUINavigationAction NavAction = FSlateApplication::Get().GetNavigationActionFromKey(InKeyEvent);
+			const EUINavigation NavDirection = FSlateApplication::Get().GetNavigationDirectionFromKey(InKeyEvent);
 			if (EnableGamepadNavigationMode)
 			{
-				EUINavigationAction NavAction = FSlateApplication::Get().GetNavigationActionForKey(KeyPressed);
 				// The controller's bottom face button must be pressed once to begin manipulating the combobox's value.
 				// Navigation away from the widget is prevented until the button has been pressed again or focus is lost.
 				if (NavAction == EUINavigationAction::Accept)
@@ -346,7 +323,7 @@ protected:
 					}
 
 				}
-				else if (NavAction == EUINavigationAction::Back || KeyPressed == EKeys::BackSpace)
+				else if (NavAction == EUINavigationAction::Back || InKeyEvent.GetKey() == EKeys::BackSpace)
 				{
 					const bool bWasInputCaptured = bControllerInputCaptured;
 
@@ -366,7 +343,7 @@ protected:
 			}
 			else
 			{
-				if (KeyPressed == EKeys::Up || KeyPressed == EKeys::Gamepad_DPad_Up || KeyPressed == EKeys::Gamepad_LeftStick_Up)
+				if (NavDirection == EUINavigation::Up)
 				{
 					NullableOptionType NullableSelected = GetSelectedItem();
 					if (TListTypeTraits<OptionType>::IsPtrValid(NullableSelected))
@@ -382,7 +359,7 @@ protected:
 
 					return FReply::Handled();
 				}
-				else if (KeyPressed == EKeys::Down || KeyPressed == EKeys::Gamepad_DPad_Down || KeyPressed == EKeys::Gamepad_LeftStick_Down)
+				else if (NavDirection == EUINavigation::Down)
 				{
 					NullableOptionType NullableSelected = GetSelectedItem();
 					if (TListTypeTraits<OptionType>::IsPtrValid(NullableSelected))
diff --git a/Engine/Source/Runtime/SlateCore/Public/Application/SlateApplicationBase.h b/Engine/Source/Runtime/SlateCore/Public/Application/SlateApplicationBase.h
--- a/Engine/Source/Runtime/SlateCore/Public/Application/SlateApplicationBase.h
+++ b/Engine/Source/Runtime/SlateCore/Public/Application/SlateApplicationBase.h
@@ -245,6 +245,7 @@ public:
 
 	virtual EUINavigation GetNavigationDirectionFromKey( const FKeyEvent& InKeyEvent ) const = 0;
 	virtual EUINavigation GetNavigationDirectionFromAnalog(const FAnalogInputEvent& InAnalogEvent) = 0;
+	virtual EUINavigationAction GetNavigationActionFromKey(const FKeyEvent& InKeyEvent) const = 0;
 	virtual EUINavigationAction GetNavigationActionForKey(const FKey& InKey) const = 0;
 
 	/** @return	Returns true if there are any pop-up menus summoned */
