PR #5527: Adding mouse user index remap support to ISlateInputManager (Contributed by [REDACTED_PERSON])

diff --git a/Engine/Source/Runtime/Engine/Private/Slate/SceneViewport.cpp b/Engine/Source/Runtime/Engine/Private/Slate/SceneViewport.cpp
--- a/Engine/Source/Runtime/Engine/Private/Slate/SceneViewport.cpp
+++ b/Engine/Source/Runtime/Engine/Private/Slate/SceneViewport.cpp
@@ -284,8 +284,8 @@ void FSceneViewport::ProcessAccumulatedPointerInput()
 	if (NumMouseSamplesX > 0 || NumMouseSamplesY > 0)
 	{
 		const float DeltaTime = FApp::GetDeltaTime();
-		ViewportClient->InputAxis( this, 0, EKeys::MouseX, MouseDelta.X, DeltaTime, NumMouseSamplesX );
-		ViewportClient->InputAxis( this, 0, EKeys::MouseY, MouseDelta.Y, DeltaTime, NumMouseSamplesY );
+		ViewportClient->InputAxis( this, MouseDeltaUserIndex, EKeys::MouseX, MouseDelta.X, DeltaTime, NumMouseSamplesX );
+		ViewportClient->InputAxis( this, MouseDeltaUserIndex, EKeys::MouseY, MouseDelta.Y, DeltaTime, NumMouseSamplesY );
 	}
 
 	if ( bCursorHiddenDueToCapture )
@@ -320,6 +320,7 @@ void FSceneViewport::ProcessAccumulatedPointerInput()
 	MouseDelta = FIntPoint::ZeroValue;
 	NumMouseSamplesX = 0;
 	NumMouseSamplesY = 0;
+	MouseDeltaUserIndex = INDEX_NONE;
 }
 
 FVector2D FSceneViewport::VirtualDesktopPixelToViewport(FIntPoint VirtualDesktopPointPx) const
@@ -655,6 +656,8 @@ FReply FSceneViewport::OnMouseMove( const FGeometry& InGeometry, const FPointerE
 
 			MouseDelta.Y -= CursorDelta.Y;
 			++NumMouseSamplesY;
+
+			MouseDeltaUserIndex = InMouseEvent.GetUserIndex();
 		}
 
 		if ( bCursorHiddenDueToCapture )
diff --git a/Engine/Source/Runtime/Engine/Public/Slate/SceneViewport.h b/Engine/Source/Runtime/Engine/Public/Slate/SceneViewport.h
--- a/Engine/Source/Runtime/Engine/Public/Slate/SceneViewport.h
+++ b/Engine/Source/Runtime/Engine/Public/Slate/SceneViewport.h
@@ -407,6 +407,8 @@ private:
 	int32 NumMouseSamplesX;
 	/** The number of input samples in Y since input was was last processed */
 	int32 NumMouseSamplesY;
+	/** User index supplied by mouse events accumulated into NumMouseSamplesX and NumMouseSamplesY */
+	int32 MouseDeltaUserIndex;
 	/** The current mouse delta */
 	FIntPoint MouseDelta;
 	/** true if the cursor is currently visible */
diff --git a/Engine/Source/Runtime/Slate/Private/Framework/Application/SlateApplication.cpp b/Engine/Source/Runtime/Slate/Private/Framework/Application/SlateApplication.cpp
--- a/Engine/Source/Runtime/Slate/Private/Framework/Application/SlateApplication.cpp
+++ b/Engine/Source/Runtime/Slate/Private/Framework/Application/SlateApplication.cpp
@@ -2637,6 +2637,11 @@ TSharedPtr<SViewport> FSlateApplication::GetGameViewport() const
 	return GameViewportWidget.Pin();
 }
 
+int32 FSlateApplication::GetUserIndexForMouse() const
+{
+	return InputManager->GetUserIndexForMouse();
+}
+
 int32 FSlateApplication::GetUserIndexForKeyboard() const
 {
 	return InputManager->GetUserIndexForKeyboard();
@@ -3702,6 +3707,7 @@ void FSlateApplication::QueryCursor()
 			const FVector2D CurrentCursorPosition = GetCursorPos();
 			const FVector2D LastCursorPosition = GetLastCursorPos();
 			const FPointerEvent CursorEvent(
+				GetUserIndexForMouse(),
 				CursorPointerIndex,
 				CurrentCursorPosition,
 				LastCursorPosition,
@@ -4330,6 +4336,7 @@ void FSlateApplication::SynthesizeMouseMove()
 
 		FPointerEvent MouseEvent
 		(
+			GetUserIndexForMouse(),
 			CursorPointerIndex,
 			GetCursorPos(),
 			GetLastCursorPos(),
@@ -5391,6 +5398,7 @@ bool FSlateApplication::OnMouseDown( const TSharedPtr< FGenericWindow >& Platfor
 	FKey Key = TranslateMouseButtonToKey( Button );
 
 	FPointerEvent MouseEvent(
+		GetUserIndexForMouse(),
 		CursorPointerIndex,
 		CursorPos,
 		GetLastCursorPos(),
@@ -6132,6 +6140,7 @@ bool FSlateApplication::OnMouseDoubleClick( const TSharedPtr< FGenericWindow >&
 	FKey Key = TranslateMouseButtonToKey( Button );
 
 	FPointerEvent MouseEvent(
+		GetUserIndexForMouse(),
 		CursorPointerIndex,
 		CursorPos,
 		GetLastCursorPos(),
@@ -6214,6 +6223,7 @@ bool FSlateApplication::OnMouseUp( const EMouseButtons::Type Button, const FVect
 	FKey Key = TranslateMouseButtonToKey( Button );
 
 	FPointerEvent MouseEvent(
+		GetUserIndexForMouse(),
 		CursorPointerIndex,
 		CursorPos,
 		GetLastCursorPos(),
@@ -6270,6 +6280,7 @@ bool FSlateApplication::OnMouseWheel( const float Delta )
 bool FSlateApplication::OnMouseWheel( const float Delta, const FVector2D CursorPos )
 {
 	FPointerEvent MouseWheelEvent(
+		GetUserIndexForMouse(),
 		CursorPointerIndex,
 		CursorPos,
 		CursorPos,
@@ -6392,6 +6403,7 @@ bool FSlateApplication::OnMouseMove()
 		LastMouseMoveTime = GetCurrentTime();
 
 		FPointerEvent MouseEvent(
+			GetUserIndexForMouse(),
 			CursorPointerIndex,
 			CurrentCursorPosition,
 			LastCursorPosition,
@@ -6430,6 +6442,7 @@ bool FSlateApplication::OnRawMouseMove( const int32 X, const int32 Y )
 	if ( X != 0 || Y != 0 )
 	{
 		FPointerEvent MouseEvent(
+			GetUserIndexForMouse(),
 			CursorPointerIndex,
 			GetCursorPos(),
 			GetLastCursorPos(),
@@ -7412,6 +7425,7 @@ EDropEffect::Type FSlateApplication::OnDragEnter( const TSharedRef< SWindow >& W
 	// Make a faux mouse event for slate, so we can initiate a drag and drop.
 	FDragDropEvent DragDropEvent(
 		FPointerEvent(
+		GetUserIndexForMouse(),
 		CursorPointerIndex,
 		CurrentCursorPosition,
 		LastCursorPosition,
@@ -7457,6 +7471,7 @@ EDropEffect::Type FSlateApplication::OnDragOver( const TSharedPtr< FGenericWindo
 		if ( LastCursorPosition != CurrentCursorPosition )
 		{
 			FPointerEvent MouseEvent(
+				GetUserIndexForMouse(),
 				CursorPointerIndex,
 				CurrentCursorPosition,
 				LastCursorPosition,
@@ -7498,6 +7513,7 @@ EDropEffect::Type FSlateApplication::OnDragDrop(const TSharedPtr< FGenericWindow
 	if (IsDragDropping())
 	{
 		FPointerEvent MouseEvent(
+			GetUserIndexForMouse(),
 			CursorPointerIndex,
 			GetCursorPos(),
 			GetLastCursorPos(),
diff --git a/Engine/Source/Runtime/Slate/Public/Framework/Application/SlateApplication.h b/Engine/Source/Runtime/Slate/Public/Framework/Application/SlateApplication.h
--- a/Engine/Source/Runtime/Slate/Public/Framework/Application/SlateApplication.h
+++ b/Engine/Source/Runtime/Slate/Public/Framework/Application/SlateApplication.h
@@ -190,6 +190,7 @@ private:
 class SLATE_API ISlateInputManager
 {
 public:
+	virtual int32 GetUserIndexForMouse() const = 0;
 	virtual int32 GetUserIndexForKeyboard() const = 0;
 	virtual int32 GetUserIndexForController(int32 ControllerId) const = 0;
 };
@@ -197,6 +198,7 @@ public:
 class SLATE_API FSlateDefaultInputMapping : public ISlateInputManager
 {
 public:
+	virtual int32 GetUserIndexForMouse() const override { return 0; }
 	virtual int32 GetUserIndexForKeyboard() const override { return 0; }
 	virtual int32 GetUserIndexForController(int32 ControllerId) const override { return ControllerId; }
 };
@@ -1528,6 +1530,11 @@ public:
 	 */
 	FReply RouteMouseWheelOrGestureEvent(const FWidgetPath& WidgetsUnderPointer, const FPointerEvent& InWheelEvent, const FPointerEvent* InGestureEvent = nullptr);
 
+	/**
+	 * @return int user index that the mouse is mapped to. -1 if the mouse isn't mapped
+	 */
+	int32 GetUserIndexForMouse() const;
+
 	/**
 	 * @return int user index that the keyboard is mapped to. -1 if the keyboard isn't mapped
 	 */
diff --git a/Engine/Source/Runtime/SlateCore/Public/Input/Events.h b/Engine/Source/Runtime/SlateCore/Public/Input/Events.h
--- a/Engine/Source/Runtime/SlateCore/Public/Input/Events.h
+++ b/Engine/Source/Runtime/SlateCore/Public/Input/Events.h
@@ -689,6 +689,31 @@ public:
 		, bIsTouchFirstMove(false)
 	{ }
 
+	FPointerEvent(
+		uint32 InUserIndex,
+		uint32 InPointerIndex,
+		const FVector2D& InScreenSpacePosition,
+		const FVector2D& InLastScreenSpacePosition,
+		const FVector2D& InDelta,
+		const TSet<FKey>& InPressedButtons,
+		const FModifierKeysState& InModifierKeys
+	)
+		: FInputEvent(InModifierKeys, InUserIndex, false)
+		, ScreenSpacePosition(InScreenSpacePosition)
+		, LastScreenSpacePosition(InLastScreenSpacePosition)
+		, CursorDelta(InDelta)
+		, PressedButtons(InPressedButtons)
+		, PointerIndex(InPointerIndex)
+		, TouchpadIndex(0)
+		, Force(1.0f)
+		, bIsTouchEvent(false)
+		, GestureType(EGestureEvent::None)
+		, WheelOrGestureDelta(0.0f, 0.0f)
+		, bIsDirectionInvertedFromDevice(false)
+		, bIsTouchForceChanged(false)
+		, bIsTouchFirstMove(false)
+	{ }
+
 	/** A constructor for raw mouse events */
 	FPointerEvent(
 		uint32 InPointerIndex,
