SBreadCrumbTrail::GetCrumbMenuContent() now returns a TSharedRef<SWidget>, since there is a HasCrumbMenuContent() to check if there should be content, and it is an error to return an invalid SWidget after saying you have crumb menu content.

diff --git a/Engine/Source/Developer/SlateFileDialogs/Private/SlateFileDlgWindow.cpp b/Engine/Source/Developer/SlateFileDialogs/Private/SlateFileDlgWindow.cpp
--- a/Engine/Source/Developer/SlateFileDialogs/Private/SlateFileDlgWindow.cpp
+++ b/Engine/Source/Developer/SlateFileDialogs/Private/SlateFileDlgWindow.cpp
@@ -984,7 +984,7 @@ void SSlateFileOpenDlg::OnPathMenuItemClicked( FString ClickedPath )
 }
 
 
-TSharedPtr<SWidget> SSlateFileOpenDlg::OnGetCrumbDelimiterContent(const FString& CrumbData) const
+TSharedRef<SWidget> SSlateFileOpenDlg::OnGetCrumbDelimiterContent(const FString& CrumbData) const
 {
 	TSharedPtr<SWidget> Widget = SNullWidget::NullWidget;
 	TArray<FString> SubDirs;
@@ -1063,7 +1063,7 @@ TSharedPtr<SWidget> SSlateFileOpenDlg::OnGetCrumbDelimiterContent(const FString&
 			];
 	}
 
-	return Widget;
+	return Widget.ToSharedRef();
 }
 
 
diff --git a/Engine/Source/Developer/SlateFileDialogs/Private/SlateFileDlgWindow.h b/Engine/Source/Developer/SlateFileDialogs/Private/SlateFileDlgWindow.h
--- a/Engine/Source/Developer/SlateFileDialogs/Private/SlateFileDlgWindow.h
+++ b/Engine/Source/Developer/SlateFileDialogs/Private/SlateFileDlgWindow.h
@@ -151,7 +151,7 @@ public:
 
 private:	
 	void OnPathClicked( const FString& CrumbData );
-	TSharedPtr<SWidget> OnGetCrumbDelimiterContent(const FString& CrumbData) const;
+	TSharedRef<SWidget> OnGetCrumbDelimiterContent(const FString& CrumbData) const;
 	void RebuildFileTable();
 	void BuildDirectoryPath();
 	void ReadDir(bool bIsRefresh = false);
diff --git a/Engine/Source/Editor/ContentBrowser/Private/SContentBrowser.cpp b/Engine/Source/Editor/ContentBrowser/Private/SContentBrowser.cpp
--- a/Engine/Source/Editor/ContentBrowser/Private/SContentBrowser.cpp
+++ b/Engine/Source/Editor/ContentBrowser/Private/SContentBrowser.cpp
@@ -1823,7 +1823,7 @@ bool SContentBrowser::OnHasCrumbDelimiterContent(const FString& CrumbData) const
 	return false;
 }
 
-TSharedPtr<SWidget> SContentBrowser::OnGetCrumbDelimiterContent(const FString& CrumbData) const
+TSharedRef<SWidget> SContentBrowser::OnGetCrumbDelimiterContent(const FString& CrumbData) const
 {
 	FSourcesData SourcesData = AssetViewPtr->GetSourcesData();
 
@@ -1929,7 +1929,7 @@ TSharedPtr<SWidget> SContentBrowser::OnGetCrumbDelimiterContent(const FString& C
 			];
 	}
 
-	return Widget;
+	return Widget.ToSharedRef();
 }
 
 TSharedRef<SWidget> SContentBrowser::GetPathPickerContent()
diff --git a/Engine/Source/Editor/ContentBrowser/Private/SContentBrowser.h b/Engine/Source/Editor/ContentBrowser/Private/SContentBrowser.h
--- a/Engine/Source/Editor/ContentBrowser/Private/SContentBrowser.h
+++ b/Engine/Source/Editor/ContentBrowser/Private/SContentBrowser.h
@@ -232,7 +232,7 @@ private:
 	 * @param CrumbData	The current directory path
 	 * @return The directory menu
 	 */
-	TSharedPtr<SWidget> OnGetCrumbDelimiterContent(const FString& CrumbData) const;
+	TSharedRef<SWidget> OnGetCrumbDelimiterContent(const FString& CrumbData) const;
 
 	/** Gets the content for the path picker combo button */
 	TSharedRef<SWidget> GetPathPickerContent();
diff --git a/Engine/Source/Editor/PropertyEditor/Private/UserInterface/PropertyTable/SPropertyTable.h b/Engine/Source/Editor/PropertyEditor/Private/UserInterface/PropertyTable/SPropertyTable.h
--- a/Engine/Source/Editor/PropertyEditor/Private/UserInterface/PropertyTable/SPropertyTable.h
+++ b/Engine/Source/Editor/PropertyEditor/Private/UserInterface/PropertyTable/SPropertyTable.h
@@ -697,7 +697,7 @@ private:
 		}
 	}
 
-	TSharedPtr< SWidget > GetCrumbMenuContent( const int32& Item )
+	TSharedRef< SWidget > GetCrumbMenuContent( const int32& Item )
 	{
 		TSharedRef< FPropertyPath > RootPath = Table->GetRootPath();
 		if ( Item == INDEX_NONE )
diff --git a/Engine/Source/Runtime/Slate/Public/Widgets/Navigation/SBreadcrumbTrail.h b/Engine/Source/Runtime/Slate/Public/Widgets/Navigation/SBreadcrumbTrail.h
--- a/Engine/Source/Runtime/Slate/Public/Widgets/Navigation/SBreadcrumbTrail.h
+++ b/Engine/Source/Runtime/Slate/Public/Widgets/Navigation/SBreadcrumbTrail.h
@@ -61,7 +61,7 @@ public:
 	DECLARE_DELEGATE_RetVal_OneParam( bool, FHasCrumbMenuContent, const ItemType& /*CrumbData*/ );
 
 	/** Callback for getting the menu content to be displayed when clicking on a crumb's delimiter arrow */
-	DECLARE_DELEGATE_RetVal_OneParam( TSharedPtr< SWidget >, FGetCrumbMenuContent, const ItemType& /*CrumbData*/ );
+	DECLARE_DELEGATE_RetVal_OneParam( TSharedRef< SWidget >, FGetCrumbMenuContent, const ItemType& /*CrumbData*/ );
 
 	DECLARE_DELEGATE_RetVal_TwoParams( FSlateColor, FOnGetCrumbColor, int32 /*CrumbId*/, bool /*bInvert*/ );
 
@@ -341,11 +341,19 @@ private:
 		{
 			if (CrumbList[CrumbListIdx].CrumbID == CrumbId)
 			{
+				if (HasCrumbMenuContentCallback.IsBound() )
+				{
+					if (!HasCrumbMenuContentCallback.Execute( CrumbList[CrumbListIdx].CrumbData ))
+					{
+						return SNullWidget::NullWidget;
+					}
+				}
+
 				MenuContent = GetCrumbMenuContentCallback.Execute( CrumbList[CrumbListIdx].CrumbData );
 				break;
 			}
 		}
-	
+
 		return MenuContent.ToSharedRef();
 	}
 
