General cleanup of some DragDrop functions and fix to JIRA UE-78679: Undocked Asset Editor Window was closed when placed over Content Browser (it is not dropable in there). Now, it is placed as an undocked element if that occurs.
General cleanup details:
1) FDockingDragOperation.cpp: Removed some unused code.
2) TabManager.cpp: Fixed an element type wrongly returned by reference to avoid potential race hazard issues.

diff --git a/Engine/Source/Editor/ContentBrowser/Private/SAssetView.cpp b/Engine/Source/Editor/ContentBrowser/Private/SAssetView.cpp
--- a/Engine/Source/Editor/ContentBrowser/Private/SAssetView.cpp
+++ b/Engine/Source/Editor/ContentBrowser/Private/SAssetView.cpp
@@ -1602,6 +1602,7 @@ FReply SAssetView::OnDrop( const FGeometry& MyGeometry, const FDragDropEvent& Dr
 		// Note: We don't test IsAssetPathSelected here as we need to prevent dropping assets on class paths
 		const FString DestPath = SourcesData.PackagePaths[0].ToString();
 
+		// If the DragDrop event is validated, continue trying to dock it to the Widget
 		bool bUnused = false;
 		if (DragDropHandler::ValidateDragDropOnAssetFolder(MyGeometry, DragDropEvent, DestPath, bUnused))
 		{
@@ -1621,8 +1622,15 @@ FReply SAssetView::OnDrop( const FGeometry& MyGeometry, const FDragDropEvent& Dr
 			{
 				OnAssetsOrPathsDragDropped(AssetDragDropOp->GetAssets(), AssetDragDropOp->GetAssetPaths(), DestPath);
 			}
+			return FReply::Handled();
+		}
+		// If the DragDropEvent is not successful, it has not been handled
+		// If it returned Handled rather than Unhandled, when a widget were dragged in there (which is not dropable nor dockable in there),
+		// that widget would disappear rather than being placed as an undocked widget
+		else
+		{
+			return FReply::Unhandled();
 		}
-		return FReply::Handled();
 	}
 	else if (HasSingleCollectionSource())
 	{
diff --git a/Engine/Source/Runtime/Slate/Private/Framework/Docking/FDockingDragOperation.cpp b/Engine/Source/Runtime/Slate/Private/Framework/Docking/FDockingDragOperation.cpp
--- a/Engine/Source/Runtime/Slate/Private/Framework/Docking/FDockingDragOperation.cpp
+++ b/Engine/Source/Runtime/Slate/Private/Framework/Docking/FDockingDragOperation.cpp
@@ -26,14 +26,7 @@ void FDockingDragOperation::OnDrop( bool bDropWasHandled, const FPointerEvent& M
 		// The event was handled, so we HAVE to have some window that we dropped onto.
 		TSharedRef<SWindow> WindowDroppedInto = MouseEvent.GetWindow();
 
-		// Let every widget under this tab manager know that this tab has found a new home.
-		TSharedPtr<SWindow> NewWindow = ( TabOwnerAreaOfOrigin->GetParentWindow() == WindowDroppedInto )
-			// Tab dropped into same window as before, meaning there is no NewWindow.
-			? TSharedPtr<SWindow>()
-			// Tab was dropped into a different window, so the tab manager needs to know in order to re-parent child windows.
-			: WindowDroppedInto;
-
-		TabOwnerAreaOfOrigin->GetTabManager()->GetPrivateApi().OnTabRelocated( TabBeingDragged.ToSharedRef(), WindowDroppedInto );
+		TabOwnerAreaOfOrigin->GetTabManager()->GetPrivateApi().OnTabRelocated(TabBeingDragged.ToSharedRef(), WindowDroppedInto);
 	}
 
 	// Destroy the CursorDecoratorWindow by calling the base class implementation because we are relocating the content into a more permanent home.
diff --git a/Engine/Source/Runtime/Slate/Private/Framework/Docking/TabManager.cpp b/Engine/Source/Runtime/Slate/Private/Framework/Docking/TabManager.cpp
--- a/Engine/Source/Runtime/Slate/Private/Framework/Docking/TabManager.cpp
+++ b/Engine/Source/Runtime/Slate/Private/Framework/Docking/TabManager.cpp
@@ -2141,7 +2141,7 @@ void FGlobalTabmanager::OnTabRelocated( const TSharedRef<SDockTab>& RelocatedTab
 			for (int32 DockAreaIndex=0; DockAreaIndex < LiveDockAreas.Num(); ++DockAreaIndex)
 			{
 				const TSharedRef<SDockingArea>& ChildDockArea = LiveDockAreas[ DockAreaIndex ].Pin().ToSharedRef();
-				const TSharedPtr<SWindow>& OldChildWindow = ChildDockArea->GetParentWindow();
+				const TSharedPtr<SWindow> OldChildWindow = ChildDockArea->GetParentWindow();
 				if ( OldChildWindow.IsValid() )
 				{
 					TSharedRef<SWindow> NewChildWindow = SNew(SWindow)
