Bug fix: Monitors to the left of monitor 0 (with negative pixel indexes) were not showing the tooltips in the right place.
To avoid future related bugs, deprecated "FSlateRect GetToolTipForceFieldRect(...)" and replaced by "bool GetToolTipForceFieldRect(..., FSlateRect&)",  which indicates whether the function found a solution, so anybody using this function can now more easily avoid the bug.

diff --git a/Engine/Source/Runtime/Slate/Private/Framework/Application/MenuStack.cpp b/Engine/Source/Runtime/Slate/Private/Framework/Application/MenuStack.cpp
--- a/Engine/Source/Runtime/Slate/Private/Framework/Application/MenuStack.cpp
+++ b/Engine/Source/Runtime/Slate/Private/Framework/Application/MenuStack.cpp
@@ -743,9 +743,10 @@ TSharedPtr<IMenu> FMenuStack::FindMenuFromWindow(TSharedRef<SWindow> InWindow) c
 	return TSharedPtr<IMenu>();
 }
 
-FSlateRect FMenuStack::GetToolTipForceFieldRect(TSharedRef<IMenu> InMenu, const FWidgetPath& PathContainingMenu) const
+bool FMenuStack::GetToolTipForceFieldRect(const TSharedRef<IMenu>& InMenu, const FWidgetPath& InPathContainingMenu, FSlateRect& OutSlateRect) const
 {
-	FSlateRect ForceFieldRect(0, 0, 0, 0);
+	bool bWasSolutionFound = false;
+	OutSlateRect = FSlateRect(0, 0, 0, 0);
 
 	int32 StackLevel = Stack.IndexOfByKey(InMenu);
 
@@ -756,7 +757,7 @@ FSlateRect FMenuStack::GetToolTipForceFieldRect(TSharedRef<IMenu> InMenu, const
 			TSharedPtr<SWidget> MenuContent = Stack[StackLevelIndex]->GetContent();
 			if (MenuContent.IsValid())
 			{
-				FWidgetPath WidgetPath = PathContainingMenu.GetPathDownTo(MenuContent.ToSharedRef());
+				FWidgetPath WidgetPath = InPathContainingMenu.GetPathDownTo(MenuContent.ToSharedRef());
 				if (!WidgetPath.IsValid())
 				{
 					FSlateApplication::Get().GeneratePathToWidgetChecked(MenuContent.ToSharedRef(), WidgetPath);
@@ -764,12 +765,30 @@ FSlateRect FMenuStack::GetToolTipForceFieldRect(TSharedRef<IMenu> InMenu, const
 				if (WidgetPath.IsValid())
 				{
 					const FGeometry& ContentGeometry = WidgetPath.Widgets.Last().Geometry;
-					ForceFieldRect = ForceFieldRect.Expand(ContentGeometry.GetLayoutBoundingRect());
+					// No first time: Expand
+					if (bWasSolutionFound)
+					{
+						OutSlateRect = OutSlateRect.Expand(ContentGeometry.GetLayoutBoundingRect());
+					}
+					// First time: assign it
+					// Otherwise, it would assume that the point [0,0,0,0] is part of the final rectangle, which is not the case in multiple scenarios.
+					// E.g., if the monitor where Slate is running is not the main one or if the Slate window is restored to the right half of the monitor.
+					else
+					{
+						OutSlateRect = ContentGeometry.GetLayoutBoundingRect();
+						bWasSolutionFound = true;
+					}
 				}
 			}
 		}
 	}
+	return bWasSolutionFound;
+}
 
+FSlateRect FMenuStack::GetToolTipForceFieldRect(TSharedRef<IMenu> InMenu, const FWidgetPath& PathContainingMenu) const
+{
+	FSlateRect ForceFieldRect(0, 0, 0, 0);
+	GetToolTipForceFieldRect(InMenu, PathContainingMenu, ForceFieldRect);
 	return ForceFieldRect;
 }
 
diff --git a/Engine/Source/Runtime/Slate/Private/Framework/Application/SlateUser.cpp b/Engine/Source/Runtime/Slate/Private/Framework/Application/SlateUser.cpp
--- a/Engine/Source/Runtime/Slate/Private/Framework/Application/SlateUser.cpp
+++ b/Engine/Source/Runtime/Slate/Private/Framework/Application/SlateUser.cpp
@@ -1153,7 +1153,12 @@ void FSlateUser::UpdateTooltip(const FMenuStack& MenuStack, bool bCanSpawnNewToo
 	{
 		if (TSharedPtr<IMenu> MenuInPath = MenuStack.FindMenuInWidgetPath(WidgetsToQueryForTooltip))
 		{
-			ForceFieldRect = ForceFieldRect->Expand(MenuStack.GetToolTipForceFieldRect(MenuInPath.ToSharedRef(), WidgetsToQueryForTooltip));
+			FSlateRect MenuStackRectangle;
+			const bool bWasSolutionFound = MenuStack.GetToolTipForceFieldRect(MenuInPath.ToSharedRef(), WidgetsToQueryForTooltip, MenuStackRectangle);
+			if (bWasSolutionFound)
+			{
+				ForceFieldRect = ForceFieldRect->Expand(MenuStackRectangle);
+			}
 		}
 	}
 
@@ -1178,11 +1183,11 @@ void FSlateUser::UpdateTooltip(const FMenuStack& MenuStack, bool bCanSpawnNewToo
 		TooltipShift.Y = (ForceFieldRect->Bottom + SlateDefs::TooltipOffsetFromForceField.Y) - DesiredLocation.Y;
 
 		// Make sure the tooltip needs to be offset
-		if (TooltipShift.X > 0.0f && TooltipShift.Y > 0.0f)
+		if (TooltipShift.X != 0.0f && TooltipShift.Y != 0.0f)
 		{
 			// Find the best edge to move the tooltip towards
 			if (ActiveTooltipInfo.OffsetDirection == ETooltipOffsetDirection::Right ||
-				(ActiveTooltipInfo.OffsetDirection == ETooltipOffsetDirection::Undetermined && TooltipShift.X < TooltipShift.Y))
+				(ActiveTooltipInfo.OffsetDirection == ETooltipOffsetDirection::Undetermined && FMath::Abs(TooltipShift.X) < FMath::Abs(TooltipShift.Y)))
 			{
 				// Move right
 				DesiredLocation.X += TooltipShift.X;
diff --git a/Engine/Source/Runtime/Slate/Public/Framework/Application/MenuStack.h b/Engine/Source/Runtime/Slate/Public/Framework/Application/MenuStack.h
--- a/Engine/Source/Runtime/Slate/Public/Framework/Application/MenuStack.h
+++ b/Engine/Source/Runtime/Slate/Public/Framework/Application/MenuStack.h
@@ -174,6 +174,16 @@ public:
 	 */
 	TSharedPtr<IMenu> FindMenuInWidgetPath(const FWidgetPath& PathToQuery) const;
 
+	/**
+	 * Called by the application when showing tooltips. It prevents tooltips from drawing over menu items. If the resulting bool is false, the resulting OutSlateRect will be [0,0,0,0] and should not be used.
+	 *
+	 * @param InMenu				A menu in the stack, used to generate the force field rect
+	 * @param InPathContainingMenu	A widget path containing InMenu. This could be generated but the application has the path so it helps performance to pass it in here.
+	 * @param OutSlateRect			The output rectangle enclosing the menu's children (created from the clipping rects in the geometry in PathContainingMenu)
+	 * @return						A bool indicating whether a solution was found. If false, the resulting OutSlateRect will be [0,0,0,0] and should not be used.
+	 */
+	bool GetToolTipForceFieldRect(const TSharedRef<IMenu>& InMenu, const FWidgetPath& InPathContainingMenu, FSlateRect& OutSlateRect) const;
+
 	/**
 	 * Called by the application when showing tooltips. It prevents tooltips from drawing over menu items.
 	 *
@@ -181,6 +191,7 @@ public:
 	 * @param PathContainingMenu	A widget path containing InMenu. This could be generated but the application has the path so it helps performance to pass it in here. 
 	 * @return						A rectangle enclosing the menu's children (created from the clipping rects in the geometry in PathContainingMenu)
 	 */
+	UE_DEPRECATED(4.24, "FSlateRect GetToolTipForceFieldRect(TSharedRef<IMenu>, const FWidgetPath&) has been deprecated to avoid bugs when no solution is found. Instead, use bool GetToolTipForceFieldRect(const TSharedRef<IMenu>&, const FWidgetPath&, FSlateRect&).")
 	FSlateRect GetToolTipForceFieldRect(TSharedRef<IMenu> InMenu, const FWidgetPath& PathContainingMenu) const;
 
 	/**
