Deprecate SetText taking FString on STextBlock

diff --git a/Engine/Plugins/NetcodeUnitTest/NetcodeUnitTest/Source/NetcodeUnitTest/Private/UI/SLogWidget.cpp b/Engine/Plugins/NetcodeUnitTest/NetcodeUnitTest/Source/NetcodeUnitTest/Private/UI/SLogWidget.cpp
--- a/Engine/Plugins/NetcodeUnitTest/NetcodeUnitTest/Source/NetcodeUnitTest/Private/UI/SLogWidget.cpp
+++ b/Engine/Plugins/NetcodeUnitTest/NetcodeUnitTest/Source/NetcodeUnitTest/Private/UI/SLogWidget.cpp
@@ -1156,11 +1156,11 @@ void SLogWidget::OnSuspendStateChanged(ESuspendState InSuspendState)
 	{
 		if (InSuspendState == ESuspendState::Active)
 		{
-			SuspendButtonText->SetText(FString(TEXT("SUSPEND")));
+			SuspendButtonText->SetText(NSLOCTEXT("SLogWidget", "SUSPEND", "SUSPEND"));
 		}
 		else if (InSuspendState == ESuspendState::Suspended)
 		{
-			SuspendButtonText->SetText(FString(TEXT("RESUME")));
+			SuspendButtonText->SetText(NSLOCTEXT("SLogWidget", "RESUME", "RESUME"));
 		}
 	}
 }
diff --git a/Engine/Source/Editor/TranslationEditor/Private/TranslationEditor.cpp b/Engine/Source/Editor/TranslationEditor/Private/TranslationEditor.cpp
--- a/Engine/Source/Editor/TranslationEditor/Private/TranslationEditor.cpp
+++ b/Engine/Source/Editor/TranslationEditor/Private/TranslationEditor.cpp
@@ -990,7 +990,7 @@ void FTranslationEditor::UpdateTranslationUnitSelection(TSet<TSharedRef<IPropert
 				UTranslationUnit* SelectedTranslationUnit = (UTranslationUnit*)UObjectPtr;
 				if (SelectedTranslationUnit != NULL)
 				{
-					PreviewTextBlock->SetText(SelectedTranslationUnit->Translation);
+					PreviewTextBlock->SetText(FText::AsCultureInvariant(SelectedTranslationUnit->Translation));
 					NamespaceTextBlock->SetText(FText::Format(LOCTEXT("TranslationNamespace", "Namespace: {0}"), FText::FromString(SelectedTranslationUnit->Namespace)));
 
 					// Add the ContextPropertyTable-specific path
diff --git a/Engine/Source/Runtime/Slate/Private/Widgets/Text/STextBlock.cpp b/Engine/Source/Runtime/Slate/Private/Widgets/Text/STextBlock.cpp
--- a/Engine/Source/Runtime/Slate/Private/Widgets/Text/STextBlock.cpp
+++ b/Engine/Source/Runtime/Slate/Private/Widgets/Text/STextBlock.cpp
@@ -117,28 +117,21 @@ void STextBlock::SetText( const TAttribute< FString >& InText )
 {
 	if (InText.IsSet() && !InText.IsBound())
 	{
-		SetText(InText.Get());
+		SetText(FText::AsCultureInvariant(InText.Get()));
 		return;
 	}
 
 	SCOPE_CYCLE_COUNTER(Stat_SlateTextBlockSetText);
-	struct Local
+	BoundText = MakeAttributeLambda([InText]()
 	{
-		static FText PassThroughAttribute( TAttribute< FString > InString )
-		{
-			return FText::FromString( InString.Get( TEXT("") ) );
-		}
-	};
-
-	BoundText = TAttribute< FText >::Create(TAttribute<FText>::FGetter::CreateStatic( &Local::PassThroughAttribute, InText) );
-
+		return FText::AsCultureInvariant(InText.Get(FString()));
+	});
 	InvalidateText(EInvalidateWidget::LayoutAndVolatility);
-
 }
 
 void STextBlock::SetText( const FString& InText )
 {
-	SetText(FText::FromString(InText));
+	SetText(FText::AsCultureInvariant(InText));
 }
 
 void STextBlock::SetText( const TAttribute< FText >& InText )
diff --git a/Engine/Source/Runtime/Slate/Public/Widgets/Text/STextBlock.h b/Engine/Source/Runtime/Slate/Public/Widgets/Text/STextBlock.h
--- a/Engine/Source/Runtime/Slate/Public/Widgets/Text/STextBlock.h
+++ b/Engine/Source/Runtime/Slate/Public/Widgets/Text/STextBlock.h
@@ -170,13 +170,16 @@ private:
 	}
 	
 public:
+	UE_DEPRECATED(4.24, "SetText taking FString is deprecated. Use the FText version instead")
+	void SetText( const TAttribute< FString >& InText );
+	UE_DEPRECATED(4.24, "SetText taking FString is deprecated. Use the FText version instead")
+	void SetText( const FString& InText );
+
 	/**
 	 * Sets the text for this text block
 	 *
 	 * @param	InText	The new text to display
 	 */
-	void SetText( const TAttribute< FString >& InText );
-	void SetText( const FString& InText );
 	void SetText( const TAttribute< FText >& InText );
 	void SetText( const FText& InText );
 
