Fix for crash when opening a static mesh or blueprint asset in VR Mode

diff --git a/Engine/Source/Runtime/Slate/Private/Framework/Docking/TabManager.cpp b/Engine/Source/Runtime/Slate/Private/Framework/Docking/TabManager.cpp
--- a/Engine/Source/Runtime/Slate/Private/Framework/Docking/TabManager.cpp
+++ b/Engine/Source/Runtime/Slate/Private/Framework/Docking/TabManager.cpp
@@ -2258,15 +2258,18 @@ void FProxyTabmanager::OpenUnmanagedTab(FName PlaceholderId, const FSearchPrefer
 		TSharedRef<SDockingArea> DockingArea = RestoreArea(NewAreaForTab, ParentWindowPtr, false);
 		ParentWindowPtr->SetContent(DockingArea);
 
-		const TSharedPtr<SDockTab> NewlyOpenedTab = DockingArea->GetAllChildTabs()[0];
-		check(NewlyOpenedTab.IsValid());
-		
-		NewlyOpenedTab->GetParent()->GetParentDockTabStack()->OpenTab(UnmanagedTab);
-		NewlyOpenedTab->RequestCloseTab();
+		if (DockingArea->GetAllChildTabs().Num() > 0)
+		{
+			const TSharedPtr<SDockTab> NewlyOpenedTab = DockingArea->GetAllChildTabs()[0];
+			check(NewlyOpenedTab.IsValid());
+
+			NewlyOpenedTab->GetParent()->GetParentDockTabStack()->OpenTab(UnmanagedTab);
+			NewlyOpenedTab->RequestCloseTab();
 
-		MainNonCloseableTab = UnmanagedTab;
+			MainNonCloseableTab = UnmanagedTab;
 
-		OnTabOpened.Broadcast(UnmanagedTab);
+			OnTabOpened.Broadcast(UnmanagedTab);
+		}
 	}
 }
 
