Menu editor, hide legacy items while editing to simplify logic
#rnx

diff --git a/Engine/Source/Developer/ToolMenus/Private/ToolMenus.cpp b/Engine/Source/Developer/ToolMenus/Private/ToolMenus.cpp
--- a/Engine/Source/Developer/ToolMenus/Private/ToolMenus.cpp
+++ b/Engine/Source/Developer/ToolMenus/Private/ToolMenus.cpp
@@ -753,7 +753,10 @@ void UToolMenus::PopulateMenuBuilder(FMenuBuilder& MenuBuilder, UToolMenu* MenuD
 		FToolMenuSection& Section = MenuData->Sections[i];
 		if (Section.Construct.NewToolMenuDelegateLegacy.IsBound())
 		{
-			Section.Construct.NewToolMenuDelegateLegacy.Execute(MenuBuilder, MenuData);
+			if (!bIsEditing)
+			{
+				Section.Construct.NewToolMenuDelegateLegacy.Execute(MenuBuilder, MenuData);
+			}
 			continue;
 		}
 
@@ -776,7 +779,10 @@ void UToolMenus::PopulateMenuBuilder(FMenuBuilder& MenuBuilder, UToolMenu* MenuD
 		{
 			if (Block.ConstructLegacy.IsBound())
 			{
-				Block.ConstructLegacy.Execute(MenuBuilder, MenuData);
+				if (!bIsEditing)
+				{
+					Block.ConstructLegacy.Execute(MenuBuilder, MenuData);
+				}
 				continue;
 			}
 
diff --git a/Engine/Source/Editor/ToolMenusEditor/Private/SEditToolMenuDialog.cpp b/Engine/Source/Editor/ToolMenusEditor/Private/SEditToolMenuDialog.cpp
--- a/Engine/Source/Editor/ToolMenusEditor/Private/SEditToolMenuDialog.cpp
+++ b/Engine/Source/Editor/ToolMenusEditor/Private/SEditToolMenuDialog.cpp
@@ -16,7 +16,6 @@
 #include "Widgets/Input/SButton.h"
 #include "Widgets/Input/SCheckBox.h"
 #include "Widgets/Layout/SScrollBox.h"
-#include "Dialogs/Dialogs.h"
 
 #include "Framework/Docking/TabManager.h"
 #include "EditorStyleSet.h"
@@ -643,7 +642,7 @@ void SEditToolMenuDialog::BuildWidget()
 				[
 					SNew(SButton)
 					.Text(LOCTEXT("ResetAllDefaultsButtonText", "Reset all to Defaults"))
-					.ToolTipText(LOCTEXT("ResetAllDefaultsButtonTooltip", "Reset all settings for all menus menu to default values"))
+					.ToolTipText(LOCTEXT("ResetAllDefaultsButtonTooltip", "Reset all settings for all menus to default values"))
 					.OnClicked(this, &SEditToolMenuDialog::HandleResetAllClicked)
 				]
 
@@ -892,9 +891,8 @@ FReply SEditToolMenuDialog::HandleResetClicked()
 
 FReply SEditToolMenuDialog::HandleResetAllClicked()
 {
-	if (OpenMsgDlgInt(EAppMsgType::YesNo,
-		LOCTEXT("ResetAllQuestion", "Remove all menu customizations for all menus?"),
-		LOCTEXT("ResetAllQuestion_Question", "Question")) == EAppReturnType::Yes)
+	FText Title = LOCTEXT("ResetAllQuestion_Question", "Question");
+	if (FMessageDialog::Open(EAppMsgType::YesNo, EAppReturnType::No, LOCTEXT("ResetAllQuestion", "Remove all menu customizations for all menus?"), &Title) == EAppReturnType::Yes)
 	{
 		UToolMenus::Get()->RemoveAllCustomizations();
 		OriginalSettings.Reset();
diff --git a/Engine/Source/Runtime/Slate/Private/Framework/MultiBox/MultiBox.cpp b/Engine/Source/Runtime/Slate/Private/Framework/MultiBox/MultiBox.cpp
--- a/Engine/Source/Runtime/Slate/Private/Framework/MultiBox/MultiBox.cpp
+++ b/Engine/Source/Runtime/Slate/Private/Framework/MultiBox/MultiBox.cpp
@@ -943,10 +943,8 @@ void SMultiBoxWidget::UpdateDropAreaPreviewBlock( TSharedRef<const FMultiBlock>
 		// Check that the command does not already exist and that we can create it or that we are dragging an existing block in this box
 		if( !ExistingBlock.IsValid() || ( ExistingBlock.IsValid() && OriginMultiBox == MultiBox->GetCustomizationName() ) )
 		{
-
 			TSharedPtr<const FMultiBlock> NewBlock = ExistingBlock;
 
-
 			if( NewBlock.IsValid() )
 			{
 				DragPreview.Reset();
@@ -961,8 +959,6 @@ void SMultiBoxWidget::UpdateDropAreaPreviewBlock( TSharedRef<const FMultiBlock>
 
 				bAddedNewBlock = true;
 			}
-
-
 		}
 		else
 		{
@@ -1004,7 +1000,6 @@ void SMultiBoxWidget::UpdateDropAreaPreviewBlock( TSharedRef<const FMultiBlock>
 			}
 		}
 
-
 		int32 CurrentIndex = DragPreview.InsertIndex;
 		DragPreview.InsertIndex = INDEX_NONE;
 		// Find the index of the multiblock being dragged over. This is where we will insert the new block
@@ -1029,9 +1024,9 @@ void SMultiBoxWidget::UpdateDropAreaPreviewBlock( TSharedRef<const FMultiBlock>
 			}
 			else if (HoverIndex != INDEX_NONE)
 			{
-				if (Blocks[HoverIndex]->IsPartOfHeading())
+				if (MultiBlock->IsPartOfHeading())
 				{
-					if (Blocks[HoverIndex]->IsSeparator())
+					if (MultiBlock->IsSeparator())
 					{
 						// Move insert index above separator of heading
 						DragPreview.InsertIndex = HoverIndex;
@@ -1056,7 +1051,6 @@ void SMultiBoxWidget::UpdateDropAreaPreviewBlock( TSharedRef<const FMultiBlock>
 			}
 		}
 	}
-
 }
 
 EVisibility SMultiBoxWidget::GetCustomizationVisibility( TWeakPtr<const FMultiBlock> BlockWeakPtr, TWeakPtr<SWidget> BlockWidgetWeakPtr ) const
diff --git a/Engine/Source/Runtime/Slate/Private/Framework/MultiBox/MultiBoxBuilder.cpp b/Engine/Source/Runtime/Slate/Private/Framework/MultiBox/MultiBoxBuilder.cpp
--- a/Engine/Source/Runtime/Slate/Private/Framework/MultiBox/MultiBoxBuilder.cpp
+++ b/Engine/Source/Runtime/Slate/Private/Framework/MultiBox/MultiBoxBuilder.cpp
@@ -308,7 +308,10 @@ void FMenuBuilder::ApplyHook(FName InExtensionHook, EExtensionHook::Position Hoo
 	auto& Extender = ExtenderStack.Top();
 	if (InExtensionHook != NAME_None && Extender.IsValid())
 	{
-		Extender->Apply(InExtensionHook, HookPosition, *this);
+		if (!MultiBox->IsInEditMode())
+		{
+			Extender->Apply(InExtensionHook, HookPosition, *this);
+		}
 	}
 }
 
