Bug fixes:
1) UE-80271: Solved the issue that plugin-dependent SWidget windows appeared  blank when the Editor was loaded with the respective plugin disabled.
2) Attemp to solve UE-80344: Avoided the out-of-bounds TArray error in FTabManager.
3) Standalone SWidgets from non-loaded plugins now behave in the same way than the docked SWidgets from non-loaded plugings: They are also re-opened when the plugin is loaded again and the Editor re-started. Before, they just disappeared.

diff --git a/Engine/Source/Editor/LocalizationDashboard/Private/LocalizationDashboard.cpp b/Engine/Source/Editor/LocalizationDashboard/Private/LocalizationDashboard.cpp
--- a/Engine/Source/Editor/LocalizationDashboard/Private/LocalizationDashboard.cpp
+++ b/Engine/Source/Editor/LocalizationDashboard/Private/LocalizationDashboard.cpp
@@ -206,8 +206,10 @@ TWeakPtr<SDockTab> SLocalizationDashboard::ShowTargetEditor(ULocalizationTarget*
 				OurTargetEditor
 			];
 
-		TabManager->InsertNewDocumentTab(DocumentsTabName, FTabManager::ESearchPreference::RequireClosedTab, NewTargetEditorTab);
-		TargetEditorDockTab = NewTargetEditorTab;
+		if (TabManager->InsertNewDocumentTab(DocumentsTabName, FTabManager::ESearchPreference::RequireClosedTab, NewTargetEditorTab))
+		{
+			TargetEditorDockTab = NewTargetEditorTab;
+		}
 	}
 	else
 	{
diff --git a/Engine/Source/Runtime/Slate/Private/Framework/Docking/TabManager.cpp b/Engine/Source/Runtime/Slate/Private/Framework/Docking/TabManager.cpp
--- a/Engine/Source/Runtime/Slate/Private/Framework/Docking/TabManager.cpp
+++ b/Engine/Source/Runtime/Slate/Private/Framework/Docking/TabManager.cpp
@@ -77,6 +77,10 @@ TSharedPtr<SDockTab> FTabManager::FLastMajorOrNomadTab::Search(const FTabManager
 	return FoundTab;
 }
 
+static const FString UE4_TABMANAGER_OPENED_TAB_STRING = TEXT("OpenedTab");
+static const FString UE4_TABMANAGER_CLOSED_TAB_STRING = TEXT("ClosedTab");
+static const FString UE4_TABMANAGER_INVALID_TAB_STRING = TEXT("InvalidTab");
+
 TSharedRef<FTabManager::FLayoutNode> FTabManager::FLayout::NewFromString_Helper( TSharedPtr<FJsonObject> JsonObject )
 {
 	struct local
@@ -129,17 +133,18 @@ TSharedRef<FTabManager::FLayoutNode> FTabManager::FLayout::NewFromString_Helper(
 
 		static ETabState::Type TabStateFromString( const FString& AsString )
 		{
-			static const FString OpenedTab_Str = TEXT("OpenedTab");
-			static const FString ClosedTab_Str = TEXT("ClosedTab");
-
-			if (AsString == OpenedTab_Str)
+			if (AsString == UE4_TABMANAGER_OPENED_TAB_STRING)
 			{
 				return ETabState::OpenedTab;
 			}
-			else if(AsString == ClosedTab_Str)
+			else if(AsString == UE4_TABMANAGER_CLOSED_TAB_STRING)
 			{
 				return ETabState::ClosedTab;
 			}
+			else if (AsString == UE4_TABMANAGER_INVALID_TAB_STRING)
+			{
+				return ETabState::InvalidTab;
+			}
 			else
 			{
 				ensureMsgf(false, TEXT("Invalid tab state."));
@@ -376,8 +381,8 @@ TSharedRef<FJsonObject> FTabManager::FLayout::PersistToString_Helper(const TShar
 			TSharedRef<FJsonObject> TabAsJson = MakeShareable( new FJsonObject() );
 			TabAsJson->SetStringField( TEXT("TabId"), NodeAsStack->Tabs[TabIndex].TabId.ToString() );
 			TabAsJson->SetStringField( TEXT("TabState"), (NodeAsStack->Tabs[TabIndex].TabState == ETabState::OpenedTab)
-				? TEXT("OpenedTab")
-				: TEXT("ClosedTab") );
+				? UE4_TABMANAGER_OPENED_TAB_STRING : NodeAsStack->Tabs[TabIndex].TabState == ETabState::ClosedTab
+				? UE4_TABMANAGER_CLOSED_TAB_STRING : UE4_TABMANAGER_INVALID_TAB_STRING);
 			TabsAsJson.Add( MakeShareable( new FJsonValueObject(TabAsJson) ) );
 		}
 		JsonObj->SetArrayField( TEXT("Tabs"), TabsAsJson );
@@ -758,6 +763,7 @@ void FTabManager::CloseAllAreas()
 	DockAreas.Empty();
 
 	CollapsedDockAreas.Empty();
+	InvalidDockAreas.Empty();
 }
 
 
@@ -797,6 +803,12 @@ TSharedRef<FTabManager::FLayout> FTabManager::PersistLayout() const
 		PersistentLayout->AddArea( CollapsedDockAreas[CollapsedAreaIndex] );
 	}
 
+	// Gather existing persistent layouts for InvalidAreas
+	for (int32 InvalidAreaIndex = 0; InvalidAreaIndex < InvalidDockAreas.Num(); ++InvalidAreaIndex)
+	{
+		PersistentLayout->AddArea(InvalidDockAreas[InvalidAreaIndex]);
+	}
+
 	return PersistentLayout;
 }
 
@@ -834,7 +846,8 @@ TSharedPtr<SWidget> FTabManager::RestoreFrom( const TSharedRef<FLayout>& Layout,
 	for (int32 AreaIndex=0; AreaIndex < Layout->Areas.Num(); ++AreaIndex )
 	{
 		const TSharedRef<FArea>& ThisArea = Layout->Areas[AreaIndex];
-
+		// Set all InvalidTab tabs to OpenedTab so the Editor tries to load them. All non-recognized tabs will be set to InvalidTab later.
+		SetTabsTo(ThisArea, ETabState::OpenedTab, ETabState::InvalidTab);
 		const bool bIsPrimaryArea = ThisArea->WindowPlacement == FArea::Placement_NoWindow;
 		const bool bShouldCreate = bIsPrimaryArea || HasValidTabs(ThisArea);
 
@@ -846,6 +859,14 @@ TSharedPtr<SWidget> FTabManager::RestoreFrom( const TSharedRef<FLayout>& Layout,
 			if ( bHasOpenTabs )
 			{
 				RestoredDockArea = RestoreArea(ThisArea, ParentWindow, bEmbedTitleAreaContent);
+				// Invalidate all tabs in ThisArea because they were not recognized
+				if (!RestoredDockArea)
+				{
+					UE_LOG(LogSlate, Warning, TEXT("The area index %d/%d (extension id \"%s\") attempted to spawn but failed."),
+						AreaIndex+1, Layout->Areas.Num(), *ThisArea->GetExtensionId().ToString());
+					SetTabsTo(ThisArea, ETabState::InvalidTab);
+					InvalidDockAreas.Add(ThisArea);
+				}
 			}
 			else
 			{
@@ -1080,26 +1101,27 @@ void FTabManager::DrawAttention( const TSharedRef<SDockTab>& TabToHighlight )
 	}
 }
 
-void FTabManager::InsertNewDocumentTab(FName PlaceholderId, const FSearchPreference& SearchPreference, const TSharedRef<SDockTab>& UnmanagedTab)
+bool FTabManager::InsertNewDocumentTab(FName PlaceholderId, const FSearchPreference& SearchPreference, const TSharedRef<SDockTab>& UnmanagedTab)
 {
-	InsertDocumentTab(PlaceholderId, SearchPreference, UnmanagedTab, true);
+	return InsertDocumentTab(PlaceholderId, SearchPreference, UnmanagedTab, true);
 }
 
-void FTabManager::InsertNewDocumentTab( FName PlaceholderId, ESearchPreference::Type SearchPreference, const TSharedRef<SDockTab>& UnmanagedTab )
+bool FTabManager::InsertNewDocumentTab( FName PlaceholderId, ESearchPreference::Type SearchPreference, const TSharedRef<SDockTab>& UnmanagedTab )
 {
 	if ( SearchPreference == ESearchPreference::PreferLiveTab )
 	{
 		FLiveTabSearch Search;
-		InsertDocumentTab(PlaceholderId, Search, UnmanagedTab, true);
+		return InsertDocumentTab(PlaceholderId, Search, UnmanagedTab, true);
 	}
 	else if ( SearchPreference == ESearchPreference::RequireClosedTab )
 	{
 		FRequireClosedTab Search;
-		InsertDocumentTab(PlaceholderId, Search, UnmanagedTab, true);
+		return InsertDocumentTab(PlaceholderId, Search, UnmanagedTab, true);
 	}
 	else
 	{
 		check(false);
+		return false;
 	}
 }
 
@@ -1124,6 +1146,7 @@ void FTabManager::RestoreDocumentTab( FName PlaceholderId, ESearchPreference::Ty
 TSharedRef<SDockTab> FTabManager::InvokeTab( const FTabId& TabId )
 {
 	TSharedPtr<SDockTab> NewTab = InvokeTab_Internal( TabId );
+	check(NewTab.IsValid());
 	TSharedPtr<SWindow> ParentWindowPtr = NewTab->GetParentWindow();
 	if ((NewTab->GetTabRole() == ETabRole::MajorTab || NewTab->GetTabRole() == ETabRole::NomadTab) && ParentWindowPtr.IsValid() && ParentWindowPtr != FGlobalTabmanager::Get()->GetRootWindow())
 	{
@@ -1185,7 +1208,7 @@ TSharedPtr<SDockTab> FTabManager::InvokeTab_Internal( const FTabId& TabId )
 	{
 		const TSharedPtr<SDockTab> NewTab = SpawnTab( TabId, TSharedPtr<SWindow>() );
 
-		if(NewTab.IsValid())
+		if (NewTab.IsValid())
 		{
 			StackToSpawnIn->OpenTab(NewTab.ToSharedRef());
 			NewTab->PlaySpawnAnim();
@@ -1208,10 +1231,17 @@ TSharedPtr<SDockTab> FTabManager::InvokeTab_Internal( const FTabId& TabId )
 			->AddTab( TabId, ETabState::OpenedTab )
 		);
 
-		TSharedRef<SDockingArea> DockingArea = RestoreArea(NewAreaForTab, GetPrivateApi().GetParentWindow());
-		const TSharedPtr<SDockTab> NewlyOpenedTab = DockingArea->GetAllChildTabs()[0];
-		check(NewlyOpenedTab.IsValid());
-		return NewlyOpenedTab.ToSharedRef();	
+		TSharedPtr<SDockingArea> DockingArea = RestoreArea(NewAreaForTab, GetPrivateApi().GetParentWindow());
+		if (DockingArea && DockingArea->GetAllChildTabs().Num() > 0)
+		{
+			const TSharedPtr<SDockTab> NewlyOpenedTab = DockingArea->GetAllChildTabs()[0];
+			check(NewlyOpenedTab.IsValid());
+			return NewlyOpenedTab.ToSharedRef();
+		}
+		else
+		{
+			return nullptr;
+		}
 	}
 }
 
@@ -1225,10 +1255,14 @@ TSharedPtr<SDockingTabStack> FTabManager::FindPotentiallyClosedTab( const FTabId
 	const int32 CollapsedAreaWithMatchingTab = FindTabInCollapsedAreas( TabMatcher );
 	if ( CollapsedAreaWithMatchingTab != INDEX_NONE )
 	{
-		TSharedRef<SDockingArea> RestoredArea = RestoreArea( CollapsedDockAreas[CollapsedAreaWithMatchingTab], GetPrivateApi().GetParentWindow() );
+		TSharedPtr<SDockingArea> RestoredArea = RestoreArea(CollapsedDockAreas[CollapsedAreaWithMatchingTab], GetPrivateApi().GetParentWindow());
+		check(RestoredArea.IsValid());
 		// We have just un-collapsed this dock area
 		CollapsedDockAreas.RemoveAt(CollapsedAreaWithMatchingTab);
-		StackWithClosedTab = FindTabInLiveArea(TabMatcher, RestoredArea);
+		if (RestoredArea.IsValid())
+		{
+			StackWithClosedTab = FindTabInLiveArea(TabMatcher, StaticCastSharedRef<SDockingArea>(RestoredArea->AsShared()));
+		}
 	}
 
 	if ( !StackWithClosedTab.IsValid() )
@@ -1245,14 +1279,17 @@ void FTabManager::InvokeTabForMenu( FName TabId )
 	InvokeTab(TabId);
 }
 
-void FTabManager::InsertDocumentTab(FName PlaceholderId, const FSearchPreference& SearchPreference, const TSharedRef<SDockTab>& UnmanagedTab, bool bPlaySpawnAnim)
+bool FTabManager::InsertDocumentTab(FName PlaceholderId, const FSearchPreference& SearchPreference, const TSharedRef<SDockTab>& UnmanagedTab, bool bPlaySpawnAnim)
 {
 	const bool bTabNotManaged = ensure( ! FindTabInLiveAreas( FTabMatcher(UnmanagedTab->GetLayoutIdentifier()) ).IsValid() );
 	UnmanagedTab->SetLayoutIdentifier( FTabId(PlaceholderId, LastDocumentUID++) );
 	
 	if (bTabNotManaged)
 	{
-		OpenUnmanagedTab(PlaceholderId, SearchPreference, UnmanagedTab);
+		if (!OpenUnmanagedTab(PlaceholderId, SearchPreference, UnmanagedTab))
+		{
+			return false;
+		}
 	}
 
 	DrawAttention(UnmanagedTab);
@@ -1260,10 +1297,13 @@ void FTabManager::InsertDocumentTab(FName PlaceholderId, const FSearchPreference
 	{
 		UnmanagedTab->PlaySpawnAnim();
 	}
+
+	return true;
 }
 
-void FTabManager::OpenUnmanagedTab(FName PlaceholderId, const FSearchPreference& SearchPreference, const TSharedRef<SDockTab>& UnmanagedTab)
+bool FTabManager::OpenUnmanagedTab(FName PlaceholderId, const FSearchPreference& SearchPreference, const TSharedRef<SDockTab>& UnmanagedTab)
 {
+	bool bWasTabOpened = true;
 	TSharedPtr<SDockTab> LiveTab = SearchPreference.Search(*this, PlaceholderId, UnmanagedTab);
 		
 	if (LiveTab.IsValid())
@@ -1273,17 +1313,25 @@ void FTabManager::OpenUnmanagedTab(FName PlaceholderId, const FSearchPreference&
 	else
 	{
 		TSharedPtr<SDockingTabStack> StackToSpawnIn = FindPotentiallyClosedTab( PlaceholderId );
-		if( StackToSpawnIn.IsValid() == false )
+		if (StackToSpawnIn.IsValid())
 		{
-			UE_LOG(LogTabManager, Warning, TEXT("Unable to insert tab '%s'."), *(PlaceholderId.ToString()));
-			LiveTab = InvokeTab_Internal( FTabId( PlaceholderId ) );
-			LiveTab->GetParent()->GetParentDockTabStack()->OpenTab( UnmanagedTab );
+			StackToSpawnIn->OpenTab(UnmanagedTab);
 		}
 		else
 		{
-			StackToSpawnIn->OpenTab( UnmanagedTab );
+			UE_LOG(LogTabManager, Warning, TEXT("Unable to insert tab '%s'."), *(PlaceholderId.ToString()));
+			LiveTab = InvokeTab_Internal( FTabId( PlaceholderId ) );
+			if (LiveTab)
+			{
+				LiveTab->GetParent()->GetParentDockTabStack()->OpenTab( UnmanagedTab );
+			}
+			else
+			{
+				bWasTabOpened = false;
+			}
 		}
 	}
+	return bWasTabOpened;
 }
 
 FTabManager::FTabManager( const TSharedPtr<SDockTab>& InOwnerTab, const TSharedRef<FTabSpawner> & InNomadTabSpawner )
@@ -1297,15 +1345,21 @@ FTabManager::FTabManager( const TSharedPtr<SDockTab>& InOwnerTab, const TSharedR
 	LocalWorkspaceMenuRoot = FWorkspaceItem::NewGroup(LOCTEXT("LocalWorkspaceRoot", "Local Workspace Root"));
 }
 
-TSharedRef<SDockingArea> FTabManager::RestoreArea( const TSharedRef<FArea>& AreaToRestore, const TSharedPtr<SWindow>& InParentWindow, const bool bEmbedTitleAreaContent  )
+TSharedPtr<SDockingArea> FTabManager::RestoreArea(const TSharedRef<FArea>& AreaToRestore, const TSharedPtr<SWindow>& InParentWindow, const bool bEmbedTitleAreaContent)
 {
-	TSharedRef<SDockingNode> RestoredNode = RestoreArea_Helper( AreaToRestore, InParentWindow, bEmbedTitleAreaContent );	
-	TSharedRef<SDockingArea> RestoredArea = StaticCastSharedRef<SDockingArea>( RestoredNode );	
-	RestoredArea->CleanUp(SDockingNode::TabRemoval_None);
-	return RestoredArea;
+	if (TSharedPtr<SDockingNode> RestoredNode = RestoreArea_Helper(AreaToRestore, InParentWindow, bEmbedTitleAreaContent))
+	{
+		TSharedRef<SDockingArea> RestoredArea = StaticCastSharedRef<SDockingArea>(RestoredNode->AsShared());
+		RestoredArea->CleanUp(SDockingNode::TabRemoval_None);
+		return RestoredArea;
+	}
+	else
+	{
+		return nullptr;
+	}
 }
 
-TSharedRef<SDockingNode> FTabManager::RestoreArea_Helper( const TSharedRef<FLayoutNode>& LayoutNode, const TSharedPtr<SWindow>& ParentWindow, const bool bEmbedTitleAreaContent )
+TSharedPtr<SDockingNode> FTabManager::RestoreArea_Helper( const TSharedRef<FLayoutNode>& LayoutNode, const TSharedPtr<SWindow>& ParentWindow, const bool bEmbedTitleAreaContent )
 {
 	TSharedPtr<FTabManager::FStack> NodeAsStack = LayoutNode->AsStack();
 	TSharedPtr<FTabManager::FSplitter> NodeAsSplitter = LayoutNode->AsSplitter();
@@ -1315,8 +1369,7 @@ TSharedRef<SDockingNode> FTabManager::RestoreArea_Helper( const TSharedRef<FLayo
 	{
 		TSharedPtr<SDockTab> WidgetToActivate;
 
-		TSharedRef<SDockingTabStack> NewStackWidget = SNew(SDockingTabStack, NodeAsStack.ToSharedRef());
-		NewStackWidget->SetSizeCoefficient(LayoutNode->GetSizeCoefficient());
+		TSharedPtr<SDockingTabStack> NewStackWidget;
 		for (int32 TabIndex=0; TabIndex < NodeAsStack->Tabs.Num(); ++TabIndex )
 		{
 			const FTab& SomeTab = NodeAsStack->Tabs[ TabIndex ];
@@ -1332,6 +1385,12 @@ TSharedRef<SDockingNode> FTabManager::RestoreArea_Helper( const TSharedRef<FLayo
 						WidgetToActivate = NewTabWidget;
 					}
 
+					// First time initialization: Only if at least a valid NewTabWidget
+					if (!NewStackWidget)
+					{
+						NewStackWidget = SNew(SDockingTabStack, NodeAsStack.ToSharedRef());
+						NewStackWidget->SetSizeCoefficient(LayoutNode->GetSizeCoefficient());
+					}
 					NewStackWidget->AddTabWidget(NewTabWidget.ToSharedRef());
 				}
 			}
@@ -1341,12 +1400,13 @@ TSharedRef<SDockingNode> FTabManager::RestoreArea_Helper( const TSharedRef<FLayo
 		{
 			WidgetToActivate->ActivateInParent(ETabActivationCause::SetDirectly);
 
-			if ((WidgetToActivate->GetTabRole() == ETabRole::MajorTab || WidgetToActivate->GetTabRole() == ETabRole::NomadTab) && ParentWindow.IsValid() && ParentWindow != FGlobalTabmanager::Get()->GetRootWindow())
+			if ((WidgetToActivate->GetTabRole() == ETabRole::MajorTab || WidgetToActivate->GetTabRole() == ETabRole::NomadTab)
+				&& ParentWindow.IsValid() && ParentWindow != FGlobalTabmanager::Get()->GetRootWindow())
 			{
 				ParentWindow->SetTitle(WidgetToActivate->GetTabLabel());
 			}
 		}
-		
+
 		return NewStackWidget;
 
 	}
@@ -1378,50 +1438,69 @@ TSharedRef<SDockingNode> FTabManager::RestoreArea_Helper( const TSharedRef<FLayo
 					.CreateTitleBar( false )
 					.IsInitiallyMaximized( NodeAsArea->bIsMaximized );
 
-			// Any windows that were "pulled out" of a dock area should be children of the window in which the parent dock area resides.
-			if (bIsChildWindow)
-			{
-				FSlateApplication::Get().AddWindowAsNativeChild( NewWindow, ParentWindow.ToSharedRef() )->SetContent(
-					SAssignNew( NewDockAreaWidget, SDockingArea, SharedThis(this), NodeAsArea.ToSharedRef() ) .ParentWindow( NewWindow )
-				);
-			}
-			else
-			{
-				FSlateApplication::Get().AddWindow( NewWindow )->SetContent(
-					SAssignNew( NewDockAreaWidget, SDockingArea, SharedThis(this), NodeAsArea.ToSharedRef() ) .ParentWindow( NewWindow )
-				);
-			}
-
 			// Set a default title; restoring the splitter content may override this if it activates a tab
 			NewWindow->SetTitle( FGlobalTabmanager::Get()->GetApplicationTitle() );
 
-			RestoreSplitterContent( NodeAsArea.ToSharedRef(), NewDockAreaWidget.ToSharedRef(), NewWindow );
+			TArray<TSharedRef<SDockingNode>> DockingNodes;
+			if (CanRestoreSplitterContent(DockingNodes, NodeAsArea.ToSharedRef(), NewWindow))
+			{
+				// Create SplitterWidget only if it will be filled with at least 1 DockingNodes
+				// Any windows that were "pulled out" of a dock area should be children of the window in which the parent dock area resides.
+				if (bIsChildWindow)
+				{
+					FSlateApplication::Get().AddWindowAsNativeChild(NewWindow, ParentWindow.ToSharedRef())->SetContent(
+						SAssignNew(NewDockAreaWidget, SDockingArea, SharedThis(this), NodeAsArea.ToSharedRef()).ParentWindow(NewWindow)
+					);
+				}
+				else
+				{
+					FSlateApplication::Get().AddWindow(NewWindow)->SetContent(
+						SAssignNew(NewDockAreaWidget, SDockingArea, SharedThis(this), NodeAsArea.ToSharedRef()).ParentWindow(NewWindow)
+					);
+				}
+				// Restore content
+				RestoreSplitterContent(DockingNodes, NewDockAreaWidget.ToSharedRef());
+			}
 		}
 		else
 		{
-			SAssignNew( NewDockAreaWidget, SDockingArea, SharedThis(this), NodeAsArea.ToSharedRef() )
-
-				// We only want to set a parent window on this dock area, if we need to have title area content
-				// embedded within it.  SDockingArea assumes that if it has a parent window set, then it needs to have
-				// title area content 
-				.ParentWindow( bEmbedTitleAreaContent ? ParentWindow : TSharedPtr<SWindow>() )
-
-				// Never manage these windows, even if a parent window is set.  The owner will take care of
-				// destroying these windows.
-				.ShouldManageParentWindow( false );
-
-			RestoreSplitterContent( NodeAsArea.ToSharedRef(), NewDockAreaWidget.ToSharedRef(), ParentWindow );
+			TArray<TSharedRef<SDockingNode>> DockingNodes;
+			if (CanRestoreSplitterContent(DockingNodes, NodeAsArea.ToSharedRef(), ParentWindow))
+			{
+				// Create SplitterWidget only if it will be filled with at least 1 DockingNodes
+				SAssignNew(NewDockAreaWidget, SDockingArea, SharedThis(this), NodeAsArea.ToSharedRef())
+
+					// We only want to set a parent window on this dock area, if we need to have title area content
+					// embedded within it.  SDockingArea assumes that if it has a parent window set, then it needs to have
+					// title area content 
+					.ParentWindow(bEmbedTitleAreaContent ? ParentWindow : TSharedPtr<SWindow>())
+
+					// Never manage these windows, even if a parent window is set.  The owner will take care of
+					// destroying these windows.
+					.ShouldManageParentWindow(false);
+				// Restore content
+				RestoreSplitterContent(DockingNodes, NewDockAreaWidget.ToSharedRef());
+			}
 		}
-		
-		return NewDockAreaWidget.ToSharedRef();
+
+		return NewDockAreaWidget;
 	}
 	else if ( NodeAsSplitter.IsValid() ) 
 	{
-		TSharedRef<SDockingSplitter> NewSplitterWidget = SNew( SDockingSplitter, NodeAsSplitter.ToSharedRef() );
-		NewSplitterWidget->SetSizeCoefficient(LayoutNode->GetSizeCoefficient());
-		RestoreSplitterContent( NodeAsSplitter.ToSharedRef(), NewSplitterWidget, ParentWindow );
-
-		return NewSplitterWidget;
+		TArray<TSharedRef<SDockingNode>> DockingNodes;
+		if (CanRestoreSplitterContent(DockingNodes, NodeAsSplitter.ToSharedRef(), ParentWindow))
+		{
+			// Create SplitterWidget only if it will be filled with at least 1 DockingNodes
+			TSharedRef<SDockingSplitter> NewSplitterWidget = SNew(SDockingSplitter, NodeAsSplitter.ToSharedRef());
+			NewSplitterWidget->SetSizeCoefficient(LayoutNode->GetSizeCoefficient());
+			// Restore content
+			RestoreSplitterContent(DockingNodes, NewSplitterWidget);
+			return NewSplitterWidget;
+		}
+		else
+		{
+			return nullptr;
+		}
 	}
 	else
 	{
@@ -1433,16 +1512,30 @@ TSharedRef<SDockingNode> FTabManager::RestoreArea_Helper( const TSharedRef<FLayo
 	}
 }
 
-void FTabManager::RestoreSplitterContent( const TSharedRef<FSplitter>& SplitterNode, const TSharedRef<SDockingSplitter>& SplitterWidget, const TSharedPtr<SWindow>& ParentWindow )
+bool FTabManager::CanRestoreSplitterContent( TArray<TSharedRef<SDockingNode>>& DockingNodes, const TSharedRef<FSplitter>& SplitterNode, const TSharedPtr<SWindow>& ParentWindow )
 {
+	DockingNodes.Empty();
 	// Restore the contents of this splitter.
 	for ( int32 ChildNodeIndex = 0; ChildNodeIndex < SplitterNode->ChildNodes.Num(); ++ChildNodeIndex )
 	{
-		TSharedRef<FLayoutNode> ThisChildNode = SplitterNode->ChildNodes[ChildNodeIndex];
+		const TSharedRef<FLayoutNode> ThisChildNode = SplitterNode->ChildNodes[ChildNodeIndex];
 
 		const bool bEmbedTitleAreaContent = false;
-		TSharedRef<SDockingNode> ThisChildNodeWidget = RestoreArea_Helper( ThisChildNode, ParentWindow, bEmbedTitleAreaContent );
-		SplitterWidget->AddChildNode( ThisChildNodeWidget, INDEX_NONE );
+		const TSharedPtr<SDockingNode> ThisChildNodeWidget = RestoreArea_Helper(ThisChildNode, ParentWindow, bEmbedTitleAreaContent);
+		if (ThisChildNodeWidget)
+		{
+			const TSharedRef<SDockingNode> ThisChildNodeWidgetRef = StaticCastSharedRef<SDockingNode>(ThisChildNodeWidget->AsShared());
+			DockingNodes.Add(ThisChildNodeWidgetRef);
+		}
+	}
+	return (DockingNodes.Num() > 0);
+}
+
+void FTabManager::RestoreSplitterContent( const TArray<TSharedRef<SDockingNode>>& DockingNodes, const TSharedRef<SDockingSplitter>& SplitterWidget )
+{
+	for (const TSharedRef<SDockingNode>& DockingNode : DockingNodes)
+	{
+		SplitterWidget->AddChildNode(DockingNode, INDEX_NONE);
 	}
 }
 
@@ -1507,7 +1600,7 @@ TSharedPtr<SDockTab> FTabManager::SpawnTab( const FTabId& TabId, const TSharedPt
 		UE_LOG(LogSlate, Warning,
 			TEXT("The tab \"%s\" attempted to spawn but failed for some reason. It will not be displayed but it will be saved in the layout settings file."),
 			*(!StringToDisplay.IsEmpty() ? StringToDisplay : FString("Unknown"))
-			);
+		);
 	}
 
 	if (NewTabWidget.IsValid())
@@ -1671,6 +1764,33 @@ bool FTabManager::HasValidTabs( const TSharedRef<FTabManager::FLayoutNode>& Some
 	return HasAnyMatchingTabs(SomeNode, FindValidTab);
 }
 
+void FTabManager::SetTabsTo(const TSharedRef<FTabManager::FLayoutNode>& SomeNode, const ETabState::Type NewTabState, const ETabState::Type OriginalTabState) const
+{
+	// Set particular tab to desired NewTabState
+	TSharedPtr<FTabManager::FStack> AsStack = SomeNode->AsStack();
+	if (AsStack.IsValid())
+	{
+		TArray<FTab>& Tabs = AsStack->Tabs;
+		for (int32 TabIndex = 0; TabIndex < Tabs.Num(); ++TabIndex)
+		{
+			if (Tabs[TabIndex].TabState == OriginalTabState || OriginalTabState == ETabState::AnyTab)
+			{
+				Tabs[TabIndex].TabState = NewTabState;
+			}
+		}
+	}
+	// Recursively set all tabs to desired NewTabState
+	else
+	{
+		TSharedPtr<FTabManager::FSplitter> AsSplitter = SomeNode->AsSplitter();
+		ensure(AsSplitter.IsValid());
+		for (int32 ChildIndex = 0; ChildIndex < AsSplitter->ChildNodes.Num(); ++ChildIndex)
+		{
+			SetTabsTo(AsSplitter->ChildNodes[ChildIndex], NewTabState, OriginalTabState);
+		}
+	}
+}
+
 void FTabManager::OnTabForegrounded( const TSharedPtr<SDockTab>& NewForegroundTab, const TSharedPtr<SDockTab>& BackgroundedTab )
 {
 	// Do nothing.
@@ -2214,15 +2334,15 @@ void FGlobalTabmanager::UpdateStats()
 	AllAreasWindowMaxCount = FMath::Max(AllAreasWindowMaxCount, ParentWindows.Num());
 }
 
-void FGlobalTabmanager::OpenUnmanagedTab(FName PlaceholderId, const FSearchPreference& SearchPreference, const TSharedRef<SDockTab>& UnmanagedTab)
+bool FGlobalTabmanager::OpenUnmanagedTab(FName PlaceholderId, const FSearchPreference& SearchPreference, const TSharedRef<SDockTab>& UnmanagedTab)
 {
 	if ( ProxyTabManager.IsValid() && ProxyTabManager->IsTabSupported( UnmanagedTab->GetLayoutIdentifier() ) )
 	{
-		ProxyTabManager->OpenUnmanagedTab(PlaceholderId, SearchPreference, UnmanagedTab);
+		return ProxyTabManager->OpenUnmanagedTab(PlaceholderId, SearchPreference, UnmanagedTab);
 	}
 	else
 	{
-		FTabManager::OpenUnmanagedTab(PlaceholderId, SearchPreference, UnmanagedTab);
+		return FTabManager::OpenUnmanagedTab(PlaceholderId, SearchPreference, UnmanagedTab);
 	}
 }
 
@@ -2242,10 +2362,10 @@ bool FProxyTabmanager::IsTabSupported( const FTabId TabId ) const
 	return bIsTabSupported;
 }
 
-void FProxyTabmanager::OpenUnmanagedTab(FName PlaceholderId, const FSearchPreference& SearchPreference, const TSharedRef<SDockTab>& UnmanagedTab)
+bool FProxyTabmanager::OpenUnmanagedTab(FName PlaceholderId, const FSearchPreference& SearchPreference, const TSharedRef<SDockTab>& UnmanagedTab)
 {
 	TSharedPtr<SWindow> ParentWindowPtr = ParentWindow.Pin();
-	if( ensure( ParentWindowPtr.IsValid() ) )
+	if (ensure(ParentWindowPtr.IsValid()))
 	{
 		// No layout info about this tab found; start 
 		TSharedRef<FArea> NewAreaForTab = FTabManager::NewPrimaryArea()
@@ -2255,22 +2375,26 @@ void FProxyTabmanager::OpenUnmanagedTab(FName PlaceholderId, const FSearchPrefer
 				->AddTab( UnmanagedTab->GetLayoutIdentifier(), ETabState::OpenedTab )
 			);
 
-		TSharedRef<SDockingArea> DockingArea = RestoreArea(NewAreaForTab, ParentWindowPtr, false);
-		ParentWindowPtr->SetContent(DockingArea);
-
-		if (DockingArea->GetAllChildTabs().Num() > 0)
+		if (TSharedPtr<SDockingArea> DockingArea = RestoreArea(NewAreaForTab, ParentWindowPtr, false))
 		{
-			const TSharedPtr<SDockTab> NewlyOpenedTab = DockingArea->GetAllChildTabs()[0];
-			check(NewlyOpenedTab.IsValid());
+			ParentWindowPtr->SetContent(StaticCastSharedRef<SDockingArea>(DockingArea->AsShared()));
+			if (DockingArea->GetAllChildTabs().Num() > 0)
+			{
+				const TSharedPtr<SDockTab> NewlyOpenedTab = DockingArea->GetAllChildTabs()[0];
+				check(NewlyOpenedTab.IsValid());
 
-			NewlyOpenedTab->GetParent()->GetParentDockTabStack()->OpenTab(UnmanagedTab);
-			NewlyOpenedTab->RequestCloseTab();
+				NewlyOpenedTab->GetParent()->GetParentDockTabStack()->OpenTab(UnmanagedTab);
+				NewlyOpenedTab->RequestCloseTab();
 
-			MainNonCloseableTab = UnmanagedTab;
+				MainNonCloseableTab = UnmanagedTab;
 
-			OnTabOpened.Broadcast(UnmanagedTab);
+				OnTabOpened.Broadcast(UnmanagedTab);
+
+				return true;
+			}
 		}
 	}
+	return false;
 }
 
 void FProxyTabmanager::DrawAttention(const TSharedRef<SDockTab>& TabToHighlight)
diff --git a/Engine/Source/Runtime/Slate/Public/Framework/Docking/TabManager.h b/Engine/Source/Runtime/Slate/Public/Framework/Docking/TabManager.h
--- a/Engine/Source/Runtime/Slate/Public/Framework/Docking/TabManager.h
+++ b/Engine/Source/Runtime/Slate/Public/Framework/Docking/TabManager.h
@@ -224,7 +224,15 @@ namespace ETabState
 	enum Type
 	{
 		OpenedTab = 0x1 << 0,
-		ClosedTab = 0x1 << 1
+		ClosedTab = 0x1 << 1,
+		/**
+		 * InvalidTab refers to tabs that were not recognized by the Editor (e.g., LiveLink when its plugin its disabled).
+		 */
+		InvalidTab = 0x1 << 2,
+		/**
+		 * AnyTab refers to any kind of tab (OpenedTab, ClosedTab, InvalidTab).
+		 */
+		AnyTab = 0x1 << 3,
 	};
 }
 
@@ -697,10 +705,10 @@ class SLATE_API FTabManager : public TSharedFromThis<FTabManager>
 		};
 
 		/** Insert a new UnmanagedTab document tab next to an existing tab (closed or open) that has the PlaceholdId. */
-		void InsertNewDocumentTab( FName PlaceholderId, const FSearchPreference& SearchPreference, const TSharedRef<SDockTab>& UnmanagedTab );
+		bool InsertNewDocumentTab( FName PlaceholderId, const FSearchPreference& SearchPreference, const TSharedRef<SDockTab>& UnmanagedTab );
 		
 		/** Insert a new UnmanagedTab document tab next to an existing tab (closed or open) that has the PlaceholdId. */
-		void InsertNewDocumentTab(FName PlaceholderId, ESearchPreference::Type SearchPreference, const TSharedRef<SDockTab>& UnmanagedTab);
+		bool InsertNewDocumentTab(FName PlaceholderId, ESearchPreference::Type SearchPreference, const TSharedRef<SDockTab>& UnmanagedTab);
 
 		/**
 		 * Much like InsertNewDocumentTab, but the UnmanagedTab is not seen by the user as newly-created.
@@ -779,9 +787,9 @@ class SLATE_API FTabManager : public TSharedFromThis<FTabManager>
 
 	protected:
 		
-		void InsertDocumentTab( FName PlaceholderId, const FSearchPreference& SearchPreference, const TSharedRef<SDockTab>& UnmanagedTab, bool bPlaySpawnAnim );
+		bool InsertDocumentTab( FName PlaceholderId, const FSearchPreference& SearchPreference, const TSharedRef<SDockTab>& UnmanagedTab, bool bPlaySpawnAnim );
 
-		virtual void OpenUnmanagedTab(FName PlaceholderId, const FSearchPreference& SearchPreference, const TSharedRef<SDockTab>& UnmanagedTab);
+		virtual bool OpenUnmanagedTab(FName PlaceholderId, const FSearchPreference& SearchPreference, const TSharedRef<SDockTab>& UnmanagedTab);
 			
 		void PopulateTabSpawnerMenu_Helper( FMenuBuilder& PopulateMe, struct FPopulateTabSpawnerMenu_Args Args );
 
@@ -803,11 +811,23 @@ class SLATE_API FTabManager : public TSharedFromThis<FTabManager>
 
 		FTabManager( const TSharedPtr<SDockTab>& InOwnerTab, const TSharedRef<FTabManager::FTabSpawner> & InNomadTabSpawner );
 
-		TSharedRef<SDockingArea> RestoreArea( const TSharedRef<FArea>& AreaToRestore, const TSharedPtr<SWindow>& InParentWindow, const bool bEmbedTitleAreaContent = false );
+		TSharedPtr<SDockingArea> RestoreArea(const TSharedRef<FArea>& AreaToRestore, const TSharedPtr<SWindow>& InParentWindow, const bool bEmbedTitleAreaContent = false);
 
-		TSharedRef<class SDockingNode> RestoreArea_Helper( const TSharedRef<FLayoutNode>& LayoutNode, const TSharedPtr<SWindow>& ParentWindow, const bool bEmbedTitleAreaContent );
+		TSharedPtr<class SDockingNode> RestoreArea_Helper(const TSharedRef<FLayoutNode>& LayoutNode, const TSharedPtr<SWindow>& ParentWindow, const bool bEmbedTitleAreaContent);
 
-		void RestoreSplitterContent( const TSharedRef<FSplitter>& SplitterNode, const TSharedRef<class SDockingSplitter>& SplitterWidget, const TSharedPtr<SWindow>& ParentWindow );
+		/**
+		 * Usage example:
+		 *		TArray<TSharedRef<SDockingNode>> DockingNodes;
+		 *		if (CanRestoreSplitterContent(DockingNodes, SplitterNode, ParentWindow))
+		 *		{
+		 *			// Create SplitterWidget only if it will be filled with at least 1 DockingNodes
+		 *			TSharedRef<SDockingSplitter> SplitterWidget = SNew(SDockingSplitter, SplitterNode);
+		 *			// Restore content
+		 *			RestoreSplitterContent(DockingNodes, SplitterWidget);
+		 *		}
+		 */
+		bool CanRestoreSplitterContent( TArray<TSharedRef<class SDockingNode>>& DockingNodes, const TSharedRef<FSplitter>& SplitterNode, const TSharedPtr<SWindow>& ParentWindow );
+		void RestoreSplitterContent( const TArray<TSharedRef<class SDockingNode>>& DockingNodes, const TSharedRef<class SDockingSplitter>& SplitterWidget );
 		
 		bool IsValidTabForSpawning( const FTab& SomeTab ) const;
 		TSharedPtr<SDockTab> SpawnTab( const FTabId& TabId, const TSharedPtr<SWindow>& ParentWindow );
@@ -818,6 +838,13 @@ class SLATE_API FTabManager : public TSharedFromThis<FTabManager>
 		template<typename MatchFunctorType> static bool HasAnyMatchingTabs( const TSharedRef<FTabManager::FLayoutNode>& SomeNode, const MatchFunctorType& Matcher );
 		bool HasOpenTabs( const TSharedRef<FTabManager::FLayoutNode>& SomeNode ) const;
 		bool HasValidTabs( const TSharedRef<FTabManager::FLayoutNode>& SomeNode ) const;
+		/**
+		 * It sets the desired (or all) tabs in the FTabManager::FLayoutNode to the desired value.
+		 * @param SomeNode The area whose tabs will be modified.
+		 * @param NewTabState The new TabState value.
+		 * @param OriginalTabState Only the tabs with this value will be modified. Use ETabState::AnyTab to modify them all.
+		 */
+		void SetTabsTo(const TSharedRef<FTabManager::FLayoutNode>& SomeNode, const ETabState::Type NewTabState, const ETabState::Type OriginalTabState = ETabState::AnyTab) const;
 
 		/**
 		 * Notify the tab manager that the NewForegroundTab was brought to front and the BackgroundedTab was send to the background as a result.
@@ -852,7 +879,16 @@ class SLATE_API FTabManager : public TSharedFromThis<FTabManager>
 		bool HasTabSpawnerFor(FName TabId) const;
 
 		TArray< TWeakPtr<SDockingArea> > DockAreas;
+		/**
+		 * CollapsedDockAreas refers to areas that were closed (e.g., by the user).
+		 * We save its location so they can be re-opened in the same location if the user opens thems again.
+		 */
 		TArray< TSharedRef<FTabManager::FArea> > CollapsedDockAreas;
+		/**
+		 * InvalidDockAreas refers to areas that were not recognized by the Editor (e.g., LiveLink when its plugin its disabled).
+		 * We save its location so they can be re-opened whenever they are recognized again (e.g., if its related plugin is re-enabled).
+		 */
+		TArray< TSharedRef<FTabManager::FArea> > InvalidDockAreas;
 
 		/** The root for the local editor's tab spawner workspace menu */
 		TSharedPtr<FWorkspaceItem> LocalWorkspaceMenuRoot;
@@ -1012,7 +1048,7 @@ protected:
 	virtual void OnTabClosing( const TSharedRef<SDockTab>& TabBeingClosed ) override;
 	virtual void UpdateStats() override;
 
-	virtual void OpenUnmanagedTab(FName PlaceholderId, const FSearchPreference& SearchPreference, const TSharedRef<SDockTab>& UnmanagedTab) override;
+	virtual bool OpenUnmanagedTab(FName PlaceholderId, const FSearchPreference& SearchPreference, const TSharedRef<SDockTab>& UnmanagedTab) override;
 
 public:
 	virtual void OnTabManagerClosing() override;
@@ -1114,7 +1150,7 @@ public:
 		bCanDoDragOperation = false;
 	}
 
-	virtual void OpenUnmanagedTab(FName PlaceholderId, const FSearchPreference& SearchPreference, const TSharedRef<SDockTab>& UnmanagedTab) override;
+	virtual bool OpenUnmanagedTab(FName PlaceholderId, const FSearchPreference& SearchPreference, const TSharedRef<SDockTab>& UnmanagedTab) override;
 	virtual void DrawAttention(const TSharedRef<SDockTab>& TabToHighlight) override;
 	
 	bool IsTabSupported( const FTabId TabId ) const;
