Bug fixed: Skeletal mesh editor and any other similar tabbed editor widgets crashed when opened.

diff --git a/Engine/Source/Editor/LocalizationDashboard/Private/LocalizationDashboard.cpp b/Engine/Source/Editor/LocalizationDashboard/Private/LocalizationDashboard.cpp
--- a/Engine/Source/Editor/LocalizationDashboard/Private/LocalizationDashboard.cpp
+++ b/Engine/Source/Editor/LocalizationDashboard/Private/LocalizationDashboard.cpp
@@ -206,10 +206,8 @@ TWeakPtr<SDockTab> SLocalizationDashboard::ShowTargetEditor(ULocalizationTarget*
 				OurTargetEditor
 			];
 
-		if (TabManager->InsertNewDocumentTab(DocumentsTabName, FTabManager::ESearchPreference::RequireClosedTab, NewTargetEditorTab))
-		{
-			TargetEditorDockTab = NewTargetEditorTab;
-		}
+		TabManager->InsertNewDocumentTab(DocumentsTabName, FTabManager::ESearchPreference::RequireClosedTab, NewTargetEditorTab);
+		TargetEditorDockTab = NewTargetEditorTab;
 	}
 	else
 	{
diff --git a/Engine/Source/Runtime/Slate/Private/Framework/Docking/TabManager.cpp b/Engine/Source/Runtime/Slate/Private/Framework/Docking/TabManager.cpp
--- a/Engine/Source/Runtime/Slate/Private/Framework/Docking/TabManager.cpp
+++ b/Engine/Source/Runtime/Slate/Private/Framework/Docking/TabManager.cpp
@@ -858,13 +858,14 @@ TSharedPtr<SWidget> FTabManager::RestoreFrom( const TSharedRef<FLayout>& Layout,
 
 			if ( bHasOpenTabs )
 			{
-				RestoredDockArea = RestoreArea(ThisArea, ParentWindow, bEmbedTitleAreaContent);
+				const bool bCanOutputBeNullptr = true;
+				RestoredDockArea = RestoreArea(ThisArea, ParentWindow, bEmbedTitleAreaContent, bCanOutputBeNullptr);
 				// Invalidate all tabs in ThisArea because they were not recognized
 				if (!RestoredDockArea)
 				{
 					UE_LOG(LogSlate, Warning, TEXT("The area index %d/%d (extension id \"%s\") attempted to spawn but failed."),
 						AreaIndex+1, Layout->Areas.Num(), *ThisArea->GetExtensionId().ToString());
-					SetTabsTo(ThisArea, ETabState::InvalidTab);
+					SetTabsTo(ThisArea, ETabState::InvalidTab, ETabState::OpenedTab);
 					InvalidDockAreas.Add(ThisArea);
 				}
 			}
@@ -1101,27 +1102,26 @@ void FTabManager::DrawAttention( const TSharedRef<SDockTab>& TabToHighlight )
 	}
 }
 
-bool FTabManager::InsertNewDocumentTab(FName PlaceholderId, const FSearchPreference& SearchPreference, const TSharedRef<SDockTab>& UnmanagedTab)
+void FTabManager::InsertNewDocumentTab(FName PlaceholderId, const FSearchPreference& SearchPreference, const TSharedRef<SDockTab>& UnmanagedTab)
 {
-	return InsertDocumentTab(PlaceholderId, SearchPreference, UnmanagedTab, true);
+	InsertDocumentTab(PlaceholderId, SearchPreference, UnmanagedTab, true);
 }
 
-bool FTabManager::InsertNewDocumentTab( FName PlaceholderId, ESearchPreference::Type SearchPreference, const TSharedRef<SDockTab>& UnmanagedTab )
+void FTabManager::InsertNewDocumentTab( FName PlaceholderId, ESearchPreference::Type SearchPreference, const TSharedRef<SDockTab>& UnmanagedTab )
 {
 	if ( SearchPreference == ESearchPreference::PreferLiveTab )
 	{
 		FLiveTabSearch Search;
-		return InsertDocumentTab(PlaceholderId, Search, UnmanagedTab, true);
+		InsertDocumentTab(PlaceholderId, Search, UnmanagedTab, true);
 	}
 	else if ( SearchPreference == ESearchPreference::RequireClosedTab )
 	{
 		FRequireClosedTab Search;
-		return InsertDocumentTab(PlaceholderId, Search, UnmanagedTab, true);
+		InsertDocumentTab(PlaceholderId, Search, UnmanagedTab, true);
 	}
 	else
 	{
 		check(false);
-		return false;
 	}
 }
 
@@ -1279,17 +1279,15 @@ void FTabManager::InvokeTabForMenu( FName TabId )
 	InvokeTab(TabId);
 }
 
-bool FTabManager::InsertDocumentTab(FName PlaceholderId, const FSearchPreference& SearchPreference, const TSharedRef<SDockTab>& UnmanagedTab, bool bPlaySpawnAnim)
+void FTabManager::InsertDocumentTab(FName PlaceholderId, const FSearchPreference& SearchPreference, const TSharedRef<SDockTab>& UnmanagedTab, bool bPlaySpawnAnim)
 {
+	bool bWasUnmanagedTabOpened = true;
 	const bool bTabNotManaged = ensure( ! FindTabInLiveAreas( FTabMatcher(UnmanagedTab->GetLayoutIdentifier()) ).IsValid() );
 	UnmanagedTab->SetLayoutIdentifier( FTabId(PlaceholderId, LastDocumentUID++) );
 	
 	if (bTabNotManaged)
 	{
-		if (!OpenUnmanagedTab(PlaceholderId, SearchPreference, UnmanagedTab))
-		{
-			return false;
-		}
+		OpenUnmanagedTab(PlaceholderId, SearchPreference, UnmanagedTab);
 	}
 
 	DrawAttention(UnmanagedTab);
@@ -1297,13 +1295,10 @@ bool FTabManager::InsertDocumentTab(FName PlaceholderId, const FSearchPreference
 	{
 		UnmanagedTab->PlaySpawnAnim();
 	}
-
-	return true;
 }
 
-bool FTabManager::OpenUnmanagedTab(FName PlaceholderId, const FSearchPreference& SearchPreference, const TSharedRef<SDockTab>& UnmanagedTab)
+void FTabManager::OpenUnmanagedTab(FName PlaceholderId, const FSearchPreference& SearchPreference, const TSharedRef<SDockTab>& UnmanagedTab)
 {
-	bool bWasTabOpened = true;
 	TSharedPtr<SDockTab> LiveTab = SearchPreference.Search(*this, PlaceholderId, UnmanagedTab);
 		
 	if (LiveTab.IsValid())
@@ -1321,17 +1316,13 @@ bool FTabManager::OpenUnmanagedTab(FName PlaceholderId, const FSearchPreference&
 		{
 			UE_LOG(LogTabManager, Warning, TEXT("Unable to insert tab '%s'."), *(PlaceholderId.ToString()));
 			LiveTab = InvokeTab_Internal( FTabId( PlaceholderId ) );
+			check(LiveTab.IsValid());
 			if (LiveTab)
 			{
 				LiveTab->GetParent()->GetParentDockTabStack()->OpenTab( UnmanagedTab );
 			}
-			else
-			{
-				bWasTabOpened = false;
-			}
 		}
 	}
-	return bWasTabOpened;
 }
 
 FTabManager::FTabManager( const TSharedPtr<SDockTab>& InOwnerTab, const TSharedRef<FTabSpawner> & InNomadTabSpawner )
@@ -1345,9 +1336,10 @@ FTabManager::FTabManager( const TSharedPtr<SDockTab>& InOwnerTab, const TSharedR
 	LocalWorkspaceMenuRoot = FWorkspaceItem::NewGroup(LOCTEXT("LocalWorkspaceRoot", "Local Workspace Root"));
 }
 
-TSharedPtr<SDockingArea> FTabManager::RestoreArea(const TSharedRef<FArea>& AreaToRestore, const TSharedPtr<SWindow>& InParentWindow, const bool bEmbedTitleAreaContent)
+TSharedPtr<SDockingArea> FTabManager::RestoreArea(
+	const TSharedRef<FArea>& AreaToRestore, const TSharedPtr<SWindow>& InParentWindow, const bool bEmbedTitleAreaContent, const bool bCanOutputBeNullptr)
 {
-	if (TSharedPtr<SDockingNode> RestoredNode = RestoreArea_Helper(AreaToRestore, InParentWindow, bEmbedTitleAreaContent))
+	if (TSharedPtr<SDockingNode> RestoredNode = RestoreArea_Helper(AreaToRestore, InParentWindow, bEmbedTitleAreaContent, bCanOutputBeNullptr))
 	{
 		TSharedRef<SDockingArea> RestoredArea = StaticCastSharedRef<SDockingArea>(RestoredNode->AsShared());
 		RestoredArea->CleanUp(SDockingNode::TabRemoval_None);
@@ -1359,7 +1351,8 @@ TSharedPtr<SDockingArea> FTabManager::RestoreArea(const TSharedRef<FArea>& AreaT
 	}
 }
 
-TSharedPtr<SDockingNode> FTabManager::RestoreArea_Helper( const TSharedRef<FLayoutNode>& LayoutNode, const TSharedPtr<SWindow>& ParentWindow, const bool bEmbedTitleAreaContent )
+TSharedPtr<SDockingNode> FTabManager::RestoreArea_Helper(
+	const TSharedRef<FLayoutNode>& LayoutNode, const TSharedPtr<SWindow>& ParentWindow, const bool bEmbedTitleAreaContent, const bool bCanOutputBeNullptr)
 {
 	TSharedPtr<FTabManager::FStack> NodeAsStack = LayoutNode->AsStack();
 	TSharedPtr<FTabManager::FSplitter> NodeAsSplitter = LayoutNode->AsSplitter();
@@ -1370,9 +1363,22 @@ TSharedPtr<SDockingNode> FTabManager::RestoreArea_Helper( const TSharedRef<FLayo
 		TSharedPtr<SDockTab> WidgetToActivate;
 
 		TSharedPtr<SDockingTabStack> NewStackWidget;
+		if (!bCanOutputBeNullptr)
+		{
+			NewStackWidget = SNew(SDockingTabStack, NodeAsStack.ToSharedRef());
+			NewStackWidget->SetSizeCoefficient(LayoutNode->GetSizeCoefficient());
+		}
 		for (int32 TabIndex=0; TabIndex < NodeAsStack->Tabs.Num(); ++TabIndex )
 		{
 			const FTab& SomeTab = NodeAsStack->Tabs[ TabIndex ];
+if (SomeTab.TabId.ToString() == FString("LiveLink"))
+{
+	UE_LOG(LogSlate, Warning, TEXT("SomeTab %d/%d = %s."), TabIndex + 1, NodeAsStack->Tabs.Num(), *(SomeTab.TabId.ToString()));
+}
+if (SomeTab.TabId.ToString() == FString("StandaloneToolkit"))
+{
+	UE_LOG(LogSlate, Warning, TEXT("SomeTab %d/%d = %s."), TabIndex + 1, NodeAsStack->Tabs.Num(), *(SomeTab.TabId.ToString()));
+}
 
 			if (SomeTab.TabState == ETabState::OpenedTab && IsValidTabForSpawning(SomeTab))
 			{
@@ -1395,7 +1401,7 @@ TSharedPtr<SDockingNode> FTabManager::RestoreArea_Helper( const TSharedRef<FLayo
 				}
 			}
 		}
-		
+
 		if(WidgetToActivate.IsValid())
 		{
 			WidgetToActivate->ActivateInParent(ETabActivationCause::SetDirectly);
@@ -1439,10 +1445,10 @@ TSharedPtr<SDockingNode> FTabManager::RestoreArea_Helper( const TSharedRef<FLayo
 					.IsInitiallyMaximized( NodeAsArea->bIsMaximized );
 
 			// Set a default title; restoring the splitter content may override this if it activates a tab
-			NewWindow->SetTitle( FGlobalTabmanager::Get()->GetApplicationTitle() );
+			NewWindow->SetTitle(FGlobalTabmanager::Get()->GetApplicationTitle());
 
 			TArray<TSharedRef<SDockingNode>> DockingNodes;
-			if (CanRestoreSplitterContent(DockingNodes, NodeAsArea.ToSharedRef(), NewWindow))
+			if (!bCanOutputBeNullptr || CanRestoreSplitterContent(DockingNodes, NodeAsArea.ToSharedRef(), NewWindow))
 			{
 				// Create SplitterWidget only if it will be filled with at least 1 DockingNodes
 				// Any windows that were "pulled out" of a dock area should be children of the window in which the parent dock area resides.
@@ -1459,53 +1465,44 @@ TSharedPtr<SDockingNode> FTabManager::RestoreArea_Helper( const TSharedRef<FLayo
 					);
 				}
 				// Restore content
-				RestoreSplitterContent(DockingNodes, NewDockAreaWidget.ToSharedRef());
+				if (!bCanOutputBeNullptr)
+				{
+					RestoreSplitterContent(NodeAsArea.ToSharedRef(), NewDockAreaWidget.ToSharedRef(), NewWindow);
+				}
+				else
+				{
+					RestoreSplitterContent(DockingNodes, NewDockAreaWidget.ToSharedRef());
+				}
 			}
 		}
 		else
 		{
-			TArray<TSharedRef<SDockingNode>> DockingNodes;
-			if (CanRestoreSplitterContent(DockingNodes, NodeAsArea.ToSharedRef(), ParentWindow))
-			{
-				// Create SplitterWidget only if it will be filled with at least 1 DockingNodes
-				SAssignNew(NewDockAreaWidget, SDockingArea, SharedThis(this), NodeAsArea.ToSharedRef())
+			SAssignNew( NewDockAreaWidget, SDockingArea, SharedThis(this), NodeAsArea.ToSharedRef() )
 
-					// We only want to set a parent window on this dock area, if we need to have title area content
-					// embedded within it.  SDockingArea assumes that if it has a parent window set, then it needs to have
-					// title area content 
-					.ParentWindow(bEmbedTitleAreaContent ? ParentWindow : TSharedPtr<SWindow>())
+				// We only want to set a parent window on this dock area, if we need to have title area content
+				// embedded within it.  SDockingArea assumes that if it has a parent window set, then it needs to have
+				// title area content 
+				.ParentWindow( bEmbedTitleAreaContent ? ParentWindow : TSharedPtr<SWindow>() )
 
-					// Never manage these windows, even if a parent window is set.  The owner will take care of
-					// destroying these windows.
-					.ShouldManageParentWindow(false);
-				// Restore content
-				RestoreSplitterContent(DockingNodes, NewDockAreaWidget.ToSharedRef());
-			}
-		}
+				// Never manage these windows, even if a parent window is set.  The owner will take care of
+				// destroying these windows.
+				.ShouldManageParentWindow( false );
 
+			RestoreSplitterContent( NodeAsArea.ToSharedRef(), NewDockAreaWidget.ToSharedRef(), ParentWindow );
+		}
+		
 		return NewDockAreaWidget;
 	}
 	else if ( NodeAsSplitter.IsValid() ) 
 	{
-		TArray<TSharedRef<SDockingNode>> DockingNodes;
-		if (CanRestoreSplitterContent(DockingNodes, NodeAsSplitter.ToSharedRef(), ParentWindow))
-		{
-			// Create SplitterWidget only if it will be filled with at least 1 DockingNodes
-			TSharedRef<SDockingSplitter> NewSplitterWidget = SNew(SDockingSplitter, NodeAsSplitter.ToSharedRef());
-			NewSplitterWidget->SetSizeCoefficient(LayoutNode->GetSizeCoefficient());
-			// Restore content
-			RestoreSplitterContent(DockingNodes, NewSplitterWidget);
-			return NewSplitterWidget;
-		}
-		else
-		{
-			return nullptr;
-		}
+		TSharedRef<SDockingSplitter> NewSplitterWidget = SNew( SDockingSplitter, NodeAsSplitter.ToSharedRef() );
+		NewSplitterWidget->SetSizeCoefficient(LayoutNode->GetSizeCoefficient());
+		RestoreSplitterContent( NodeAsSplitter.ToSharedRef(), NewSplitterWidget, ParentWindow );
+		return NewSplitterWidget;
 	}
 	else
 	{
 		ensureMsgf( false, TEXT("Unexpected node type") );
-
 		TSharedRef<SDockingTabStack> NewStackWidget = SNew(SDockingTabStack, FTabManager::NewStack());
 		NewStackWidget->OpenTab(SpawnTab( FName(NAME_None), ParentWindow ).ToSharedRef());
 		return NewStackWidget;
@@ -1521,7 +1518,8 @@ bool FTabManager::CanRestoreSplitterContent( TArray<TSharedRef<SDockingNode>>& D
 		const TSharedRef<FLayoutNode> ThisChildNode = SplitterNode->ChildNodes[ChildNodeIndex];
 
 		const bool bEmbedTitleAreaContent = false;
-		const TSharedPtr<SDockingNode> ThisChildNodeWidget = RestoreArea_Helper(ThisChildNode, ParentWindow, bEmbedTitleAreaContent);
+		const bool bCanOutputBeNullptr = true;
+		const TSharedPtr<SDockingNode> ThisChildNodeWidget = RestoreArea_Helper(ThisChildNode, ParentWindow, bEmbedTitleAreaContent, bCanOutputBeNullptr);
 		if (ThisChildNodeWidget)
 		{
 			const TSharedRef<SDockingNode> ThisChildNodeWidgetRef = StaticCastSharedRef<SDockingNode>(ThisChildNodeWidget->AsShared());
@@ -1539,6 +1537,24 @@ void FTabManager::RestoreSplitterContent( const TArray<TSharedRef<SDockingNode>>
 	}
 }
 
+void FTabManager::RestoreSplitterContent(const TSharedRef<FSplitter>& SplitterNode, const TSharedRef<SDockingSplitter>& SplitterWidget, const TSharedPtr<SWindow>& ParentWindow)
+{
+	// Restore the contents of this splitter.
+	for ( int32 ChildNodeIndex = 0; ChildNodeIndex < SplitterNode->ChildNodes.Num(); ++ChildNodeIndex )
+	{
+		TSharedRef<FLayoutNode> ThisChildNode = SplitterNode->ChildNodes[ChildNodeIndex];
+
+		const bool bEmbedTitleAreaContent = false;
+		TSharedPtr<SDockingNode> ThisChildNodeWidget = RestoreArea_Helper(ThisChildNode, ParentWindow, bEmbedTitleAreaContent);
+		check(ThisChildNodeWidget.IsValid());
+		if (ThisChildNodeWidget)
+		{
+			const TSharedRef<SDockingNode> ThisChildNodeWidgetRef = StaticCastSharedRef<SDockingNode>(ThisChildNodeWidget->AsShared());
+			SplitterWidget->AddChildNode( ThisChildNodeWidgetRef, INDEX_NONE );
+		}
+	}
+}
+
 bool FTabManager::CanSpawnTab(FName TabId) const
 {
 	return HasTabSpawner(TabId);
@@ -1773,7 +1789,7 @@ void FTabManager::SetTabsTo(const TSharedRef<FTabManager::FLayoutNode>& SomeNode
 		TArray<FTab>& Tabs = AsStack->Tabs;
 		for (int32 TabIndex = 0; TabIndex < Tabs.Num(); ++TabIndex)
 		{
-			if (Tabs[TabIndex].TabState == OriginalTabState || OriginalTabState == ETabState::AnyTab)
+			if (Tabs[TabIndex].TabState == OriginalTabState)
 			{
 				Tabs[TabIndex].TabState = NewTabState;
 			}
@@ -2334,15 +2350,15 @@ void FGlobalTabmanager::UpdateStats()
 	AllAreasWindowMaxCount = FMath::Max(AllAreasWindowMaxCount, ParentWindows.Num());
 }
 
-bool FGlobalTabmanager::OpenUnmanagedTab(FName PlaceholderId, const FSearchPreference& SearchPreference, const TSharedRef<SDockTab>& UnmanagedTab)
+void FGlobalTabmanager::OpenUnmanagedTab(FName PlaceholderId, const FSearchPreference& SearchPreference, const TSharedRef<SDockTab>& UnmanagedTab)
 {
 	if ( ProxyTabManager.IsValid() && ProxyTabManager->IsTabSupported( UnmanagedTab->GetLayoutIdentifier() ) )
 	{
-		return ProxyTabManager->OpenUnmanagedTab(PlaceholderId, SearchPreference, UnmanagedTab);
+		ProxyTabManager->OpenUnmanagedTab(PlaceholderId, SearchPreference, UnmanagedTab);
 	}
 	else
 	{
-		return FTabManager::OpenUnmanagedTab(PlaceholderId, SearchPreference, UnmanagedTab);
+		FTabManager::OpenUnmanagedTab(PlaceholderId, SearchPreference, UnmanagedTab);
 	}
 }
 
@@ -2362,7 +2378,7 @@ bool FProxyTabmanager::IsTabSupported( const FTabId TabId ) const
 	return bIsTabSupported;
 }
 
-bool FProxyTabmanager::OpenUnmanagedTab(FName PlaceholderId, const FSearchPreference& SearchPreference, const TSharedRef<SDockTab>& UnmanagedTab)
+void FProxyTabmanager::OpenUnmanagedTab(FName PlaceholderId, const FSearchPreference& SearchPreference, const TSharedRef<SDockTab>& UnmanagedTab)
 {
 	TSharedPtr<SWindow> ParentWindowPtr = ParentWindow.Pin();
 	if (ensure(ParentWindowPtr.IsValid()))
@@ -2389,12 +2405,9 @@ bool FProxyTabmanager::OpenUnmanagedTab(FName PlaceholderId, const FSearchPrefer
 				MainNonCloseableTab = UnmanagedTab;
 
 				OnTabOpened.Broadcast(UnmanagedTab);
-
-				return true;
 			}
 		}
 	}
-	return false;
 }
 
 void FProxyTabmanager::DrawAttention(const TSharedRef<SDockTab>& TabToHighlight)
diff --git a/Engine/Source/Runtime/Slate/Public/Framework/Docking/TabManager.h b/Engine/Source/Runtime/Slate/Public/Framework/Docking/TabManager.h
--- a/Engine/Source/Runtime/Slate/Public/Framework/Docking/TabManager.h
+++ b/Engine/Source/Runtime/Slate/Public/Framework/Docking/TabManager.h
@@ -228,11 +228,7 @@ namespace ETabState
 		/**
 		 * InvalidTab refers to tabs that were not recognized by the Editor (e.g., LiveLink when its plugin its disabled).
 		 */
-		InvalidTab = 0x1 << 2,
-		/**
-		 * AnyTab refers to any kind of tab (OpenedTab, ClosedTab, InvalidTab).
-		 */
-		AnyTab = 0x1 << 3,
+		InvalidTab = 0x1 << 2
 	};
 }
 
@@ -705,10 +701,10 @@ class SLATE_API FTabManager : public TSharedFromThis<FTabManager>
 		};
 
 		/** Insert a new UnmanagedTab document tab next to an existing tab (closed or open) that has the PlaceholdId. */
-		bool InsertNewDocumentTab( FName PlaceholderId, const FSearchPreference& SearchPreference, const TSharedRef<SDockTab>& UnmanagedTab );
+		void InsertNewDocumentTab( FName PlaceholderId, const FSearchPreference& SearchPreference, const TSharedRef<SDockTab>& UnmanagedTab );
 		
 		/** Insert a new UnmanagedTab document tab next to an existing tab (closed or open) that has the PlaceholdId. */
-		bool InsertNewDocumentTab(FName PlaceholderId, ESearchPreference::Type SearchPreference, const TSharedRef<SDockTab>& UnmanagedTab);
+		void InsertNewDocumentTab(FName PlaceholderId, ESearchPreference::Type SearchPreference, const TSharedRef<SDockTab>& UnmanagedTab);
 
 		/**
 		 * Much like InsertNewDocumentTab, but the UnmanagedTab is not seen by the user as newly-created.
@@ -787,9 +783,9 @@ class SLATE_API FTabManager : public TSharedFromThis<FTabManager>
 
 	protected:
 		
-		bool InsertDocumentTab( FName PlaceholderId, const FSearchPreference& SearchPreference, const TSharedRef<SDockTab>& UnmanagedTab, bool bPlaySpawnAnim );
+		void InsertDocumentTab( FName PlaceholderId, const FSearchPreference& SearchPreference, const TSharedRef<SDockTab>& UnmanagedTab, bool bPlaySpawnAnim );
 
-		virtual bool OpenUnmanagedTab(FName PlaceholderId, const FSearchPreference& SearchPreference, const TSharedRef<SDockTab>& UnmanagedTab);
+		virtual void OpenUnmanagedTab(FName PlaceholderId, const FSearchPreference& SearchPreference, const TSharedRef<SDockTab>& UnmanagedTab);
 			
 		void PopulateTabSpawnerMenu_Helper( FMenuBuilder& PopulateMe, struct FPopulateTabSpawnerMenu_Args Args );
 
@@ -811,11 +807,14 @@ class SLATE_API FTabManager : public TSharedFromThis<FTabManager>
 
 		FTabManager( const TSharedPtr<SDockTab>& InOwnerTab, const TSharedRef<FTabManager::FTabSpawner> & InNomadTabSpawner );
 
-		TSharedPtr<SDockingArea> RestoreArea(const TSharedRef<FArea>& AreaToRestore, const TSharedPtr<SWindow>& InParentWindow, const bool bEmbedTitleAreaContent = false);
+		TSharedPtr<SDockingArea> RestoreArea(
+			const TSharedRef<FArea>& AreaToRestore, const TSharedPtr<SWindow>& InParentWindow, const bool bEmbedTitleAreaContent = false, const bool bCanOutputBeNullptr = false);
 
-		TSharedPtr<class SDockingNode> RestoreArea_Helper(const TSharedRef<FLayoutNode>& LayoutNode, const TSharedPtr<SWindow>& ParentWindow, const bool bEmbedTitleAreaContent);
+		TSharedPtr<class SDockingNode> RestoreArea_Helper(
+			const TSharedRef<FLayoutNode>& LayoutNode, const TSharedPtr<SWindow>& ParentWindow, const bool bEmbedTitleAreaContent, const bool bCanOutputBeNullptr = false);
 
 		/**
+		 * Use CanRestoreSplitterContent + RestoreSplitterContent when the output of its internal RestoreArea_Helper can be a nullptr.
 		 * Usage example:
 		 *		TArray<TSharedRef<SDockingNode>> DockingNodes;
 		 *		if (CanRestoreSplitterContent(DockingNodes, SplitterNode, ParentWindow))
@@ -826,8 +825,13 @@ class SLATE_API FTabManager : public TSharedFromThis<FTabManager>
 		 *			RestoreSplitterContent(DockingNodes, SplitterWidget);
 		 *		}
 		 */
-		bool CanRestoreSplitterContent( TArray<TSharedRef<class SDockingNode>>& DockingNodes, const TSharedRef<FSplitter>& SplitterNode, const TSharedPtr<SWindow>& ParentWindow );
-		void RestoreSplitterContent( const TArray<TSharedRef<class SDockingNode>>& DockingNodes, const TSharedRef<class SDockingSplitter>& SplitterWidget );
+		bool CanRestoreSplitterContent(TArray<TSharedRef<class SDockingNode>>& DockingNodes, const TSharedRef<FSplitter>& SplitterNode, const TSharedPtr<SWindow>& ParentWindow);
+		void RestoreSplitterContent(const TArray<TSharedRef<class SDockingNode>>& DockingNodes, const TSharedRef<class SDockingSplitter>& SplitterWidget);
+
+		/**
+		 * Use this standalone RestoreSplitterContent when the output of its internal RestoreArea_Helper cannot be a nullptr.
+		 */
+		void RestoreSplitterContent(const TSharedRef<FSplitter>& SplitterNode, const TSharedRef<class SDockingSplitter>& SplitterWidget, const TSharedPtr<SWindow>& ParentWindow);
 		
 		bool IsValidTabForSpawning( const FTab& SomeTab ) const;
 		TSharedPtr<SDockTab> SpawnTab( const FTabId& TabId, const TSharedPtr<SWindow>& ParentWindow );
@@ -844,7 +848,7 @@ class SLATE_API FTabManager : public TSharedFromThis<FTabManager>
 		 * @param NewTabState The new TabState value.
 		 * @param OriginalTabState Only the tabs with this value will be modified. Use ETabState::AnyTab to modify them all.
 		 */
-		void SetTabsTo(const TSharedRef<FTabManager::FLayoutNode>& SomeNode, const ETabState::Type NewTabState, const ETabState::Type OriginalTabState = ETabState::AnyTab) const;
+		void SetTabsTo(const TSharedRef<FTabManager::FLayoutNode>& SomeNode, const ETabState::Type NewTabState, const ETabState::Type OriginalTabState) const;
 
 		/**
 		 * Notify the tab manager that the NewForegroundTab was brought to front and the BackgroundedTab was send to the background as a result.
@@ -1048,7 +1052,7 @@ protected:
 	virtual void OnTabClosing( const TSharedRef<SDockTab>& TabBeingClosed ) override;
 	virtual void UpdateStats() override;
 
-	virtual bool OpenUnmanagedTab(FName PlaceholderId, const FSearchPreference& SearchPreference, const TSharedRef<SDockTab>& UnmanagedTab) override;
+	virtual void OpenUnmanagedTab(FName PlaceholderId, const FSearchPreference& SearchPreference, const TSharedRef<SDockTab>& UnmanagedTab) override;
 
 public:
 	virtual void OnTabManagerClosing() override;
@@ -1150,7 +1154,7 @@ public:
 		bCanDoDragOperation = false;
 	}
 
-	virtual bool OpenUnmanagedTab(FName PlaceholderId, const FSearchPreference& SearchPreference, const TSharedRef<SDockTab>& UnmanagedTab) override;
+	virtual void OpenUnmanagedTab(FName PlaceholderId, const FSearchPreference& SearchPreference, const TSharedRef<SDockTab>& UnmanagedTab) override;
 	virtual void DrawAttention(const TSharedRef<SDockTab>& TabToHighlight) override;
 	
 	bool IsTabSupported( const FTabId TabId ) const;
