Bug fix: Reverted code so that VR can still open an "unrecognized tab" so that it does not crash on opening.

diff --git a/Engine/Source/Runtime/Slate/Private/Framework/Docking/TabManager.cpp b/Engine/Source/Runtime/Slate/Private/Framework/Docking/TabManager.cpp
--- a/Engine/Source/Runtime/Slate/Private/Framework/Docking/TabManager.cpp
+++ b/Engine/Source/Runtime/Slate/Private/Framework/Docking/TabManager.cpp
@@ -1347,6 +1347,7 @@ TSharedPtr<SDockingArea> FTabManager::RestoreArea(
 	}
 	else
 	{
+		check(bCanOutputBeNullptr);
 		return nullptr;
 	}
 }
@@ -1371,18 +1372,10 @@ TSharedPtr<SDockingNode> FTabManager::RestoreArea_Helper(
 		for (int32 TabIndex=0; TabIndex < NodeAsStack->Tabs.Num(); ++TabIndex )
 		{
 			const FTab& SomeTab = NodeAsStack->Tabs[ TabIndex ];
-if (SomeTab.TabId.ToString() == FString("LiveLink"))
-{
-	UE_LOG(LogSlate, Warning, TEXT("SomeTab %d/%d = %s."), TabIndex + 1, NodeAsStack->Tabs.Num(), *(SomeTab.TabId.ToString()));
-}
-if (SomeTab.TabId.ToString() == FString("StandaloneToolkit"))
-{
-	UE_LOG(LogSlate, Warning, TEXT("SomeTab %d/%d = %s."), TabIndex + 1, NodeAsStack->Tabs.Num(), *(SomeTab.TabId.ToString()));
-}
 
 			if (SomeTab.TabState == ETabState::OpenedTab && IsValidTabForSpawning(SomeTab))
 			{
-				const TSharedPtr<SDockTab> NewTabWidget = SpawnTab( SomeTab.TabId, ParentWindow );
+				const TSharedPtr<SDockTab> NewTabWidget = SpawnTab(SomeTab.TabId, ParentWindow, bCanOutputBeNullptr);
 
 				if (NewTabWidget.IsValid())
 				{
@@ -1504,7 +1497,7 @@ if (SomeTab.TabId.ToString() == FString("StandaloneToolkit"))
 	{
 		ensureMsgf( false, TEXT("Unexpected node type") );
 		TSharedRef<SDockingTabStack> NewStackWidget = SNew(SDockingTabStack, FTabManager::NewStack());
-		NewStackWidget->OpenTab(SpawnTab( FName(NAME_None), ParentWindow ).ToSharedRef());
+		NewStackWidget->OpenTab(SpawnTab(FName(NAME_None), ParentWindow, bCanOutputBeNullptr).ToSharedRef());
 		return NewStackWidget;
 	}
 }
@@ -1579,7 +1572,7 @@ bool FTabManager::IsValidTabForSpawning( const FTab& SomeTab ) const
 	return ( !NomadSpawner || !NomadSpawner->Get().IsSoleTabInstanceSpawned() );
 }
 
-TSharedPtr<SDockTab> FTabManager::SpawnTab( const FTabId& TabId, const TSharedPtr<SWindow>& ParentWindow )
+TSharedPtr<SDockTab> FTabManager::SpawnTab(const FTabId& TabId, const TSharedPtr<SWindow>& ParentWindow, const bool bCanOutputBeNullptr)
 {
 	TSharedPtr<SDockTab> NewTabWidget;
 
@@ -1617,6 +1610,25 @@ TSharedPtr<SDockTab> FTabManager::SpawnTab( const FTabId& TabId, const TSharedPt
 			TEXT("The tab \"%s\" attempted to spawn but failed for some reason. It will not be displayed but it will be saved in the layout settings file."),
 			*(!StringToDisplay.IsEmpty() ? StringToDisplay : FString("Unknown"))
 		);
+
+		// If an output must be generated, create an "unrecognized tab"
+		if (!bCanOutputBeNullptr)
+		{
+			NewTabWidget = SNew(SDockTab)
+				.Label( TabId.ToText() )
+				.ShouldAutosize( false )
+				[
+					SNew(SBox)
+					.HAlign(HAlign_Center)
+					.VAlign(VAlign_Center)
+					[
+						SNew(STextBlock)
+						.Text( NSLOCTEXT("TabManagement", "Unrecognized", "unrecognized tab") )
+					]
+				];
+
+			NewTabWidget->SetLayoutIdentifier(TabId);
+		}
 	}
 
 	if (NewTabWidget.IsValid())
@@ -2391,7 +2403,7 @@ void FProxyTabmanager::OpenUnmanagedTab(FName PlaceholderId, const FSearchPrefer
 				->AddTab( UnmanagedTab->GetLayoutIdentifier(), ETabState::OpenedTab )
 			);
 
-		if (TSharedPtr<SDockingArea> DockingArea = RestoreArea(NewAreaForTab, ParentWindowPtr, false))
+		if (TSharedPtr<SDockingArea> DockingArea = RestoreArea(NewAreaForTab, ParentWindowPtr))
 		{
 			ParentWindowPtr->SetContent(StaticCastSharedRef<SDockingArea>(DockingArea->AsShared()));
 			if (DockingArea->GetAllChildTabs().Num() > 0)
diff --git a/Engine/Source/Runtime/Slate/Public/Framework/Docking/TabManager.h b/Engine/Source/Runtime/Slate/Public/Framework/Docking/TabManager.h
--- a/Engine/Source/Runtime/Slate/Public/Framework/Docking/TabManager.h
+++ b/Engine/Source/Runtime/Slate/Public/Framework/Docking/TabManager.h
@@ -834,7 +834,7 @@ class SLATE_API FTabManager : public TSharedFromThis<FTabManager>
 		void RestoreSplitterContent(const TSharedRef<FSplitter>& SplitterNode, const TSharedRef<class SDockingSplitter>& SplitterWidget, const TSharedPtr<SWindow>& ParentWindow);
 		
 		bool IsValidTabForSpawning( const FTab& SomeTab ) const;
-		TSharedPtr<SDockTab> SpawnTab( const FTabId& TabId, const TSharedPtr<SWindow>& ParentWindow );
+		TSharedPtr<SDockTab> SpawnTab(const FTabId& TabId, const TSharedPtr<SWindow>& ParentWindow, const bool bCanOutputBeNullptr = false);
 
 		TSharedPtr<class SDockingTabStack> FindTabInLiveAreas( const FTabMatcher& TabMatcher ) const;
 		static TSharedPtr<class SDockingTabStack> FindTabInLiveArea( const FTabMatcher& TabMatcher, const TSharedRef<SDockingArea>& InArea );
