Updated GetKeyboardFocusableWidget function on SScrollBox to ignore children of invisible widgets when looking for a focus candidate.

diff --git a/Engine/Source/Runtime/Slate/Private/Widgets/Layout/SScrollBox.cpp b/Engine/Source/Runtime/Slate/Private/Widgets/Layout/SScrollBox.cpp
--- a/Engine/Source/Runtime/Slate/Private/Widgets/Layout/SScrollBox.cpp
+++ b/Engine/Source/Runtime/Slate/Private/Widgets/Layout/SScrollBox.cpp
@@ -998,20 +998,23 @@ FNavigationReply SScrollBox::OnNavigation(const FGeometry& MyGeometry, const FNa
 
 TSharedPtr<SWidget> SScrollBox::GetKeyboardFocusableWidget(TSharedPtr<SWidget> InWidget)
 {
-	if (InWidget->SupportsKeyboardFocus() && EVisibility::DoesVisibilityPassFilter(InWidget->GetVisibility(), EVisibility::Visible))
+	if (EVisibility::DoesVisibilityPassFilter(InWidget->GetVisibility(), EVisibility::Visible))
 	{
-		return InWidget;
-	}
-	else
-	{
-		FChildren* Children = InWidget->GetChildren();
-		for (int32 i = 0; i < Children->Num(); ++i)
+		if (InWidget->SupportsKeyboardFocus())
+		{
+			return InWidget;
+		}
+		else
 		{
-			TSharedPtr<SWidget> ChildWidget = Children->GetChildAt(i);
-			TSharedPtr<SWidget> FoucusableWidget = GetKeyboardFocusableWidget(ChildWidget);
-			if (FoucusableWidget.IsValid() && EVisibility::DoesVisibilityPassFilter(FoucusableWidget->GetVisibility(), EVisibility::Visible))
+			FChildren* Children = InWidget->GetChildren();
+			for (int32 i = 0; i < Children->Num(); ++i)
 			{
-				return FoucusableWidget;
+				TSharedPtr<SWidget> ChildWidget = Children->GetChildAt(i);
+				TSharedPtr<SWidget> FoucusableWidget = GetKeyboardFocusableWidget(ChildWidget);
+				if (FoucusableWidget.IsValid() && EVisibility::DoesVisibilityPassFilter(FoucusableWidget->GetVisibility(), EVisibility::Visible))
+				{
+					return FoucusableWidget;
+				}
 			}
 		}
 	}
