Enhancement/bug fix: Unrecognized Docked SWidgets from non-loaded plugins are not displayed on the UI rather than being showed up as "unrecognized tab", similarly to how standalone widgets behave.
Also removed 1 of 2 very similar warnings.
This code should also fix the 1-off crash of crash when cliking on "Window" in the Menu Bar in QA Game, given that we have reverted even more changes and have limited it only to the layout generated in MainFrameModule.

diff --git a/Engine/Source/Editor/LevelEditor/Private/SLevelEditor.cpp b/Engine/Source/Editor/LevelEditor/Private/SLevelEditor.cpp
--- a/Engine/Source/Editor/LevelEditor/Private/SLevelEditor.cpp
+++ b/Engine/Source/Editor/LevelEditor/Private/SLevelEditor.cpp
@@ -1319,7 +1319,9 @@ TSharedRef<SWidget> SLevelEditor::RestoreContentArea( const TSharedRef<SDockTab>
 	LevelEditorModule.OnRegisterLayoutExtensions().Broadcast(LayoutExtender);
 	Layout->ProcessExtensions(LayoutExtender);
 
-	return LevelEditorTabManager->RestoreFrom( Layout, OwnerWindow ).ToSharedRef();
+	const bool bEmbedTitleAreaContent = false;
+	const bool bCanRestoreAreaOutputBeNullptr = true;
+	return LevelEditorTabManager->RestoreFrom(Layout, OwnerWindow, bEmbedTitleAreaContent, bCanRestoreAreaOutputBeNullptr).ToSharedRef();
 }
 
 void SLevelEditor::HandleExperimentalSettingChanged(FName PropertyName)
diff --git a/Engine/Source/Editor/MainFrame/Private/MainFrameModule.cpp b/Engine/Source/Editor/MainFrame/Private/MainFrameModule.cpp
--- a/Engine/Source/Editor/MainFrame/Private/MainFrameModule.cpp
+++ b/Engine/Source/Editor/MainFrame/Private/MainFrameModule.cpp
@@ -204,7 +204,8 @@ void FMainFrameModule::CreateDefaultMainFrame( const bool bStartImmersive, const
 				)
 			);
 
-			MainFrameContent = FGlobalTabmanager::Get()->RestoreFrom( LoadedLayout, RootWindow, bEmbedTitleAreaContent );
+			const bool bCanRestoreAreaOutputBeNullptr = true;
+			MainFrameContent = FGlobalTabmanager::Get()->RestoreFrom(LoadedLayout, RootWindow, bEmbedTitleAreaContent, bCanRestoreAreaOutputBeNullptr);
 			bLevelEditorIsMainTab = true;
 		}
 
diff --git a/Engine/Source/Runtime/Slate/Private/Framework/Docking/TabManager.cpp b/Engine/Source/Runtime/Slate/Private/Framework/Docking/TabManager.cpp
--- a/Engine/Source/Runtime/Slate/Private/Framework/Docking/TabManager.cpp
+++ b/Engine/Source/Runtime/Slate/Private/Framework/Docking/TabManager.cpp
@@ -838,7 +838,7 @@ void FTabManager::UnregisterAllTabSpawners()
 	TabSpawner.Empty();
 }
 
-TSharedPtr<SWidget> FTabManager::RestoreFrom( const TSharedRef<FLayout>& Layout, const TSharedPtr<SWindow>& ParentWindow, const bool bEmbedTitleAreaContent )
+TSharedPtr<SWidget> FTabManager::RestoreFrom(const TSharedRef<FLayout>& Layout, const TSharedPtr<SWindow>& ParentWindow, const bool bEmbedTitleAreaContent, const bool bCanRestoreAreaOutputBeNullptr)
 {
 	ActiveLayoutName = Layout->LayoutName;
 
@@ -858,13 +858,10 @@ TSharedPtr<SWidget> FTabManager::RestoreFrom( const TSharedRef<FLayout>& Layout,
 
 			if ( bHasOpenTabs )
 			{
-				const bool bCanOutputBeNullptr = true;
-				RestoredDockArea = RestoreArea(ThisArea, ParentWindow, bEmbedTitleAreaContent, bCanOutputBeNullptr);
+				RestoredDockArea = RestoreArea(ThisArea, ParentWindow, bEmbedTitleAreaContent, bCanRestoreAreaOutputBeNullptr);
 				// Invalidate all tabs in ThisArea because they were not recognized
 				if (!RestoredDockArea)
 				{
-					UE_LOG(LogSlate, Warning, TEXT("The area index %d/%d (extension id \"%s\") attempted to spawn but failed."),
-						AreaIndex+1, Layout->Areas.Num(), *ThisArea->GetExtensionId().ToString());
 					SetTabsTo(ThisArea, ETabState::InvalidTab, ETabState::OpenedTab);
 					InvalidDockAreas.Add(ThisArea);
 				}
@@ -1470,28 +1467,54 @@ TSharedPtr<SDockingNode> FTabManager::RestoreArea_Helper(
 		}
 		else
 		{
-			SAssignNew( NewDockAreaWidget, SDockingArea, SharedThis(this), NodeAsArea.ToSharedRef() )
-
-				// We only want to set a parent window on this dock area, if we need to have title area content
-				// embedded within it.  SDockingArea assumes that if it has a parent window set, then it needs to have
-				// title area content 
-				.ParentWindow( bEmbedTitleAreaContent ? ParentWindow : TSharedPtr<SWindow>() )
-
-				// Never manage these windows, even if a parent window is set.  The owner will take care of
-				// destroying these windows.
-				.ShouldManageParentWindow( false );
+			TArray<TSharedRef<SDockingNode>> DockingNodes;
+			if (!bCanOutputBeNullptr || CanRestoreSplitterContent(DockingNodes, NodeAsArea.ToSharedRef(), ParentWindow))
+			{
+				SAssignNew(NewDockAreaWidget, SDockingArea, SharedThis(this), NodeAsArea.ToSharedRef())
+					// We only want to set a parent window on this dock area, if we need to have title area content
+					// embedded within it.  SDockingArea assumes that if it has a parent window set, then it needs to have
+					// title area content 
+					.ParentWindow(bEmbedTitleAreaContent ? ParentWindow : TSharedPtr<SWindow>())
+					// Never manage these windows, even if a parent window is set.  The owner will take care of
+					// destroying these windows.
+					.ShouldManageParentWindow(false);
 
-			RestoreSplitterContent( NodeAsArea.ToSharedRef(), NewDockAreaWidget.ToSharedRef(), ParentWindow );
+				// Restore content
+				if (!bCanOutputBeNullptr)
+				{
+					RestoreSplitterContent(NodeAsArea.ToSharedRef(), NewDockAreaWidget.ToSharedRef(), ParentWindow);
+				}
+				else
+				{
+					RestoreSplitterContent(DockingNodes, NewDockAreaWidget.ToSharedRef());
+				}
+			}
 		}
 		
 		return NewDockAreaWidget;
 	}
 	else if ( NodeAsSplitter.IsValid() ) 
 	{
-		TSharedRef<SDockingSplitter> NewSplitterWidget = SNew( SDockingSplitter, NodeAsSplitter.ToSharedRef() );
-		NewSplitterWidget->SetSizeCoefficient(LayoutNode->GetSizeCoefficient());
-		RestoreSplitterContent( NodeAsSplitter.ToSharedRef(), NewSplitterWidget, ParentWindow );
-		return NewSplitterWidget;
+		TArray<TSharedRef<SDockingNode>> DockingNodes;
+		if (!bCanOutputBeNullptr || CanRestoreSplitterContent(DockingNodes, NodeAsSplitter.ToSharedRef(), ParentWindow))
+		{
+			TSharedRef<SDockingSplitter> NewSplitterWidget = SNew( SDockingSplitter, NodeAsSplitter.ToSharedRef() );
+			NewSplitterWidget->SetSizeCoefficient(LayoutNode->GetSizeCoefficient());
+			// Restore content
+			if (!bCanOutputBeNullptr)
+			{
+				RestoreSplitterContent( NodeAsSplitter.ToSharedRef(), NewSplitterWidget, ParentWindow );
+			}
+			else
+			{
+				RestoreSplitterContent(DockingNodes, NewSplitterWidget);
+			}
+			return NewSplitterWidget;
+		}
+		else
+		{
+			return nullptr;
+		}
 	}
 	else
 	{
@@ -1603,17 +1626,21 @@ TSharedPtr<SDockTab> FTabManager::SpawnTab(const FTabId& TabId, const TSharedPtr
 	if (bSpawningAllowedBySpawner && !NewTabWidget.IsValid())
 	{
 		// We don't know how to spawn this tab. 2 alternatives:
-		// 1) Make a dummy tab so that things aren't entirely broken (previous versions did this).
-		// 2) Do not open the widget, but keep it saved in the layout, so this function should still return it.
-		const FString StringToDisplay = (Spawner.IsValid() && !Spawner->GetDisplayName().IsEmpty() ? Spawner->GetDisplayName().ToString() : TabId.TabType.ToString());
-		UE_LOG(LogSlate, Warning,
-			TEXT("The tab \"%s\" attempted to spawn but failed for some reason. It will not be displayed but it will be saved in the layout settings file."),
-			*(!StringToDisplay.IsEmpty() ? StringToDisplay : FString("Unknown"))
-		);
+		// 1) Make a dummy tab so that things aren't entirely broken (previous versions of UE did this in all cases).
+		// 2) Do not open the widget and return nullptr, but keep the unknown widget saved in the layout. E.g., applied when calling RestoreFrom() from MainFrameModule.
 
+		FString StringToDisplay = (Spawner.IsValid() && !Spawner->GetDisplayName().IsEmpty() ? Spawner->GetDisplayName().ToString() : TabId.TabType.ToString());
+		if (StringToDisplay.IsEmpty())
+		{
+			StringToDisplay = FString("Unknown");
+		}
 		// If an output must be generated, create an "unrecognized tab"
 		if (!bCanOutputBeNullptr)
 		{
+			UE_LOG(LogSlate, Warning,
+				TEXT("The tab \"%s\" attempted to spawn but failed for some reason. An \"unrecognized tab\" will be returned instead."), *StringToDisplay
+			);
+
 			NewTabWidget = SNew(SDockTab)
 				.Label( TabId.ToText() )
 				.ShouldAutosize( false )
@@ -1629,6 +1656,13 @@ TSharedPtr<SDockTab> FTabManager::SpawnTab(const FTabId& TabId, const TSharedPtr
 
 			NewTabWidget->SetLayoutIdentifier(TabId);
 		}
+		// If we can return nullptr, report it in the log
+		else
+		{
+			UE_LOG(LogSlate, Warning,
+				TEXT("The tab \"%s\" attempted to spawn but failed for some reason. It will not be displayed but it will still be saved in the layout settings file."), *StringToDisplay
+			);
+		}
 	}
 
 	if (NewTabWidget.IsValid())
diff --git a/Engine/Source/Runtime/Slate/Public/Framework/Docking/TabManager.h b/Engine/Source/Runtime/Slate/Public/Framework/Docking/TabManager.h
--- a/Engine/Source/Runtime/Slate/Public/Framework/Docking/TabManager.h
+++ b/Engine/Source/Runtime/Slate/Public/Framework/Docking/TabManager.h
@@ -679,7 +679,7 @@ class SLATE_API FTabManager : public TSharedFromThis<FTabManager>
 		 */
 		void UnregisterAllTabSpawners();
 
-		TSharedPtr<SWidget> RestoreFrom( const TSharedRef<FLayout>& Layout, const TSharedPtr<SWindow>& ParentWindow, const bool bEmbedTitleAreaContent = false );
+		TSharedPtr<SWidget> RestoreFrom(const TSharedRef<FLayout>& Layout, const TSharedPtr<SWindow>& ParentWindow, const bool bEmbedTitleAreaContent = false, const bool bCanRestoreAreaOutputBeNullptr = false);
 
 		void PopulateLocalTabSpawnerMenu( FMenuBuilder& PopulateMe );
 
