Fixed issue with SScrollBar widget on touch devices where scroll thumb would maintain capture, preventing it from becoming unhovered.

diff --git a/Engine/Source/Runtime/Slate/Private/Widgets/Layout/SScrollBar.cpp b/Engine/Source/Runtime/Slate/Private/Widgets/Layout/SScrollBar.cpp
--- a/Engine/Source/Runtime/Slate/Private/Widgets/Layout/SScrollBar.cpp
+++ b/Engine/Source/Runtime/Slate/Private/Widgets/Layout/SScrollBar.cpp
@@ -110,7 +110,7 @@ FReply SScrollBar::OnMouseButtonDown( const FGeometry& MyGeometry, const FPointe
 	{
 		FGeometry ThumbGeometry = FindChildGeometry(MyGeometry, DragThumb.ToSharedRef());
 
-		if (DragThumb->IsHovered())
+		if (DragThumb->IsDirectlyHovered())
 		{
 			// Clicking on the scrollbar drag thumb
 			if( Orientation == Orient_Horizontal )
@@ -277,7 +277,7 @@ const FSlateBrush* SScrollBar::GetDragThumbImage() const
 	{
 		return DraggedThumbImage;
 	}
-	else if (DragThumb->IsHovered())
+	else if (DragThumb->IsDirectlyHovered())
 	{
 		return HoveredThumbImage;
 	}
