Fixed regression suffered in CL 9452975 for experimental "Modeling Tools Editor Mode" plugin.

diff --git a/Engine/Source/Editor/LevelEditor/Private/SLevelEditor.cpp b/Engine/Source/Editor/LevelEditor/Private/SLevelEditor.cpp
--- a/Engine/Source/Editor/LevelEditor/Private/SLevelEditor.cpp
+++ b/Engine/Source/Editor/LevelEditor/Private/SLevelEditor.cpp
@@ -1320,8 +1320,8 @@ TSharedRef<SWidget> SLevelEditor::RestoreContentArea( const TSharedRef<SDockTab>
 	Layout->ProcessExtensions(LayoutExtender);
 
 	const bool bEmbedTitleAreaContent = false;
-	const bool bCanRestoreAreaOutputBeNullptr = true;
-	return LevelEditorTabManager->RestoreFrom(Layout, OwnerWindow, bEmbedTitleAreaContent, bCanRestoreAreaOutputBeNullptr).ToSharedRef();
+	const EOutputCanBeNullptr OutputCanBeNullptr = EOutputCanBeNullptr::IfNoTabValid;
+	return LevelEditorTabManager->RestoreFrom(Layout, OwnerWindow, bEmbedTitleAreaContent, OutputCanBeNullptr).ToSharedRef();
 }
 
 void SLevelEditor::HandleExperimentalSettingChanged(FName PropertyName)
diff --git a/Engine/Source/Editor/MainFrame/Private/MainFrameModule.cpp b/Engine/Source/Editor/MainFrame/Private/MainFrameModule.cpp
--- a/Engine/Source/Editor/MainFrame/Private/MainFrameModule.cpp
+++ b/Engine/Source/Editor/MainFrame/Private/MainFrameModule.cpp
@@ -204,8 +204,8 @@ void FMainFrameModule::CreateDefaultMainFrame( const bool bStartImmersive, const
 				)
 			);
 
-			const bool bCanRestoreAreaOutputBeNullptr = true;
-			MainFrameContent = FGlobalTabmanager::Get()->RestoreFrom(LoadedLayout, RootWindow, bEmbedTitleAreaContent, bCanRestoreAreaOutputBeNullptr);
+			const EOutputCanBeNullptr OutputCanBeNullptr = EOutputCanBeNullptr::IfNoOpenTabValid;
+			MainFrameContent = FGlobalTabmanager::Get()->RestoreFrom(LoadedLayout, RootWindow, bEmbedTitleAreaContent, OutputCanBeNullptr);
 			bLevelEditorIsMainTab = true;
 		}
 
diff --git a/Engine/Source/Runtime/Slate/Private/Framework/Docking/TabManager.cpp b/Engine/Source/Runtime/Slate/Private/Framework/Docking/TabManager.cpp
--- a/Engine/Source/Runtime/Slate/Private/Framework/Docking/TabManager.cpp
+++ b/Engine/Source/Runtime/Slate/Private/Framework/Docking/TabManager.cpp
@@ -838,7 +838,7 @@ void FTabManager::UnregisterAllTabSpawners()
 	TabSpawner.Empty();
 }
 
-TSharedPtr<SWidget> FTabManager::RestoreFrom(const TSharedRef<FLayout>& Layout, const TSharedPtr<SWindow>& ParentWindow, const bool bEmbedTitleAreaContent, const bool bCanRestoreAreaOutputBeNullptr)
+TSharedPtr<SWidget> FTabManager::RestoreFrom(const TSharedRef<FLayout>& Layout, const TSharedPtr<SWindow>& ParentWindow, const bool bEmbedTitleAreaContent, const EOutputCanBeNullptr RestoreAreaOutputCanBeNullptr)
 {
 	ActiveLayoutName = Layout->LayoutName;
 
@@ -858,7 +858,7 @@ TSharedPtr<SWidget> FTabManager::RestoreFrom(const TSharedRef<FLayout>& Layout,
 
 			if ( bHasOpenTabs )
 			{
-				RestoredDockArea = RestoreArea(ThisArea, ParentWindow, bEmbedTitleAreaContent, bCanRestoreAreaOutputBeNullptr);
+				RestoredDockArea = RestoreArea(ThisArea, ParentWindow, bEmbedTitleAreaContent, RestoreAreaOutputCanBeNullptr);
 				// Invalidate all tabs in ThisArea because they were not recognized
 				if (!RestoredDockArea)
 				{
@@ -1334,9 +1334,9 @@ FTabManager::FTabManager( const TSharedPtr<SDockTab>& InOwnerTab, const TSharedR
 }
 
 TSharedPtr<SDockingArea> FTabManager::RestoreArea(
-	const TSharedRef<FArea>& AreaToRestore, const TSharedPtr<SWindow>& InParentWindow, const bool bEmbedTitleAreaContent, const bool bCanOutputBeNullptr)
+	const TSharedRef<FArea>& AreaToRestore, const TSharedPtr<SWindow>& InParentWindow, const bool bEmbedTitleAreaContent, const EOutputCanBeNullptr OutputCanBeNullptr)
 {
-	if (TSharedPtr<SDockingNode> RestoredNode = RestoreArea_Helper(AreaToRestore, InParentWindow, bEmbedTitleAreaContent, bCanOutputBeNullptr))
+	if (TSharedPtr<SDockingNode> RestoredNode = RestoreArea_Helper(AreaToRestore, InParentWindow, bEmbedTitleAreaContent, OutputCanBeNullptr))
 	{
 		TSharedRef<SDockingArea> RestoredArea = StaticCastSharedRef<SDockingArea>(RestoredNode->AsShared());
 		RestoredArea->CleanUp(SDockingNode::TabRemoval_None);
@@ -1344,32 +1344,59 @@ TSharedPtr<SDockingArea> FTabManager::RestoreArea(
 	}
 	else
 	{
-		check(bCanOutputBeNullptr);
+		check(OutputCanBeNullptr != EOutputCanBeNullptr::Never);
 		return nullptr;
 	}
 }
 
 TSharedPtr<SDockingNode> FTabManager::RestoreArea_Helper(
-	const TSharedRef<FLayoutNode>& LayoutNode, const TSharedPtr<SWindow>& ParentWindow, const bool bEmbedTitleAreaContent, const bool bCanOutputBeNullptr)
+	const TSharedRef<FLayoutNode>& LayoutNode, const TSharedPtr<SWindow>& ParentWindow, const bool bEmbedTitleAreaContent, const EOutputCanBeNullptr OutputCanBeNullptr)
 {
 	TSharedPtr<FTabManager::FStack> NodeAsStack = LayoutNode->AsStack();
 	TSharedPtr<FTabManager::FSplitter> NodeAsSplitter = LayoutNode->AsSplitter();
 	TSharedPtr<FTabManager::FArea> NodeAsArea = LayoutNode->AsArea();
+	const bool bCanOutputBeNullptr = (OutputCanBeNullptr != EOutputCanBeNullptr::Never);
 
 	if (NodeAsStack.IsValid())
 	{
 		TSharedPtr<SDockTab> WidgetToActivate;
 
 		TSharedPtr<SDockingTabStack> NewStackWidget;
-		if (!bCanOutputBeNullptr)
+		// Should we init NewStackWidget before the for loop? It depends on OutputCanBeNullptr
+		bool bIsNewStackWidgetInit = false;
+		// 1. If EOutputCanBeNullptr::Never, function cannot return nullptr
+		if (OutputCanBeNullptr == EOutputCanBeNullptr::Never)
+		{
+			bIsNewStackWidgetInit = true;
+		}
+		// 2. If EOutputCanBeNullptr::IfNoTabValid, we must init the SWidget as soon as any tab is valid for spawning
+		else if (OutputCanBeNullptr == EOutputCanBeNullptr::IfNoTabValid)
+		{
+			// Note: IsValidTabForSpawning does not check whether SpawnTab() will return nullptr
+			for (const FTab& SomeTab : NodeAsStack->Tabs)
+			{
+				if (IsValidTabForSpawning(SomeTab))
+				{
+					bIsNewStackWidgetInit = true;
+					break;
+				}
+			}
+		}
+		// 3. If EOutputCanBeNullptr::IfNoOpenTabValid, we must init the SWidget as soon as any open tab is valid for spawning. For efficiency, done in the for loop
+		// 4. Else, case not handled --> error
+		else if (OutputCanBeNullptr != EOutputCanBeNullptr::IfNoOpenTabValid)
+		{
+			check(false);
+		}
+		// Initialize the SWidget already?
+		if (bIsNewStackWidgetInit)
 		{
 			NewStackWidget = SNew(SDockingTabStack, NodeAsStack.ToSharedRef());
 			NewStackWidget->SetSizeCoefficient(LayoutNode->GetSizeCoefficient());
 		}
-		for (int32 TabIndex=0; TabIndex < NodeAsStack->Tabs.Num(); ++TabIndex )
+		// Open Tabs
+		for (const FTab& SomeTab : NodeAsStack->Tabs)
 		{
-			const FTab& SomeTab = NodeAsStack->Tabs[ TabIndex ];
-
 			if (SomeTab.TabState == ETabState::OpenedTab && IsValidTabForSpawning(SomeTab))
 			{
 				const TSharedPtr<SDockTab> NewTabWidget = SpawnTab(SomeTab.TabId, ParentWindow, bCanOutputBeNullptr);
@@ -1438,7 +1465,7 @@ TSharedPtr<SDockingNode> FTabManager::RestoreArea_Helper(
 			NewWindow->SetTitle(FGlobalTabmanager::Get()->GetApplicationTitle());
 
 			TArray<TSharedRef<SDockingNode>> DockingNodes;
-			if (!bCanOutputBeNullptr || CanRestoreSplitterContent(DockingNodes, NodeAsArea.ToSharedRef(), NewWindow))
+			if (CanRestoreSplitterContent(DockingNodes, NodeAsArea.ToSharedRef(), NewWindow, OutputCanBeNullptr))
 			{
 				// Create SplitterWidget only if it will be filled with at least 1 DockingNodes
 				// Any windows that were "pulled out" of a dock area should be children of the window in which the parent dock area resides.
@@ -1468,7 +1495,7 @@ TSharedPtr<SDockingNode> FTabManager::RestoreArea_Helper(
 		else
 		{
 			TArray<TSharedRef<SDockingNode>> DockingNodes;
-			if (!bCanOutputBeNullptr || CanRestoreSplitterContent(DockingNodes, NodeAsArea.ToSharedRef(), ParentWindow))
+			if (CanRestoreSplitterContent(DockingNodes, NodeAsArea.ToSharedRef(), ParentWindow, OutputCanBeNullptr))
 			{
 				SAssignNew(NewDockAreaWidget, SDockingArea, SharedThis(this), NodeAsArea.ToSharedRef())
 					// We only want to set a parent window on this dock area, if we need to have title area content
@@ -1496,14 +1523,14 @@ TSharedPtr<SDockingNode> FTabManager::RestoreArea_Helper(
 	else if ( NodeAsSplitter.IsValid() ) 
 	{
 		TArray<TSharedRef<SDockingNode>> DockingNodes;
-		if (!bCanOutputBeNullptr || CanRestoreSplitterContent(DockingNodes, NodeAsSplitter.ToSharedRef(), ParentWindow))
+		if (CanRestoreSplitterContent(DockingNodes, NodeAsSplitter.ToSharedRef(), ParentWindow, OutputCanBeNullptr))
 		{
 			TSharedRef<SDockingSplitter> NewSplitterWidget = SNew( SDockingSplitter, NodeAsSplitter.ToSharedRef() );
 			NewSplitterWidget->SetSizeCoefficient(LayoutNode->GetSizeCoefficient());
 			// Restore content
 			if (!bCanOutputBeNullptr)
 			{
-				RestoreSplitterContent( NodeAsSplitter.ToSharedRef(), NewSplitterWidget, ParentWindow );
+				RestoreSplitterContent(NodeAsSplitter.ToSharedRef(), NewSplitterWidget, ParentWindow);
 			}
 			else
 			{
@@ -1525,8 +1552,12 @@ TSharedPtr<SDockingNode> FTabManager::RestoreArea_Helper(
 	}
 }
 
-bool FTabManager::CanRestoreSplitterContent( TArray<TSharedRef<SDockingNode>>& DockingNodes, const TSharedRef<FSplitter>& SplitterNode, const TSharedPtr<SWindow>& ParentWindow )
+bool FTabManager::CanRestoreSplitterContent(TArray<TSharedRef<SDockingNode>>& DockingNodes, const TSharedRef<FSplitter>& SplitterNode, const TSharedPtr<SWindow>& ParentWindow, const EOutputCanBeNullptr OutputCanBeNullptr)
 {
+	if (OutputCanBeNullptr == EOutputCanBeNullptr::Never)
+	{
+		return true;
+	}
 	DockingNodes.Empty();
 	// Restore the contents of this splitter.
 	for ( int32 ChildNodeIndex = 0; ChildNodeIndex < SplitterNode->ChildNodes.Num(); ++ChildNodeIndex )
@@ -1534,8 +1565,7 @@ bool FTabManager::CanRestoreSplitterContent( TArray<TSharedRef<SDockingNode>>& D
 		const TSharedRef<FLayoutNode> ThisChildNode = SplitterNode->ChildNodes[ChildNodeIndex];
 
 		const bool bEmbedTitleAreaContent = false;
-		const bool bCanOutputBeNullptr = true;
-		const TSharedPtr<SDockingNode> ThisChildNodeWidget = RestoreArea_Helper(ThisChildNode, ParentWindow, bEmbedTitleAreaContent, bCanOutputBeNullptr);
+		const TSharedPtr<SDockingNode> ThisChildNodeWidget = RestoreArea_Helper(ThisChildNode, ParentWindow, bEmbedTitleAreaContent, OutputCanBeNullptr);
 		if (ThisChildNodeWidget)
 		{
 			const TSharedRef<SDockingNode> ThisChildNodeWidgetRef = StaticCastSharedRef<SDockingNode>(ThisChildNodeWidget->AsShared());
diff --git a/Engine/Source/Runtime/Slate/Public/Framework/Docking/TabManager.h b/Engine/Source/Runtime/Slate/Public/Framework/Docking/TabManager.h
--- a/Engine/Source/Runtime/Slate/Public/Framework/Docking/TabManager.h
+++ b/Engine/Source/Runtime/Slate/Public/Framework/Docking/TabManager.h
@@ -232,6 +232,29 @@ namespace ETabState
 	};
 }
 
+enum class EOutputCanBeNullptr
+{
+	/**
+	 * RestoreArea_Helper() will always return a SWidget. It will return an "Unrecognized Tab" dummy SWidget if it cannot find the way to create the desired SWidget.
+	 * Default behavior and the only one used before UE 4.24.
+	 * This is the most strict condition and will never return nullptr.
+	 */
+	Never,
+	/**
+	 * RestoreArea_Helper() will return nullptr if its parent FTabManager::FStack does not contain at least a valid tab.
+	 * Useful for docked SWidgets that contain closed Tabs.
+	 * With IfNoTabValid, if a previously docked Tab is re-opened (i.e., its ETabState value changes from ClosedTab to OpenedTab), it will be displayed in the same place
+	 * it was placed before being initially closed. However, with IfNoOpenTabValid, if that Tab were re-opened, it might be instead displayed in a new standalone SWidget.
+	 */
+	IfNoTabValid,
+	/**
+	 * RestoreArea_Helper() will return nullptr if its parent FTabManager::FStack does not contain at least a valid tab whose ETabState is set to OpenedTab.
+	 * Useful for standalone SWidgets that otherwise will display a blank UI with no Tabs on it.
+	 * This is the most relaxed condition and the one that will return nullptr the most.
+	 */
+	IfNoOpenTabValid
+};
+
 
 class SLATE_API FTabManager : public TSharedFromThis<FTabManager>
 {
@@ -679,7 +702,7 @@ class SLATE_API FTabManager : public TSharedFromThis<FTabManager>
 		 */
 		void UnregisterAllTabSpawners();
 
-		TSharedPtr<SWidget> RestoreFrom(const TSharedRef<FLayout>& Layout, const TSharedPtr<SWindow>& ParentWindow, const bool bEmbedTitleAreaContent = false, const bool bCanRestoreAreaOutputBeNullptr = false);
+		TSharedPtr<SWidget> RestoreFrom(const TSharedRef<FLayout>& Layout, const TSharedPtr<SWindow>& ParentWindow, const bool bEmbedTitleAreaContent = false, const EOutputCanBeNullptr RestoreAreaOutputCanBeNullptr = EOutputCanBeNullptr::Never);
 
 		void PopulateLocalTabSpawnerMenu( FMenuBuilder& PopulateMe );
 
@@ -808,16 +831,16 @@ class SLATE_API FTabManager : public TSharedFromThis<FTabManager>
 		FTabManager( const TSharedPtr<SDockTab>& InOwnerTab, const TSharedRef<FTabManager::FTabSpawner> & InNomadTabSpawner );
 
 		TSharedPtr<SDockingArea> RestoreArea(
-			const TSharedRef<FArea>& AreaToRestore, const TSharedPtr<SWindow>& InParentWindow, const bool bEmbedTitleAreaContent = false, const bool bCanOutputBeNullptr = false);
+			const TSharedRef<FArea>& AreaToRestore, const TSharedPtr<SWindow>& InParentWindow, const bool bEmbedTitleAreaContent = false, const EOutputCanBeNullptr OutputCanBeNullptr = EOutputCanBeNullptr::Never);
 
 		TSharedPtr<class SDockingNode> RestoreArea_Helper(
-			const TSharedRef<FLayoutNode>& LayoutNode, const TSharedPtr<SWindow>& ParentWindow, const bool bEmbedTitleAreaContent, const bool bCanOutputBeNullptr = false);
+			const TSharedRef<FLayoutNode>& LayoutNode, const TSharedPtr<SWindow>& ParentWindow, const bool bEmbedTitleAreaContent, const EOutputCanBeNullptr OutputCanBeNullptr = EOutputCanBeNullptr::Never);
 
 		/**
 		 * Use CanRestoreSplitterContent + RestoreSplitterContent when the output of its internal RestoreArea_Helper can be a nullptr.
 		 * Usage example:
 		 *		TArray<TSharedRef<SDockingNode>> DockingNodes;
-		 *		if (CanRestoreSplitterContent(DockingNodes, SplitterNode, ParentWindow))
+		 *		if (CanRestoreSplitterContent(DockingNodes, SplitterNode, ParentWindow, OutputCanBeNullptr))
 		 *		{
 		 *			// Create SplitterWidget only if it will be filled with at least 1 DockingNodes
 		 *			TSharedRef<SDockingSplitter> SplitterWidget = SNew(SDockingSplitter, SplitterNode);
@@ -825,7 +848,7 @@ class SLATE_API FTabManager : public TSharedFromThis<FTabManager>
 		 *			RestoreSplitterContent(DockingNodes, SplitterWidget);
 		 *		}
 		 */
-		bool CanRestoreSplitterContent(TArray<TSharedRef<class SDockingNode>>& DockingNodes, const TSharedRef<FSplitter>& SplitterNode, const TSharedPtr<SWindow>& ParentWindow);
+		bool CanRestoreSplitterContent(TArray<TSharedRef<class SDockingNode>>& DockingNodes, const TSharedRef<FSplitter>& SplitterNode, const TSharedPtr<SWindow>& ParentWindow, const EOutputCanBeNullptr OutputCanBeNullptr);
 		void RestoreSplitterContent(const TArray<TSharedRef<class SDockingNode>>& DockingNodes, const TSharedRef<class SDockingSplitter>& SplitterWidget);
 
 		/**
