PR #6287: More editor tooltips show property values (Contributed by [REDACTED_PERSON])

diff --git a/Engine/Source/Editor/DetailCustomizations/Private/MathStructCustomizations.cpp b/Engine/Source/Editor/DetailCustomizations/Private/MathStructCustomizations.cpp
--- a/Engine/Source/Editor/DetailCustomizations/Private/MathStructCustomizations.cpp
+++ b/Engine/Source/Editor/DetailCustomizations/Private/MathStructCustomizations.cpp
@@ -86,6 +86,12 @@ void FMathStructCustomization::MakeHeaderRow(TSharedRef<class IPropertyHandle>&
 		TSharedRef<SWidget> NumericEntryBox = MakeChildWidget(StructPropertyHandle, ChildHandle);
 		NumericEntryBoxWidgetList.Add(NumericEntryBox);
 
+		NumericEntryBox->SetToolTipText(MakeAttributeLambda([StructPropertyHandle] {
+			FText result;
+			StructPropertyHandle->GetValueAsDisplayText(result);
+			return result;
+		}));
+
 		HorizontalBox->AddSlot()
 		.Padding(FMargin(0.0f, 2.0f, bLastChild ? 0.0f : 3.0f, 2.0f))
 		[
diff --git a/Engine/Source/Editor/PropertyEditor/Private/PropertyEditorHelpers.cpp b/Engine/Source/Editor/PropertyEditor/Private/PropertyEditorHelpers.cpp
--- a/Engine/Source/Editor/PropertyEditor/Private/PropertyEditorHelpers.cpp
+++ b/Engine/Source/Editor/PropertyEditor/Private/PropertyEditorHelpers.cpp
@@ -82,7 +82,7 @@ void SPropertyValueWidget::Construct( const FArguments& InArgs, TSharedPtr<FProp
 
 	if ( !ValueEditorWidget->GetToolTip().IsValid() )
 	{
-		ValueEditorWidget->SetToolTipText( PropertyEditor->GetToolTipText() );
+		ValueEditorWidget->SetToolTipText(TAttribute<FText>(PropertyEditor.ToSharedRef(), &FPropertyEditor::GetValueAsText));
 	}
 
 
diff --git a/Engine/Source/Runtime/Slate/Private/Widgets/Input/SRotatorInputBox.cpp b/Engine/Source/Runtime/Slate/Private/Widgets/Input/SRotatorInputBox.cpp
--- a/Engine/Source/Runtime/Slate/Private/Widgets/Input/SRotatorInputBox.cpp
+++ b/Engine/Source/Runtime/Slate/Private/Widgets/Input/SRotatorInputBox.cpp
@@ -21,7 +21,10 @@ void SRotatorInputBox::Construct( const SRotatorInputBox::FArguments& InArgs )
 	const FLinearColor LabelColorY = InArgs._bColorAxisLabels ? SNumericEntryBox<float>::GreenLabelBackgroundColor : FLinearColor(0.0f, 0.0f, 0.0f, .5f);
 	const FLinearColor LabelColorZ = InArgs._bColorAxisLabels ? SNumericEntryBox<float>::BlueLabelBackgroundColor : FLinearColor(0.0f, 0.0f, 0.0f, .5f);
 
-	
+	TAttribute<TOptional<float>> RollValue = InArgs._Roll;
+	TAttribute<TOptional<float>> PitchValue = InArgs._Pitch;
+	TAttribute<TOptional<float>> YawValue = InArgs._Yaw;
+
 	this->ChildSlot
 	[
 		SNew(SHorizontalBox)
@@ -46,7 +49,10 @@ void SRotatorInputBox::Construct( const SRotatorInputBox::FArguments& InArgs )
 			.OnBeginSliderMovement( InArgs._OnBeginSliderMovement )
 			.OnEndSliderMovement( InArgs._OnEndSliderMovement )
 			.UndeterminedString( LOCTEXT("MultipleValues", "Multiple Values") )
-			.ToolTipText( LOCTEXT("Roll_ToolTip", "Roll Value") )
+			.ToolTipText(MakeAttributeLambda([RollValue]
+			{
+				return FText::Format(LOCTEXT("Roll_ToolTip", "Roll Value = {0}"), RollValue.Get().GetValue());
+			}))
 			.TypeInterface(InArgs._TypeInterface)
 		]
 		+SHorizontalBox::Slot()
@@ -70,7 +76,10 @@ void SRotatorInputBox::Construct( const SRotatorInputBox::FArguments& InArgs )
 			.OnBeginSliderMovement( InArgs._OnBeginSliderMovement )
 			.OnEndSliderMovement( InArgs._OnEndSliderMovement )
 			.UndeterminedString( LOCTEXT("MultipleValues", "Multiple Values") )
-			.ToolTipText( LOCTEXT("Pitch_ToolTip", "Pitch Value") )
+			.ToolTipText(MakeAttributeLambda([PitchValue]
+			{
+				return FText::Format(LOCTEXT("Pitch_ToolTip", "Pitch Value = {0}"), PitchValue.Get().GetValue());
+			}))
 			.TypeInterface(InArgs._TypeInterface)
 		]
 		+SHorizontalBox::Slot()
@@ -94,7 +103,10 @@ void SRotatorInputBox::Construct( const SRotatorInputBox::FArguments& InArgs )
 			.OnBeginSliderMovement( InArgs._OnBeginSliderMovement )
 			.OnEndSliderMovement( InArgs._OnEndSliderMovement )
 			.UndeterminedString( LOCTEXT("MultipleValues", "Multiple Values") )
-			.ToolTipText( LOCTEXT("Yaw_ToolTip", "Yaw Value") )
+			.ToolTipText(MakeAttributeLambda([YawValue]
+			{
+				return FText::Format(LOCTEXT("Yaw_ToolTip", "Yaw Value = {0}"), YawValue.Get().GetValue());
+			}))
 			.TypeInterface(InArgs._TypeInterface)
 		]
 	];
diff --git a/Engine/Source/Runtime/Slate/Private/Widgets/Input/SVectorInputBox.cpp b/Engine/Source/Runtime/Slate/Private/Widgets/Input/SVectorInputBox.cpp
--- a/Engine/Source/Runtime/Slate/Private/Widgets/Input/SVectorInputBox.cpp
+++ b/Engine/Source/Runtime/Slate/Private/Widgets/Input/SVectorInputBox.cpp
@@ -39,6 +39,8 @@ void SVectorInputBox::ConstructX( const FArguments& InArgs, TSharedRef<SHorizont
 		MarginAttribute = TAttribute<FMargin>::Create(TAttribute<FMargin>::FGetter::CreateSP(this, &SVectorInputBox::GetTextMargin));
 	}
 
+	TAttribute<TOptional<float>> Value = InArgs._X;
+
 	HorizontalBox->AddSlot()
 	.VAlign( VAlign_Center )
 	.FillWidth( 1.0f )
@@ -49,8 +51,11 @@ void SVectorInputBox::ConstructX( const FArguments& InArgs, TSharedRef<SHorizont
 		.Font( InArgs._Font )
 		.Value( InArgs._X )
 		.OnValueChanged( InArgs._OnXChanged )
-		.OnValueCommitted( InArgs._OnXCommitted )		
-		.ToolTipText( LOCTEXT("X_ToolTip", "X Value") )
+		.OnValueCommitted( InArgs._OnXCommitted )
+		.ToolTipText(MakeAttributeLambda([Value]
+		{
+			return FText::Format(LOCTEXT("X_ToolTip", "X Value = {0}"), Value.Get().GetValue());
+		}))
 		.UndeterminedString( LOCTEXT("MultipleValues", "Multiple Values") )
 		.LabelPadding(0)
 		.OverrideTextMargin(MarginAttribute)
@@ -82,6 +87,8 @@ void SVectorInputBox::ConstructY( const FArguments& InArgs, TSharedRef<SHorizont
 		MarginAttribute = TAttribute<FMargin>::Create(TAttribute<FMargin>::FGetter::CreateSP(this, &SVectorInputBox::GetTextMargin));
 	}
 
+	TAttribute<TOptional<float>> Value = InArgs._Y;
+
 	HorizontalBox->AddSlot()
 	.VAlign( VAlign_Center )
 	.FillWidth( 1.0f )
@@ -93,7 +100,10 @@ void SVectorInputBox::ConstructY( const FArguments& InArgs, TSharedRef<SHorizont
 		.Value( InArgs._Y )
 		.OnValueChanged( InArgs._OnYChanged )
 		.OnValueCommitted( InArgs._OnYCommitted )
-		.ToolTipText( LOCTEXT("Y_ToolTip", "Y Value") )
+		.ToolTipText(MakeAttributeLambda([Value]
+		{
+			return FText::Format(LOCTEXT("Y_ToolTip", "Y Value = {0}"), Value.Get().GetValue());
+		}))
 		.UndeterminedString( LOCTEXT("MultipleValues", "Multiple Values") )
 		.LabelPadding(0)
 		.OverrideTextMargin(MarginAttribute)
@@ -125,6 +135,8 @@ void SVectorInputBox::ConstructZ( const FArguments& InArgs, TSharedRef<SHorizont
 		MarginAttribute = TAttribute<FMargin>::Create(TAttribute<FMargin>::FGetter::CreateSP(this, &SVectorInputBox::GetTextMargin));
 	}
 
+	TAttribute<TOptional<float>> Value = InArgs._Z;
+
 	HorizontalBox->AddSlot()
 	.VAlign( VAlign_Center )
 	.FillWidth( 1.0f )
@@ -136,7 +148,10 @@ void SVectorInputBox::ConstructZ( const FArguments& InArgs, TSharedRef<SHorizont
 		.Value( InArgs._Z )
 		.OnValueChanged( InArgs._OnZChanged )
 		.OnValueCommitted( InArgs._OnZCommitted )
-		.ToolTipText( LOCTEXT("Z_ToolTip", "Z Value") )
+		.ToolTipText(MakeAttributeLambda([Value]
+		{
+			return FText::Format(LOCTEXT("Z_ToolTip", "Z Value = {0}"), Value.Get().GetValue());
+		}))
 		.UndeterminedString( LOCTEXT("MultipleValues", "Multiple Values") )
 		.LabelPadding(0)
 		.OverrideTextMargin(MarginAttribute)
